Modified UWorld::PurgeScene to internally flush Slate scene ptr cache instead of SetPlayInEditorWorld to avoid repeated calls per frame in some cases.

diff --git a/Engine/Source/Editor/UnrealEd/Private/Editor.cpp b/Engine/Source/Editor/UnrealEd/Private/Editor.cpp
--- a/Engine/Source/Editor/UnrealEd/Private/Editor.cpp
+++ b/Engine/Source/Editor/UnrealEd/Private/Editor.cpp
@@ -978,12 +978,6 @@ UWorld* SetPlayInEditorWorld( UWorld* PlayInEditorWorld )
 	// Purge the existing scene interface from the editor world to avoid 2x GPU allocations with the additional play-in-editor world
 	if (GEditor->EditorWorld != nullptr)
 	{
-		// Clear out Slate's active scenes list since these ptrs no longer reference valid scenes.
-		if (!GEditor->EditorWorld->bPurgedScene && FSlateApplication::IsInitialized())
-		{
-			FSlateApplication::Get().FlushRenderState();
-		}
-
 		// Tear down the scene interface for the editor world
 		GEditor->EditorWorld->PurgeScene();
 	}
diff --git a/Engine/Source/Runtime/Engine/Private/World.cpp b/Engine/Source/Runtime/Engine/Private/World.cpp
--- a/Engine/Source/Runtime/Engine/Private/World.cpp
+++ b/Engine/Source/Runtime/Engine/Private/World.cpp
@@ -8623,6 +8623,12 @@ void UWorld::PurgeScene()
 
 	if (!bPurgedScene && this->IsEditorWorld())
 	{
+		// Clear out Slate's active scenes list since these ptrs no longer reference valid scenes.
+		if (FSlateApplication::IsInitialized())
+		{
+			FSlateApplication::Get().FlushRenderState();
+		}
+
 		FGlobalComponentReregisterContext RecreateComponents;
 
 		// Finish any deferred / async render cleanup work.
