Fix crash on exit in stall detector thread when editor closed by call to RequestEngineExit()

diff --git a/Engine/Source/Editor/UnrealEd/Classes/Editor/UnrealEdEngine.h b/Engine/Source/Editor/UnrealEd/Classes/Editor/UnrealEdEngine.h
--- a/Engine/Source/Editor/UnrealEd/Classes/Editor/UnrealEdEngine.h
+++ b/Engine/Source/Editor/UnrealEd/Classes/Editor/UnrealEdEngine.h
@@ -950,4 +950,7 @@ private:
 
 	/** Delegate Called after files have been deleted to perform necessary cleanups. */
 	FDelegateHandle SourceControlFilesDeletedHandle;
+
+	/** Keep track of need to stop stall detector thread */
+	bool bRequiresStallDetectorShutdown = false;
 };
diff --git a/Engine/Source/Editor/UnrealEd/Private/UnrealEdEngine.cpp b/Engine/Source/Editor/UnrealEd/Private/UnrealEdEngine.cpp
--- a/Engine/Source/Editor/UnrealEd/Private/UnrealEdEngine.cpp
+++ b/Engine/Source/Editor/UnrealEd/Private/UnrealEdEngine.cpp
@@ -269,15 +269,17 @@ void UUnrealEdEngine::Init(IEngineLoop* InEngineLoop)
 
 #if STALL_DETECTOR
 	// Start tracking stalls when we open the Main Frame
-	IMainFrameModule::Get().OnMainFrameCreationFinished().AddLambda([](TSharedPtr<SWindow>, bool)
+	IMainFrameModule::Get().OnMainFrameCreationFinished().AddWeakLambda(this, [this](TSharedPtr<SWindow>, bool)
 		{
+			bRequiresStallDetectorShutdown = true;
 			UE::FStallDetector::Startup();
 		});
 
 	// Stop tracking stalls when we close the Main Frame, the closest event around
-	OnEditorClose().AddLambda([]()
+	OnEditorClose().AddLambda([this]()
 		{
 			// We conditionally shutdown here because for now there is no UnreadEdEngine specific shutdown 
+			bRequiresStallDetectorShutdown = false;
 			if (UE::FStallDetector::IsRunning())
 			{
 				UE::FStallDetector::Shutdown();
@@ -468,6 +470,16 @@ UUnrealEdEngine::~UUnrealEdEngine()
 	{
 		GUnrealEd = NULL; 
 	}
+
+#if STALL_DETECTOR
+	// Handle cases where OnEditorClose() was not called such as UE_LOG fatal errors or any call to RequestEngineExit()
+	// Without this stall detector thread will continue to run until it crashes trying to access cleaned up critical section
+	if (bRequiresStallDetectorShutdown && UE::FStallDetector::IsRunning())
+	{
+		bRequiresStallDetectorShutdown = false;
+		UE::FStallDetector::Shutdown();
+	}
+#endif // STALL_DETECTOR
 }
 PRAGMA_ENABLE_DEPRECATION_WARNINGS
 
