Fix non-unity compile error caused by missing includes.

diff --git a/Engine/Source/Editor/UnrealEd/Private/Cooker/CookWorkerClient.cpp b/Engine/Source/Editor/UnrealEd/Private/Cooker/CookWorkerClient.cpp
--- a/Engine/Source/Editor/UnrealEd/Private/Cooker/CookWorkerClient.cpp
+++ b/Engine/Source/Editor/UnrealEd/Private/Cooker/CookWorkerClient.cpp
@@ -11,6 +11,8 @@
 #include "CookWorkerServer.h"
 #include "CoreGlobals.h"
 #include "HAL/PlatformTime.h"
+#include "Interfaces/ITargetPlatform.h"
+#include "Interfaces/ITargetPlatformManagerModule.h"
 #include "Misc/CoreMisc.h"
 #include "PackageResultsMessage.h"
 #include "PackageTracker.h"
