Cooker: OnlyEditorOnly references: Change FPackageData to track reachability per platform, and change FRequestCluster to do a GraphPaintingSearch with one color per platform instead of a GraphSearch with a single reachability color.
#rnx

diff --git a/Engine/Source/Editor/UnrealEd/Classes/CookOnTheSide/CookOnTheFlyServer.h b/Engine/Source/Editor/UnrealEd/Classes/CookOnTheSide/CookOnTheFlyServer.h
--- a/Engine/Source/Editor/UnrealEd/Classes/CookOnTheSide/CookOnTheFlyServer.h
+++ b/Engine/Source/Editor/UnrealEd/Classes/CookOnTheSide/CookOnTheFlyServer.h
@@ -178,6 +178,7 @@ namespace UE::Cook
 	struct FCookerTimer;
 	struct FCookGenerationInfo;
 	struct FCookSavePackageContext;
+	struct FDiscoveredPlatformSet;
 	struct FGeneratorPackage;
 	struct FInitializeConfigSettings;
 	struct FPackageData;
@@ -516,38 +517,30 @@ private:
 	 * Report the number of requests that were completed (either skipped or successfully saved or failed to save)
 	 */
 	void PumpSaves(UE::Cook::FTickStackData& StackData, uint32 DesiredQueueLength, int32& OutNumPushed, bool& bOutBusy);
-	/**
-	 * Inspect the given package and queue it for saving if necessary.
-	 *
-	 * @param Package				The package to be considered for saving.
-	 * @param bOutWasInProgress		Report whether the package was already in progress.
-	 * @return						Returns the PackageData for this queued package.
-	 */
-	UE::Cook::FPackageData* QueueDiscoveredPackage(UPackage* Package, UE::Cook::FInstigator&& Instigator,
-		bool* bOutWasInProgress = nullptr);
+
 	/**
 	 * Inspect the given package and queue it for saving if necessary.
 	 *
 	 * @param PackageData			The PackageData to be considered for saving.
-	 * @param bLoadReady			If true send on to LoadReady, otherwise add it at the Request stage.
 	 */
-	void QueueDiscoveredPackageData(UE::Cook::FPackageData& PackageData, UE::Cook::FInstigator&& Instigator,
-		bool bLoadReady = false);
+	void QueueDiscoveredPackage(UE::Cook::FPackageData& PackageData, UE::Cook::FInstigator&& Instigator, 
+		UE::Cook::FDiscoveredPlatformSet&& ReachablePlatforms);
+	void QueueDiscoveredPackageOnDirector(UE::Cook::FPackageData& PackageData, UE::Cook::FInstigator&& Instigator,
+		UE::Cook::FDiscoveredPlatformSet&& ReachablePlatforms);
 
 	/** Called when a package is cancelled and returned to idle. Notifies CookDirector when on a CookWorker. */
 	void DemoteToIdle(UE::Cook::FPackageData& PackageData, UE::Cook::ESendFlags SendFlags, UE::Cook::ESuppressCookReason Reason);
 	/** Called when a package completes its save and returns to idle. Notifies the CookDirector when on a CookWorker. */
 	void PromoteToSaveComplete(UE::Cook::FPackageData& PackageData, UE::Cook::ESendFlags SendFlags);
 
-	/** If the package filter has changed, call QueueDiscoveredPackage again on each existing package. */
-	void UpdatePackageFilter();
-
 	/**
 	 * Remove all request data about the given platform from any data in the CookOnTheFlyServer.
 	 * Called when a platform is removed from the list of session platforms e.g. because it has not been recently
 	 * used by CookOnTheFly. Does not modify Cooked platforms.
 	 */
-	void OnRemoveSessionPlatform(const ITargetPlatform* TargetPlatform);
+	void OnRemoveSessionPlatform(const ITargetPlatform* TargetPlatform, int32 RemovedIndex);
+	/** Update structures that have data per platfrom when COTF adds a new SessionPlatform. */
+	void OnPlatformAddedToSession(const ITargetPlatform* TargetPlatform);
 
 	void InitializePollables();
 	void PumpPollables(UE::Cook::FTickStackData& StackData, bool bIsIdle);
@@ -1404,8 +1397,9 @@ private:
 	/** Data tracking the package currently having calls made (Load/Save/Other) from COTFS. Used for diagnostics. */
 	FActivePackageData ActivePackageData;
 
-	/** True when the which packages we need to cook changes because e.g. a platform was added to the sessionplatforms. */
-	bool bPackageFilterDirty = false;
+	/** Package -> Package dependencies that were hidden and discovered when a package loaded or ran system-specific code. */
+	TMap<FName, TArray<FName>> DiscoveredDependencies;
+
 	/** True when PumpLoads has detected it is blocked on async work and CookOnTheFlyServer should do work elsewhere. */
 	bool bLoadBusy = false;
 	/** True when PumpSaves has detected it is blocked on async work and CookOnTheFlyServer should do work elsewhere. */
diff --git a/Engine/Source/Editor/UnrealEd/Private/CookOnTheFlyServer.cpp b/Engine/Source/Editor/UnrealEd/Private/CookOnTheFlyServer.cpp
--- a/Engine/Source/Editor/UnrealEd/Private/CookOnTheFlyServer.cpp
+++ b/Engine/Source/Editor/UnrealEd/Private/CookOnTheFlyServer.cpp
@@ -381,7 +381,7 @@ public:
 				{
 					return;
 				}
-				PackageData->ClearCookProgress();
+				PackageData->ClearCookResults();
 			});
 	}
 
@@ -744,7 +744,7 @@ void UCookOnTheFlyServer::StartCookOnTheFlySessionFromGameThread(ITargetPlatform
 		InitializeSession();
 	}
 	PlatformManager->AddSessionPlatform(*this, TargetPlatform);
-	bPackageFilterDirty = true;
+	ResetCook({ TPair<const ITargetPlatform*,bool>{TargetPlatform, true /* bResetResults */} });
 
 	// Blocking on the AssetRegistry needs to wait until the session starts because it needs all plugins loaded.
 	// AddCookOnTheFlyPlatformFromGameThread can be called on cooker startup which occurs in UUnrealEdEngine::Init
@@ -754,12 +754,22 @@ void UCookOnTheFlyServer::StartCookOnTheFlySessionFromGameThread(ITargetPlatform
 
 void UCookOnTheFlyServer::OnTargetPlatformsInvalidated()
 {
+	using namespace UE::Cook;
+
 	check(IsInGameThread());
 	TMap<ITargetPlatform*, ITargetPlatform*> Remap = PlatformManager->RemapTargetPlatforms();
 
 	PackageDatas->RemapTargetPlatforms(Remap);
 	PackageTracker->RemapTargetPlatforms(Remap);
 	WorkerRequests->RemapTargetPlatforms(Remap);
+	for (UE::Cook::FRequestCluster& Cluster : PackageDatas->GetRequestQueue().GetRequestClusters())
+	{
+		Cluster.RemapTargetPlatforms(Remap);
+	}
+	for (FDiscoveryQueueElement& Element : PackageDatas->GetRequestQueue().GetDiscoveryQueue())
+	{
+		Element.ReachablePlatforms.RemapTargetPlatforms(Remap);
+	}
 
 	if (PlatformManager->GetArePlatformsPrepopulated())
 	{
@@ -1300,7 +1310,7 @@ void UCookOnTheFlyServer::RunCookList(ECookListOptions CookListOptions)
 	PumpExternalRequests(StackData.Timer);
 	ProcessUnsolicitedPackages();
 	FRequestQueue& RequestQueue = PackageDatas->GetRequestQueue();
-	while (RequestQueue.GetRequestClusters().Num() != 0 || RequestQueue.GetUnclusteredRequests().Num() != 0)
+	while (RequestQueue.HasRequestsToExplore())
 	{
 		int32 NumPushed;
 		PumpRequests(StackData, NumPushed);
@@ -1312,7 +1322,7 @@ void UCookOnTheFlyServer::RunCookList(ECookListOptions CookListOptions)
 		bool bIncludePackage;
 		if (EnumHasAnyFlags(CookListOptions, ECookListOptions::ShowRejected))
 		{
-			bIncludePackage = PackageData->GetInstigator().Category != EInstigator::NotYetRequested;
+			bIncludePackage = PackageData->HasInstigator();
 			if (bIncludePackage)
 			{
 				// Skip printing out a message for external actors
@@ -1468,7 +1478,6 @@ void UCookOnTheFlyServer::TickCookStatus(UE::Cook::FTickStackData& StackData)
 	UpdateDisplay(StackData, false /* bForceDisplay */);
 	ProcessAsyncLoading(false, false, 0.0f);
 	ProcessUnsolicitedPackages();
-	UpdatePackageFilter();
 	PumpExternalRequests(StackData.Timer);
 }
 
@@ -2213,7 +2222,7 @@ UCookOnTheFlyServer::ECookAction UCookOnTheFlyServer::DecideNextCookAction(UE::C
 	}
 
 	UE::Cook::FRequestQueue& RequestQueue = PackageDatas->GetRequestQueue();
-	if (RequestQueue.GetRequestClusters().Num() != 0 || RequestQueue.GetUnclusteredRequests().Num() != 0)
+	if (RequestQueue.HasRequestsToExplore())
 	{
 		SetIdleStatus(StackData, EIdleStatus::Active);
 		return ECookAction::Request;
@@ -2415,7 +2424,7 @@ void UCookOnTheFlyServer::PumpExternalRequests(const UE::Cook::FCookerTimer& Coo
 #endif
 			bool bRequestsAreUrgent = IsCookOnTheFlyMode() && IsUsingLegacyCookOnTheFlyScheduling();
 			TRingBuffer<UE::Cook::FRequestCluster>& RequestClusters = PackageDatas->GetRequestQueue().GetRequestClusters();
-			UE::Cook::FRequestCluster::AddClusters(*this, MoveTemp(BuildRequests), bRequestsAreUrgent, RequestClusters);
+			RequestClusters.Add(UE::Cook::FRequestCluster(*this, MoveTemp(BuildRequests), bRequestsAreUrgent));
 		}
 	}
 }
@@ -2424,7 +2433,7 @@ bool UCookOnTheFlyServer::TryCreateRequestCluster(UE::Cook::FPackageData& Packag
 {
 	check(PackageData.IsInProgress()); // This should only be called from Pump functions, and only on in-progress Packages
 	using namespace UE::Cook;
-	if (!PackageData.AreAllRequestedPlatformsExplored())
+	if (!PackageData.AreAllReachablePlatformsVisitedByCluster())
 	{
 		PackageData.SendToState(EPackageState::Request, ESendFlags::QueueAdd);
 		return true;
@@ -2439,59 +2448,76 @@ void UCookOnTheFlyServer::PumpRequests(UE::Cook::FTickStackData& StackData, int3
 
 	OutNumPushed = 0;
 	FRequestQueue& RequestQueue = PackageDatas->GetRequestQueue();
-	FPackageDataSet& UnclusteredRequests = RequestQueue.GetUnclusteredRequests();
+	FPackageDataSet& RestartedRequests = RequestQueue.GetRestartedRequests();
+	TRingBuffer<FDiscoveryQueueElement>& DiscoveryQueue = RequestQueue.GetDiscoveryQueue();
 	TRingBuffer<FRequestCluster>& RequestClusters = RequestQueue.GetRequestClusters();
 	const FCookerTimer& CookerTimer = StackData.Timer;
-	if (!UnclusteredRequests.IsEmpty())
+
+	// First pump all requestclusters and unclustered/discovered requests that need to create a requestcluster
+	for (;;)
 	{
-		FRequestCluster::AddClusters(*this, UnclusteredRequests, RequestClusters, RequestQueue);
-		UnclusteredRequests.Empty();
-		if (CookerTimer.IsTimeUp())
+		// We completely finish the first cluster before moving on to new or remaining clusters.
+		// This prevents the problem of an infinite loop due to having two clusters steal a PackageData
+		// back and forth from each other.
+		if (!RequestClusters.IsEmpty())
 		{
-			return;
+			FRequestCluster& RequestCluster = RequestClusters.First();
+			bool bComplete = false;
+			RequestCluster.Process(CookerTimer, bComplete);
+			if (bComplete)
+			{
+				TArray<FPackageData*> RequestsToLoad;
+				TArray<TPair<FPackageData*, ESuppressCookReason>> RequestsToDemote;
+				TMap<FPackageData*, TArray<FPackageData*>> RequestGraph;
+				RequestCluster.ClearAndDetachOwnedPackageDatas(RequestsToLoad, RequestsToDemote, RequestGraph);
+				AssignRequests(RequestsToLoad, RequestQueue, MoveTemp(RequestGraph));
+				for (TPair<FPackageData*, ESuppressCookReason>& Pair : RequestsToDemote)
+				{
+					DemoteToIdle(*Pair.Key, ESendFlags::QueueAdd, Pair.Value);
+				}
+				RequestClusters.PopFront();
+				OnRequestClusterCompleted(RequestCluster);
+			}
 		}
-	}
-
-	while (RequestClusters.Num() > 0)
-	{
-		FRequestCluster& RequestCluster = RequestClusters.First();
-		bool bComplete = false;
-		RequestCluster.Process(CookerTimer, bComplete);
-		if (bComplete)
+		else if (!RestartedRequests.IsEmpty())
+		{
+			RequestClusters.Add(FRequestCluster(*this, MoveTemp(RestartedRequests)));
+			RestartedRequests.Empty();
+		}
+		else if (!DiscoveryQueue.IsEmpty())
 		{
-			TArray<FPackageData*> RequestsToLoad;
-			TArray<TPair<FPackageData*, ESuppressCookReason>> RequestsToDemote;
-			TMap<FPackageData*, TArray<FPackageData*>> RequestGraph;
-			RequestCluster.ClearAndDetachOwnedPackageDatas(RequestsToLoad, RequestsToDemote, RequestGraph);
-			AssignRequests(RequestsToLoad, RequestQueue, MoveTemp(RequestGraph));
-			for (TPair<FPackageData*, ESuppressCookReason>& Pair : RequestsToDemote)
+			FRequestCluster Cluster(*this, DiscoveryQueue);
+			if (Cluster.NumPackageDatas() > 0)
 			{
-				DemoteToIdle(*Pair.Key, ESendFlags::QueueAdd, Pair.Value);
+				RequestClusters.Add(MoveTemp(Cluster));
 			}
-			RequestClusters.PopFront();
-			OnRequestClusterCompleted(RequestCluster);
 		}
+		else
+		{
+			break;
+		}
+
 		if (CookerTimer.IsTimeUp())
 		{
 			return;
 		}
 	}
 
-	COOK_STAT(DetailedCookStats::PeakRequestQueueSize = FMath::Max(DetailedCookStats::PeakRequestQueueSize, static_cast<int32>(RequestQueue.ReadyRequestsNum())));
+	// After all clusters have been processed, pull a batch of readyrequests into the load state
+	COOK_STAT(DetailedCookStats::PeakRequestQueueSize = FMath::Max(DetailedCookStats::PeakRequestQueueSize,
+		static_cast<int32>(RequestQueue.ReadyRequestsNum())));
 	uint32 NumInBatch = 0;
 	while (!RequestQueue.IsReadyRequestsEmpty() && NumInBatch < RequestBatchSize)
 	{
 		FPackageData* PackageData = RequestQueue.PopReadyRequest();
+		check(PackageData->GetState() == EPackageState::Request);
 		FPoppedPackageDataScope Scope(*PackageData);
 		if (TryCreateRequestCluster(*PackageData))
 		{
 			continue;
 		}
-		if (PackageData->AreAllRequestedPlatformsCooked(true /* bAllowFailedCooks */))
+		if (PackageData->GetPlatformsNeedingCookingNum() == 0)
 		{
-#if DEBUG_COOKONTHEFLY
-			UE_LOG(LogCook, Display, TEXT("Package for platform already cooked %s, discarding request"), *PackageData.GetFileName().ToString());
-#endif
 			DemoteToIdle(*PackageData, ESendFlags::QueueAdd, ESuppressCookReason::AlreadyCooked);
 			continue;
 		}
@@ -2572,10 +2598,10 @@ void UCookOnTheFlyServer::DemoteToIdle(UE::Cook::FPackageData& PackageData, UE::
 						ULevel::GetExternalActorsFolderName(), ESearchCase::IgnoreCase) != INDEX_NONE))
 				{
 					UE_CLOG((GCookProgressDisplay & (int32)ECookProgressDisplayMode::Instigators), LogCook, Display,
-						TEXT("Cooking %s, Instigator: { %s } -> Skipped %s"), *PackageNameStr,
+						TEXT("Cooking %s, Instigator: { %s } -> Rejected %s"), *PackageNameStr,
 						*(PackageData.GetInstigator().ToString()), LexToString(Reason));
 					UE_CLOG(GCookProgressDisplay & (int32)ECookProgressDisplayMode::PackageNames, LogCook, Display,
-						TEXT("Cooking %s -> Skipped %s"), *PackageNameStr, LexToString(Reason));
+						TEXT("Cooking %s -> Rejected %s"), *PackageNameStr, LexToString(Reason));
 				}
 			}
 		}
@@ -2723,9 +2749,7 @@ void UCookOnTheFlyServer::LoadPackageInQueue(UE::Cook::FPackageData& PackageData
 			FPackageData& OtherPackageData = PackageDatas->AddPackageDataByPackageNameChecked(LoadedPackage->GetFName());
 			UE_LOG(LogCook, Verbose, TEXT("Request for %s received going to save %s"), *PackageFileName.ToString(),
 				*OtherPackageData.GetFileName().ToString());
-			TArray<const ITargetPlatform*, TInlineAllocator<ExpectedMaxNumPlatforms>> RequestedPlatforms;
-			PackageData.GetRequestedPlatforms(RequestedPlatforms);
-			QueueDiscoveredPackageData(OtherPackageData, FInstigator(PackageData.GetInstigator()));
+			QueueDiscoveredPackage(OtherPackageData, FInstigator(PackageData.GetInstigator()), EDiscoveredPlatformSet::CopyFromInstigator);
 
 			PackageData.SetPlatformsCooked(PlatformManager->GetSessionPlatforms(), ECookResult::Succeeded);
 			RejectPackageToLoad(PackageData, TEXT("is redirected to another filename"), ESuppressCookReason::Redirected);
@@ -2785,7 +2809,7 @@ void UCookOnTheFlyServer::LoadPackageInQueue(UE::Cook::FPackageData& PackageData
 		}
 	}
 
-	if (PackageData.AreAllRequestedPlatformsCooked(true))
+	if (PackageData.GetPlatformsNeedingCookingNum() == 0)
 	{
 		// Already cooked. This can happen if we needed to load a package that was previously cooked and garbage collected because it is a loaddependency of a new request.
 		// Send the package back to idle, nothing further to do with it.
@@ -2802,9 +2826,9 @@ void UCookOnTheFlyServer::LoadPackageInQueue(UE::Cook::FPackageData& PackageData
 void UCookOnTheFlyServer::RejectPackageToLoad(UE::Cook::FPackageData& PackageData, const TCHAR* ReasonText, UE::Cook::ESuppressCookReason Reason)
 {
 	// make sure this package doesn't exist
-	for (const TPair<const ITargetPlatform*, UE::Cook::FPackageData::FPlatformData>& Pair : PackageData.GetPlatformDatas())
+	for (const TPair<const ITargetPlatform*, UE::Cook::FPackagePlatformData>& Pair : PackageData.GetPlatformDatas())
 	{
-		if (!Pair.Value.IsRequested())
+		if (Pair.Key == CookerLoadingPlatformKey || !Pair.Value.NeedsCooking(Pair.Key))
 		{
 			continue;
 		}
@@ -2823,54 +2847,35 @@ void UCookOnTheFlyServer::RejectPackageToLoad(UE::Cook::FPackageData& PackageDat
 	DemoteToIdle(PackageData, UE::Cook::ESendFlags::QueueAdd, Reason);
 }
 
-//////////////////////////////////////////////////////////////////////////
-
-UE::Cook::FPackageData* UCookOnTheFlyServer::QueueDiscoveredPackage(UPackage* Package,
-	UE::Cook::FInstigator&& Instigator, bool* bOutWasInProgress)
+void UCookOnTheFlyServer::QueueDiscoveredPackage(UE::Cook::FPackageData& PackageData,
+	UE::Cook::FInstigator&& Instigator, UE::Cook::FDiscoveredPlatformSet&& ReachablePlatforms)
 {
-	check(Package != nullptr);
+	using namespace UE::Cook;
 
-	UE::Cook::FPackageData* PackageData = PackageDatas->TryAddPackageDataByPackageName(Package->GetFName());
-	if (!PackageData)
+	TArray<const ITargetPlatform*, TInlineAllocator<ExpectedMaxNumPlatforms>> BufferPlatforms;
+	if (PackageData.HasReachablePlatforms(
+		ReachablePlatforms.GetPlatforms(*this, &Instigator,
+			TConstArrayView<const ITargetPlatform*>(), &BufferPlatforms, true /* bAllowPartialInstigatorResults */)))
 	{
-		return nullptr;	// Getting the PackageData will fail if e.g. it is a script package
+		// Not a new discovery; ignore
+		return;
 	}
 
-	if (bOutWasInProgress)
+	if (bHiddenDependenciesDebug)
 	{
-		*bOutWasInProgress = PackageData->IsInProgress();
+		if (!PackageData.IsGenerated())
+		{
+			OnDiscoveredPackageDebug(PackageData.GetPackageName(), Instigator);
+		}
 	}
-	QueueDiscoveredPackageData(*PackageData, MoveTemp(Instigator), true /* bIsLoadReady */);
-	return PackageData;
+	WorkerRequests->QueueDiscoveredPackage(*this, PackageData, MoveTemp(Instigator), MoveTemp(ReachablePlatforms));
 }
 
-void UCookOnTheFlyServer::QueueDiscoveredPackageData(UE::Cook::FPackageData& PackageData,
-	UE::Cook::FInstigator&& Instigator, bool bLoadReady)
+void UCookOnTheFlyServer::QueueDiscoveredPackageOnDirector(UE::Cook::FPackageData& PackageData,
+	UE::Cook::FInstigator&& Instigator, UE::Cook::FDiscoveredPlatformSet&& ReachablePlatforms)
 {
 	using namespace UE::Cook;
 
-	const TArray<const ITargetPlatform*>& TargetPlatforms = PlatformManager->GetSessionPlatforms();
-	if (PackageData.HasAllCookedPlatforms(TargetPlatforms, true /* bIncludeFailed */))
-	{
-		// All SessionPlatforms have already been cooked for the package, so we don't need to save it again
-		return;
-	}
-
-	if (!PackageData.CanCookForPlatforms())
-	{
-		return;
-	}
-
-	if (PackageData.IsInProgress())
-	{
-		return;
-	}
-
-	if (bHiddenDependenciesDebug)
-	{
-		OnDiscoveredPackageDebug(PackageData.GetPackageName(), Instigator);
-	}
-
 	if (CookOnTheFlyRequestManager)
 	{
 		if (PackageData.IsGenerated())
@@ -2883,65 +2888,41 @@ void UCookOnTheFlyServer::QueueDiscoveredPackageData(UE::Cook::FPackageData& Pac
 		}
 	}
 
-	bool bShouldAddToQueue;
-	WorkerRequests->AddDiscoveredPackage(PackageData, Instigator, bLoadReady, bShouldAddToQueue);
-	if (!bShouldAddToQueue)
+	if (!CookByTheBookOptions->bSkipHardReferences ||
+		(Instigator.Category == EInstigator::GeneratedPackage))
 	{
-		return;
-	}
-
-	if ((Instigator.Category == EInstigator::GeneratedPackage) || 
-		(Instigator.Category == EInstigator::Unspecified) || 
-		(Instigator.Category == EInstigator::Unsolicited && !CookByTheBookOptions->bSkipHardReferences) ||
-		!CookByTheBookOptions->bSkipSoftReferences)
-	{
-		PackageData.SetRequestData(TargetPlatforms, /*bIsUrgent*/ false, UE::Cook::FCompletionCallback(),
-			MoveTemp(Instigator));
-		if (bLoadReady)
-		{
-			// Send this package into the LoadReadyQueue to fully load it and send it on to the SaveQueue
-			PackageData.SendToState(UE::Cook::EPackageState::LoadReady, UE::Cook::ESendFlags::QueueRemove);
-			// Send it to the front of the LoadReadyQueue since it is mostly loaded already
-			PackageDatas->GetLoadReadyQueue().AddFront(&PackageData);
-		}
-		else
-		{
-			PackageData.SendToState(UE::Cook::EPackageState::Request, UE::Cook::ESendFlags::QueueAddAndRemove);
-		}
+		PackageData.QueueAsDiscovered(MoveTemp(Instigator), MoveTemp(ReachablePlatforms));
 	}
 }
 
-FName GInstigatorUpdatePackageFilter(TEXT("UpdatePackageFilter"));
-
-void UCookOnTheFlyServer::UpdatePackageFilter()
+void UCookOnTheFlyServer::OnRemoveSessionPlatform(const ITargetPlatform* TargetPlatform, int32 RemovedIndex)
 {
-	if (!bPackageFilterDirty)
+	using namespace UE::Cook;
+
+	for (FDiscoveryQueueElement& Element : PackageDatas->GetRequestQueue().GetDiscoveryQueue())
 	{
-		return;
+		Element.ReachablePlatforms.OnRemoveSessionPlatform(TargetPlatform, RemovedIndex);
 	}
-	bPackageFilterDirty = false;
-
-	UE_SCOPED_COOKTIMER(UpdatePackageFilter);
-	const TArray<const ITargetPlatform*>& TargetPlatforms = PlatformManager->GetSessionPlatforms();
-
-	PackageTracker->ForEachLoadedPackage(
-		[this, &TargetPlatforms](UPackage* Package)
-		{
-			UE::Cook::FPackageData* PackageData = QueueDiscoveredPackage(Package,
-			UE::Cook::FInstigator(UE::Cook::EInstigator::Unspecified, GInstigatorUpdatePackageFilter));
-			if (PackageData && PackageData->IsInProgress())
-			{
-				PackageData->UpdateRequestData(TargetPlatforms, /*bIsUrgent*/ false, UE::Cook::FCompletionCallback(),
-					UE::Cook::FInstigator(UE::Cook::EInstigator::Unspecified, GInstigatorUpdatePackageFilter));
-			}
-		}
-	);
+	for (UE::Cook::FRequestCluster& Cluster : PackageDatas->GetRequestQueue().GetRequestClusters())
+	{
+		Cluster.OnRemoveSessionPlatform(TargetPlatform);
+	}
+	PackageDatas->OnRemoveSessionPlatform(TargetPlatform);
+	WorkerRequests->OnRemoveSessionPlatform(TargetPlatform);
 }
 
-void UCookOnTheFlyServer::OnRemoveSessionPlatform(const ITargetPlatform* TargetPlatform)
+void UCookOnTheFlyServer::OnPlatformAddedToSession(const ITargetPlatform* TargetPlatform)
 {
-	PackageDatas->OnRemoveSessionPlatform(TargetPlatform);
-	WorkerRequests->OnRemoveSessionPlatform(TargetPlatform);
+	using namespace UE::Cook;
+
+	for (FDiscoveryQueueElement& Element : PackageDatas->GetRequestQueue().GetDiscoveryQueue())
+	{
+		Element.ReachablePlatforms.OnPlatformAddedToSession(TargetPlatform);
+	}
+	for (UE::Cook::FRequestCluster& Cluster : PackageDatas->GetRequestQueue().GetRequestClusters())
+	{
+		Cluster.OnPlatformAddedToSession(TargetPlatform);
+	}
 }
 
 void UCookOnTheFlyServer::TickNetwork()
@@ -3114,10 +3095,13 @@ UE::Cook::EPollStatus UCookOnTheFlyServer::QueueGeneratedPackages(UE::Cook::FGen
 	FName OwnerName = Owner->GetFName();
 	if (Info.GetSaveState() <= FCookGenerationInfo::ESaveState::QueueGeneratedPackages)
 	{
+		TArray<const ITargetPlatform*, TInlineAllocator<ExpectedMaxNumPlatforms>> ReachablePlatforms;
+		PackageData.GetReachablePlatforms(ReachablePlatforms);
 		for (const FCookGenerationInfo& ChildInfo: Generator.GetPackagesToGenerate())
 		{
 			FPackageData* ChildPackageData = ChildInfo.PackageData;
-			QueueDiscoveredPackageData(*ChildPackageData, FInstigator(EInstigator::GeneratedPackage, OwnerName));
+			QueueDiscoveredPackage(*ChildPackageData, FInstigator(EInstigator::GeneratedPackage, OwnerName),
+				EDiscoveredPlatformSet::CopyFromInstigator);
 		}
 		Info.SetSaveStateComplete(FCookGenerationInfo::ESaveState::QueueGeneratedPackages);
 	}
@@ -3305,7 +3289,7 @@ UE::Cook::EPollStatus UCookOnTheFlyServer::BeginCacheObjectsToMove(UE::Cook::FGe
 	{
 		return Result;
 	}
-	Info.BeginCacheObjects.EndRound(PackageData.GetNumRequestedPlatforms());
+	Info.BeginCacheObjects.EndRound(PackageData.GetCachedObjectsInOuterNumPlatforms());
 	return EPollStatus::Success;
 }
 
@@ -3444,7 +3428,7 @@ UE::Cook::EPollStatus UCookOnTheFlyServer::BeginCachePostMove(UE::Cook::FGenerat
 	{
 		return Result;
 	}
-	Info.BeginCacheObjects.EndRound(PackageData.GetNumRequestedPlatforms());
+	Info.BeginCacheObjects.EndRound(PackageData.GetCachedObjectsInOuterNumPlatforms());
 
 	return EPollStatus::Success;
 }
@@ -3696,7 +3680,7 @@ UE::Cook::EPollStatus UCookOnTheFlyServer::CallBeginCacheOnObjects(UE::Cook::FPa
 	check(Package);
 
 	TArray<const ITargetPlatform*, TInlineAllocator<ExpectedMaxNumPlatforms>> TargetPlatforms;
-	PackageData.GetRequestedPlatforms(TargetPlatforms);
+	PackageData.GetCachedObjectsInOuterPlatforms(TargetPlatforms);
 
 	FWeakObjectPtr* ObjectsData = Objects.GetData();
 	int NumPlatforms = TargetPlatforms.Num();
@@ -3833,7 +3817,7 @@ void UCookOnTheFlyServer::ReleaseCookedPlatformData(UE::Cook::FPackageData& Pack
 		// and this PackageData may therefore still have GetNumPendingCookedPlatformData > 0
 		if (!IsCookingInEditor()) // ClearAllCachedCookedPlatformData calls are only used when not in editor.
 		{
-			int32 NumPlatforms = PackageData.GetNumRequestedPlatforms();
+			int32 NumPlatforms = PackageData.GetCachedObjectsInOuterNumPlatforms();
 			if (NumPlatforms > 0) // Shouldn't happen because PumpSaves checks for this, but avoid a divide by 0 if it does.
 			{
 				// We have only called BeginCacheForCookedPlatformData on Object,Platform pairs up to GetCookedPlatformDataNextIndex.
@@ -4095,17 +4079,48 @@ void UCookOnTheFlyServer::ProcessUnsolicitedPackages(TArray<FName>* OutDiscovere
 	for (auto& PackageWithInstigator : NewPackages)
 	{
 		UPackage* Package = PackageWithInstigator.Key;
+		check(Package != nullptr);
 		FInstigator& Instigator = PackageWithInstigator.Value;
 
-		bool bWasInProgress;
-		FPackageData* PackageData = QueueDiscoveredPackage(Package, FInstigator(Instigator), &bWasInProgress);
-		if (PackageData && OutDiscoveredPackageNames && !bWasInProgress)
+		UE::Cook::FPackageData* PackageData = PackageDatas->TryAddPackageDataByPackageName(Package->GetFName());
+		if (!PackageData)
 		{
-			FInstigator& Existing = OutInstigators->FindOrAdd(PackageData->GetPackageName());
-			if (Existing.Category == EInstigator::InvalidCategory)
+			continue; // Getting the PackageData will fail if e.g. it is a script package
+		}
+		if (PackageData->IsGenerated())
+		{
+			continue; // Generated pacakges are queued separately (with the correct instigator) in QueueGeneratedPackages
+		}
+
+		if (OutDiscoveredPackageNames)
+		{
+			if (!PackageData->IsInProgress())
+			{
+				FInstigator& Existing = OutInstigators->FindOrAdd(PackageData->GetPackageName());
+				if (Existing.Category == EInstigator::InvalidCategory)
+				{
+					OutDiscoveredPackageNames->Add(PackageData->GetPackageName());
+					Existing = Instigator;
+				}
+			}
+		}
+
+		if (!PackageData->FindOrAddPlatformData(CookerLoadingPlatformKey).IsReachable())
+		{
+			// Possibly an unexpected load, but it might be an expected load after we finish processing some other unexpected loads that were
+			// discovered earlier during the same LoadPackage call. Queue it for discovery and if it is still unreachable when its turn comes
+			// we will add the hidden dependency for it.
+			// as reachable in all of its instigator's reachable platforms
+			QueueDiscoveredPackage(*PackageData, MoveTemp(Instigator), EDiscoveredPlatformSet::CopyFromInstigator);
+		}
+		else
+		{
+			// This load was expected so we do not need to add a hidden dependency for it.
+			// ONLYEDITORONLY_TODO: Read the list of Imports in FinishPlatform out of the SavePackage results instead, and then turn off bMarkExpectedLoadsReachableAsWell
+			constexpr bool bMarkExpectedLoadsReachableAsWell = true;
+			if (bMarkExpectedLoadsReachableAsWell)
 			{
-				OutDiscoveredPackageNames->Add(PackageData->GetPackageName());
-				Existing = Instigator;
+				QueueDiscoveredPackage(*PackageData, MoveTemp(Instigator), EDiscoveredPlatformSet::CopyFromInstigator);
 			}
 		}
 	}
@@ -4223,7 +4238,7 @@ void UCookOnTheFlyServer::PumpSaves(UE::Cook::FTickStackData& StackData, uint32
 		}
 
 		// Cook only the session platforms that have not yet been cooked for the given package
-		PackageData.GetUncookedPlatforms(PlatformsForPackage);
+		PackageData.GetPlatformsNeedingCooking(PlatformsForPackage);
 		if (PlatformsForPackage.Num() == 0)
 		{
 			// We've already saved all possible platforms for this package; this should not be possible.
@@ -4452,7 +4467,8 @@ void UCookOnTheFlyServer::PostLoadPackageFixup(UE::Cook::FPackageData& PackageDa
 		UE::Cook::FPackageData* NewPackageData = PackageDatas->TryAddPackageDataByPackageName(FName(*PackageName));
 		if (NewPackageData)
 		{
-			QueueDiscoveredPackageData(*NewPackageData, UE::Cook::FInstigator(UE::Cook::EInstigator::Dependency, OwnerName));
+			QueueDiscoveredPackage(*NewPackageData, UE::Cook::FInstigator(UE::Cook::EInstigator::Dependency, OwnerName),
+				UE::Cook::EDiscoveredPlatformSet::CopyFromInstigator);
 		}
 	}
 }
@@ -5478,14 +5494,16 @@ void UCookOnTheFlyServer::MarkPackageDirtyForCookerFromSchedulerThread(const FNa
 	{
 		check(IsInGameThread()); // We're editing scheduler data, which is only allowable from the scheduler thread
 		bool bHadCookedPlatforms = PackageData->HasAnyCookedPlatform();
-		PackageData->ClearCookProgress();
+		PackageData->ClearCookResults();
 		if (PackageData->IsInProgress())
 		{
 			PackageData->SendToState(UE::Cook::EPackageState::Request, UE::Cook::ESendFlags::QueueAddAndRemove);
 		}
 		else if (IsCookByTheBookMode() && IsInSession() && bHadCookedPlatforms)
 		{
-			QueueDiscoveredPackageData(*PackageData, UE::Cook::FInstigator(UE::Cook::EInstigator::Unspecified, GInstigatorMarkPackageDirty));
+			QueueDiscoveredPackage(*PackageData,
+				UE::Cook::FInstigator(UE::Cook::EInstigator::Unspecified, GInstigatorMarkPackageDirty),
+				UE::Cook::EDiscoveredPlatformSet::CopyFromInstigator);
 		}
 
 		if ( IsCookOnTheFlyMode() && FileModifiedDelegate.IsBound())
@@ -5847,7 +5865,7 @@ void FSaveCookedPackageContext::SetupPlatform(const ITargetPlatform* InTargetPla
 	Info.LooseFilePath = PlatFilename;
 	PackageWriter->BeginPackage(Info);
 	// Set platform-specific save flags
-	FPackageData::FPlatformData& PlatformData = PackageData.FindOrAddPlatformData(TargetPlatform);
+	FPackagePlatformData& PlatformData = PackageData.FindOrAddPlatformData(TargetPlatform);
 	uint32 PlatformSaveFlagsMask = SAVE_AllowTimeout;
 	SaveFlags = (SaveFlags & ~PlatformSaveFlagsMask);
 	if (!PlatformData.IsSaveTimedOut())
@@ -6058,17 +6076,21 @@ void FSaveCookedPackageContext::FinishPackage()
 	FName PackageFName = Package->GetFName();
 	if (!COTFS.CookByTheBookOptions->bSkipSoftReferences)
 	{
+		TConstArrayView<const ITargetPlatform*> ReachablePlatforms = PlatformsForPackage;
 		// Also request any localized variants of this package
 		for (FName LocalizedPackageName : FRequestCluster::GetLocalizationReferences(PackageFName, COTFS))
 		{
 			UE::Cook::FPackageData* LocalizedPackageData = COTFS.PackageDatas->TryAddPackageDataByPackageName(LocalizedPackageName);
 			if (LocalizedPackageData)
 			{
-				COTFS.QueueDiscoveredPackageData(*LocalizedPackageData, UE::Cook::FInstigator(UE::Cook::EInstigator::SoftDependency, PackageFName));
+				COTFS.QueueDiscoveredPackage(*LocalizedPackageData,
+					UE::Cook::FInstigator(UE::Cook::EInstigator::SoftDependency, PackageFName),
+					FDiscoveredPlatformSet(ReachablePlatforms));
 			}
 		}
 
 		// Add SoftObjectPaths to the cook. This has to be done after the package save to catch any SoftObjectPaths that were added during save.
+		// ONLYEDITORONLY_TODO: Read the list of SoftPackagePaths in FinishPlatform out of the SavePackage results instead
 		TSet<FName> SoftObjectPackages;
 		GRedirectCollector.ProcessSoftObjectPathPackageList(PackageFName, false /* bGetEditorOnly */, SoftObjectPackages);
 		for (FName SoftObjectPackage : SoftObjectPackages)
@@ -6087,7 +6109,9 @@ void FSaveCookedPackageContext::FinishPackage()
 			UE::Cook::FPackageData* SoftObjectPackageData = COTFS.PackageDatas->TryAddPackageDataByPackageName(SoftObjectPackage);
 			if (SoftObjectPackageData)
 			{
-				COTFS.QueueDiscoveredPackageData(*SoftObjectPackageData, UE::Cook::FInstigator(UE::Cook::EInstigator::SoftDependency, PackageFName));
+				COTFS.QueueDiscoveredPackage(*SoftObjectPackageData,
+					UE::Cook::FInstigator(UE::Cook::EInstigator::SoftDependency, PackageFName),
+					FDiscoveredPlatformSet(ReachablePlatforms));
 			}
 		}
 	}
@@ -9311,12 +9335,13 @@ void UCookOnTheFlyServer::CancelAllQueues()
 	}
 	PackageDatas->GetAssignedToWorkerSet().Empty();
 	FRequestQueue& RequestQueue = PackageDatas->GetRequestQueue();
-	FPackageDataSet& UnclusteredRequests = RequestQueue.GetUnclusteredRequests();
-	for (FPackageData* PackageData : UnclusteredRequests)
+	RequestQueue.GetDiscoveryQueue().Empty();
+	FPackageDataSet& RestartedRequests = RequestQueue.GetRestartedRequests();
+	for (FPackageData* PackageData : RestartedRequests)
 	{
 		DemoteToIdle(*PackageData, ESendFlags::QueueAdd, ESuppressCookReason::CookCanceled);
 	}
-	UnclusteredRequests.Empty();
+	RestartedRequests.Empty();
 	TRingBuffer<FRequestCluster>& RequestClusters = RequestQueue.GetRequestClusters();
 	for (FRequestCluster& RequestCluster : RequestClusters)
 	{
@@ -9365,18 +9390,19 @@ void UCookOnTheFlyServer::ResetCook(TConstArrayView<TPair<const ITargetPlatform*
 {
 	PackageDatas->LockAndEnumeratePackageDatas([TargetPlatforms](UE::Cook::FPackageData* PackageData)
 	{
+		PackageData->FindOrAddPlatformData(CookerLoadingPlatformKey).ResetReachable();
+
 		for (const TPair<const ITargetPlatform*, bool>& Pair : TargetPlatforms)
 		{
 			const ITargetPlatform* TargetPlatform = Pair.Key;
-			UE::Cook::FPackageData::FPlatformData* PlatformData = PackageData->FindPlatformData(TargetPlatform);
+			UE::Cook::FPackagePlatformData* PlatformData = PackageData->FindPlatformData(TargetPlatform);
 			if (PlatformData)
 			{
 				bool bResetResults = Pair.Value;
-				PlatformData->SetExplored(false);
+				PlatformData->ResetReachable();
 				if (bResetResults)
 				{
-					PlatformData->SetSaveTimedOut(false);
-					PackageData->ClearCookProgress(TargetPlatform);
+					PackageData->ClearCookResults(TargetPlatform);
 				}
 			}
 		}
@@ -9393,6 +9419,8 @@ void UCookOnTheFlyServer::ResetCook(TConstArrayView<TPair<const ITargetPlatform*
 			PackageTracker->UnsolicitedCookedPackages.GetPackagesForPlatformAndRemove(TargetPlatform, PackageNames);
 		}
 	}
+
+	PackageTracker->MarkLoadedPackagesAsNew();
 }
 
 void UCookOnTheFlyServer::ClearCachedCookedPlatformDataForPlatform(const ITargetPlatform* TargetPlatform)
@@ -10620,10 +10648,7 @@ void UCookOnTheFlyServer::SelectSessionPlatforms(FBeginCookContext& BeginContext
 {
 	InitializeSession();
 	PlatformManager->SelectSessionPlatforms(*this, BeginContext.TargetPlatforms);
-	if (PackageTracker->HasBeenConsumed())
-	{
-		bPackageFilterDirty = true;
-	}
+
 	FindOrCreateSaveContexts(BeginContext.TargetPlatforms);
 	for (FBeginCookContextPlatform& PlatformContext : BeginContext.PlatformContexts)
 	{
@@ -10884,9 +10909,9 @@ void UCookOnTheFlyServer::MaybeMarkPackageAsAlreadyLoaded(UPackage *Package)
 		if (IsCookFlagSet(ECookInitializationFlags::LogDebugInfo))
 		{
 			FString Platforms;
-			for (const TPair<const ITargetPlatform*, UE::Cook::FPackageData::FPlatformData>& Pair : PackageData->GetPlatformDatas())
+			for (const TPair<const ITargetPlatform*, UE::Cook::FPackagePlatformData>& Pair : PackageData->GetPlatformDatas())
 			{
-				if (Pair.Value.IsCookAttempted())
+				if (Pair.Key != CookerLoadingPlatformKey && Pair.Value.IsCookAttempted())
 				{
 					Platforms += TEXT(" ");
 					Platforms += Pair.Key->PlatformName();
@@ -10904,12 +10929,6 @@ void UCookOnTheFlyServer::MaybeMarkPackageAsAlreadyLoaded(UPackage *Package)
 	}
 
 	check(IsInGameThread());
-	if (PackageTracker->NeverCookPackageList.Contains(StandardName))
-	{
-		bShouldMarkAsAlreadyProcessed = true;
-		UE_LOG(LogCook, Verbose, TEXT("Marking %s as reloading for cooker because it was requested as never cook package."), *StandardName.ToString());
-	}
-
 	if (bShouldMarkAsAlreadyProcessed)
 	{
 		if (Package->IsFullyLoaded() == false)
@@ -11407,13 +11426,6 @@ void UCookOnTheFlyServer::OnDiscoveredPackageDebug(FName PackageName, const UE::
 					return;
 				}
 
-				if (!Value.bCookable)
-				{
-					// Skipped during exploration because it is not cookable; we will skip it again in QueueDiscoveredPackageData.
-					bShouldReport = false;
-					return;
-				}
-
 				if (Value.bHasLoggedDiscoveryWarning)
 				{
 					// Discovered and warned earlier; has not yet completed the request phase so Instigator is not yet set
diff --git a/Engine/Source/Editor/UnrealEd/Private/Cooker/CookDirector.cpp b/Engine/Source/Editor/UnrealEd/Private/Cooker/CookDirector.cpp
--- a/Engine/Source/Editor/UnrealEd/Private/Cooker/CookDirector.cpp
+++ b/Engine/Source/Editor/UnrealEd/Private/Cooker/CookDirector.cpp
@@ -1646,10 +1646,10 @@ void FCookDirector::FRetractionHandler::ReassignPackages(const FWorkerId& FromWo
 	if (WorkersToSplitOver.IsEmpty())
 	{
 		// Send the packages back to the Director for reassignment
-		FPackageDataSet& UnclusteredRequests = Director.COTFS.PackageDatas->GetRequestQueue().GetUnclusteredRequests();
+		FPackageDataSet& RestartedRequests = Director.COTFS.PackageDatas->GetRequestQueue().GetRestartedRequests();
 		for (FPackageData* PackageData : AssignmentPackages)
 		{
-			UnclusteredRequests.Add(PackageData);
+			RestartedRequests.Add(PackageData);
 		}
 		// MPCOOKTODO: Add a method to PumpRequests long enough to assign the packages
 		UE_LOG(LogCook, Display, TEXT("%d packages retracted from %s. No workers are currently idle so the packages were assigned evenly to all CookWorkers."),
diff --git a/Engine/Source/Editor/UnrealEd/Private/Cooker/CookPackageData.cpp b/Engine/Source/Editor/UnrealEd/Private/Cooker/CookPackageData.cpp
--- a/Engine/Source/Editor/UnrealEd/Private/Cooker/CookPackageData.cpp
+++ b/Engine/Source/Editor/UnrealEd/Private/Cooker/CookPackageData.cpp
@@ -51,12 +51,34 @@ static FAutoConsoleVariableRef CVarPollAsyncPeriod(
 	
 //////////////////////////////////////////////////////////////////////////
 // FPackageData
-FPackageData::FPlatformData::FPlatformData()
-	: bRequested(0), bExplored(0), bSaveTimedOut(0), bCookable(1), bIterativelySkipped(0)
-	, CookResults((uint8)ECookResult::NotAttempted)
+FPackagePlatformData::FPackagePlatformData()
+	: bReachable(0), bVisitedByCluster(0), bSaveTimedOut(0), bCookable(1), bExplorable(1), bIterativelySkipped(0)
+	, bRegisteredForCachedObjectsInOuter(0), CookResults((uint8)ECookResult::NotAttempted)
 {
 }
 
+void FPackagePlatformData::ResetReachable()
+{
+	SetReachable(false);
+	SetVisitedByCluster(false);
+	SetCookable(true);
+	SetExplorable(true);
+}
+
+void FPackagePlatformData::MarkCookableForWorker(FCookWorkerClient& CookWorkerClient)
+{
+	SetReachable(true);
+	SetVisitedByCluster(true);
+	SetExplorable(true);
+	SetCookable(true);
+	SetCookResults(ECookResult::NotAttempted);
+}
+
+bool FPackagePlatformData::NeedsCooking(const ITargetPlatform* PlatformItBelongsTo) const
+{
+	return IsReachable() && PlatformItBelongsTo != CookerLoadingPlatformKey && IsCookable() && !IsCookAttempted();
+}
+
 FPackageData::FPackageData(FPackageDatas& PackageDatas, const FName& InPackageName, const FName& InFileName)
 	: GeneratedOwner(nullptr), PackageName(InPackageName), FileName(InFileName), PackageDatas(PackageDatas)
 	, Instigator(EInstigator::NotYetRequested), bIsUrgent(0)
@@ -75,7 +97,7 @@ FPackageData::~FPackageData()
 	// ClearReferences should have been called earlier, but call it here in case it was missed
 	ClearReferences();
 	// We need to send OnLastCookedPlatformRemoved message to the monitor, so call SetPlatformsNotCooked
-	ClearCookProgress();
+	ClearCookResults();
 	// Update the monitor's counters and call exit functions
 	SendToState(EPackageState::Idle, ESendFlags::QueueNone);
 }
@@ -100,33 +122,26 @@ void FPackageData::SetFileName(const FName& InFileName)
 	FileName = InFileName;
 }
 
-int32 FPackageData::GetNumRequestedPlatforms() const
+int32 FPackageData::GetPlatformsNeedingCookingNum() const
 {
 	int32 Result = 0;
-	for (const TPair<const ITargetPlatform*, FPlatformData>& Pair : PlatformDatas)
+	for (const TPair<const ITargetPlatform*, FPackagePlatformData>& Pair : PlatformDatas)
 	{
-		Result += Pair.Value.IsRequested() ? 1 : 0;
+		if (Pair.Value.NeedsCooking(Pair.Key))
+		{
+			++Result;
+		}
 	}
 	return Result;
 }
 
-void FPackageData::SetPlatformsRequested(TConstArrayView<const ITargetPlatform*> TargetPlatforms, bool bRequested)
-{
-	for (const ITargetPlatform* TargetPlatform : TargetPlatforms)
-	{
-		PlatformDatas.FindOrAdd(TargetPlatform).SetRequested(true);
-	}
-}
-
-void FPackageData::ClearRequestedPlatforms()
+bool FPackageData::IsPlatformVisitedByCluster(const ITargetPlatform* Platform) const
 {
-	for (TPair<const ITargetPlatform*, FPlatformData>& Pair : PlatformDatas)
-	{
-		Pair.Value.SetRequested(false);
-	}
+	const FPackagePlatformData* PlatformData = FindPlatformData(Platform);
+	return PlatformData && PlatformData->IsVisitedByCluster();
 }
 
-bool FPackageData::HasAllRequestedPlatforms(const TArrayView<const ITargetPlatform* const>& Platforms) const
+bool FPackageData::HasReachablePlatforms(const TArrayView<const ITargetPlatform* const>& Platforms) const
 {
 	if (Platforms.Num() == 0)
 	{
@@ -139,8 +154,8 @@ bool FPackageData::HasAllRequestedPlatforms(const TArrayView<const ITargetPlatfo
 
 	for (const ITargetPlatform* QueryPlatform : Platforms)
 	{
-		const FPlatformData* PlatformData = PlatformDatas.Find(QueryPlatform);
-		if (!PlatformData || !PlatformData->IsRequested())
+		const FPackagePlatformData* PlatformData = PlatformDatas.Find(QueryPlatform);
+		if (!PlatformData || !PlatformData->IsReachable())
 		{
 			return false;
 		}
@@ -148,12 +163,11 @@ bool FPackageData::HasAllRequestedPlatforms(const TArrayView<const ITargetPlatfo
 	return true;
 }
 
-bool FPackageData::AreAllRequestedPlatformsCooked(bool bAllowFailedCooks) const
+bool FPackageData::AreAllReachablePlatformsVisitedByCluster() const
 {
-	for (const TPair<const ITargetPlatform*, FPlatformData>& Pair : PlatformDatas)
+	for (const TPair<const ITargetPlatform*, FPackagePlatformData>& Pair : PlatformDatas)
 	{
-		if (Pair.Value.IsRequested() &&
-			(!Pair.Value.IsCookAttempted() || (!bAllowFailedCooks && !Pair.Value.IsCookSucceeded())))
+		if (Pair.Value.IsReachable() && !Pair.Value.IsVisitedByCluster())
 		{
 			return false;
 		}
@@ -161,64 +175,69 @@ bool FPackageData::AreAllRequestedPlatformsCooked(bool bAllowFailedCooks) const
 	return true;
 }
 
-bool FPackageData::AreAllRequestedPlatformsExplored() const
+void FPackageData::GetReachablePlatformsForInstigator(UCookOnTheFlyServer& COTFS, FName InInstigator,
+	TArray<const ITargetPlatform*>& Platforms, bool bAllowPartialInstigatorResults)
 {
-	for (const TPair<const ITargetPlatform*, FPlatformData>& Pair : PlatformDatas)
-	{
-		if (Pair.Value.IsRequested() && !Pair.Value.IsExplored())
-		{
-			return false;
-		}
-	}
-	return true;
+	return GetReachablePlatformsForInstigator(COTFS,
+		COTFS.PackageDatas->TryAddPackageDataByPackageName(InInstigator), Platforms, bAllowPartialInstigatorResults);
 }
 
-bool FPackageData::HasAllExploredPlatforms(const TArrayView<const ITargetPlatform* const>& Platforms) const
+void FPackageData::GetReachablePlatformsForInstigator(UCookOnTheFlyServer& COTFS, UE::Cook::FPackageData* InInstigator,
+	TArray<const ITargetPlatform*>& Platforms, bool bAllowPartialInstigatorResults)
 {
-	if (Platforms.Num() == 0)
-	{
-		return true;
-	}
-	if (PlatformDatas.Num() == 0)
-	{
-		return false;
-	}
-
-	for (const ITargetPlatform* QueryPlatform : Platforms)
+	if (InInstigator)
 	{
-		const FPlatformData* PlatformData = FindPlatformData(QueryPlatform);
-		if (!PlatformData || !PlatformData->IsExplored())
+		InInstigator->GetReachablePlatforms(Platforms);
+		if (Platforms.IsEmpty() && !bAllowPartialInstigatorResults)
 		{
-			return false;
+			UE_LOG(LogCook, Error, TEXT("An instigator package is not marked as reachable for any platform, should be impossible: %s"),
+				*InInstigator->GetPackageName().ToString());
 		}
 	}
-	return true;
+	else
+	{
+		const TArray<const ITargetPlatform*>& SessionPlatforms = COTFS.PlatformManager->GetSessionPlatforms();
+		Platforms.Reset(SessionPlatforms.Num() + 1);
+		Platforms.Append(SessionPlatforms);
+	}
 }
 
-bool FPackageData::CanCookForPlatforms() const
+
+void FPackageData::AddReachablePlatforms(FRequestCluster& RequestCluster,
+	TConstArrayView<const ITargetPlatform*> Platforms, FInstigator&& InInstigator)
 {
-	if (PlatformDatas.Num() == 0)
-	{
-		return true;
-	}
+	AddReachablePlatformsInternal(*this, Platforms, MoveTemp(InInstigator));
+}
 
-	if (AreAllRequestedPlatformsExplored())
+void FPackageData::AddReachablePlatformsInternal(FPackageData& PackageData,
+	TConstArrayView<const ITargetPlatform*> Platforms, FInstigator&& InInstigator)
+{
+	// This is a static helper function to make it impossible to make a typo and use this->Instigator instead of InInstigator
+	bool bSessionPlatformModified = false;
+	for (const ITargetPlatform* Platform : Platforms)
 	{
-		for (const TPair<const ITargetPlatform*, FPlatformData>& Pair : PlatformDatas)
-		{
-			if (Pair.Value.IsCookable())
-			{
-				return true;
-			}
-		}
-		return false;
+		FPackagePlatformData& PlatformData = PackageData.FindOrAddPlatformData(Platform);
+		bSessionPlatformModified |= (Platform != CookerLoadingPlatformKey && !PlatformData.IsReachable());
+		PlatformData.SetReachable(true);
 	}
-	else
+	if (bSessionPlatformModified)
 	{
-		return true;
+		PackageData.SetInstigatorInternal(MoveTemp(InInstigator));
 	}
 }
 
+void FPackageData::QueueAsDiscovered(FInstigator&& InInstigator, FDiscoveredPlatformSet&& ReachablePlatforms)
+{
+	QueueAsDiscoveredInternal(*this, MoveTemp(InInstigator), MoveTemp(ReachablePlatforms));
+}
+
+void FPackageData::QueueAsDiscoveredInternal(FPackageData& PackageData, FInstigator&& InInstigator,
+	FDiscoveredPlatformSet&& ReachablePlatforms)
+{
+	// This is a static helper function to make it impossible to make a typo and use this->Instigator instead of InInstigator
+	TRingBuffer<FDiscoveryQueueElement>& Queue = PackageData.PackageDatas.GetRequestQueue().GetDiscoveryQueue();
+	Queue.Add(FDiscoveryQueueElement{ &PackageData, MoveTemp(InInstigator), MoveTemp(ReachablePlatforms) });
+}
 
 void FPackageData::SetIsUrgent(bool Value)
 {
@@ -226,58 +245,45 @@ void FPackageData::SetIsUrgent(bool Value)
 	if (OldValue != Value)
 	{
 		bIsUrgent = Value != 0;
+		check(IsInProgress() || !bIsUrgent);
 		PackageDatas.GetMonitor().OnUrgencyChanged(*this);
 	}
 }
 
-void FPackageData::UpdateRequestData(const TConstArrayView<const ITargetPlatform*> InRequestedPlatforms,
-	bool bInIsUrgent, FCompletionCallback&& InCompletionCallback, FInstigator&& InInstigator, bool bAllowUpdateUrgency)
+void FPackageData::AddUrgency(bool bUrgent, bool bAllowUpdateState)
 {
-	check(IsInProgress());
-	AddCompletionCallback(MoveTemp(InCompletionCallback));
-
-	bool bUrgencyChanged = false;
-	if (bInIsUrgent && !GetIsUrgent())
-	{
-		bUrgencyChanged = true;
-		SetIsUrgent(true);
-	}
-
-	if (!HasAllRequestedPlatforms(InRequestedPlatforms))
+	if (!bUrgent)
 	{
-		// Send back to the Request state (canceling any current operations) and then add the new platforms
-		if (GetState() != EPackageState::Request)
-		{
-			SendToState(EPackageState::Request, ESendFlags::QueueAddAndRemove);
-		}
-		SetPlatformsRequested(InRequestedPlatforms, true);
+		return;
 	}
-	else if (bUrgencyChanged && bAllowUpdateUrgency)
+	bool bWasUrgent = GetIsUrgent();
+	SetIsUrgent(true);
+	if (!bWasUrgent && bAllowUpdateState)
 	{
 		SendToState(GetState(), ESendFlags::QueueAddAndRemove);
 	}
 }
 
-void FPackageData::SetRequestData(const TArrayView<const ITargetPlatform* const>& InRequestedPlatforms,
-	bool bInIsUrgent, FCompletionCallback&& InCompletionCallback, FInstigator&& InInstigator)
+void FPackageData::SetInstigator(FRequestCluster& Cluster, FInstigator&& InInstigator)
 {
-	check(!CompletionCallback);
-	check(GetNumRequestedPlatforms() == 0)
-	check(!bIsUrgent);
+	SetInstigatorInternal(MoveTemp(InInstigator));
+}
 
-	check(InRequestedPlatforms.Num() != 0);
-	SetPlatformsRequested(InRequestedPlatforms, true);
-	SetIsUrgent(bInIsUrgent);
-	AddCompletionCallback(MoveTemp(InCompletionCallback));
-	if (Instigator.Category == EInstigator::NotYetRequested)
+void FPackageData::SetInstigator(FCookWorkerClient& Cluster, FInstigator&& InInstigator)
+{
+	SetInstigatorInternal(MoveTemp(InInstigator));
+}
+
+void FPackageData::SetInstigatorInternal(FInstigator&& InInstigator)
+{
+	if (this->Instigator.Category == EInstigator::NotYetRequested)
 	{
-		OnPackageDataFirstRequested(MoveTemp(InInstigator));
+		OnPackageDataFirstMarkedReachable(MoveTemp(InInstigator));
 	}
 }
 
 void FPackageData::ClearInProgressData()
 {
-	ClearRequestedPlatforms();
 	SetIsUrgent(false);
 	CompletionCallback = FCompletionCallback();
 }
@@ -307,7 +313,7 @@ void FPackageData::SetPlatformCooked(const ITargetPlatform* TargetPlatform, ECoo
 	bool bModifiedCookAttempted = false;
 	bool bHasAnyOtherCookAttempted = false;
 	bool bExists = false;
-	for (TPair<const ITargetPlatform*, FPlatformData>& Pair : PlatformDatas)
+	for (TPair<const ITargetPlatform*, FPackagePlatformData>& Pair : PlatformDatas)
 	{
 		if (Pair.Key == TargetPlatform)
 		{
@@ -324,7 +330,7 @@ void FPackageData::SetPlatformCooked(const ITargetPlatform* TargetPlatform, ECoo
 	}
 	if (!bExists && bNewCookAttemptedValue)
 	{
-		FPlatformData& Value = PlatformDatas.FindOrAdd(TargetPlatform);
+		FPackagePlatformData& Value = PlatformDatas.FindOrAdd(TargetPlatform);
 		Value.SetCookResults(CookResult);
 		Value.SetSaveTimedOut(false);
 		bModifiedCookAttempted = true;
@@ -342,22 +348,29 @@ void FPackageData::SetPlatformCooked(const ITargetPlatform* TargetPlatform, ECoo
 	}
 }
 
-void FPackageData::ClearCookProgress(const TConstArrayView<const ITargetPlatform*> TargetPlatforms)
+void FPackageData::ClearCookResults(const TConstArrayView<const ITargetPlatform*> TargetPlatforms)
 {
 	for (const ITargetPlatform* TargetPlatform : TargetPlatforms)
 	{
-		ClearCookProgress(TargetPlatform);
+		ClearCookResults(TargetPlatform);
 	}
 }
 
-void FPackageData::ClearCookProgress()
+void FPackageData::ResetReachable()
+{
+	for (TPair<const ITargetPlatform*, FPackagePlatformData>& Pair : PlatformDatas)
+	{
+		Pair.Value.ResetReachable();
+	}
+}
+
+void FPackageData::ClearCookResults()
 {
 	bool bModifiedCookAttempted = false;
-	for (TPair<const ITargetPlatform*, FPlatformData>& Pair : PlatformDatas)
+	for (TPair<const ITargetPlatform*, FPackagePlatformData>& Pair : PlatformDatas)
 	{
 		bModifiedCookAttempted = bModifiedCookAttempted | Pair.Value.IsCookAttempted();
 		Pair.Value.SetCookResults(ECookResult::NotAttempted);
-		Pair.Value.SetExplored(false);
 		Pair.Value.SetSaveTimedOut(false);
 	}
 	if (bModifiedCookAttempted)
@@ -366,17 +379,16 @@ void FPackageData::ClearCookProgress()
 	}
 }
 
-void FPackageData::ClearCookProgress(const ITargetPlatform* TargetPlatform)
+void FPackageData::ClearCookResults(const ITargetPlatform* TargetPlatform)
 {
 	bool bHasAnyOthers = false;
 	bool bModifiedCookAttempted = false;
-	for (TPair<const ITargetPlatform*, FPlatformData>& Pair : PlatformDatas)
+	for (TPair<const ITargetPlatform*, FPackagePlatformData>& Pair : PlatformDatas)
 	{
 		if (Pair.Key == TargetPlatform)
 		{
 			bModifiedCookAttempted = bModifiedCookAttempted | Pair.Value.IsCookAttempted();
 			Pair.Value.SetCookResults(ECookResult::NotAttempted);
-			Pair.Value.SetExplored(false);
 			Pair.Value.SetSaveTimedOut(false);
 		}
 		else
@@ -390,22 +402,27 @@ void FPackageData::ClearCookProgress(const ITargetPlatform* TargetPlatform)
 	}
 }
 
-const TSortedMap<const ITargetPlatform*, FPackageData::FPlatformData, TInlineAllocator<1>>& FPackageData::GetPlatformDatas() const
+const TSortedMap<const ITargetPlatform*, FPackagePlatformData, TInlineAllocator<1>>& FPackageData::GetPlatformDatas() const
+{
+	return PlatformDatas;
+}
+
+TSortedMap<const ITargetPlatform*, FPackagePlatformData, TInlineAllocator<1>>& FPackageData::GetPlatformDatasConstKeysMutableValues()
 {
 	return PlatformDatas;
 }
 
-FPackageData::FPlatformData& FPackageData::FindOrAddPlatformData(const ITargetPlatform* TargetPlatform)
+FPackagePlatformData& FPackageData::FindOrAddPlatformData(const ITargetPlatform* TargetPlatform)
 {
 	return PlatformDatas.FindOrAdd(TargetPlatform);
 }
 
-FPackageData::FPlatformData* FPackageData::FindPlatformData(const ITargetPlatform* TargetPlatform)
+FPackagePlatformData* FPackageData::FindPlatformData(const ITargetPlatform* TargetPlatform)
 {
 	return PlatformDatas.Find(TargetPlatform);
 }
 
-const FPackageData::FPlatformData* FPackageData::FindPlatformData(const ITargetPlatform* TargetPlatform) const
+const FPackagePlatformData* FPackageData::FindPlatformData(const ITargetPlatform* TargetPlatform) const
 {
 	return PlatformDatas.Find(TargetPlatform);
 }
@@ -413,7 +430,10 @@ const FPackageData::FPlatformData* FPackageData::FindPlatformData(const ITargetP
 bool FPackageData::HasAnyCookedPlatform() const
 {
 	return Algo::AnyOf(PlatformDatas,
-		[](const TPair<const ITargetPlatform*, FPlatformData>& Pair) { return Pair.Value.IsCookAttempted(); });
+		[](const TPair<const ITargetPlatform*, FPackagePlatformData>& Pair)
+		{
+			return Pair.Key != CookerLoadingPlatformKey && Pair.Value.IsCookAttempted();
+		});
 }
 
 bool FPackageData::HasAnyCookedPlatforms(const TArrayView<const ITargetPlatform* const>& Platforms,
@@ -464,7 +484,7 @@ bool FPackageData::HasCookedPlatform(const ITargetPlatform* Platform, bool bIncl
 
 ECookResult FPackageData::GetCookResults(const ITargetPlatform* Platform) const
 {
-	const FPlatformData* PlatformData = PlatformDatas.Find(Platform);
+	const FPackagePlatformData* PlatformData = PlatformDatas.Find(Platform);
 	if (PlatformData)
 	{
 		return PlatformData->GetCookResults();
@@ -744,9 +764,6 @@ void FPackageData::OnExitIdle()
 
 void FPackageData::OnEnterRequest()
 {
-	// It is not valid to enter the request state without requested platforms; it indicates a bug due to e.g.
-	// calling SendToState without UpdateRequestData from Idle
-	check(GetNumRequestedPlatforms() > 0);
 }
 
 void FPackageData::OnExitRequest()
@@ -853,11 +870,12 @@ void FPackageData::OnExitHasPackage()
 	SetPackage(nullptr);
 }
 
-void FPackageData::OnPackageDataFirstRequested(FInstigator&& InInstigator)
+void FPackageData::OnPackageDataFirstMarkedReachable(FInstigator&& InInstigator)
 {
 	TracePackage(GetPackageName().ToUnstableInt(), GetPackageName().ToString());
 	Instigator = MoveTemp(InInstigator);
 	PackageDatas.DebugInstigator(*this);
+	PackageDatas.UpdateThreadsafePackageData(*this);
 }
 
 void FPackageData::SetState(EPackageState NextState)
@@ -870,9 +888,26 @@ FCompletionCallback& FPackageData::GetCompletionCallback()
 	return CompletionCallback;
 }
 
-void FPackageData::AddCompletionCallback(FCompletionCallback&& InCompletionCallback)
+void FPackageData::AddCompletionCallback(TConstArrayView<const ITargetPlatform*> TargetPlatforms, 
+	FCompletionCallback&& InCompletionCallback)
 {
-	if (InCompletionCallback)
+	if (!InCompletionCallback)
+	{
+		return;
+	}
+
+	for (const ITargetPlatform* TargetPlatform : TargetPlatforms)
+	{
+		FPackagePlatformData* PlatformData = FindPlatformData(TargetPlatform);
+		// Adding a completion callback is only allowed after marking the requested platforms as reachable
+		check(PlatformData);
+		check(PlatformData->IsReachable());
+		// Adding a completion callback is only allowed after putting the PackageData in progress.
+		// If it's not in progress because it already finished the desired platforms, that is allowed.
+		check(IsInProgress() || PlatformData->IsCookAttempted() || !PlatformData->IsCookable());
+	}
+
+	if (IsInProgress())
 	{
 		// We don't yet have a mechanism for calling two completion callbacks.
 		// CompletionCallbacks only come from external requests, and it should not be possible to request twice,
@@ -880,6 +915,11 @@ void FPackageData::AddCompletionCallback(FCompletionCallback&& InCompletionCallb
 		check(!CompletionCallback);
 		CompletionCallback = MoveTemp(InCompletionCallback);
 	}
+	else
+	{
+		// Already done; call the completioncallback immediately
+		InCompletionCallback(this);
+	}
 }
 
 bool FPackageData::TryPreload()
@@ -1067,6 +1107,19 @@ TArray<FWeakObjectPtr>& FPackageData::GetCachedObjectsInOuter()
 	return CachedObjectsInOuter;
 }
 
+int32 FPackageData::GetCachedObjectsInOuterNumPlatforms() const
+{
+	int32 Result = 0;
+	for (const TPair<const ITargetPlatform*, FPackagePlatformData>& Pair : PlatformDatas)
+	{
+		if (Pair.Value.IsRegisteredForCachedObjectsInOuter())
+		{
+			++Result;
+		}
+	}
+	return Result;
+}
+
 void FPackageData::CheckObjectCacheEmpty() const
 {
 	check(CachedObjectsInOuter.Num() == 0);
@@ -1095,6 +1148,17 @@ void FPackageData::CreateObjectCache()
 			check(ObjectWeakPointer.Get()); // GetObjectsWithOuter with Garbage filtered out should only return valid-for-weakptr objects
 			CachedObjectsInOuter.Emplace(MoveTemp(ObjectWeakPointer));
 		}
+
+		for (TPair<const ITargetPlatform*, FPackagePlatformData>& Pair : PlatformDatas)
+		{
+			FPackagePlatformData& PlatformData = Pair.Value;
+			check(!PlatformData.IsRegisteredForCachedObjectsInOuter());
+			if (PlatformData.NeedsCooking(Pair.Key))
+			{
+				PlatformData.SetRegisteredForCachedObjectsInOuter(true);
+			}
+		}
+
 		SetHasSaveCache(true);
 	}
 	else
@@ -1146,7 +1210,7 @@ static TArray<UObject*> SetDifference(TArray<UObject*>& A, TArray<UObject*>& B)
 EPollStatus FPackageData::RefreshObjectCache(bool& bOutFoundNewObjects)
 {
 	// Caller will only call this function after CallBeginCacheOnObjects finished successfully
-	int32 NumPlatforms = GetNumRequestedPlatforms();
+	int32 NumPlatforms = GetCachedObjectsInOuterNumPlatforms();
 	check(NumPlatforms > 0);
 	check(GetCookedPlatformDataNextIndex()/NumPlatforms == GetCachedObjectsInOuter().Num());
 	check(Package.Get() != nullptr);
@@ -1194,6 +1258,10 @@ EPollStatus FPackageData::RefreshObjectCache(bool& bOutFoundNewObjects)
 void FPackageData::ClearObjectCache()
 {
 	CachedObjectsInOuter.Empty();
+	for (TPair<const ITargetPlatform*, FPackagePlatformData>& Pair : PlatformDatas)
+	{
+		Pair.Value.SetRegisteredForCachedObjectsInOuter(false);
+	}
 	SetHasSaveCache(false);
 }
 
@@ -1264,10 +1332,10 @@ bool FPackageData::HasReferencedObjects() const
 
 void FPackageData::RemapTargetPlatforms(const TMap<ITargetPlatform*, ITargetPlatform*>& Remap)
 {
-	typedef TSortedMap<const ITargetPlatform*, FPlatformData, TInlineAllocator<1>> MapType;
+	typedef TSortedMap<const ITargetPlatform*, FPackagePlatformData, TInlineAllocator<1>> MapType;
 	MapType NewPlatformDatas;
 	NewPlatformDatas.Reserve(PlatformDatas.Num());
-	for (TPair<const ITargetPlatform*, FPlatformData>& ExistingPair : PlatformDatas)
+	for (TPair<const ITargetPlatform*, FPackagePlatformData>& ExistingPair : PlatformDatas)
 	{
 		ITargetPlatform* NewKey = Remap[ExistingPair.Key];
 		NewPlatformDatas.FindOrAdd(NewKey) = MoveTemp(ExistingPair.Value);
@@ -1417,17 +1485,21 @@ void FGeneratorPackage::InitializeSave(const UObject* InSplitDataObject,
 {
 	if (InCookPackageSplitterInstance)
 	{
-		if (IsInitialized())
+		// If we already have a splitter, keep the old and throw out the new. The old one
+		// still contains some state.
+		if (!CookPackageSplitterInstance)
 		{
-			check(CookPackageSplitterInstance.Get() == InCookPackageSplitterInstance);
-			check(SplitDataObjectName == FName(*InSplitDataObject->GetFullName()));
+			CookPackageSplitterInstance.Reset(InCookPackageSplitterInstance);
 		}
 		else
 		{
-			CookPackageSplitterInstance.Reset(InCookPackageSplitterInstance);
-			SplitDataObjectName = *InSplitDataObject->GetFullName();
-			SetOwnerPackage(GetOwner().GetPackage());
+			delete InCookPackageSplitterInstance;
 		}
+
+		FName InSplitDataObjectName = *InSplitDataObject->GetFullName();
+		check(SplitDataObjectName.IsNone() || SplitDataObjectName == InSplitDataObjectName);
+		SplitDataObjectName = InSplitDataObjectName;
+		SetOwnerPackage(GetOwner().GetPackage());
 	}
 }
 
@@ -1481,8 +1553,8 @@ bool FGeneratorPackage::TryGenerateList(UObject* OwnerObject, FPackageDatas& Pac
 		GeneratorDatas = CookPackageSplitterInstance->GetGenerateList(LocalOwnerPackage, OwnerObject);
 	}
 	PackagesToGenerate.Reset(GeneratorDatas.Num());
-	TArray<const ITargetPlatform*, TInlineAllocator<1>> RequestedPlatforms;
-	OwnerPackageData.GetRequestedPlatforms(RequestedPlatforms);
+	TArray<const ITargetPlatform*, TInlineAllocator<1>> PlatformsToCook;
+	OwnerPackageData.GetPlatformsNeedingCooking(PlatformsToCook);
 	bool bHybridIterativeEnabled = COTFS.bHybridIterativeEnabled;
 
 	int32 NumIterativeModified = 0; 
@@ -1553,7 +1625,7 @@ bool FGeneratorPackage::TryGenerateList(UObject* OwnerObject, FPackageDatas& Pac
 		if (PreviousGeneratedPackages.RemoveAndCopyValue(PackageFName, PreviousGuid) && !bHybridIterativeEnabled)
 		{
 			bool bIdentical;
-			GeneratedInfo.IterativeCookValidateOrClear(*this, RequestedPlatforms, PreviousGuid, bIdentical);
+			GeneratedInfo.IterativeCookValidateOrClear(*this, PlatformsToCook, PreviousGuid, bIdentical);
 			if (bIdentical)
 			{
 				IdenticalGenerated.Add(PackageFName);
@@ -1569,7 +1641,7 @@ bool FGeneratorPackage::TryGenerateList(UObject* OwnerObject, FPackageDatas& Pac
 		NumIterativeRemoved = PreviousGeneratedPackages.Num();
 		for (TPair<FName, FGuid>& Pair : PreviousGeneratedPackages)
 		{
-			for (const ITargetPlatform* TargetPlatform : RequestedPlatforms)
+			for (const ITargetPlatform* TargetPlatform : PlatformsToCook)
 			{
 				COTFS.DeleteOutputForPackage(Pair.Key, TargetPlatform);
 			}
@@ -1584,7 +1656,7 @@ bool FGeneratorPackage::TryGenerateList(UObject* OwnerObject, FPackageDatas& Pac
 			IdenticalGenerated.Num(), NumIterativeModified, NumIterativeRemoved);
 		COOK_STAT(DetailedCookStats::NumPackagesIterativelySkipped += IdenticalGenerated.Num());
 
-		for (const ITargetPlatform* TargetPlatform : RequestedPlatforms)
+		for (const ITargetPlatform* TargetPlatform : PlatformsToCook)
 		{
 			COTFS.FindOrCreatePackageWriter(TargetPlatform).MarkPackagesUpToDate(IdenticalGenerated);
 		}
@@ -2168,7 +2240,7 @@ void FCookGenerationInfo::IterativeCookValidateOrClear(FGeneratorPackage& Genera
 			}
 			for (const ITargetPlatform* TargetPlatform : RequestedPlatforms)
 			{
-				FPackageData::FPlatformData* DependencyPlatformData = DependencyData->FindPlatformData(TargetPlatform);
+				FPackagePlatformData* DependencyPlatformData = DependencyData->FindPlatformData(TargetPlatform);
 				if (!DependencyPlatformData || !DependencyPlatformData->IsIterativelySkipped())
 				{
 					bOutIdentical = false;
@@ -2879,7 +2951,6 @@ FPackageData* FPackageDatas::UpdateFileName(FName PackageName)
 
 FThreadsafePackageData::FThreadsafePackageData()
 	: bInitialized(false)
-	, bCookable(true)
 	, bHasLoggedDiscoveryWarning(false)
 	, bHasLoggedDependencyWarning(false)
 {
@@ -2893,7 +2964,6 @@ void FPackageDatas::UpdateThreadsafePackageData(const FPackageData& PackageData)
 			ThreadsafeData.Instigator = PackageData.GetInstigator();
 			FGeneratorPackage* Generator = PackageData.GetGeneratedOwner();
 			ThreadsafeData.Generator = Generator ? Generator->GetOwner().GetPackageName() : NAME_None;
-			ThreadsafeData.bCookable = PackageData.CanCookForPlatforms();
 		});
 }
 
@@ -2954,7 +3024,8 @@ void FPackageDatas::ClearCookedPlatforms()
 {
 	LockAndEnumeratePackageDatas([](FPackageData* PackageData)
 	{
-		PackageData->ClearCookProgress();
+		PackageData->ResetReachable();
+		PackageData->ClearCookResults();
 	});
 }
 
@@ -3134,7 +3205,7 @@ bool FRequestQueue::IsEmpty() const
 
 uint32 FRequestQueue::Num() const
 {
-	uint32 Count = UnclusteredRequests.Num() + ReadyRequestsNum();
+	uint32 Count = RestartedRequests.Num() + ReadyRequestsNum();
 	for (const FRequestCluster& RequestCluster : RequestClusters)
 	{
 		Count += RequestCluster.NumPackageDatas();
@@ -3145,7 +3216,8 @@ uint32 FRequestQueue::Num() const
 bool FRequestQueue::Contains(const FPackageData* InPackageData) const
 {
 	FPackageData* PackageData = const_cast<FPackageData*>(InPackageData);
-	if (UnclusteredRequests.Contains(PackageData) || NormalRequests.Contains(PackageData) || UrgentRequests.Contains(PackageData))
+	if (RestartedRequests.Contains(PackageData) || NormalRequests.Contains(PackageData) ||
+		UrgentRequests.Contains(PackageData))
 	{
 		return true;
 	}
@@ -3159,21 +3231,35 @@ bool FRequestQueue::Contains(const FPackageData* InPackageData) const
 	return false;
 }
 
-uint32 FRequestQueue::RemoveRequest(FPackageData* PackageData)
+bool FRequestQueue::DiscoveryQueueContains(FPackageData* PackageData) const
+{
+	return Algo::FindBy(DiscoveryQueue, PackageData,
+		[](const FDiscoveryQueueElement& Element) { return Element.PackageData; }) != nullptr;
+}
+
+uint32 FRequestQueue::RemoveRequestExceptFromCluster(FPackageData* PackageData, FRequestCluster* ExceptFromCluster)
 {
 	uint32 OriginalNum = Num();
-	UnclusteredRequests.Remove(PackageData);
+	RestartedRequests.Remove(PackageData);
 	NormalRequests.Remove(PackageData);
 	UrgentRequests.Remove(PackageData);
 	for (FRequestCluster& RequestCluster : RequestClusters)
 	{
-		RequestCluster.RemovePackageData(PackageData);
+		if (&RequestCluster != ExceptFromCluster)
+		{
+			RequestCluster.RemovePackageData(PackageData);
+		}
 	}
 	uint32 Result = OriginalNum - Num();
 	check(Result == 0 || Result == 1);
 	return Result;
 }
 
+uint32 FRequestQueue::RemoveRequest(FPackageData* PackageData)
+{
+	return RemoveRequestExceptFromCluster(PackageData, nullptr);
+}
+
 uint32 FRequestQueue::Remove(FPackageData* PackageData)
 {
 	return RemoveRequest(PackageData);
@@ -3184,6 +3270,11 @@ bool FRequestQueue::IsReadyRequestsEmpty() const
 	return ReadyRequestsNum() == 0;
 }
 
+bool FRequestQueue::HasRequestsToExplore() const
+{
+	return !RequestClusters.IsEmpty() | !RestartedRequests.IsEmpty() | !DiscoveryQueue.IsEmpty();
+}
+
 uint32 FRequestQueue::ReadyRequestsNum() const
 {
 	return UrgentRequests.Num() + NormalRequests.Num();
@@ -3208,14 +3299,7 @@ FPackageData* FRequestQueue::PopReadyRequest()
 
 void FRequestQueue::AddRequest(FPackageData* PackageData, bool bForceUrgent)
 {
-	if (!PackageData->AreAllRequestedPlatformsExplored())
-	{
-		UnclusteredRequests.Add(PackageData);
-	}
-	else
-	{
-		AddReadyRequest(PackageData, bForceUrgent);
-	}
+	RestartedRequests.Add(PackageData);
 }
 
 void FRequestQueue::AddReadyRequest(FPackageData* PackageData, bool bForceUrgent)
@@ -3287,4 +3371,4 @@ FPoppedPackageDataScope::~FPoppedPackageDataScope()
 }
 #endif
 
-}
+} // namespace UE::Cook
diff --git a/Engine/Source/Editor/UnrealEd/Private/Cooker/CookPackageData.h b/Engine/Source/Editor/UnrealEd/Private/Cooker/CookPackageData.h
--- a/Engine/Source/Editor/UnrealEd/Private/Cooker/CookPackageData.h
+++ b/Engine/Source/Editor/UnrealEd/Private/Cooker/CookPackageData.h
@@ -36,7 +36,10 @@ class FCbWriter;
 class UCookOnTheFlyServer;
 class UObject;
 class UPackage;
+namespace UE::Cook { class FCookWorkerClient; }
+namespace UE::Cook { class FRequestCluster; }
 namespace UE::Cook { struct FConstructPackageData; }
+namespace UE::Cook { struct FDiscoveredPlatformSet; }
 
 FCbWriter& operator<<(FCbWriter& Writer, const UE::Cook::FConstructPackageData& PackageData);
 bool LoadFromCompactBinary(FCbFieldView Field, UE::Cook::FConstructPackageData& PackageData);
@@ -88,24 +91,83 @@ struct FConstructPackageData
 	FName NormalizedFileName;
 };
 
+// This should be a constexpr, but some compilers give an error that 0x1 is an unevaluable pointer value
+#define CookerLoadingPlatformKey ((ITargetPlatform*)0x1)
+
+/** Data about a platform that has been interacted with (marked reachable, etc) by a PackageData. */
+struct FPackagePlatformData
+{
+	FPackagePlatformData();
+
+	/** Requested from StartCookByTheBook or CookOnTheFly, or is a transitive dependency of those requests. */
+	bool IsReachable() const { return bReachable != 0; }
+	void SetReachable(bool bValue) { bReachable = (uint32)bValue; }
+
+	/** The package has passed through a RequestCluster and its transitive dependencies were added. */
+	bool IsVisitedByCluster() const { return bVisitedByCluster != 0; }
+	void SetVisitedByCluster(bool bValue) { bVisitedByCluster = (uint32)bValue; }
+
+	/** UPackage::Save was called on the package, but timedout and needs to be retried. */
+	bool IsSaveTimedOut() const { return bSaveTimedOut != 0; }
+	void SetSaveTimedOut(bool bValue) { bSaveTimedOut = (uint32)bValue; }
+
+	/** Whether the package is allowed to cook. Initially true, may be set false during Cluster search. */
+	bool IsCookable() const { return bCookable != 0; }
+	void SetCookable(bool bValue) { bCookable = (uint32)bValue; }
+
+	/**
+	 * Whether the package is searchable for transitive dependencies during Cluster evaluation.
+	 * This can be true even if the package is not cookable. Initially true, may be set false during Cluster search */
+	bool IsExplorable() const { return bExplorable != 0; }
+	void SetExplorable(bool bValue) { bExplorable = (uint32)bValue; }
+
+	/** All flags modified by reachability calculations are returned to default. */
+	void ResetReachable();
+	/** Called on CookWorkers to indicate reachable,cookable,etc for packages sent from Director. */
+	void MarkCookableForWorker(FCookWorkerClient& CookWorkerClient);
+
+	/** The package was found to be unmodified in the current iterative cook. */
+	bool IsIterativelySkipped() const { return bIterativelySkipped != 0; }
+	void SetIterativelySkipped(bool bValue) { bIterativelySkipped = (uint32)bValue; }
+
+	ECookResult GetCookResults() const { return (ECookResult)CookResults; }
+	bool IsCookAttempted() const { return CookResults != (uint32)ECookResult::NotAttempted; }
+	bool IsCookSucceeded() const { return CookResults == (uint32)ECookResult::Succeeded; }
+	void SetCookResults(ECookResult Value) { CookResults = (uint32)Value; }
+
+	/** Return if we need to call SavePackage on it (reachable, cookable, not yet cooked) */
+	bool NeedsCooking(const ITargetPlatform* PlatformItBelongsTo) const;
+
+	bool IsRegisteredForCachedObjectsInOuter() const { return bRegisteredForCachedObjectsInOuter != 0; }
+	void SetRegisteredForCachedObjectsInOuter(bool bValue) { bRegisteredForCachedObjectsInOuter = bValue; }
+
+private:
+	uint32 bReachable : 1;
+	uint32 bVisitedByCluster : 1;
+	uint32 bSaveTimedOut : 1;
+	uint32 bCookable : 1;
+	uint32 bExplorable : 1;
+	uint32 bIterativelySkipped : 1;
+	uint32 bRegisteredForCachedObjectsInOuter : 1;
+	uint32 CookResults : (int)ECookResult::NumBits;
+};
+
 /**
  * Contains all the information the cooker uses for a package, during request, load, or save.  Once allocated, this
  * structure is never deallocated or moved for a given package; it is deallocated only when the CookOnTheFlyServer
  * is destroyed.
  *
- * Requested platforms - Requested platforms are the Platforms that will be Saved during cooking.
- *   Platforms are only requested on the PackageData when the Package is InProgress (e.g. is in the RequestQueue,
- *   LoadReadyQueue, or other state containers). Once a Package finishes cooking, the platforms are set to
- *   unrequested, and may be set again later if the Package is requested for another platform. If multiple requests
- *   occur at the same time, their requested platforms are merged. Since modifying the requested platforms can put
- *   an InProgress package into an invalid state, direct write access is private; use UpdateRequestData or
- *   SetRequestData to write.
-
- * Cooked platforms - Cooked platforms are the platforms that have already been saved in the lifetime of the 
- *   CookOnTheFlyServer; note this extends outside of the current CookOnTheFlyServer session.
- *   Cooked platforms also store a flag indicating whether the cook for the given platform was successful or not.
- *   Cooked platforms can be added to a PackageData for reasons other than normal Save, such as when a Package is
- *   detected as not applicable for a Platform and its "Cooked" operation is therefore a noop.
+ * PlatformDatas - Per-platform information about the package's cook state. Whether it has been marked reachable,
+ *   whether it has been cooked and what the result was, etc.
+ *   Direct write access to the platformdata is pseudo-private: it's available for convenience of the RequestCluster,
+*    but when modifiying it the package may need to change state, so calling code should either update the state itself
+*    or should call the functions on FPackageData that take a TargetPlatform and handle modifying the state.
+
+ * Cooked platforms - Platforms are marked cooked if they have been saved in the lifetime of the 
+ *   CookOnTheFlyServer; note this extends outside of the current CookOnTheFlyServer session for in-editor cooks.
+ *   Cooked platforms also store the CookResults - success,error,other.
+ *   Other: cooked platforms can be added to a PackageData for reasons other than normal Save, such as when a Package
+ *   is marked not cookable for a Platform and its "Cooked" operation is therefore a noop.
  *   Cooked platforms can be cleared for a PackageData if the Package is modified, e.g. during editor operations
  *   when the CookOnTheFlyServer is run from the editor.
  *
@@ -127,58 +189,19 @@ struct FConstructPackageData
  *   The lifetime of this counter extends for the lifetime of the PackageData; it is shared across
  *   CookOnTheFlyServer sessions.
  *
- * CookedPlatformDataNextIndex - Index for the next (Object, RequestedPlatform) pair that needs to have
+ * CookedPlatformDataNextIndex - Index for the next (Object, ReachablePlatform) pair that needs to have
  *   BeginCacheForCookedPlatformData called on it for the current PackageSave. This field is only non-zero during
  *   the save state; it is cleared when successfully or unsucessfully leaving the save state. The field is an
  *   object-major index into the two-dimensional array given by
- *       {Each Object in GetCachedObjectsInOuter} x {each Platform that is requested}
- *   e.g. the index into GetCachedObjectsInOuter is GetCookedPlatformDataNextIndex() / GetNumRequestedPlatforms()
- *   and the index into GetRequestedPlatforms() is GetCookedPlatformDataNextIndex() % GetNumRequestedPlatforms()
+ *       {Each Object in GetCachedObjectsInOuter} x {each Platform that is reachable}
+ *   e.g. the index into GetCachedObjectsInOuter is GetCookedPlatformDataNextIndex() / GetCachedObjectsInOuterNumPlatforms()
+ *   and the index into GetCachedObjectsInOuterPlatforms() is GetCookedPlatformDataNextIndex() % GetCachedObjectsInOuterNumPlatforms()
  *
  * Other fields with explanation inline
 */
 struct FPackageData
 {
 public:
-	/** Data about a platform that has been interacted with (requested, etc) by the cook of a PackageData. */
-	struct FPlatformData
-	{
-		FPlatformData();
-
-		/** Requested from StartCookByTheBook, CookOnTheFly, or as a transitive dependency. */
-		bool IsRequested() const { return bRequested != 0; }
-		void SetRequested(bool bValue) { bRequested = (uint32) bValue; }
-
-		/** The package has passed through a CookRequestCluster and its transitive dependencies were added. */
-		bool IsExplored() const { return bExplored != 0; }
-		void SetExplored(bool bValue) { bExplored = (uint32) bValue; }
-
-		/** UPackage::Save was called on the package, but timedout and needs to be retried. */
-		bool IsSaveTimedOut() const { return bSaveTimedOut != 0; }
-		void SetSaveTimedOut(bool bValue) { bSaveTimedOut = (uint32) bValue; }
-
-		/** True if the package has either not passed through a CookRequestCluster, or it has and it was found cookable. */
-		bool IsCookable() const { return bCookable != 0; }
-		void SetCookable(bool bValue) { bCookable = (uint32) bValue; }
-
-		/** The package was found to be unmodified in the current iterative cook. */
-		bool IsIterativelySkipped() const { return bIterativelySkipped != 0; }
-		void SetIterativelySkipped(bool bValue) { bIterativelySkipped = (uint32) bValue; }
-
-		ECookResult GetCookResults() const { return (ECookResult) CookResults; }
-		bool IsCookAttempted() const { return CookResults != (uint32) ECookResult::NotAttempted; }
-		bool IsCookSucceeded() const { return CookResults == (uint32) ECookResult::Succeeded; }
-		void SetCookResults(ECookResult Value) { CookResults = (uint32) Value; }
-
-	private:
-		uint32 bRequested : 1;
-		uint32 bExplored : 1;
-		uint32 bSaveTimedOut : 1;
-		uint32 bCookable : 1;
-		uint32 bIterativelySkipped : 1;
-		uint32 CookResults : (int)ECookResult::NumBits;
-	};
-
 	FPackageData(FPackageDatas& PackageDatas, const FName& InPackageName, const FName& InFileName);
 	~FPackageData();
 	static void* operator new(size_t Size)
@@ -217,71 +240,47 @@ public:
 	 */
 	const FName& GetFileName() const;
 
-	/** Reset OutPlatforms and copy current set of requested platforms into it. */
+	/** Reset OutPlatforms and copy current set of reachable & cookable & not yet cooked platforms into it. */
 	template <typename ArrayType>
-	void GetRequestedPlatforms(ArrayType& OutPlatforms) const
-	{
-		OutPlatforms.Reset(PlatformDatas.Num());
-		for (const TPair<const ITargetPlatform*, FPlatformData>& Pair : PlatformDatas)
-		{
-			if (Pair.Value.IsRequested())
-			{
-				OutPlatforms.Add(Pair.Key);
-			}
-		}
-	}
+	void GetPlatformsNeedingCooking(ArrayType& OutPlatforms) const;
+	/** Number of platforms that would be returned by GetPlatformsNeedingCooking. */
+	int32 GetPlatformsNeedingCookingNum() const;
 
-	/** Get the number of currently requested platforms. */
-	int32 GetNumRequestedPlatforms() const;
+	/** Reset OutPlatforms and copy current set of reachable platforms into it. */
+	template <typename ArrayType>
+	void GetReachablePlatforms(ArrayType& OutPlatforms) const;
 
-	/**
-	 * Return true if and only if every element of Platforms is currently requested.
-	 * Returns true if Platforms is empty.
-	 */
-	bool HasAllRequestedPlatforms(const TArrayView<const ITargetPlatform* const>& Platforms) const;
+	/** Return true if and only if the Platform has been visited by a cluster (and hence is reachable and explored). */
+	bool IsPlatformVisitedByCluster(const ITargetPlatform* Platform) const;
 
 	/**
-	 * Return whether all Platforms requested have also been cooked.
-	 * 
-	 * @param bAllowFailedCooks Iff false, failed cooks cause function to return false.
+	 * Return true if and only if every element of Platforms is currently reachable.
+	 * Returns true if Platforms is empty.
 	 */
-	bool AreAllRequestedPlatformsCooked(bool bAllowFailedCooks) const;
-
-	/** Return whether all Platforms requested have also been explored for transitive references. */
-	bool AreAllRequestedPlatformsExplored() const;
-
-	/** Return true if and only if every element of Platforms has been explored. Returns true if Platforms is empty. */
-	bool HasAllExploredPlatforms(const TArrayView<const ITargetPlatform* const>& Platforms) const;
+	bool HasReachablePlatforms(const TArrayView<const ITargetPlatform* const>& Platforms) const;
 
-	/**
-	 * Return true if this package is cookable for at least one of the platforms.
-	 * If this package has not been explored will return true
-	 */
-	bool CanCookForPlatforms() const;
+	/** Return whether all Platforms that have been marked reachable have also been explored in an FRequestCluster. */
+	bool AreAllReachablePlatformsVisitedByCluster() const;
 
 	/**
 	 * Get the flag for whether this InProgress PackageData has been marked as an urgent request
 	 * (e.g. because it has been requested from the game client during cookonthefly.).
 	 * Always false for PackageDatas that are not InProgress.
-	 * Since modifying the urgency can put an InProgress package into an invalid state, direct write access is
-	 * private; use UpdateRequestData or SetRequestData to write.
 	 */
 	bool GetIsUrgent() const { return static_cast<bool>(bIsUrgent); }
 	/**
-	 * Add the given request data onto the existing request data.
-	 * If the PackageData is not in progress, just calls SetRequestData; skipped if InRequestedPlatforms is empty.
-	 * If the PackageData is in progress it is demoted back to an earlier state if necessary.
-	 * RequestData is set to the union of the previous and new values.
+	 * Mark this PackageData as urgent if bUrgent is true. Noop if bUrgent is false.
+	 * Move it to front of its current container if bAllowUpdateState.
 	 */
-	void UpdateRequestData(const TConstArrayView<const ITargetPlatform*> InRequestedPlatforms, bool bInIsUrgent,
-		FCompletionCallback&& InCompletionCallback, FInstigator&& Instigator, bool bAllowUpdateUrgency=true);
-	/*
-	 * Set the given data onto the existing request data.
-	 * It is invalid to call with empty InRequestedPlatforms.
-	 * It is invalid to call on a PackageData that is already InProgress. 
-	 */
-	void SetRequestData(const TArrayView<const ITargetPlatform* const>& InRequestedPlatforms, bool bInIsUrgent,
-		FCompletionCallback&& InCompletionCallback, FInstigator&& Instigator);
+	void AddUrgency(bool bUrgent, bool bAllowUpdateState);
+
+	/** Accessor for RequestClusters to add reachable platforms directly without modifying dependent data. */
+	void AddReachablePlatforms(FRequestCluster& RequestCluster, TConstArrayView<const ITargetPlatform*> Platforms,
+		FInstigator&& InInstigator);
+
+	/** Add the given reachable platforms to this PackageData and send it back to Request state for exploration. */
+	void QueueAsDiscovered(FInstigator&& InInstigator, FDiscoveredPlatformSet&& ReachablePlatforms);
+
 	/**
 	 * Clear all the inprogress variables from the current PackageData. It is invalid to call this except when
 	 * the PackageData is transitioning out of InProgress.
@@ -300,21 +299,24 @@ public:
 	 * FindOrAdd each TargetPlatform and set its flags: CookAttempted=false.
 	 * In Version that takes no TargetPlatform, CookAttempted is cleared from all existing platforms.
 	 */
-	void ClearCookProgress(const TConstArrayView<const ITargetPlatform*> TargetPlatforms);
-	void ClearCookProgress();
-	void ClearCookProgress(const ITargetPlatform* TargetPlatform);
+	void ClearCookResults(const TConstArrayView<const ITargetPlatform*> TargetPlatforms);
+	void ClearCookResults();
+	void ClearCookResults(const ITargetPlatform* TargetPlatform);
+	/** Clear reachable and related fields from all platforms. */
+	void ResetReachable();
 
 	/** Access the information about platforms interacted with by *this. */
-	const TSortedMap<const ITargetPlatform*, FPlatformData, TInlineAllocator<1>>& GetPlatformDatas() const;
+	const TSortedMap<const ITargetPlatform*, FPackagePlatformData, TInlineAllocator<1>>& GetPlatformDatas() const;
+	TSortedMap<const ITargetPlatform*, FPackagePlatformData, TInlineAllocator<1>>& GetPlatformDatasConstKeysMutableValues();
 
 	/** Add a platform if not already existing and return a writable pointer to its flags. */
-	FPlatformData& FindOrAddPlatformData(const ITargetPlatform* TargetPlatform);
+	FPackagePlatformData& FindOrAddPlatformData(const ITargetPlatform* TargetPlatform);
 
 	/** Find a platform if it exists and return a writable pointer to its flags. */
-	FPlatformData* FindPlatformData(const ITargetPlatform* TargetPlatform);
+	FPackagePlatformData* FindPlatformData(const ITargetPlatform* TargetPlatform);
 
 	/** Find a platform if it exists and return a writable pointer to its flags. */
-	const FPlatformData* FindPlatformData(const ITargetPlatform* TargetPlatform) const;
+	const FPackagePlatformData* FindPlatformData(const ITargetPlatform* TargetPlatform) const;
 
 	/** Return true if and only if at least one platform has been cooked. */
 	bool HasAnyCookedPlatform() const;
@@ -339,20 +341,6 @@ public:
 	 */
 	ECookResult GetCookResults(const ITargetPlatform* Platform) const;
 
-	/** Empties and then sets OutPlatforms to contain all platforms that have been requested and not cook-attempted. */
-	template <typename ArrayType>
-	void GetUncookedPlatforms(ArrayType& OutPlatforms)
-	{
-		OutPlatforms.Reset(PlatformDatas.Num());
-		for (const TPair<const ITargetPlatform*, FPlatformData>& Pair : PlatformDatas)
-		{
-			if (Pair.Value.IsRequested() && !Pair.Value.IsCookAttempted())
-			{
-				OutPlatforms.Add(Pair.Key);
-			}
-		}
-	}
-
 	/**
 	 * Return the package pointer. By contract it will be non-null if and only if the PackageData's state is
 	 * >= EPackageState::Load.
@@ -396,7 +384,8 @@ public:
 	 * Add the given callback into this PackageData's callback field. It is invalid to call this function with a
 	 * non-empty callback if this PackageData already has a CompletionCallback.
 	 */
-	void AddCompletionCallback(FCompletionCallback&& InCompletionCallback);
+	void AddCompletionCallback(TConstArrayView<const ITargetPlatform*> TargetPlatforms,
+		FCompletionCallback&& InCompletionCallback);
 
 	/**
 	 * Get/Set a visited flag used when searching graphs of PackageData. User of the graph is responsible for 
@@ -423,6 +412,12 @@ public:
 	 * TryCreateObjectCache and is cleared when leaving the save state.
 	 */
 	TArray<FWeakObjectPtr>& GetCachedObjectsInOuter();
+	template <typename ArrayType>
+	/** The list of platforms that were recorded as needscooking when CachedObjeObjectsInOuter was recorded. */
+	void GetCachedObjectsInOuterPlatforms(ArrayType& OutPlatforms) const;
+	/** The number of platforms that would be returned by GetCachedObjectsInOuterPlatforms. */
+	int32 GetCachedObjectsInOuterNumPlatforms() const;
+
 	/** Validate that the CachedObjectsInOuter-dependent variables are empty, when entering save. */
 	void CheckObjectCacheEmpty() const;
 	/** Populate CachedObjectsInOuter if not already populated. Invalid to call except when in the save state. */
@@ -489,7 +484,7 @@ public:
 	ECookResult GetMonitorCookResult() const { return (ECookResult)MonitorCookResult; }
 	void SetMonitorCookResult(ECookResult Value) { MonitorCookResult = (uint8) Value; }
 
-	/** Remove all request data about the given platform from all fields in this PackageData. */
+	/** Remove all data about the given platform from all fields in this PackageData. */
 	void OnRemoveSessionPlatform(const ITargetPlatform* Platform);
 
 	/** Report whether this PackageData holds references to UObjects and would be affected by GarbageCollection. */
@@ -527,6 +522,10 @@ public:
 	 * referencing package that causes the package to enter the requested state.
 	 */
 	const FInstigator& GetInstigator() const { return Instigator; }
+	bool HasInstigator() const { return Instigator.Category != EInstigator::NotYetRequested; }
+	/** Setting the instigator is mostly private - it should only be done during clustering. */
+	void SetInstigator(FRequestCluster& Cluster, FInstigator&& InInstigator);
+	void SetInstigator(FCookWorkerClient& Cluster, FInstigator&& InInstigator);
 
 	/** Get whether COTFS is keeping this package referenced referenced during GC. */
 	bool IsKeepReferencedDuringGC() const { return static_cast<bool>(bKeepReferencedDuringGC); }
@@ -547,14 +546,22 @@ public:
 
 	/** Marshall this PackageData to a ConstructData that can be used later or on a remote machine to reconstruct it. */
 	FConstructPackageData CreateConstructData();
+
+	/**
+	 * Return the platforms for which the given Package has been marked reachable.
+	 * If the package does not exist, return the COTFS's list of Session platforms
+	 */
+	static void GetReachablePlatformsForInstigator(UCookOnTheFlyServer& COTFS, FPackageData* InInstigator,
+		TArray<const ITargetPlatform*>& Platforms, bool bAllowPartialInstigatorResults=false);
+	static void GetReachablePlatformsForInstigator(UCookOnTheFlyServer& COTFS, FName InInstigator,
+		TArray<const ITargetPlatform*>& Platforms, bool bAllowPartialInstigatorResults=false);
 private:
 	friend struct UE::Cook::FPackageDatas;
 
-	/** FindOrAdd each TargetPlatform and set its bRequested value. */
-	void SetPlatformsRequested(const TConstArrayView<const ITargetPlatform*> TargetPlatforms, bool bRequested);
-	/** Set all platforms to not requested. */
-	void ClearRequestedPlatforms();
 	void SetIsUrgent(bool Value);
+	void SetInstigatorInternal(FInstigator&& InInstigator);
+	static void AddReachablePlatformsInternal(FPackageData& PackageData, TConstArrayView<const ITargetPlatform*> Platforms, FInstigator&& InInstigator);
+	static void QueueAsDiscoveredInternal(FPackageData& PackageData, FInstigator&& InInstigator, FDiscoveredPlatformSet&& ReachablePlatforms);
 
 	/**
 	 * Set the FileName of the file that contains the package. This member is private because FPackageDatas
@@ -614,12 +621,12 @@ private:
 	void OnExitHasPackage();
 	void OnEnterHasPackage();
 
-	void OnPackageDataFirstRequested(FInstigator&& InInstigator);
+	void OnPackageDataFirstMarkedReachable(FInstigator&& InInstigator);
 
 	TUniquePtr<FGeneratorPackage> GeneratorPackage;
 	FGeneratorPackage* GeneratedOwner;
 	/** Data for each platform that has been interacted with by *this. */
-	TSortedMap<const ITargetPlatform*, FPlatformData, TInlineAllocator<1>> PlatformDatas;
+	TSortedMap<const ITargetPlatform*, FPackagePlatformData, TInlineAllocator<1>> PlatformDatas;
 	TArray<FWeakObjectPtr> CachedObjectsInOuter;
 	FCompletionCallback CompletionCallback;
 	FName PackageName;
@@ -1013,6 +1020,12 @@ private:
 	int32 NumUrgentInState[static_cast<uint32>(EPackageState::Count)];
 };
 
+struct FDiscoveryQueueElement
+{
+	FPackageData* PackageData;
+	FInstigator Instigator;
+	FDiscoveredPlatformSet ReachablePlatforms;
+};
 
 /**
  * A container for FPackageDatas in the Request state. This container needs to support fast find and remove,
@@ -1026,22 +1039,32 @@ public:
 	uint32 Num() const;
 	uint32 Remove(FPackageData* PackageData);
 	bool Contains(const FPackageData* PackageData) const;
+	bool DiscoveryQueueContains(FPackageData* PackageData) const;
 	void Empty();
 
 	void AddRequest(FPackageData* PackageData, bool bForceUrgent=false);
 
+	bool HasRequestsToExplore() const;
 	uint32 ReadyRequestsNum() const;
 	bool IsReadyRequestsEmpty() const;
 	FPackageData* PopReadyRequest();
 	void AddReadyRequest(FPackageData* PackageData, bool bForceUrgent = false);
 	uint32 RemoveRequest(FPackageData* PackageData);
+	uint32 RemoveRequestExceptFromCluster(FPackageData* PackageData, FRequestCluster* ExceptFromCluster);
 
-	FPackageDataSet& GetUnclusteredRequests() { return UnclusteredRequests; }
+	FPackageDataSet& GetRestartedRequests() { return RestartedRequests; }
+	/**
+	 * Unlike other containers on PackageData, GetDiscoveryQueue is not an ownership container.
+	 * PackageDatas in the DiscoveryQueue are PackageDatas we need to look at - in the right order during PumpRequests -
+	 * they can be in any state and are owned by another container (or are in the idle state).
+	 */
+	TRingBuffer<FDiscoveryQueueElement>& GetDiscoveryQueue() { return DiscoveryQueue; }
 	TRingBuffer<FRequestCluster>& GetRequestClusters() { return RequestClusters; }
 	FPackageDataSet& GetReadyRequestsUrgent() { return UrgentRequests; }
 	FPackageDataSet& GetReadyRequestsNormal() { return NormalRequests; }
 private:
-	FPackageDataSet UnclusteredRequests;
+	FPackageDataSet RestartedRequests;
+	TRingBuffer<FDiscoveryQueueElement> DiscoveryQueue;
 	TRingBuffer<FRequestCluster> RequestClusters;
 	FPackageDataSet UrgentRequests;
 	FPackageDataSet NormalRequests;
@@ -1071,7 +1094,6 @@ struct FThreadsafePackageData
 	FInstigator Instigator;
 	FName Generator;
 	bool bInitialized : 1;
-	bool bCookable : 1;
 	bool bHasLoggedDiscoveryWarning : 1;
 	bool bHasLoggedDependencyWarning : 1;
 
@@ -1082,7 +1104,7 @@ typedef TArray<FPendingCookedPlatformData> FPendingCookedPlatformDataContainer;
 
 /*
  * Class that manages the list of all PackageDatas for a CookOnTheFlyServer. PackageDatas is an associative
- * array for extra data about a package (e.g. the requested platforms) that is needed by the CookOnTheFlyServer.
+ * array for extra data about a package (e.g. the cook results) that is needed by the CookOnTheFlyServer.
  * FPackageData are allocated once and never destroyed or moved until the CookOnTheFlyServer is destroyed.
  * Memory on the FPackageData is allocated and deallocated as necessary for its current state.
  * FPackageData are mapped by PackageName and by FileName.
@@ -1281,7 +1303,7 @@ public:
 	void Clear();
 	/** Set all platforms to not cooked in all PackageDatas. Used to e.g. invalidate previous cooks. */
 	void ClearCookedPlatforms();
-	/** Remove all request data about the given platform from all PackageDatas and other memory used by *this. */
+	/** Remove all data about the given platform from all PackageDatas and other memory used by *this. */
 	void OnRemoveSessionPlatform(const ITargetPlatform* TargetPlatform);
 
 	/** Enumerate PendingPlatformDatas: the list of pending calls to BeginCacheForCookedPlatformData. */
@@ -1402,4 +1424,43 @@ inline TOptional<FThreadsafePackageData> FPackageDatas::FindThreadsafePackageDat
 	return Value ? TOptional<FThreadsafePackageData>(*Value) : TOptional<FThreadsafePackageData>();
 }
 
+template <typename ArrayType>
+inline void FPackageData::GetPlatformsNeedingCooking(ArrayType& OutPlatforms) const
+{
+	OutPlatforms.Reset(PlatformDatas.Num());
+	for (const TPair<const ITargetPlatform*, FPackagePlatformData>& Pair : PlatformDatas)
+	{
+		if (Pair.Value.NeedsCooking(Pair.Key))
+		{
+			OutPlatforms.Add(Pair.Key);
+		}
+	}
+}
+
+template <typename ArrayType>
+inline void FPackageData::GetCachedObjectsInOuterPlatforms(ArrayType& OutPlatforms) const
+{
+	OutPlatforms.Reset(PlatformDatas.Num());
+	for (const TPair<const ITargetPlatform*, FPackagePlatformData>& Pair : PlatformDatas)
+	{
+		if (Pair.Value.IsRegisteredForCachedObjectsInOuter())
+		{
+			OutPlatforms.Add(Pair.Key);
+		}
+	}
+}
+
+template <typename ArrayType>
+inline void FPackageData::GetReachablePlatforms(ArrayType& OutPlatforms) const
+{
+	OutPlatforms.Reset(PlatformDatas.Num());
+	for (const TPair<const ITargetPlatform*, FPackagePlatformData>& Pair : PlatformDatas)
+	{
+		if (Pair.Value.IsReachable())
+		{
+			OutPlatforms.Add(Pair.Key);
+		}
+	}
+}
+
 } // namespace UE::Cook
\ No newline at end of file
diff --git a/Engine/Source/Editor/UnrealEd/Private/Cooker/CookPlatformManager.cpp b/Engine/Source/Editor/UnrealEd/Private/Cooker/CookPlatformManager.cpp
--- a/Engine/Source/Editor/UnrealEd/Private/Cooker/CookPlatformManager.cpp
+++ b/Engine/Source/Editor/UnrealEd/Private/Cooker/CookPlatformManager.cpp
@@ -157,6 +157,10 @@ void FPlatformManager::AddSessionPlatform(UCookOnTheFlyServer& COTFS, FPlatformI
 		return;
 	}
 	SessionPlatforms.Add(TargetPlatform);
+	if (COTFS.bSessionRunning)
+	{
+		COTFS.OnPlatformAddedToSession(TargetPlatform);
+	}
 	COTFS.bSessionRunning = true;
 	bHasSelectedSessionPlatforms = true;
 }
@@ -269,14 +273,17 @@ void FPlatformManager::PruneUnreferencedSessionPlatforms(UCookOnTheFlyServer& Co
 			if (PlatformData->LastReferenceTime > 0. && PlatformData->ReferenceCount == 0
 				&& PlatformData->LastReferenceTime < OldestKeepTime)
 			{
+				int32 RemovedIndex = SessionPlatforms.IndexOfByKey(TargetPlatform);
+				check(RemovedIndex != INDEX_NONE);
+
 				// Mark that the platform no longer needs to be inspected for pruning because we have removed it
 				// from CookOnTheFly's SessionPlatforms
 				PlatformData->LastReferenceTime = 0.;
 
 				// Remove the SessionPlatform
-				CookOnTheFlyServer.OnRemoveSessionPlatform(TargetPlatform);
+				CookOnTheFlyServer.OnRemoveSessionPlatform(TargetPlatform, RemovedIndex);
 
-				SessionPlatforms.Remove(TargetPlatform);
+				SessionPlatforms.RemoveAt(RemovedIndex);
 				if (SessionPlatforms.Num() == 0)
 				{
 					CookOnTheFlyServer.bSessionRunning = false;
diff --git a/Engine/Source/Editor/UnrealEd/Private/Cooker/CookProfiling.cpp b/Engine/Source/Editor/UnrealEd/Private/Cooker/CookProfiling.cpp
--- a/Engine/Source/Editor/UnrealEd/Private/Cooker/CookProfiling.cpp
+++ b/Engine/Source/Editor/UnrealEd/Private/Cooker/CookProfiling.cpp
@@ -246,12 +246,10 @@ namespace DetailedCookStats
 	int32 PeakSaveQueueSize = 0;
 	std::atomic<int32> NumDetectedLoads{ 0 };
 	int32 NumRequestedLoads = 0;
-	uint32 NumPreloadedDependencies = 0;
 	uint32 NumPackagesIterativelySkipped = 0;
 	FCookStatsManager::FAutoRegisterCallback RegisterCookOnTheFlyServerStats([](FCookStatsManager::AddStatFuncRef AddStat)
 		{
 			AddStat(TEXT("Package.Load"), FCookStatsManager::CreateKeyValueArray(
-				TEXT("NumPreloadedDependencies"), NumPreloadedDependencies,
 				TEXT("NumInlineLoads"), NumDetectedLoads - NumRequestedLoads));
 			AddStat(TEXT("Package.Save"), FCookStatsManager::CreateKeyValueArray(TEXT("NumPackagesIterativelySkipped"), NumPackagesIterativelySkipped));
 			AddStat(TEXT("CookOnTheFlyServer"), FCookStatsManager::CreateKeyValueArray(
diff --git a/Engine/Source/Editor/UnrealEd/Private/Cooker/CookProfiling.h b/Engine/Source/Editor/UnrealEd/Private/Cooker/CookProfiling.h
--- a/Engine/Source/Editor/UnrealEd/Private/Cooker/CookProfiling.h
+++ b/Engine/Source/Editor/UnrealEd/Private/Cooker/CookProfiling.h
@@ -156,7 +156,6 @@ extern bool IsFastCook;
 extern bool IsUnversioned;
 
 // Stats tracked through FAutoRegisterCallback
-extern uint32 NumPreloadedDependencies;
 extern std::atomic<int32> NumDetectedLoads;
 extern int32 NumRequestedLoads;
 extern uint32 NumPackagesIterativelySkipped;
diff --git a/Engine/Source/Editor/UnrealEd/Private/Cooker/CookRequestCluster.cpp b/Engine/Source/Editor/UnrealEd/Private/Cooker/CookRequestCluster.cpp
--- a/Engine/Source/Editor/UnrealEd/Private/Cooker/CookRequestCluster.cpp
+++ b/Engine/Source/Editor/UnrealEd/Private/Cooker/CookRequestCluster.cpp
@@ -2,6 +2,7 @@
 
 #include "Cooker/CookRequestCluster.h"
 
+#include "Algo/BinarySearch.h"
 #include "Algo/Sort.h"
 #include "Algo/TopologicalSort.h"
 #include "Algo/Unique.h"
@@ -17,179 +18,25 @@
 #include "DerivedDataCache.h"
 #include "DerivedDataRequestOwner.h"
 #include "EditorDomain/EditorDomainUtils.h"
+#include "Engine/Level.h"
 #include "HAL/Event.h"
 #include "HAL/PlatformProcess.h"
 #include "Interfaces/ITargetPlatform.h"
 #include "Misc/RedirectCollector.h"
 #include "Misc/StringBuilder.h"
+#include "String/Find.h"
 #include "UObject/CoreRedirects.h"
 #include "UObject/SavePackage.h"
 
 namespace UE::Cook
 {
 
-FRequestCluster::FRequestCluster(UCookOnTheFlyServer& InCOTFS, TArray<const ITargetPlatform*>&& InPlatforms)
-	: Platforms(MoveTemp(InPlatforms))
-	, COTFS(InCOTFS)
+FRequestCluster::FRequestCluster(UCookOnTheFlyServer& InCOTFS)
+	: COTFS(InCOTFS)
 	, PackageDatas(*InCOTFS.PackageDatas)
 	, AssetRegistry(*IAssetRegistry::Get())
 	, PackageTracker(*InCOTFS.PackageTracker)
 	, BuildDefinitions(*InCOTFS.BuildDefinitions)
-{
-}
-
-FRequestCluster::FRequestCluster(UCookOnTheFlyServer& InCOTFS, TConstArrayView<const ITargetPlatform*> InPlatforms)
-	: Platforms(InPlatforms)
-	, COTFS(InCOTFS)
-	, PackageDatas(*InCOTFS.PackageDatas)
-	, AssetRegistry(*IAssetRegistry::Get())
-	, PackageTracker(*InCOTFS.PackageTracker)
-	, BuildDefinitions(*InCOTFS.BuildDefinitions)
-{
-}
-
-void FRequestCluster::AddClusters(UCookOnTheFlyServer& InCOTFS, TArray<FFilePlatformRequest>&& InRequests,
-	bool bRequestsAreUrgent, TRingBuffer<FRequestCluster>& OutClusters)
-{
-	if (InRequests.Num() == 0)
-	{
-		return;
-	}
-
-	TArray<FRequestCluster, TInlineAllocator<1>> AddedClusters;
-	auto FindOrAddClusterForPlatforms =
-		[&AddedClusters, &InCOTFS](TArray<const ITargetPlatform*>&& InPlatforms)
-	{
-		for (FRequestCluster& Existing : AddedClusters)
-		{
-			if (Existing.GetPlatforms() == InPlatforms)
-			{
-				InPlatforms.Reset();
-				return &Existing;
-			}
-		}
-		return &AddedClusters.Emplace_GetRef(InCOTFS, MoveTemp(InPlatforms));
-	};
-
-	UE::Cook::FRequestCluster* MRUCluster = FindOrAddClusterForPlatforms(MoveTemp(InRequests[0].GetPlatforms()));
-	// The usual case is all platforms are the same, so reserve the first Cluster's size assuming it will get all requests
-	MRUCluster->InRequests.Reserve(InRequests.Num());
-	// Add the first Request to it. Since we've already taken away the Platforms from the first request, we have to handle it specially.
-	MRUCluster->InRequests.Add(FFileNameRequest(MoveTemp(InRequests[0]), bRequestsAreUrgent));
-
-	for (FFilePlatformRequest& Request : TArrayView<FFilePlatformRequest>(InRequests).Slice(1, InRequests.Num() - 1))
-	{
-		if (Request.GetPlatforms() != MRUCluster->GetPlatforms())
-		{
-			// MRUCluster points to data inside AddedClusters, so we have to recalculate MRUCluster whenever we add
-			MRUCluster = FindOrAddClusterForPlatforms(MoveTemp(Request.GetPlatforms()));
-		}
-		else
-		{
-			Request.GetPlatforms().Reset();
-		}
-
-		MRUCluster->InRequests.Add(FFileNameRequest(MoveTemp(Request), bRequestsAreUrgent));
-	}
-
-	for (UE::Cook::FRequestCluster& AddedCluster : AddedClusters)
-	{
-		AddedCluster.Initialize(InCOTFS);
-		OutClusters.Add(MoveTemp(AddedCluster));
-	}
-}
-
-FName GInstigatorRequestCluster(TEXT("RequestCluster"));
-
-void FRequestCluster::AddClusters(UCookOnTheFlyServer& InCOTFS, FPackageDataSet& UnclusteredRequests,
-	TRingBuffer<FRequestCluster>& OutClusters, FRequestQueue& QueueForReadyRequests)
-{
-	if (UnclusteredRequests.Num() == 0)
-	{
-		return;
-	}
-
-	TArray<FRequestCluster, TInlineAllocator<1>> AddedClusters;
-	TArray<const ITargetPlatform*, TInlineAllocator<ExpectedMaxNumPlatforms>> RequestedPlatforms;
-	auto FindOrAddClusterForPlatforms = [&AddedClusters, &InCOTFS, &RequestedPlatforms, &UnclusteredRequests]()
-	{
-		if (AddedClusters.Num() == 0)
-		{
-			FRequestCluster& Cluster = AddedClusters.Emplace_GetRef(InCOTFS, RequestedPlatforms);
-			// The usual case is all platforms are the same, so reserve the first Cluster's size assuming it will get all requests
-			Cluster.Requests.Reserve(UnclusteredRequests.Num());
-			return &Cluster;
-		}
-		for (FRequestCluster& Existing : AddedClusters)
-		{
-			if (Existing.GetPlatforms() == RequestedPlatforms)
-			{
-				return &Existing;
-			}
-		}
-		return &AddedClusters.Emplace_GetRef(InCOTFS, RequestedPlatforms);
-	};
-
-	bool bErrorOnEngineContentUse = false;
-	bool bAllowUncookedAssetReferences = false;
-	UE::Cook::FCookByTheBookOptions& Options = *InCOTFS.CookByTheBookOptions;
-	FString DLCPath;
-	bErrorOnEngineContentUse = Options.bErrorOnEngineContentUse;
-	bAllowUncookedAssetReferences = Options.bAllowUncookedAssetReferences;
-	if (bErrorOnEngineContentUse)
-	{
-		DLCPath = FPaths::Combine(*InCOTFS.GetBaseDirectoryForDLC(), TEXT("Content"));
-		FPaths::MakeStandardFilename(DLCPath);
-	}
-
-	UE::Cook::FRequestCluster* MRUCluster = nullptr;
-	for (FPackageData* PackageData : UnclusteredRequests)
-	{
-		if (PackageData->AreAllRequestedPlatformsExplored())
-		{
-			QueueForReadyRequests.AddReadyRequest(PackageData);
-			continue;
-		}
-
-		// For non-cookable packages that we skipped in an earlier cluster but loaded because
-		// they are hard dependencies, avoid the work of creating a cluster just for them,
-		// by checking non-cookable and sending the package to idle instead of adding to
-		// a cluster
-		PackageData->GetRequestedPlatforms(RequestedPlatforms);
-		ESuppressCookReason SuppressCookReason;
-		if (!IsRequestCookable(PackageData->GetPackageName(), PackageData, *InCOTFS.PackageDatas, *InCOTFS.PackageTracker,
-			DLCPath, bErrorOnEngineContentUse, bAllowUncookedAssetReferences, RequestedPlatforms, SuppressCookReason))
-		{
-			InCOTFS.DemoteToIdle(*PackageData, ESendFlags::QueueAdd, SuppressCookReason);
-			continue;
-		}
-
-		if (!MRUCluster || RequestedPlatforms != MRUCluster->GetPlatforms())
-		{
-			// MRUCluster points to data inside AddedClusters, so we have to recalculate MRUCluster whenever we add
-			MRUCluster = FindOrAddClusterForPlatforms();
-		}
-		CA_ASSUME(MRUCluster);
-
-		MRUCluster->Requests.Add(PackageData);
-	}
-
-	for (UE::Cook::FRequestCluster& AddedCluster : AddedClusters)
-	{
-		AddedCluster.Initialize(InCOTFS);
-		OutClusters.Add(MoveTemp(AddedCluster));
-	}
-}
-
-FRequestCluster::FFileNameRequest::FFileNameRequest(FFilePlatformRequest&& FileRequest, bool bInUrgent)
-	: FileName(FileRequest.GetFilename())
-	, Instigator(MoveTemp(FileRequest.GetInstigator()))
-	, CompletionCallback(MoveTemp(FileRequest.GetCompletionCallback()))
-	, bUrgent(bInUrgent)
-{
-}
-
-void FRequestCluster::Initialize(UCookOnTheFlyServer& InCOTFS)
 {
 	if (!COTFS.IsCookOnTheFlyMode())
 	{
@@ -214,13 +61,10 @@ void FRequestCluster::Initialize(UCookOnTheFlyServer& InCOTFS)
 	}
 	GConfig->GetBool(TEXT("CookSettings"), TEXT("PreQueueBuildDefinitions"), bPreQueueBuildDefinitions, GEditorIni);
 
-
-	PackageWriters.Reserve(Platforms.Num());
 	bFullBuild = false;
 	bool bFirst = true;
-	for (const ITargetPlatform* TargetPlatform : Platforms)
+	for (const ITargetPlatform* TargetPlatform : COTFS.PlatformManager->GetSessionPlatforms())
 	{
-		PackageWriters.Add(&COTFS.FindOrCreatePackageWriter(TargetPlatform));
 		FPlatformData* PlatformData = COTFS.PlatformManager->GetPlatformData(TargetPlatform);
 		if (bFirst)
 		{
@@ -236,24 +80,100 @@ void FRequestCluster::Initialize(UCookOnTheFlyServer& InCOTFS)
 			}
 		}
 	}
+}
+
+FRequestCluster::FRequestCluster(UCookOnTheFlyServer& InCOTFS, TArray<FFilePlatformRequest>&& InRequests, bool bInExternalRequestsAreUrgent)
+	: FRequestCluster(InCOTFS)
+{
+	bExternalRequestsAreUrgent = bInExternalRequestsAreUrgent;
+	ReserveInitialRequests(InRequests.Num());
+	FilePlatformRequests = MoveTemp(InRequests);
+	InRequests.Empty();
+}
 
-	bool bRemoveNulls = false;
-	for (FPackageData*& PackageData : Requests)
+FRequestCluster::FRequestCluster(UCookOnTheFlyServer& InCOTFS, FPackageDataSet&& InRequests)
+	: FRequestCluster(InCOTFS)
+{
+	ReserveInitialRequests(InRequests.Num());
+	for (FPackageData* PackageData : InRequests)
 	{
-		bool bAlreadyExists = false;
-		OwnedPackageDatas.Add(PackageData, &bAlreadyExists);
-		if (bAlreadyExists)
-		{
-			bRemoveNulls = true;
-			PackageData = nullptr;
-		}
+		ESuppressCookReason& Existing = OwnedPackageDatas.FindOrAdd(PackageData, ESuppressCookReason::Invalid);
+		check(Existing == ESuppressCookReason::Invalid);
+		Existing = ESuppressCookReason::NotSuppressed;
 	}
-	if (bRemoveNulls)
+	InRequests.Empty();
+}
+
+FRequestCluster::FRequestCluster(UCookOnTheFlyServer& InCOTFS, TRingBuffer<FDiscoveryQueueElement>& DiscoveryQueue)
+	: FRequestCluster(InCOTFS)
+{
+	TArray<const ITargetPlatform*, TInlineAllocator<ExpectedMaxNumPlatforms>> BufferPlatforms;
+	while (!DiscoveryQueue.IsEmpty())
 	{
-		Requests.Remove(nullptr);
+		FDiscoveryQueueElement* Discovery = &DiscoveryQueue.First();
+		FPackageData& PackageData = *Discovery->PackageData;
+
+		TConstArrayView<const ITargetPlatform*> NewReachablePlatforms =
+			Discovery->ReachablePlatforms.GetPlatforms(COTFS, &Discovery->Instigator,
+				TConstArrayView<const ITargetPlatform*>(), &BufferPlatforms);
+		if (PackageData.HasReachablePlatforms(NewReachablePlatforms))
+		{
+			// If there are no new reachable platforms, add it to the cluster for cooking if it needs
+			// it, otherwise let it remain where it is
+			FDiscoveryQueueElement PoppedDiscovery = DiscoveryQueue.PopFrontValue();
+			Discovery = &PoppedDiscovery;
+			if (!PackageData.IsInProgress() && PackageData.GetPlatformsNeedingCookingNum() == 0)
+			{
+				PackageData.SendToState(EPackageState::Request, ESendFlags::QueueRemove);
+				OwnedPackageDatas.Add(&PackageData, ESuppressCookReason::NotSuppressed);
+			}
+			continue;
+		}
+
+		if (NewReachablePlatforms.Contains(CookerLoadingPlatformKey) &&
+			!PackageData.FindOrAddPlatformData(CookerLoadingPlatformKey).IsReachable())
+		{
+			// We are adding the expectation that this package will be loaded for something during the cook. Doing so
+			// is only expected from a few types of instigators, or from external package requests, or from cluster
+			// exploration. If not expected, add a diagnostic message, and add it as a hidden dependency.
+			if (Discovery->Instigator.Category != EInstigator::StartupPackage &&
+				Discovery->Instigator.Category != EInstigator::GeneratedPackage)
+			{
+				// If there are other discovered packages we have already added to this cluster, then defer this one
+				// until we have explored those; add this one to the next cluster. Exploring those earlier discoveries
+				// might uncover this one and make it expected.
+				if (!OwnedPackageDatas.IsEmpty())
+				{
+					break;
+				}
+
+				COTFS.OnDiscoveredPackageDebug(PackageData.GetPackageName(), Discovery->Instigator);
+				FPackageData* InstigatorPackageData = Discovery->Instigator.Referencer.IsNone() ? nullptr
+					: COTFS.PackageDatas->TryAddPackageDataByPackageName(Discovery->Instigator.Referencer);
+				if (InstigatorPackageData)
+				{
+					COTFS.DiscoveredDependencies.FindOrAdd(InstigatorPackageData->GetPackageName())
+						.Add(PackageData.GetPackageName());
+				}
+			}
+		}
+		// Add the new reachable platforms
+		PackageData.AddReachablePlatforms(*this, NewReachablePlatforms, MoveTemp(Discovery->Instigator));
+
+		// Pop it off the list; note that this invalidates the pointers we had into the DiscoveryQueueElement
+		FDiscoveryQueueElement PoppedDiscovery = DiscoveryQueue.PopFrontValue();
+		Discovery = &PoppedDiscovery;
+		NewReachablePlatforms = TConstArrayView<const ITargetPlatform*>();
+
+		// Send it to the Request state if it's not already there, remove it from its old container
+		// and add it to this cluster.
+		PackageData.SendToState(EPackageState::Request, ESendFlags::QueueRemove);
+		OwnedPackageDatas.Add(&PackageData, ESuppressCookReason::NotSuppressed);
 	}
 }
 
+FName GInstigatorRequestCluster(TEXT("RequestCluster"));
+
 void FRequestCluster::Process(const FCookerTimer& CookerTimer, bool& bOutComplete)
 {
 	bOutComplete = true;
@@ -263,7 +183,7 @@ void FRequestCluster::Process(const FCookerTimer& CookerTimer, bool& bOutComplet
 	{
 		return;
 	}
-	FetchDependencies(CookerTimer, bOutComplete);
+	PumpExploration(CookerTimer, bOutComplete);
 	if (!bOutComplete)
 	{
 		return;
@@ -279,108 +199,90 @@ void FRequestCluster::FetchPackageNames(const FCookerTimer& CookerTimer, bool& b
 	}
 
 	constexpr int32 TimerCheckPeriod = 100; // Do not incur the cost of checking the timer on every package
-	int32 InRequestsNum = InRequests.Num();
-	if (this->NextRequest == 0)
-	{
-		Requests.Reserve(Requests.Num() + InRequestsNum);
-	}
-	for (;NextRequest < InRequestsNum; ++NextRequest)
+	int32 NextRequest = 0;
+	for (; NextRequest < FilePlatformRequests.Num(); ++NextRequest)
 	{
 		if ((NextRequest+1) % TimerCheckPeriod == 0 && CookerTimer.IsTimeUp())
 		{
 			break;
 		}
 
-		FFileNameRequest& Request = InRequests[NextRequest];
-		FName OriginalName = Request.FileName;
-#if DEBUG_COOKONTHEFLY
-		UE_LOG(LogCook, Display, TEXT("Processing request for package %s"), *OriginalName.ToString());
-#endif
+		FFilePlatformRequest& Request = FilePlatformRequests[NextRequest];
+		FName OriginalName = Request.GetFilename();
+
 		// The input filenames are normalized, but might be missing their extension, so allow PackageDatas
 		// to correct the filename if the package is found with a different filename
 		bool bExactMatchRequired = false;
-		FPackageData* PackageData = PackageDatas.TryAddPackageDataByStandardFileName(OriginalName, bExactMatchRequired,
-			&Request.FileName);
+		FPackageData* PackageData = PackageDatas.TryAddPackageDataByStandardFileName(OriginalName, bExactMatchRequired);
 		if (!PackageData)
 		{
 			LogCookerMessage(FString::Printf(TEXT("Could not find package at file %s!"),
 				*OriginalName.ToString()), EMessageSeverity::Error);
 			UE_LOG(LogCook, Error, TEXT("Could not find package at file %s!"), *OriginalName.ToString());
-			UE::Cook::FCompletionCallback CompletionCallback(MoveTemp(Request.CompletionCallback));
+			FCompletionCallback CompletionCallback(MoveTemp(Request.GetCompletionCallback()));
 			if (CompletionCallback)
 			{
 				CompletionCallback(nullptr);
 			}
+			continue;
+		}
+
+		// If it has new reachable platforms we definitely need to explore it
+		if (!PackageData->HasReachablePlatforms(Request.GetPlatforms()))
+		{
+			PackageData->AddReachablePlatforms(*this, Request.GetPlatforms(), MoveTemp(Request.GetInstigator()));
+			PullIntoCluster(*PackageData);
+			PackageData->AddUrgency(bExternalRequestsAreUrgent, false /* bAllowUpdateState */);
 		}
 		else
 		{
-			if (TryTakeOwnership(*PackageData, Request.bUrgent, MoveTemp(Request.CompletionCallback), Request.Instigator))
+			if (PackageData->IsInProgress())
 			{
-				bool bAlreadyExists;
-				OwnedPackageDatas.Add(PackageData, &bAlreadyExists);
-				if (!bAlreadyExists)
-				{
-					Requests.Add(PackageData);
-				}
+				// If it's already in progress with no new platforms, we don't need to add it to the cluster, but add
+				// add on our urgency setting
+				PackageData->AddUrgency(bExternalRequestsAreUrgent, true /* bAllowUpdateState */);
+			}
+			else if (PackageData->GetPlatformsNeedingCookingNum() > 0)
+			{
+				// If it's missing cookable platforms and not in progress we need to add it to the cluster for cooking
+				PullIntoCluster(*PackageData);
+				PackageData->AddUrgency(bExternalRequestsAreUrgent, true /* bAllowUpdateState */);
 			}
 		}
+		// Add on our completion callback, or call it immediately if already done
+		PackageData->AddCompletionCallback(Request.GetPlatforms(), MoveTemp(Request.GetCompletionCallback()));
 	}
-	if (NextRequest < InRequestsNum)
+	if (NextRequest < FilePlatformRequests.Num())
 	{
+		FilePlatformRequests.RemoveAt(0, NextRequest);
 		bOutComplete = false;
 		return;
 	}
 
-	InRequests.Empty();
-	NextRequest = 0;
+	FilePlatformRequests.Empty();
 	bPackageNamesComplete = true;
 }
 
-bool FRequestCluster::TryTakeOwnership(FPackageData& PackageData, bool bUrgent, UE::Cook::FCompletionCallback && CompletionCallback,
-	const FInstigator& InInstigator)
+void FRequestCluster::ReserveInitialRequests(int32 RequestNum)
+{
+	OwnedPackageDatas.Reserve(FMath::Max(RequestNum, 1024));
+}
+
+void FRequestCluster::PullIntoCluster(FPackageData& PackageData)
 {
-	if (!PackageData.IsInProgress())
+	ESuppressCookReason& Existing = OwnedPackageDatas.FindOrAdd(&PackageData, ESuppressCookReason::Invalid);
+	if (Existing == ESuppressCookReason::Invalid)
 	{
-		check(!COTFS.IsCookWorkerMode()); // CookWorkers skip dependency traversal and should only call TryTakeOwnership on incoming requests, which are already in progress
-		check(GetPlatforms().Num() != 0); // This is required for SetRequestData
-		if (PackageData.HasAllExploredPlatforms(GetPlatforms()))
+		// Steal it from wherever it is and send it to Request State. It has already been added to this cluster
+		if (PackageData.GetState() == EPackageState::Request)
 		{
-			// Leave it in idle - it's already been processed by a cluster
-			if (CompletionCallback)
-			{
-				CompletionCallback(&PackageData);
-			}
-			return false;
+			COTFS.PackageDatas->GetRequestQueue().RemoveRequestExceptFromCluster(&PackageData, this);
 		}
 		else
 		{
-			PackageData.SetRequestData(GetPlatforms(), bUrgent, MoveTemp(CompletionCallback), FInstigator(InInstigator));
 			PackageData.SendToState(EPackageState::Request, ESendFlags::QueueRemove);
-			return true;
-		}
-	}
-	else
-	{
-		if (PackageData.HasAllExploredPlatforms(GetPlatforms()))
-		{
-			// Leave it where it is - it's already been processed by a cluster - but update with our request data
-			// This might demote it back to Request, and add it to Normal or Urgent request, but that will not
-			// impact this RequestCluster
-			PackageData.UpdateRequestData(GetPlatforms(), bUrgent, MoveTemp(CompletionCallback), FInstigator(InInstigator));
-			return false;
-		}
-		else
-		{
-			if (!OwnedPackageDatas.Contains(&PackageData))
-			{
-				// Steal it from wherever it is and add it to this cluster
-				// This might steal it from another RequestCluster or from the UnclusteredRequests if it's in request
-				// Doing that steal is wasteful but okay; one of the RequestClusters will win it and keep it
-				PackageData.SendToState(EPackageState::Request, ESendFlags::QueueRemove);
-				PackageData.UpdateRequestData(GetPlatforms(), bUrgent, MoveTemp(CompletionCallback), FInstigator(InInstigator), false /* bAllowUpdateUrgency */);
-			}
-			return true;
 		}
+		Existing = ESuppressCookReason::NotSuppressed;
 	}
 }
 
@@ -403,10 +305,13 @@ void FRequestCluster::StartAsync(const FCookerTimer& CookerTimer, bool& bOutComp
 		{
 			// If the EditorDomain is active, then batch-download all packages to cook from remote cache into local
 			TArray<FName> BatchDownload;
-			BatchDownload.Reserve(Requests.Num());
-			for (FPackageData* PackageData : Requests)
+			BatchDownload.Reserve(OwnedPackageDatas.Num());
+			for (TPair<FPackageData*, ESuppressCookReason>& Pair : OwnedPackageDatas)
 			{
-				BatchDownload.Add(PackageData->GetPackageName());
+				if (Pair.Value == ESuppressCookReason::NotSuppressed)
+				{
+					BatchDownload.Add(Pair.Key->GetPackageName());
+				}
 			};
 			EditorDomain->BatchDownload(BatchDownload);
 		}
@@ -429,16 +334,54 @@ void FRequestCluster::RemovePackageData(FPackageData* PackageData)
 
 	if (GraphSearch)
 	{
-		GraphSearch->GetTransitiveRequests().Remove(PackageData);
-		TMap<FPackageData*, TArray<FPackageData*>>& GraphEdges = GraphSearch->GetGraphEdges();
-		GraphEdges.Remove(PackageData);
-		for (TPair<FPackageData*, TArray<FPackageData*>>& Pair : GraphEdges)
+		GraphSearch->RemovePackageData(PackageData);
+	}
+}
+
+void FRequestCluster::OnNewReachablePlatforms(FPackageData* PackageData)
+{
+	if (GraphSearch)
+	{
+		GraphSearch->OnNewReachablePlatforms(PackageData);
+	}
+}
+
+void FRequestCluster::OnPlatformAddedToSession(const ITargetPlatform* TargetPlatform)
+{
+	if (GraphSearch)
+	{
+		FCookerTimer CookerTimer(FCookerTimer::Forever);
+		bool bComplete;
+		while (PumpExploration(CookerTimer, bComplete), !bComplete)
+		{
+			UE_LOG(LogCook, Display, TEXT("Waiting for RequestCluster to finish before adding platform to session."));
+			FPlatformProcess::Sleep(.001f);
+		}
+	}
+}
+
+void FRequestCluster::OnRemoveSessionPlatform(const ITargetPlatform* TargetPlatform)
+{
+	if (GraphSearch)
+	{
+		FCookerTimer CookerTimer(FCookerTimer::Forever);
+		bool bComplete;
+		while (PumpExploration(CookerTimer, bComplete), !bComplete)
 		{
-			Pair.Value.Remove(PackageData);
+			UE_LOG(LogCook, Display, TEXT("Waiting for RequestCluster to finish before removing platform from session."));
+			FPlatformProcess::Sleep(.001f);
 		}
 	}
-	Requests.Remove(PackageData);
-	RequestsToDemote.RemoveAll([PackageData](const TPair<FPackageData*, ESuppressCookReason>& Pair) { return Pair.Key == PackageData; });
+}
+
+void FRequestCluster::RemapTargetPlatforms(TMap<ITargetPlatform*, ITargetPlatform*>& Remap)
+{
+	if (GraphSearch)
+	{
+		// The platforms have already been invalidated, which means we can't wait for GraphSearch to finish
+		// Need to wait for all async operations to finish, then remap all the platforms
+		checkNoEntry(); // Not yet implemented
+	}
 }
 
 bool FRequestCluster::Contains(FPackageData* PackageData) const
@@ -453,88 +396,63 @@ void FRequestCluster::ClearAndDetachOwnedPackageDatas(TArray<FPackageData*>& Out
 	if (bStartAsyncComplete)
 	{
 		check(!GraphSearch);
-		OutRequestsToLoad = MoveTemp(Requests);
-		OutRequestsToDemote = MoveTemp(RequestsToDemote);
+		OutRequestsToLoad.Reset();
+		OutRequestsToDemote.Reset();
+		for (TPair<FPackageData*, ESuppressCookReason>& Pair : OwnedPackageDatas)
+		{
+			if (Pair.Value == ESuppressCookReason::NotSuppressed)
+			{
+				OutRequestsToLoad.Add(Pair.Key);
+			}
+			else
+			{
+				OutRequestsToDemote.Add(Pair);
+			}
+		}
 		OutRequestGraph = MoveTemp(RequestGraph);
-		check(OutRequestsToLoad.Num() + OutRequestsToDemote.Num() == OwnedPackageDatas.Num())
 	}
 	else
 	{
-		OutRequestsToLoad = OwnedPackageDatas.Array();
+		OutRequestsToLoad.Reset();
+		for (TPair<FPackageData*, ESuppressCookReason>& Pair : OwnedPackageDatas)
+		{
+			OutRequestsToLoad.Add(Pair.Key);
+		}
 		OutRequestsToDemote.Reset();
 		OutRequestGraph.Reset();
 	}
-	InRequests.Empty();
-	Requests.Empty();
-	RequestsToDemote.Empty();
+	FilePlatformRequests.Empty();
 	OwnedPackageDatas.Empty();
 	GraphSearch.Reset();
 	RequestGraph.Reset();
-	NextRequest = 0;
 }
 
-void FRequestCluster::FetchDependencies(const FCookerTimer& CookerTimer, bool& bOutComplete)
+void FRequestCluster::PumpExploration(const FCookerTimer& CookerTimer, bool& bOutComplete)
 {
 	if (bDependenciesComplete)
 	{
 		return;
 	}
 
-	if (COTFS.IsCookWorkerMode())
-	{
-		bDependenciesComplete = true;
-		return;
-	}
-
 	if (!GraphSearch)
 	{
 		GraphSearch.Reset(new FGraphSearch(*this));
-	}
 
-	if (!bAllowHardDependencies)
-	{
-		// FetchDependencies is responsible for marking all requests as explored and demoting
-		// the ones that are not cookable. If we're skipping the dependencies search, handle that
-		// responsibility for the initial requests and return.
-		for (FPackageData* PackageData : Requests)
+		if (!bAllowHardDependencies || COTFS.IsCookWorkerMode())
 		{
-			GraphSearch->VisitVertex(FVertexData(FVertexData::ESkipDependencies, *PackageData, *this));
+			GraphSearch->VisitWithoutDependencies();
+			GraphSearch.Reset();
+			bDependenciesComplete = true;
+			return;
 		}
-		Swap(Requests, GraphSearch->GetTransitiveRequests());
-		GraphSearch.Reset();
-		bDependenciesComplete = true;
-		return;
+		GraphSearch->StartSearch();
 	}
 
 	constexpr double WaitTime = 0.50;
-	for (;;)
+	bool bDone;
+	while (GraphSearch->TickExploration(bDone), !bDone)
 	{
-		bool bIsDone;
-		TArray<TUniquePtr<FVertexData>> CompletedVertices;
-		GraphSearch->Poll(CompletedVertices, bIsDone);
-		if (bIsDone)
-		{
-			break;
-		}
-		if (CompletedVertices.Num())
-		{
-			FHasPendingVerticesScope HasPendingVerticesScope(*GraphSearch);
-			for (TUniquePtr<FVertexData>& VertexData : CompletedVertices)
-			{
-				// The PackageData may be removed from this due to external requests while the VertexData was processing;
-				// discard the vertex now if we no longer own the PackageData.
-				if (OwnedPackageDatas.Contains(VertexData->PackageData))
-				{
-					GraphSearch->VisitVertex(*VertexData);
-				}
-				GraphSearch->FreeVertex(MoveTemp(VertexData));
-			}
-		}
-		else
-		{
-			GraphSearch->UpdateDisplay();
-			GraphSearch->WaitForPollAvailability(WaitTime);
-		}
+		GraphSearch->WaitForAsyncQueue(WaitTime);
 		if (CookerTimer.IsTimeUp())
 		{
 			bOutComplete = false;
@@ -542,9 +460,15 @@ void FRequestCluster::FetchDependencies(const FCookerTimer& CookerTimer, bool& b
 		}
 	}
 
-	TArray<FPackageData*>& NewRequests = GraphSearch->GetTransitiveRequests();
-	check(NewRequests.Num() + RequestsToDemote.Num() == OwnedPackageDatas.Num());
-	COOK_STAT(DetailedCookStats::NumPreloadedDependencies += FMath::Max(0, NewRequests.Num() - Requests.Num()));
+	TArray<FPackageData*> SortedPackages;
+	SortedPackages.Reserve(OwnedPackageDatas.Num());
+	for (TPair<FPackageData*, ESuppressCookReason>& Pair : OwnedPackageDatas)
+	{
+		if (Pair.Value == ESuppressCookReason::NotSuppressed)
+		{
+			SortedPackages.Add(Pair.Key);
+		}
+	}
 
 	// Sort the NewRequests in leaf to root order and replace the requests list with NewRequests
 	TArray<FPackageData*> Empty;
@@ -553,8 +477,33 @@ void FRequestCluster::FetchDependencies(const FCookerTimer& CookerTimer, bool& b
 		const TArray<FPackageData*>* VertexEdges = GraphSearch->GetGraphEdges().Find(PackageData);
 		return VertexEdges ? *VertexEdges : Empty;
 	};
-	Algo::TopologicalSort(NewRequests, GetElementDependencies, Algo::ETopologicalSort::AllowCycles);
-	Requests = MoveTemp(NewRequests);
+
+	Algo::TopologicalSort(SortedPackages, GetElementDependencies, Algo::ETopologicalSort::AllowCycles);
+	TMap<FPackageData*, int32> SortOrder;
+	int32 Counter = 0;
+	SortOrder.Reserve(SortedPackages.Num());
+	for (FPackageData* PackageData : SortedPackages)
+	{
+		SortOrder.Add(PackageData, Counter++);
+	}
+	OwnedPackageDatas.KeySort([&SortOrder](const FPackageData& A, const FPackageData& B)
+		{
+			int32* CounterA = SortOrder.Find(&A);
+			int32* CounterB = SortOrder.Find(&B);
+			if ((CounterA != nullptr) != (CounterB != nullptr))
+			{
+				// Sort the demotes to occur last
+				return CounterB == nullptr;
+			}
+			else if (CounterA)
+			{
+				return *CounterA < *CounterB;
+			}
+			else
+			{
+				return false; // demotes are unsorted
+			}
+		});
 
 	RequestGraph = MoveTemp(GraphSearch->GetGraphEdges());
 	GraphSearch.Reset();
@@ -563,104 +512,256 @@ void FRequestCluster::FetchDependencies(const FCookerTimer& CookerTimer, bool& b
 
 FRequestCluster::FGraphSearch::FGraphSearch(FRequestCluster& InCluster)
 	: Cluster(InCluster)
-	, PollReadyEvent(EEventMode::ManualReset)
+	, AsyncResultsReadyEvent(EEventMode::ManualReset)
 {
-	PollReadyEvent->Trigger();
+	AsyncResultsReadyEvent->Trigger();
 	bCookAttachmentsEnabled = !Cluster.bFullBuild && Cluster.bHybridIterativeEnabled;
+	LastActivityTime = FPlatformTime::Seconds();
+	VertexAllocator.SetMaxBlockSize(1024);
+	VertexAllocator.SetMaxBlockSize(65536);
+	VertexQueryAllocator.SetMinBlockSize(1024);
+	VertexQueryAllocator.SetMaxBlockSize(1024);
+	BatchAllocator.SetMaxBlockSize(16);
+	BatchAllocator.SetMaxBlockSize(16);
+
+	TConstArrayView<const ITargetPlatform*> SessionPlatforms = Cluster.COTFS.PlatformManager->GetSessionPlatforms();
+	check(SessionPlatforms.Num() > 0);
+	FetchPlatforms.SetNum(SessionPlatforms.Num() + 2);
+	FetchPlatforms[PlatformAgnosticPlatformIndex].bIsPlatformAgnosticPlatform = true;
+	FetchPlatforms[CookerLoadingPlatformIndex].Platform = CookerLoadingPlatformKey;
+	FetchPlatforms[CookerLoadingPlatformIndex].bIsCookerLoadingPlatform = true;
+	for (int32 SessionPlatformIndex = 0; SessionPlatformIndex < SessionPlatforms.Num(); ++SessionPlatformIndex)
+	{
+		FFetchPlatformData& FetchPlatform = FetchPlatforms[SessionPlatformIndex + 2];
+		FetchPlatform.Platform = SessionPlatforms[SessionPlatformIndex];
+		FetchPlatform.Writer = &Cluster.COTFS.FindOrCreatePackageWriter(FetchPlatform.Platform);
+	}
+	Algo::Sort(FetchPlatforms, [](const FFetchPlatformData& A, const FFetchPlatformData& B)
+		{
+			return A.Platform < B.Platform;
+		});
+	check(FetchPlatforms[PlatformAgnosticPlatformIndex].bIsPlatformAgnosticPlatform);
+	check(FetchPlatforms[CookerLoadingPlatformIndex].bIsCookerLoadingPlatform);
+}
 
-	TArray<TUniquePtr<FVertexData>> Vertices;
-	Vertices.Reserve(Cluster.Requests.Num());
-	for (FPackageData* PackageData : Cluster.Requests)
+void FRequestCluster::FGraphSearch::VisitWithoutDependencies()
+{
+	// PumpExploration is responsible for marking all requests as explored and cookable/uncoookable.
+	// If we're skipping the dependencies search, handle that responsibility for the initial requests and return.
+	for (TPair<FPackageData*, ESuppressCookReason>& Pair : Cluster.OwnedPackageDatas)
 	{
-		TUniquePtr<FVertexData>& VertexData = Vertices.Emplace_GetRef();
-		FindOrAddVertex(PackageData->GetPackageName(), PackageData, true /* bInitialRequest */,
-			false /* bHardDependency */, FInstigator(EInstigator::Unspecified, GInstigatorRequestCluster), VertexData);
-		check(VertexData); // Initial requests should always yield a valid new vertex
+		FVertexData Vertex;
+		Vertex.PackageData = Pair.Key;
+		VisitVertex(Vertex, true /* bSkipDependencies */);
+	}
+}
+
+void FRequestCluster::FGraphSearch::StartSearch()
+{
+	Frontier.Reserve(Cluster.OwnedPackageDatas.Num());
+	for (TPair<FPackageData*, ESuppressCookReason>& Pair : Cluster.OwnedPackageDatas)
+	{
+		FVertexData& Vertex = FindOrAddVertex(Pair.Key->GetPackageName(), *Pair.Key);
+		check(Vertex.PackageData);
+		// We're iterating over OwnedPackageDatas and the Vertex is already in the Cluster so we don't need to call AddToFrontier;
+		// just add it directly.
+		check(Pair.Value != ESuppressCookReason::Invalid);
+		Frontier.Add(&Vertex);
 	}
-	AddVertices(MoveTemp(Vertices));
-	LastActivityTime = FPlatformTime::Seconds();
 }
 
 FRequestCluster::FGraphSearch::~FGraphSearch()
 {
-	// Wait for any asynchronous tasks to complete by calling Poll and dropping the vertices it returns
-	bool bIsDone = false;
-	while (!bIsDone)
+	for (;;)
 	{
-		TArray<TUniquePtr<FVertexData>> VerticesToDrop;
-		Poll(VerticesToDrop, bIsDone);
-		if (!bIsDone)
+		bool bHadActivity = false;
+		bool bAsyncBatchesEmpty = false;
+		{
+			FScopeLock ScopeLock(&Lock);
+			bAsyncBatchesEmpty = AsyncQueueBatches.IsEmpty();
+			if (!bAsyncBatchesEmpty)
+			{
+				// It is safe to Reset AsyncResultsReadyEvent and wait on it later because we are inside the lock and there is a
+				// remaining batch, so it will be triggered after the Reset when that batch completes.
+				AsyncResultsReadyEvent->Reset();
+			}
+		}
+		for (;;)
+		{
+			TOptional<FVertexData*> Vertex = AsyncQueueResults.Dequeue();
+			if (!Vertex)
+			{
+				break;
+			}
+			FreeQueryData((**Vertex).QueryData);
+			(**Vertex).QueryData = nullptr;
+			bHadActivity = true;
+		}
+		if (bAsyncBatchesEmpty)
+		{
+			break;
+		}
+		if (bHadActivity)
+		{
+			LastActivityTime = FPlatformTime::Seconds();
+		}
+		else
 		{
 			UpdateDisplay();
-			constexpr double WaitTime = 1.0;
-			WaitForPollAvailability(WaitTime);
 		}
+		constexpr double WaitTime = 1.0;
+		WaitForAsyncQueue(WaitTime);
 	}
 }
 
-void FRequestCluster::FGraphSearch::AddVertices(TArray<TUniquePtr<FVertexData>>&& Vertices)
+void FRequestCluster::FGraphSearch::RemovePackageData(FPackageData* PackageData)
 {
-	TArray<FQueryVertexBatch*> NewBatches;
+	check(PackageData);
+	FVertexData** Vertex = Vertices.Find(PackageData->GetPackageName());
+	if (Vertex)
 	{
-		FScopeLock ScopeLock(&Lock);
-		VerticesToRead.Reserve(VerticesToRead.Num() + Vertices.Num());
-		for (TUniquePtr<FVertexData>& Vertex : Vertices)
-		{
-			VerticesToRead.Add(MoveTemp(Vertex));
-		}
-		NewBatches = CreateAvailableBatches();
+		(**Vertex).PackageData = nullptr;
 	}
-	for (FQueryVertexBatch* NewBatch : NewBatches)
+
+	GraphEdges.Remove(PackageData);
+	for (TPair<FPackageData*, TArray<FPackageData*>>& Pair : GraphEdges)
 	{
-		NewBatch->Send();
+		Pair.Value.Remove(PackageData);
 	}
-	Vertices.Reset();
 }
 
-void FRequestCluster::FGraphSearch::WaitForPollAvailability(double WaitTimeSeconds)
+void FRequestCluster::FGraphSearch::OnNewReachablePlatforms(FPackageData* PackageData)
 {
-	uint32 WaitTime = (WaitTimeSeconds > 0.0) ? static_cast<uint32>(FMath::Floor(WaitTimeSeconds * 1000)) : MAX_uint32;
-	PollReadyEvent->Wait(WaitTime);
+	FVertexData** VertexPtr = Vertices.Find(PackageData->GetPackageName());
+	if (!VertexPtr)
+	{
+		return;
+	}
+	// Already in OwnedPackageDatas, so just add to Frontier directly
+	Frontier.Add(*VertexPtr);
 }
 
-void FRequestCluster::FGraphSearch::Poll(TArray<TUniquePtr<FVertexData>>& OutCompletedVertices, bool& bOutIsDone)
+void FRequestCluster::FGraphSearch::QueueEdgesFetch(FVertexData& Vertex, TConstArrayView<const ITargetPlatform*> Platforms)
 {
-	bOutIsDone = true;
-	TArray<FQueryVertexBatch*> NewBatches;
+	check(!Vertex.QueryData);
+	FVertexQueryData& QueryData = *AllocateQueryData();
+	Vertex.QueryData = &QueryData;
+	
+	QueryData.PackageName = Vertex.PackageData->GetPackageName();
+	QueryData.Platforms.SetNum(FetchPlatforms.Num());
+
+	// Store Platforms in QueryData->Platforms.bActive. All bActive values start false from constructor or from Reset
+	bool bHasPlatformAgnostic = false;
+	for (const ITargetPlatform* Platform : Platforms)
 	{
-		FScopeLock ScopeLock(&Lock);
-		// We check for batches to send after adding dependencies for each vertex, so we can get them processing on the async thread
-		// while we work on other vertices on the main thread.
-		// We additionally need to check for them on every poll, in case the last batch we sent completes without finding any new vertices
-		// and there are no batches remaining so we need to send an incomplete batch of the vertices remaining.
-		NewBatches = CreateAvailableBatches();
-		if (Batches.Num() != 0)
-		{
-			bOutIsDone = false;
-			PollReadyEvent->Reset();
+		int32 Index = Algo::BinarySearchBy(FetchPlatforms, Platform, [](const FFetchPlatformData& D) { return D.Platform; });
+		check(Index != INDEX_NONE);
+		QueryData.Platforms[Index].bActive = true;
+		if (Platform != CookerLoadingPlatformKey)
+		{
+			bHasPlatformAgnostic = true;
 		}
 	}
-	for (FQueryVertexBatch* NewBatch : NewBatches)
+	if (bHasPlatformAgnostic)
 	{
-		NewBatch->Send();
+		QueryData.Platforms[PlatformAgnosticPlatformIndex].bActive = true;
 	}
+	int32 NumPendingPlatforms = Platforms.Num() + (bHasPlatformAgnostic ? 1 : 0);
+	QueryData.PendingPlatforms.store(NumPendingPlatforms, std::memory_order_release);
+
+	PreAsyncQueue.Add(&Vertex);
+	CreateAvailableBatches(false /* bAllowIncompleteBatch */);
+}
+
+void FRequestCluster::FGraphSearch::WaitForAsyncQueue(double WaitTimeSeconds)
+{
+	uint32 WaitTime = (WaitTimeSeconds > 0.0) ? static_cast<uint32>(FMath::Floor(WaitTimeSeconds * 1000)) : MAX_uint32;
+	AsyncResultsReadyEvent->Wait(WaitTime);
+}
 
-	OutCompletedVertices.Reset();
+void FRequestCluster::FGraphSearch::TickExploration(bool& bOutDone)
+{
+	bool bHadActivity = false;
 	for (;;)
 	{
-		TOptional<TUniquePtr<FVertexData>> VertexData = CompletedVertices.Dequeue();
-		if (!VertexData.IsSet())
+		TOptional<FVertexData*> Vertex = AsyncQueueResults.Dequeue();
+		if (!Vertex.IsSet())
 		{
 			break;
 		}
-		OutCompletedVertices.Add(MoveTemp(*VertexData));
+		ExploreVertexEdges(**Vertex);
+		FreeQueryData((**Vertex).QueryData);
+		(**Vertex).QueryData = nullptr;
+		bHadActivity = true;
 	}
 
-	bOutIsDone = bOutIsDone && OutCompletedVertices.Num() == 0;
-	if (OutCompletedVertices.Num() > 0)
+	if (!Frontier.IsEmpty())
+	{
+		TArray<FVertexData*> BusyVertices;
+		for (FVertexData* Vertex : Frontier)
+		{
+			if (Vertex->QueryData)
+			{
+				// Vertices that are already in the AsyncQueue can not be added again; we would clobber their QueryData. Postpone them.
+				BusyVertices.Add(Vertex);
+			}
+			else
+			{
+				VisitVertex(*Vertex);
+			}
+		}
+		bHadActivity |= BusyVertices.Num() != Frontier.Num();
+		Frontier.Reset();
+		Frontier.Append(BusyVertices);
+	}
+
+	if (bHadActivity)
 	{
 		LastActivityTime = FPlatformTime::Seconds();
+		bOutDone = false;
+		return;
+	}
+
+	bool bAsyncQueueEmpty;
+	{
+		FScopeLock ScopeLock(&Lock);
+		if (!AsyncQueueResults.IsEmpty())
+		{
+			bAsyncQueueEmpty = false;
+		}
+		else
+		{
+			bAsyncQueueEmpty = AsyncQueueBatches.IsEmpty();
+			// AsyncResultsReadyEvent can only be Reset when either the AsyncQueue is empty or it is non-empty and we
+			// know the AsyncResultsReadyEvent will be triggered again "later".
+			// The guaranteed place where it will be Triggered is when a batch completes. To guarantee that
+			// place will be called "later", the batch completion trigger and this reset have to both
+			// be done inside the lock.
+			AsyncResultsReadyEvent->Reset();
+		}
+	}
+	if (!bAsyncQueueEmpty)
+	{
+		// Waiting on the AsyncQueue; give a warning if we have been waiting for long with no AsyncQueueResults.
+		UpdateDisplay();
+		bOutDone = false;
+		return;
 	}
+
+	// No more work coming in the future from the AsyncQueue, and we are out of work to do
+	// without it. If we have any queued vertices in the PreAsyncQueue, send them now and continue
+	// waiting. Otherwise we are done.
+	if (!PreAsyncQueue.IsEmpty())
+	{
+		CreateAvailableBatches(true /* bAllowInCompleteBatch */);
+		bOutDone = false;
+		return;
+	}
+
+	// Frontier was reset above, and it cannot be modified between there and here.
+	// If it were non-empty we would not be done.
+	check(Frontier.IsEmpty());
+	bOutDone = true;
 }
 
 void FRequestCluster::FGraphSearch::UpdateDisplay()
@@ -670,10 +771,10 @@ void FRequestCluster::FGraphSearch::UpdateDisplay()
 	{
 		FScopeLock ScopeLock(&Lock);
 		int32 NumVertices = 0;
-		int32 NumBatches = Batches.Num();
-		for (TPair<FQueryVertexBatch*, TUniquePtr<FQueryVertexBatch>>& Pair : Batches)
+		int32 NumBatches = AsyncQueueBatches.Num();
+		for (FQueryVertexBatch* Batch : AsyncQueueBatches)
 		{
-			NumVertices += Pair.Key->PendingVertices;
+			NumVertices += Batch->PendingVertices;
 		}
 
 		UE_LOG(LogCook, Warning, TEXT("FRequestCluster waited more than %.0lfs for previous build results from the oplog. ")
@@ -683,423 +784,605 @@ void FRequestCluster::FGraphSearch::UpdateDisplay()
 	}
 }
 
-void FRequestCluster::FGraphSearch::VisitVertex(const FVertexData& VertexData)
+void FRequestCluster::FGraphSearch::VisitVertex(FVertexData& Vertex, bool bSkipDependencies)
 {
-	// Only called from Process thread
-	using namespace UE::AssetRegistry;
-	using namespace UE::TargetDomain;
+	// Only called from PumpExploration thread
 
-	TArray<FName>& HardDependencies(Scratch.HardDependencies);
-	TArray<FName>& SoftDependencies(Scratch.SoftDependencies);
-	HardDependencies.Reset();
-	SoftDependencies.Reset();
-	FPackageData* PackageData = VertexData.PackageData;
-	if (VertexData.bExploreDependencies)
+	// The PackageData will not exist if the package does not exist on disk or
+	// the PackageData was removed from the FRequestCluster due to changes in the PackageData's
+	// state elsewhere in the cooker.
+	if (!Vertex.PackageData)
 	{
-		FName PackageName = PackageData->GetPackageName();
+		return;
+	}
 
-		// TODO EditorOnly References: When bCanSkipEditorReferencedPackagesWhenCooking is active,
-		// we only fetch Game dependencies, because the cooker explicitly loads all of
-		// the dependencies that we report. And if we explicitly load an EditorOnly dependency, that causes
-		// StaticLoadObjectInternal to SetLoadedByEditorPropertiesOnly(false), which then treats the editor-only package
-		// as needed in-game.
-		EDependencyQuery DependencyQuery = Cluster.COTFS.bCanSkipEditorReferencedPackagesWhenCooking ?
-			EDependencyQuery::Game : EDependencyQuery::NoRequirements;
+	TArray<const ITargetPlatform*, TInlineAllocator<1>> ExplorePlatforms;
+	FPackagePlatformData* CookerLoadingPlatform = nullptr;
+	const ITargetPlatform* FirstReachableSessionPlatform = nullptr;
+	ESuppressCookReason SuppressCookReason = ESuppressCookReason::Invalid;
+	bool bAllReachablesUncookable = true;
+	for (TPair<const ITargetPlatform*, FPackagePlatformData>& Pair : Vertex.PackageData->GetPlatformDatasConstKeysMutableValues())
+	{
+		FPackagePlatformData& PlatformData = Pair.Value;
+		if (Pair.Key == CookerLoadingPlatformKey)
+		{
+			CookerLoadingPlatform = &Pair.Value;
+		}
+		else if (PlatformData.IsReachable())
+		{
+			if (!FirstReachableSessionPlatform)
+			{
+				FirstReachableSessionPlatform = Pair.Key;
+			}
+			if (!PlatformData.IsVisitedByCluster())
+			{
+				VisitVertexForPlatform(Vertex, Pair.Key, PlatformData, SuppressCookReason);
+				if (!bSkipDependencies && PlatformData.IsExplorable())
+				{
+					ExplorePlatforms.Add(Pair.Key);
+				}
+			}
+			if (PlatformData.IsCookable())
+			{
+				bAllReachablesUncookable = false;
+				SuppressCookReason = ESuppressCookReason::NotSuppressed;
+			}
+		}
+	}
+	bool bAnyCookable = (FirstReachableSessionPlatform == nullptr) | !bAllReachablesUncookable;
+	if (bAnyCookable != Vertex.bAnyCookable)
+	{
+		if (!bAnyCookable)
+		{
+			if (SuppressCookReason == ESuppressCookReason::Invalid)
+			{
+				// We need the SuppressCookReason for reporting. If we didn't calculate it this Visit and
+				// we don't have it stored in this->OwnedPackageDatas, then we must have calculated it in
+				// a previous cluster, but we don't store it anywhere. Recalculate it from the
+				// FirstReachableSessionPlatform. FirstReachableSessionPlatform must be non-null, otherwise
+				// bAnyCookable would be true.
+				check(FirstReachableSessionPlatform);
+				bool bCookable;
+				bool bExplorable;
+				Cluster.IsRequestCookable(FirstReachableSessionPlatform, Vertex.PackageData->GetPackageName(),
+					*Vertex.PackageData, SuppressCookReason, bCookable, bExplorable);
+				check(!bCookable); // We don't support bCookable changing for a given package and platform
+				check(SuppressCookReason != ESuppressCookReason::Invalid);
+			}
+		}
+		else
+		{
+			check(SuppressCookReason == ESuppressCookReason::NotSuppressed);
+		}
+		Cluster.OwnedPackageDatas.FindOrAdd(Vertex.PackageData) = SuppressCookReason;
+		Vertex.bAnyCookable = bAnyCookable;
+	}
 
-		Cluster.AssetRegistry.GetDependencies(PackageName, HardDependencies, EDependencyCategory::Package,
-			DependencyQuery | EDependencyQuery::Hard);
-		// Skip soft dependencies if the cook commandline is set to skip soft references or if the vertex is uncookable.
-		if (VertexData.bExploreSoftDependencies)
+	// If any target platform is cookable, then we need to mark the CookerLoadingPlatform as reachable because we will need
+	// to load the package to cook it
+	if (bAnyCookable)
+	{
+		if (!CookerLoadingPlatform)
+		{
+			CookerLoadingPlatform = &Vertex.PackageData->FindOrAddPlatformData(CookerLoadingPlatformKey);
+		}
+		CookerLoadingPlatform->SetReachable(true);
+	}
+	if (CookerLoadingPlatform && CookerLoadingPlatform->IsReachable() && !CookerLoadingPlatform->IsVisitedByCluster())
+	{
+		CookerLoadingPlatform->SetCookable(true);
+		CookerLoadingPlatform->SetExplorable(true);
+		CookerLoadingPlatform->SetVisitedByCluster(true);
+		if (!bSkipDependencies)
 		{
-			EDependencyQuery SoftDependencyQuery = DependencyQuery | EDependencyQuery::Soft
-				// Don't explore editor-only soft dependencies; the cooker should not try to load those
-				| EDependencyQuery::Game;
+			ExplorePlatforms.Add(CookerLoadingPlatformKey);
+		}
+	}
 
-			Cluster.AssetRegistry.GetDependencies(PackageName, SoftDependencies, EDependencyCategory::Package,
-				SoftDependencyQuery);
+	if (!ExplorePlatforms.IsEmpty())
+	{
+		check(!bSkipDependencies);
+		QueueEdgesFetch(Vertex, ExplorePlatforms);
+	}
+}
+
+void FRequestCluster::FGraphSearch::VisitVertexForPlatform(FVertexData& Vertex, const ITargetPlatform* Platform,
+	FPackagePlatformData& PlatformData, ESuppressCookReason& AccumulatedSuppressCookReason)
+{
+	FPackageData& PackageData = *Vertex.PackageData;
+	ESuppressCookReason SuppressCookReason = ESuppressCookReason::Invalid;
+	bool bCookable;
+	bool bExplorable;
+	Cluster.IsRequestCookable(Platform, Vertex.PackageData->GetPackageName(), PackageData, SuppressCookReason,
+		bCookable, bExplorable);
+	PlatformData.SetCookable(bCookable);
+	PlatformData.SetExplorable(bExplorable);
+	if (bCookable)
+	{
+		AccumulatedSuppressCookReason = ESuppressCookReason::NotSuppressed;
+	}
+	else
+	{
+		check(SuppressCookReason != ESuppressCookReason::Invalid && SuppressCookReason != ESuppressCookReason::NotSuppressed);
+		if (AccumulatedSuppressCookReason == ESuppressCookReason::Invalid)
+		{
+			AccumulatedSuppressCookReason = SuppressCookReason;
+		}
+	}
+	PlatformData.SetVisitedByCluster(true);
+}
+
+void FRequestCluster::FGraphSearch::ExploreVertexEdges(FVertexData& Vertex)
+{
+	// Only called from PumpExploration thread
+	using namespace UE::AssetRegistry;
+	using namespace UE::TargetDomain;
+
+	// The PackageData will not exist if the package does not exist on disk or
+	// the PackageData was removed from the FRequestCluster due to changes in the PackageData's
+	// state elsewhere in the cooker.
+	if (!Vertex.PackageData)
+	{
+		return;
+	}
+
+	TArray<FName>& HardGameDependencies(Scratch.HardGameDependencies);
+	TArray<FName>& SoftGameDependencies(Scratch.SoftGameDependencies);
+	TSet<FName>& HardDependenciesSet(Scratch.HardDependenciesSet);
+	HardGameDependencies.Reset();
+	SoftGameDependencies.Reset();
+	HardDependenciesSet.Reset();
+	FPackageData& PackageData = *Vertex.PackageData;
+	FName PackageName = PackageData.GetPackageName();
+	bool bFetchAnyTargetPlatform = Vertex.QueryData->Platforms[PlatformAgnosticPlatformIndex].bActive;
+	TArray<FName>* DiscoveredDependencies = Cluster.COTFS.DiscoveredDependencies.Find(PackageName);
+	if (bFetchAnyTargetPlatform)
+	{
+		EDependencyQuery FlagsForHardDependencyQuery;
+		if (Cluster.COTFS.bCanSkipEditorReferencedPackagesWhenCooking)
+		{
+			FlagsForHardDependencyQuery = EDependencyQuery::Game | EDependencyQuery::Hard;
+		}
+		else
+		{
+			// We're not allowed to skip editoronly imports, so include all hard dependencies
+			FlagsForHardDependencyQuery = EDependencyQuery::Hard;
+		}
+		Cluster.AssetRegistry.GetDependencies(PackageName, HardGameDependencies, EDependencyCategory::Package,
+			FlagsForHardDependencyQuery);
+		HardDependenciesSet.Append(HardGameDependencies);
+		if (DiscoveredDependencies)
+		{
+			HardDependenciesSet.Append(*DiscoveredDependencies);
+		}
+		if (Cluster.bAllowSoftDependencies)
+		{
+			// bCanSkipEditorReferencedPackagesWhenCooking does not affect soft dependencies; skip editoronly soft dependencies
+			Cluster.AssetRegistry.GetDependencies(PackageName, SoftGameDependencies, EDependencyCategory::Package,
+				EDependencyQuery::Game | EDependencyQuery::Soft);
 
 			// Even if we're following soft references in general, we need to check with the SoftObjectPath registry
 			// for any startup packages that marked their softobjectpaths as excluded, and not follow those
 			TSet<FName>& SkippedPackages(Scratch.SkippedPackages);
 			if (GRedirectCollector.RemoveAndCopySoftObjectPathExclusions(PackageName, SkippedPackages))
 			{
-				SoftDependencies.RemoveAll([&SkippedPackages](FName SoftDependency)
+				SoftGameDependencies.RemoveAll([&SkippedPackages](FName SoftDependency)
 					{
 						return SkippedPackages.Contains(SoftDependency);
 					});
 			}
 
-			SoftDependencies.Append(GetLocalizationReferences(PackageName, Cluster.COTFS));
+			// LocalizationReferences are a source of SoftGameDependencies that are not present in the AssetRegistry
+			SoftGameDependencies.Append(GetLocalizationReferences(PackageName, Cluster.COTFS));
+		}
+	}
+
+	int32 LocalNumFetchPlatforms = NumFetchPlatforms();
+	TMap<FName, FScratchPlatformDependencyBits>& PlatformDependencyMap(Scratch.PlatformDependencyMap);
+	PlatformDependencyMap.Reset();
+	auto AddPlatformDependency = [&PlatformDependencyMap, LocalNumFetchPlatforms](FName DependencyName, int32 PlatformIndex, bool bHardDependency)
+	{
+		FScratchPlatformDependencyBits& PlatformDependencyBits = PlatformDependencyMap.FindOrAdd(DependencyName);
+		if (PlatformDependencyBits.HasPlatformByIndex.Num() != LocalNumFetchPlatforms)
+		{
+			PlatformDependencyBits.HasPlatformByIndex.Init(false, LocalNumFetchPlatforms);
+			PlatformDependencyBits.bHardDependency = false;
+		}
+		PlatformDependencyBits.HasPlatformByIndex[PlatformIndex] = true;
+		if (bHardDependency)
+		{
+			PlatformDependencyBits.bHardDependency = true;
 		}
 
-		bool bFoundCachedTargetDomain = false;
-		bool bFoundBuildDefinitions = false;
-		int32 NumPlatforms = GetNumPlatforms();
-		// Platform 0 is the platform agnostic data, which we check after the loop
-		for (int32 PlatformIndex = 1; PlatformIndex < NumPlatforms; ++PlatformIndex)
+	};
+	auto AddPlatformDependencyRange = [&AddPlatformDependency](TConstArrayView<FName> Range, int32 PlatformIndex, bool bHardDependency)
+	{
+		for (FName DependencyName : Range)
+		{
+			AddPlatformDependency(DependencyName, PlatformIndex, bHardDependency);
+		}
+	};
+
+	FQueryPlatformData& PlatformAgnosticQueryPlatformData = Vertex.QueryData->Platforms[PlatformAgnosticPlatformIndex];
+	for (int32 PlatformIndex = 0; PlatformIndex < LocalNumFetchPlatforms; ++PlatformIndex)
+	{
+		FQueryPlatformData& QueryPlatformData = Vertex.QueryData->Platforms[PlatformIndex];
+		if (!QueryPlatformData.bActive || PlatformIndex == PlatformAgnosticPlatformIndex)
 		{
-			const FCookAttachments& PlatformAttachments = VertexData.CookAttachments[PlatformIndex];
-			if (!IsCookAttachmentsValid(PackageName, PlatformAttachments))
+			continue;
+		}
+
+		if (PlatformIndex == CookerLoadingPlatformIndex)
+		{
+			TArray<FName>& CookerLoadingDependencies(Scratch.CookerLoadingDependencies);
+			CookerLoadingDependencies.Reset();
+
+			Cluster.AssetRegistry.GetDependencies(PackageName, CookerLoadingDependencies, EDependencyCategory::Package,
+				EDependencyQuery::Hard);
+			CookerLoadingDependencies.Reset();
+
+			// ITERATIVECOOK_TODO: Build dependencies need to be stored and used to mark package loads as expected
+			// But we can't use them to explore packages that will be loaded during cook because they might not be;
+			// some build dependencies might be a conservative list but unused by the asset, or unused on targetplatform
+			// Adding BuildDependencies also sets up many circular dependencies, because maps declare their external
+			// actors as build dependencies and the external actors declare the map as a build or hard dependency.
+			// Topological sort done at the end of the Cluster has poor performance when there are 100k+ circular dependencies.
+			constexpr bool bAddBuildDependenciesToGraph = false;
+			if (bAddBuildDependenciesToGraph)
 			{
-				continue;
+				Cluster.AssetRegistry.GetDependencies(PackageName, CookerLoadingDependencies, EDependencyCategory::Package,
+					EDependencyQuery::Build);
 			}
-			const ITargetPlatform* TargetPlatform = Cluster.Platforms[PlatformIndex - 1];
-			ICookedPackageWriter* PackageWriter = Cluster.PackageWriters[PlatformIndex - 1];
-			if (!Cluster.bFullBuild && Cluster.bHybridIterativeEnabled)
+			// CookerLoadingPlatform does not cause SetInstigator so it does not modify bHardDependency
+			AddPlatformDependencyRange(CookerLoadingDependencies, PlatformIndex, false /* bHardDependency */);
+		}
+		else
+		{
+			FFetchPlatformData& FetchPlatformData = FetchPlatforms[PlatformIndex];
+			const ITargetPlatform* TargetPlatform = FetchPlatformData.Platform;
+
+			AddPlatformDependencyRange(HardGameDependencies, PlatformIndex, true /* bHardDependency */);
+			AddPlatformDependencyRange(SoftGameDependencies, PlatformIndex, false /* bHardDependency */);
+
+			const FCookAttachments& PlatformAttachments = QueryPlatformData.CookAttachments;
+			bool bFoundBuildDefinitions = false;
+			if (IsCookAttachmentsValid(PackageName, PlatformAttachments))
 			{
-				if (IsIterativeEnabled(PackageName))
-				{
-					bFoundCachedTargetDomain = true;
-					PackageData->SetPlatformCooked(TargetPlatform, ECookResult::Succeeded);
-					PackageWriter->MarkPackagesUpToDate({ PackageName });
-					// Declare the package to the EDLCookInfo verification so we don't warn about missing exports from it
-					UE::SavePackageUtilities::EDLCookInfoAddIterativelySkippedPackage(PackageName);
-				}
-				HardDependencies.Append(PlatformAttachments.BuildDependencies);
-				if (VertexData.bExploreSoftDependencies)
+				ICookedPackageWriter* PackageWriter = FetchPlatformData.Writer;
+				if (!Cluster.bFullBuild && Cluster.bHybridIterativeEnabled)
 				{
-					SoftDependencies.Append(PlatformAttachments.RuntimeOnlyDependencies);
-				}
+					if (IsIterativeEnabled(PackageName))
+					{
+						if (PlatformIndex == FirstSessionPlatformIndex)
+						{
+							COOK_STAT(++DetailedCookStats::NumPackagesIterativelySkipped);
+						}
+						PackageData.SetPlatformCooked(TargetPlatform, ECookResult::Succeeded);
+						PackageWriter->MarkPackagesUpToDate({ PackageName });
+						// Declare the package to the EDLCookInfo verification so we don't warn about missing exports from it
+						UE::SavePackageUtilities::EDLCookInfoAddIterativelySkippedPackage(PackageName);
+					}
+					AddPlatformDependencyRange(PlatformAttachments.BuildDependencies, PlatformIndex,
+						true /* bHardDependency */);
+					if (Cluster.bAllowSoftDependencies)
+					{
+						AddPlatformDependencyRange(PlatformAttachments.RuntimeOnlyDependencies, PlatformIndex,
+							true /* bHardDependency */);
+					}
 
-				if (Cluster.bPreQueueBuildDefinitions)
+					if (Cluster.bPreQueueBuildDefinitions)
+					{
+						bFoundBuildDefinitions = true;
+						Cluster.BuildDefinitions.AddBuildDefinitionList(PackageName, TargetPlatform,
+							PlatformAttachments.BuildDefinitionList);
+					}
+				}
+			}
+			if (Cluster.bPreQueueBuildDefinitions && !bFoundBuildDefinitions)
+			{
+				if (PlatformAgnosticQueryPlatformData.bActive &&
+					IsCookAttachmentsValid(PackageName, PlatformAgnosticQueryPlatformData.CookAttachments))
 				{
-					bFoundBuildDefinitions = true;
 					Cluster.BuildDefinitions.AddBuildDefinitionList(PackageName, TargetPlatform,
-						PlatformAttachments.BuildDefinitionList);
+						PlatformAgnosticQueryPlatformData.CookAttachments.BuildDefinitionList);
 				}
 			}
 		}
-		if (Cluster.bPreQueueBuildDefinitions && !bFoundBuildDefinitions)
+		if (DiscoveredDependencies)
 		{
-			const FCookAttachments& PlatformAgnosticAttachments = VertexData.CookAttachments[0];
-			if (IsCookAttachmentsValid(PackageName, PlatformAgnosticAttachments))
-			{
-				Cluster.BuildDefinitions.AddBuildDefinitionList(PackageName, nullptr, PlatformAgnosticAttachments.BuildDefinitionList);
-			}
+			AddPlatformDependencyRange(*DiscoveredDependencies, PlatformIndex, true /* bHardDependency */);
 		}
-		if (bFoundCachedTargetDomain)
+	}
+	if (PlatformDependencyMap.IsEmpty())
+	{
+		return;
+	}
+
+	TArray<FPackageData*>* Edges = nullptr;
+	for (TPair<FName, FScratchPlatformDependencyBits>& PlatformDependencyPair : PlatformDependencyMap)
+	{
+		FName DependencyName = PlatformDependencyPair.Key;
+		TBitArray<>& HasPlatformByIndex = PlatformDependencyPair.Value.HasPlatformByIndex;
+		bool bHardDependency = PlatformDependencyPair.Value.bHardDependency;
+
+		// Process any CoreRedirects before checking whether the package exists
+		FName Redirected = FCoreRedirects::GetRedirectedName(ECoreRedirectFlags::Type_Package,
+			FCoreRedirectObjectName(NAME_None, NAME_None, DependencyName)).PackageName;
+		DependencyName = Redirected;
+
+		FVertexData& DependencyVertex = FindOrAddVertex(DependencyName);
+		if (!DependencyVertex.PackageData)
 		{
-			COOK_STAT(++DetailedCookStats::NumPackagesIterativelySkipped);
+			continue;
 		}
+		FPackageData& DependencyPackageData(*DependencyVertex.PackageData);
+		bool bAddToFrontier = false;
 
-		// Sort the lists of Dependencies to check for uniqueness and make them deterministic
-		Algo::Sort(HardDependencies, FNameLexicalLess());
-		HardDependencies.SetNum(Algo::Unique(HardDependencies));
-		Algo::Sort(SoftDependencies, FNameLexicalLess());
-		SoftDependencies.SetNum(Algo::Unique(SoftDependencies));
-		SoftDependencies.RemoveAll([&HardDependencies](FName Dependency)
+		for (int32 PlatformIndex = 0; PlatformIndex < LocalNumFetchPlatforms; ++PlatformIndex)
+		{
+			if (!HasPlatformByIndex[PlatformIndex])
 			{
-				return Algo::BinarySearch(HardDependencies, Dependency, FNameLexicalLess()) != INDEX_NONE;
-			});
+				continue;
+			}
+			FFetchPlatformData& FetchPlatformData = FetchPlatforms[PlatformIndex];
+			const ITargetPlatform* TargetPlatform = FetchPlatformData.Platform;
+			FPackagePlatformData& PlatformData = DependencyPackageData.FindOrAddPlatformData(TargetPlatform);
 
-		if (HardDependencies.Num() || SoftDependencies.Num())
-		{
-			TArray<FPackageData*>& Edges = GraphEdges.FindOrAdd(PackageData);
-			check(Edges.Num() == 0);
-			Edges.Reserve(HardDependencies.Num());
-			TArray<TUniquePtr<FVertexData>>& NewVertices(Scratch.NewVertices);
-			NewVertices.Reset(HardDependencies.Num() + SoftDependencies.Num());
-			for (TArray<FName>* Dependencies : { &HardDependencies, &SoftDependencies })
+			if (PlatformIndex == CookerLoadingPlatformIndex)
 			{
-				bool bHardDependency = Dependencies == &HardDependencies;
-				for (FName Dependency : *Dependencies)
+				if (!Edges)
 				{
-					// Process any CoreRedirects before checking whether the package exists
-					FName Redirected = FCoreRedirects::GetRedirectedName(ECoreRedirectFlags::Type_Package,
-						FCoreRedirectObjectName(NAME_None, NAME_None, Dependency)).PackageName;
-					Dependency = Redirected;
+					Edges = &GraphEdges.FindOrAdd(&PackageData);
+					Edges->Reset(PlatformDependencyMap.Num());
+				}
+				Edges->Add(&DependencyPackageData);
+			}
 
-					TUniquePtr<FVertexData> NewVertexData;
+			if (!PlatformData.IsReachable())
+			{
+				PlatformData.SetReachable(true);
+				if (!DependencyPackageData.HasInstigator() && TargetPlatform != CookerLoadingPlatformKey)
+				{
 					EInstigator InstigatorType = bHardDependency ? EInstigator::HardDependency : EInstigator::SoftDependency;
-					FPackageData* DependencyPackageData = FindOrAddVertex(Dependency, nullptr,
-						false /* bInitialRequest */, bHardDependency, FInstigator(InstigatorType, PackageName),
-						NewVertexData);
-					if (NewVertexData)
-					{
-						NewVertices.Add(MoveTemp(NewVertexData));
-					}
-					if (bHardDependency && DependencyPackageData)
-					{
-						Edges.Add(DependencyPackageData);
-					}
+					DependencyPackageData.SetInstigator(Cluster, FInstigator(InstigatorType, PackageName));
 				}
 			}
-			if (NewVertices.Num())
+			if (!PlatformData.IsVisitedByCluster())
 			{
-				AddVertices(MoveTemp(NewVertices));
+				bAddToFrontier = true;
 			}
 		}
+		if (bAddToFrontier)
+		{
+			AddToFrontier(DependencyVertex);
+		}
 	}
+}
 
-	MarkExplored(*PackageData, VertexData.bCookable);
-	bool bAlreadyCooked = PackageData->AreAllRequestedPlatformsCooked(true /* bAllowFailedCooks */);
-
-	if (VertexData.bCookable && !bAlreadyCooked)
-	{
-		TransitiveRequests.Add(PackageData);
-	}
-	else
+void FRequestCluster::FVertexQueryData::Reset()
+{
+	for (FQueryPlatformData& PlatformData : Platforms)
 	{
-		ESuppressCookReason SuppressCookReason = VertexData.bCookable ? ESuppressCookReason::AlreadyCooked : VertexData.SuppressCookReason;
-		Cluster.RequestsToDemote.Emplace(PackageData, SuppressCookReason);
+		PlatformData.CookAttachments.Reset();
+		PlatformData.bActive = false;
 	}
 }
 
-void FRequestCluster::FGraphSearch::MarkExplored(FPackageData& PackageData, bool bCookable)
+FRequestCluster::FVertexData* FRequestCluster::FGraphSearch::AllocateVertex()
 {
-	for (const ITargetPlatform* TargetPlatform : Cluster.Platforms)
-	{
-		FPackageData::FPlatformData& PlatformData = PackageData.FindOrAddPlatformData(TargetPlatform);
-		PlatformData.SetExplored(true);
-		PlatformData.SetCookable(bCookable);
-	}
-	Cluster.COTFS.PackageDatas->UpdateThreadsafePackageData(PackageData);
+	return VertexAllocator.NewElement();
 }
 
-FRequestCluster::FVertexData::FVertexData(EAsyncType, int32 NumPlatforms)
+FRequestCluster::FVertexQueryData* FRequestCluster::FGraphSearch::AllocateQueryData()
 {
-	CookAttachments.SetNum(NumPlatforms, true /* bAllowShrinking */);
-	Reset();
+	// VertexQueryAllocator uses DeferredDestruction, so this might be a resused Batch, but we don't need to Reset it
+	// during allocation because Batches are Reset during Free.
+	return VertexQueryAllocator.NewElement();
 }
 
-FRequestCluster::FVertexData::FVertexData(ESkipDependenciesType, FPackageData& InPackageData, FRequestCluster& Cluster)
+void FRequestCluster::FGraphSearch::FreeQueryData(FVertexQueryData* QueryData)
 {
-	PackageName = InPackageData.GetPackageName();
-	PackageData = &InPackageData;
-	bInitialRequest = true;
-	bCookable = Cluster.IsRequestCookable(PackageData->GetPackageName(), PackageData, SuppressCookReason);
-	bExploreDependencies = false;
-	bExploreSoftDependencies = false;
+	QueryData->Reset();
+	VertexQueryAllocator.Free(QueryData);
 }
 
-void FRequestCluster::FVertexData::Reset()
+FRequestCluster::FVertexData&
+FRequestCluster::FGraphSearch::FindOrAddVertex(FName PackageName)
 {
-	PackageName = NAME_None;
-	for (UE::TargetDomain::FCookAttachments& PlatformAttachments : CookAttachments)
+	// Only called from PumpExploration thread
+	FVertexData*& ExistingVertex = Vertices.FindOrAdd(PackageName);
+	if (ExistingVertex)
+	{
+		return *ExistingVertex;
+	}
+
+	ExistingVertex = AllocateVertex();
+	TStringBuilder<256> NameBuffer;
+	PackageName.ToString(NameBuffer);
+	ExistingVertex->PackageData = nullptr;
+	if (!FPackageName::IsScriptPackage(NameBuffer))
 	{
-		PlatformAttachments.Empty();
+		ExistingVertex->PackageData = Cluster.COTFS.PackageDatas->TryAddPackageDataByPackageName(PackageName);
 	}
-	PackageData = nullptr;
-	bInitialRequest = false;
-	bCookable = false;
-	bExploreDependencies = false;
-	bExploreSoftDependencies = false;
-	SuppressCookReason = ESuppressCookReason::InvalidSuppressCookReason;
+	return *ExistingVertex;
 }
 
-TUniquePtr<FRequestCluster::FVertexData> FRequestCluster::FGraphSearch::AllocateVertex()
+FRequestCluster::FVertexData&
+FRequestCluster::FGraphSearch::FindOrAddVertex(FName PackageName, FPackageData& PackageData)
 {
-	// Only called from Process thread
-	TUniquePtr<FVertexData> Result;
-	Result = VertexAllocationPool.Num()
-		? VertexAllocationPool.Pop(false /* bAllowShrinking */)
-		: TUniquePtr<FVertexData>(new FVertexData(FVertexData::EAsync, GetNumPlatforms()));
-	// Vertices are Reset when constructed or when returned to pool, so we do not need to reset here
-	return Result;
+	// Only called from PumpExploration thread
+	FVertexData*& ExistingVertex = Vertices.FindOrAdd(PackageName);
+	if (ExistingVertex)
+	{
+		check(ExistingVertex->PackageData == &PackageData);
+		return *ExistingVertex;
+	}
+
+	ExistingVertex = AllocateVertex();
+	ExistingVertex->PackageData = &PackageData;
+	return *ExistingVertex;
 }
 
-void FRequestCluster::FGraphSearch::FreeVertex(TUniquePtr<FVertexData>&& Vertex)
+void FRequestCluster::FGraphSearch::AddToFrontier(FVertexData& Vertex)
 {
-	// Only called from Process thread
-	Vertex->Reset();
-	VertexAllocationPool.Add(MoveTemp(Vertex));
+	if (Vertex.PackageData)
+	{
+		Cluster.PullIntoCluster(*Vertex.PackageData);
+	}
+	Frontier.Add(&Vertex);
 }
 
-FPackageData* FRequestCluster::FGraphSearch::FindOrAddVertex(FName PackageName, FPackageData* PackageData,
-	bool bInitialRequest, bool bHardDependency, const FInstigator& InInstigator, TUniquePtr<FVertexData>& OutNewVertex)
+void FRequestCluster::FGraphSearch::CreateAvailableBatches(bool bAllowIncompleteBatch)
 {
-	// Only called from Process thread
-	OutNewVertex.Reset();
-	FVisitStatus& VisitStatus = Visited.FindOrAdd(PackageName);
-	if (VisitStatus.bVisited)
+	constexpr int32 BatchSize = 1000;
+	if (PreAsyncQueue.IsEmpty() || (!bAllowIncompleteBatch && PreAsyncQueue.Num() < BatchSize))
 	{
-		return VisitStatus.PackageData;
+		return;
 	}
-	VisitStatus.bVisited = true;
 
-	ESuppressCookReason SuppressCookReason;
-	bool bCookable = Cluster.IsRequestCookable(PackageName, PackageData, SuppressCookReason);
-	if (!bInitialRequest)
+	TArray<FQueryVertexBatch*> NewBatches;
+	NewBatches.Reserve((PreAsyncQueue.Num() + BatchSize - 1) / BatchSize);
 	{
-		if (!PackageData)
+		FScopeLock ScopeLock(&Lock);
+		while (PreAsyncQueue.Num() >= BatchSize)
 		{
-			return nullptr;
+			NewBatches.Add(CreateBatchOfPoppedVertices(BatchSize));
 		}
-		// Uncookable packages that are hard referenced by requested packages need to be explored even though they will not
-		// be cooked, because both they and any hard references they in turn have will be loaded when the original referencer
-		// is loaded and we need to precache those loads for performance and avoid marking them as discovered packages.
-		// This exploration unfortunately causes the inclusion in the cook of any cookable hard references of the uncookable
-		// package; those references are not needed but we end up including them because we do not track that they are referenced
-		// only by uncookable packages. (ONLY_EDITORONLY_TODO: Exclude these packages from the cook.)
-		// We can at least avoid cooking their soft references, since those do not automatically get loaded.
-		if (!bHardDependency && !bCookable)
-		{
-			// Mark this uncookable SoftDependency package as "explored" (a noop) and as uncookable.
-			MarkExplored(*PackageData, bCookable);
-			return nullptr;
-		}
-		if (!Cluster.TryTakeOwnership(*PackageData, false /* bUrgent */, FCompletionCallback(), InInstigator))
+		if (PreAsyncQueue.Num() > 0 && bAllowIncompleteBatch)
 		{
-			return nullptr;
+			NewBatches.Add(CreateBatchOfPoppedVertices(PreAsyncQueue.Num()));
 		}
-		Cluster.OwnedPackageDatas.Add(PackageData);
 	}
-
-	VisitStatus.PackageData = PackageData;
-	OutNewVertex = AllocateVertex();
-	OutNewVertex->PackageName = PackageName;
-	OutNewVertex->PackageData = PackageData;
-	OutNewVertex->bInitialRequest = bInitialRequest;
-	OutNewVertex->bCookable = bCookable;
-	OutNewVertex->SuppressCookReason = SuppressCookReason;
-	OutNewVertex->bExploreDependencies = bCookable || bHardDependency;
-	OutNewVertex->bExploreSoftDependencies = Cluster.bAllowSoftDependencies && bCookable;
-
-	return PackageData;
-}
-
-TArray<FRequestCluster::FQueryVertexBatch*> FRequestCluster::FGraphSearch::CreateAvailableBatches()
-{
-	// Called from inside this->Lock
-	TArray<FQueryVertexBatch*> Results;
-	constexpr int32 BatchSize = 1000;
-	Results.Reserve((VerticesToRead.Num() + BatchSize - 1) / BatchSize);
-	while (VerticesToRead.Num() >= BatchSize)
-	{
-		Results.Add(CreateBatchOfPoppedVertices(BatchSize));
-	}
-	if (Batches.Num() == 0 && !bHasPendingVertices && VerticesToRead.Num() > 0)
+	for (FQueryVertexBatch* NewBatch : NewBatches)
 	{
-		// No more vertices are coming into VerticesToRead; send the batch that we have
-		Results.Add(CreateBatchOfPoppedVertices(VerticesToRead.Num()));
+		NewBatch->Send();
 	}
-	return Results;
 }
 
-FRequestCluster::FHasPendingVerticesScope::FHasPendingVerticesScope(FGraphSearch& InGraphSearch)
-	:GraphSearch(InGraphSearch)
-{
-	checkf(!GraphSearch.bHasPendingVertices, TEXT("Nested HasPendingVerticesScope is not yet supported"));
-	GraphSearch.bHasPendingVertices = true;
-}
-
-FRequestCluster::FHasPendingVerticesScope::~FHasPendingVerticesScope()
-{
-	checkf(GraphSearch.bHasPendingVertices, TEXT("bHasPendingVertices was corrupted during FHasPendingVerticesScope"));
-	GraphSearch.bHasPendingVertices = false;
-	TArray<TUniquePtr<FVertexData>> Empty;
-	// Call AddVertices to check whether a batch needs to be sent because no further vertices are pending
-	GraphSearch.AddVertices(MoveTemp(Empty));
-}
-
-TUniquePtr<FRequestCluster::FQueryVertexBatch> FRequestCluster::FGraphSearch::AllocateBatch()
+FRequestCluster::FQueryVertexBatch* FRequestCluster::FGraphSearch::AllocateBatch()
 {
 	// Called from inside this->Lock
-	TUniquePtr<FQueryVertexBatch> Result;
-	Result = BatchAllocationPool.Num()
-		? BatchAllocationPool.Pop(false /* bAllowShrinking */)
-		: TUniquePtr<FQueryVertexBatch>(new FQueryVertexBatch(*this));
-	// Batches are Reset when constructed or when returned to pool, so we do not need to Reset here
-	return Result;
+	// BatchAllocator uses DeferredDestruction, so this might be a resused Batch, but we don't need to Reset it during
+	// allocation because Batches are Reset during Free.
+	return BatchAllocator.NewElement(*this);
 }
 
-void FRequestCluster::FGraphSearch::FreeBatch(TUniquePtr<FQueryVertexBatch>&& Batch)
+void FRequestCluster::FGraphSearch::FreeBatch(FQueryVertexBatch* Batch)
 {
 	// Called from inside this->Lock
 	Batch->Reset();
-	BatchAllocationPool.Add(MoveTemp(Batch));
+	BatchAllocator.Free(Batch);
 }
 
 FRequestCluster::FQueryVertexBatch* FRequestCluster::FGraphSearch::CreateBatchOfPoppedVertices(int32 BatchSize)
 {
 	// Called from inside this->Lock
-	check(BatchSize <= VerticesToRead.Num());
-	TUniquePtr<FQueryVertexBatch> BatchData = AllocateBatch();
+	check(BatchSize <= PreAsyncQueue.Num());
+	FQueryVertexBatch* BatchData = AllocateBatch();
 	BatchData->Vertices.Reserve(BatchSize);
 	for (int32 BatchIndex = 0; BatchIndex < BatchSize; ++BatchIndex)
 	{
-		TUniquePtr<FVertexData> VertexData = VerticesToRead.PopFrontValue();
-		TUniquePtr<FVertexData>& ExistingVert = BatchData->Vertices.FindOrAdd(VertexData->PackageName);
-		check(!ExistingVert); // We should not have any duplicate names in VerticesToRead
-		ExistingVert = MoveTemp(VertexData);
+		FVertexData* Vertex = PreAsyncQueue.PopFrontValue();
+		FVertexData*& ExistingVert = BatchData->Vertices.FindOrAdd(Vertex->QueryData->PackageName);
+		check(!ExistingVert); // We should not have any duplicate names in PreAsyncQueue
+		ExistingVert = Vertex;
 	}
-	TUniquePtr<FQueryVertexBatch>& ExistingBatch = Batches.FindOrAdd(BatchData.Get());
-	check(!ExistingBatch);
-	ExistingBatch = MoveTemp(BatchData);
-	return ExistingBatch.Get();
+	AsyncQueueBatches.Add(BatchData);
+	return BatchData;
 }
 
 void FRequestCluster::FGraphSearch::OnBatchCompleted(FQueryVertexBatch* Batch)
 {
 	FScopeLock ScopeLock(&Lock);
-	FreeBatch(Batches.FindAndRemoveChecked(Batch));
-	PollReadyEvent->Trigger();
+	AsyncQueueBatches.Remove(Batch);
+	FreeBatch(Batch);
+	AsyncResultsReadyEvent->Trigger();
 }
 
 void FRequestCluster::FGraphSearch::OnVertexCompleted()
 {
-	// The trigger occurs outside of the lock, and might get clobbered and incorrectly ignored by a call to Poll which
-	// consumes the vertices before our caller added a vertex but calls PollReadyEvent->Reset after this PollReadyEvent->Trigger.
-	// This clobbering will not cause a deadlock, because eventually DestroyBatch will be called which triggers inside
-	// the lock. Doing the per-vertex trigger outside the lock is good for performance.
-	PollReadyEvent->Trigger();
+	// The trigger occurs outside of the lock, and might get clobbered and incorrectly ignored by a call from the
+	// consumer thread if the consumer tried to consume and found the vertices empty before our caller added a vertex
+	// but then pauses and calls AsyncResultsReadyEvent->Reset after this AsyncResultsReadyEvent->Trigger.
+	// This clobbering will not cause a deadlock, because eventually DestroyBatch will be called which triggers it
+	// inside the lock. Doing the per-vertex trigger outside the lock is good for performance.
+	AsyncResultsReadyEvent->Trigger();
 }
 
 FRequestCluster::FQueryVertexBatch::FQueryVertexBatch(FGraphSearch& InGraphSearch)
 	: ThreadSafeOnlyVars(InGraphSearch)
-	, NumPlatforms(InGraphSearch.GetNumPlatforms())
 {
-	Reset();
+	PlatformDatas.SetNum(InGraphSearch.FetchPlatforms.Num());
 }
 
 void FRequestCluster::FQueryVertexBatch::Reset()
 {
-	Scratch.PackageNames.Reset();
+	for (FPlatformData& PlatformData : PlatformDatas)
+	{
+		PlatformData.PackageNames.Reset();
+	}
 	Vertices.Reset();
 }
 
 void FRequestCluster::FQueryVertexBatch::Send()
 {
-	Scratch.PackageNames.Reserve(Vertices.Num());
-	for (const TPair<FName, TUniquePtr<FVertexData>>& Pair : Vertices)
+	for (const TPair<FName, FVertexData*>& Pair : Vertices)
 	{
-		Scratch.PackageNames.Add(Pair.Key);
-		Pair.Value->PendingPlatforms.store(NumPlatforms, std::memory_order_release);
+		FVertexData* Vertex = Pair.Value;
+		TArray<FQueryPlatformData>& QueryPlatforms = Vertex->QueryData->Platforms;
+		bool bAtLeastOnePlatform = false;
+		for (int32 PlatformIndex = 0; PlatformIndex < PlatformDatas.Num(); ++PlatformIndex)
+		{
+			if (QueryPlatforms[PlatformIndex].bActive)
+			{
+				PlatformDatas[PlatformIndex].PackageNames.Add(Pair.Key);
+			}
+			bAtLeastOnePlatform = true;
+		}
+		// We only check for the vertex's completion when the vertex receives a callback from the completion of a
+		// platform. Therefore we do not support Vertices in the batch that have no platforms.
+		check(bAtLeastOnePlatform);
 	}
 	PendingVertices.store(Vertices.Num(), std::memory_order_release);
 
-	for (int32 PlatformIndex = 0; PlatformIndex < NumPlatforms; ++PlatformIndex)
+	for (int32 PlatformIndex = 0; PlatformIndex < PlatformDatas.Num(); ++PlatformIndex)
 	{
-		const ITargetPlatform* TargetPlatform = nullptr;
-		ICookedPackageWriter* PackageWriter = nullptr;
-		// Platform 0 is the platform-agnostic platform
-		if (PlatformIndex > 0)
+		FPlatformData& PlatformData = PlatformDatas[PlatformIndex];
+		if (PlatformData.PackageNames.IsEmpty())
 		{
-			TargetPlatform = ThreadSafeOnlyVars.Cluster.Platforms[PlatformIndex - 1];
-			PackageWriter = ThreadSafeOnlyVars.Cluster.PackageWriters[PlatformIndex - 1];
+			continue;
 		}
+		FFetchPlatformData& FetchPlatformData = ThreadSafeOnlyVars.FetchPlatforms[PlatformIndex];
 
-		TUniqueFunction<void(FName PackageName, UE::TargetDomain::FCookAttachments&& Result)> Callback =
-			[this, PlatformIndex](FName PackageName, UE::TargetDomain::FCookAttachments&& Attachments)
+		if (ThreadSafeOnlyVars.bCookAttachmentsEnabled // Only FetchCookAttachments if our cookmode supports it. Otherwise keep them all empty
+			&& !FetchPlatformData.bIsCookerLoadingPlatform // The CookerLoadingPlatform has no stored CookAttachments; always use empty
+			)
 		{
-			RecordCacheResults(PackageName, PlatformIndex, MoveTemp(Attachments));
-		};
-		if (ThreadSafeOnlyVars.bCookAttachmentsEnabled)
-		{
-			UE::TargetDomain::FetchCookAttachments(Scratch.PackageNames, TargetPlatform,
-				PackageWriter, MoveTemp(Callback));
+			TUniqueFunction<void(FName PackageName, UE::TargetDomain::FCookAttachments&& Result)> Callback =
+				[this, PlatformIndex](FName PackageName, UE::TargetDomain::FCookAttachments&& Attachments)
+			{
+				RecordCacheResults(PackageName, PlatformIndex, MoveTemp(Attachments));
+			};
+			UE::TargetDomain::FetchCookAttachments(PlatformData.PackageNames, FetchPlatformData.Platform,
+				FetchPlatformData.Writer, MoveTemp(Callback));
 		}
 		else
 		{
 			// When we do not need to asynchronously fetch, we record empty cache results from an AsyncTask.
 			// Using an AsyncTask keeps the threading flow similar to the FetchCookAttachments case
 			AsyncTask(ENamedThreads::AnyThread,
-				[PackageNames = Scratch.PackageNames, Callback = MoveTemp(Callback)]()
+				[this, PlatformIndex]()
 			{
-				for (FName PackageName : PackageNames)
+				FPlatformData& PlatformData = PlatformDatas[PlatformIndex];
+				// Don't use a ranged-for, as we are not allowed to access this or this->PackageNames after the
+				// last index, and ranged-for != at the end of the final loop iteration can read from PackageNames
+				int32 NumPackageNames = PlatformData.PackageNames.Num();
+				FName* PackageNamesData = PlatformData.PackageNames.GetData();
+				for (int32 PackageNameIndex = 0; PackageNameIndex < NumPackageNames; ++PackageNameIndex)
 				{
+					FName PackageName = PackageNamesData[PackageNameIndex];
 					UE::TargetDomain::FCookAttachments Attachments;
-					Callback(PackageName, MoveTemp(Attachments));
+					RecordCacheResults(PackageName, PlatformIndex, MoveTemp(Attachments));
 				}
 			});
 		}
@@ -1109,12 +1392,13 @@ void FRequestCluster::FQueryVertexBatch::Send()
 void FRequestCluster::FQueryVertexBatch::RecordCacheResults(FName PackageName, int32 PlatformIndex,
 	UE::TargetDomain::FCookAttachments&& CookAttachments)
 {
-	TUniquePtr<FVertexData>& VertexData = Vertices.FindChecked(PackageName);
-	check(VertexData);
-	VertexData->CookAttachments[PlatformIndex] = MoveTemp(CookAttachments);
-	if (VertexData->PendingPlatforms.fetch_sub(1, std::memory_order_acq_rel) == 1)
+	FVertexData* Vertex = Vertices.FindChecked(PackageName);
+	check(Vertex->QueryData);
+	FVertexQueryData& QueryData = *Vertex->QueryData;
+	QueryData.Platforms[PlatformIndex].CookAttachments = MoveTemp(CookAttachments);
+	if (QueryData.PendingPlatforms.fetch_sub(1, std::memory_order_acq_rel) == 1)
 	{
-		ThreadSafeOnlyVars.CompletedVertices.Enqueue(MoveTemp(VertexData));
+		ThreadSafeOnlyVars.AsyncQueueResults.Enqueue(Vertex);
 		bool bBatchComplete = PendingVertices.fetch_sub(1, std::memory_order_relaxed) == 1;
 		if (!bBatchComplete)
 		{
@@ -1123,38 +1407,33 @@ void FRequestCluster::FQueryVertexBatch::RecordCacheResults(FName PackageName, i
 		else
 		{
 			ThreadSafeOnlyVars.OnBatchCompleted(this);
-			// this is no longer accessible
+			// *this is no longer accessible
 		}
 	}
 }
 
-int32 FRequestCluster::FGraphSearch::GetNumPlatforms() const
-{
-	return Cluster.Platforms.Num() + 1;
-}
-
-TArray<FPackageData*>& FRequestCluster::FGraphSearch::GetTransitiveRequests()
-{
-	return TransitiveRequests;
-}
-
 TMap<FPackageData*, TArray<FPackageData*>>& FRequestCluster::FGraphSearch::GetGraphEdges()
 {
 	return GraphEdges;
 }
 
-bool FRequestCluster::IsRequestCookable(FName PackageName, FPackageData*& InOutPackageData,
-	ESuppressCookReason& OutReason)
+void FRequestCluster::IsRequestCookable(const ITargetPlatform* Platform, FName PackageName, FPackageData& PackageData,
+	ESuppressCookReason& OutReason, bool& bOutCookable, bool& bOutExplorable)
 {
-	return IsRequestCookable(PackageName, InOutPackageData, PackageDatas, PackageTracker,
-		DLCPath, bErrorOnEngineContentUse, bAllowUncookedAssetReferences, GetPlatforms(), OutReason);
+	return IsRequestCookable(Platform, PackageName, PackageData, PackageDatas, PackageTracker,
+		DLCPath, bErrorOnEngineContentUse, bAllowUncookedAssetReferences,
+		COTFS.bCanSkipEditorReferencedPackagesWhenCooking,
+		OutReason, bOutCookable, bOutExplorable);
 }
 
-bool FRequestCluster::IsRequestCookable(FName PackageName, FPackageData*& InOutPackageData,
+void FRequestCluster::IsRequestCookable(const ITargetPlatform* Platform, FName PackageName, FPackageData& PackageData,
 	FPackageDatas& InPackageDatas, FPackageTracker& InPackageTracker,
 	FStringView InDLCPath, bool bInErrorOnEngineContentUse, bool bInAllowUncookedAssetReferences,
-	TConstArrayView<const ITargetPlatform*> RequestPlatforms, ESuppressCookReason& OutReason)
+	bool bCanSkipEditorReferencedPackagesWhenCooking,
+	ESuppressCookReason& OutReason, bool& bOutCookable, bool& bOutExplorable)
 {
+	check(Platform != CookerLoadingPlatformKey); // IsRequestCookable should not be called for The CookerLoadingPlatform; it has different rules
+
 	TStringBuilder<256> NameBuffer;
 	// We need to reject packagenames from adding themselves or their transitive dependencies using all the same rules that
 	// UCookOnTheFlyServer::ProcessRequest uses. Packages that are rejected from cook do not add their dependencies to the cook.
@@ -1162,35 +1441,50 @@ bool FRequestCluster::IsRequestCookable(FName PackageName, FPackageData*& InOutP
 	if (FPackageName::IsScriptPackage(NameBuffer))
 	{
 		OutReason = ESuppressCookReason::ScriptPackage;
-		return false;
+		bOutCookable = false;
+		bOutExplorable = false;
+		return;
 	}
 
-	if (!InOutPackageData)
+	FName FileName = PackageData.GetFileName();
+	if (InPackageTracker.NeverCookPackageList.Contains(FileName))
 	{
-		InOutPackageData = InPackageDatas.TryAddPackageDataByPackageName(PackageName);
-		if (!InOutPackageData)
+		if (INDEX_NONE != UE::String::FindFirst(NameBuffer, ULevel::GetExternalActorsFolderName(), ESearchCase::IgnoreCase))
 		{
-			// Package does not exist on disk
-			OutReason = ESuppressCookReason::DoesNotExistInWorkspaceDomain;
-			return false;
+			// EXTERNALACTOR_TODO: Add a separate category for ExternalActors rather than putting them in
+			// NeverCookPackageList and checking naming convention here.
+			OutReason = ESuppressCookReason::NeverCook;
+			bOutCookable = false;
+			bOutExplorable = true;
+			if (bCanSkipEditorReferencedPackagesWhenCooking)
+			{
+				bOutExplorable = true;
+			}
+			else
+			{
+				// ONLYEDITORONLY_TODO
+				// Workaround to avoid changing behavior from legacy graph search. In legacy graph search ExternalActors
+				// were uncookable and unexplorable, and so packages referenced through externalactors referenced by e.g. GameFeatureData
+				// would not be cooked.
+				bOutExplorable = false;
+			}
 		}
+		else
+		{
+			UE_LOG(LogCook, Verbose, TEXT("Package %s is referenced but is in the never cook package list, discarding request"), *NameBuffer);
+			OutReason = ESuppressCookReason::NeverCook;
+			bOutCookable = false;
+			bOutExplorable = false;
+		}
+		return;
 	}
 
-	FName FileName = InOutPackageData->GetFileName();
-	if (InPackageTracker.NeverCookPackageList.Contains(FileName))
-	{
-		UE_LOG(LogCook, Verbose, TEXT("Package %s is referenced but is in the never cook package list, discarding request"), *NameBuffer);
-		OutReason = ESuppressCookReason::NeverCook;
-		return false;
-	}
-
-
 	if (bInErrorOnEngineContentUse && !InDLCPath.IsEmpty())
 	{
 		FileName.ToString(NameBuffer);
 		if (!FStringView(NameBuffer).StartsWith(InDLCPath))
 		{
-			if (!InOutPackageData->HasAllCookedPlatforms(RequestPlatforms, true /* bIncludeFailed */))
+			if (!PackageData.HasCookedPlatform(Platform, true /* bIncludeFailed */))
 			{
 				// AllowUncookedAssetReferences should only be used when the DLC plugin to cook is going to be mounted where uncooked packages are available.
 				// This will allow a DLC plugin to be recooked continually and mounted in an uncooked editor which is useful for CI.
@@ -1200,12 +1494,15 @@ bool FRequestCluster::IsRequestCookable(FName PackageName, FPackageData*& InOutP
 				}
 			}
 			OutReason = ESuppressCookReason::NotInCurrentPlugin;
-			return false;
+			bOutCookable = false;
+			bOutExplorable = false;
+			return;
 		}
 	}
 
-	OutReason = ESuppressCookReason::InvalidSuppressCookReason;
-	return true;
+	OutReason = ESuppressCookReason::NotSuppressed;
+	bOutCookable = true;
+	bOutExplorable = true;
 }
 
 TConstArrayView<FName> FRequestCluster::GetLocalizationReferences(FName PackageName, UCookOnTheFlyServer& InCOTFS)
diff --git a/Engine/Source/Editor/UnrealEd/Private/Cooker/CookRequestCluster.h b/Engine/Source/Editor/UnrealEd/Private/Cooker/CookRequestCluster.h
--- a/Engine/Source/Editor/UnrealEd/Private/Cooker/CookRequestCluster.h
+++ b/Engine/Source/Editor/UnrealEd/Private/Cooker/CookRequestCluster.h
@@ -2,6 +2,7 @@
 
 #pragma once
 
+#include "Containers/BitArray.h"
 #include "Containers/ContainerAllocationPolicies.h"
 #include "Containers/ContainersFwd.h"
 #include "Containers/Set.h"
@@ -16,6 +17,7 @@
 #include "Containers/MpscQueue.h"
 #include "Containers/RingBuffer.h"
 #include "Cooker/CookTypes.h"
+#include "Cooker/TypedBlockAllocator.h"
 #include "HAL/CriticalSection.h"
 #include "TargetDomain/TargetDomainUtils.h"
 #include "Templates/UniquePtr.h"
@@ -28,13 +30,15 @@ class ICookedPackageWriter;
 class ITargetPlatform;
 class UCookOnTheFlyServer;
 namespace UE::Cook { class FRequestQueue; }
+namespace UE::Cook { struct FDiscoveryQueueElement; }
 namespace UE::Cook { struct FFilePlatformRequest; }
+namespace UE::Cook { struct FPackageData; }
 namespace UE::Cook { struct FPackageDatas; }
+namespace UE::Cook { struct FPackagePlatformData; }
 namespace UE::Cook { struct FPackageTracker; }
 
 namespace UE::Cook
 {
-struct FPackageData;
 
 /**
  * A group of external requests sent to CookOnTheFlyServer's tick loop. Transitive dependencies are found and all of the
@@ -43,18 +47,10 @@ struct FPackageData;
 class FRequestCluster
 {
 public:
-	/** The cluster's data about a package from the external requests or discovered dependency. */
-	struct FFileNameRequest
-	{
-		FName FileName;
-		FInstigator Instigator;
-		FCompletionCallback CompletionCallback;
-		bool bUrgent;
-
-		FFileNameRequest(FFilePlatformRequest&& FileRequest, bool bInUrgent);
-	};
-	explicit FRequestCluster(UCookOnTheFlyServer& COTFS, TArray<const ITargetPlatform*>&& Platforms);
-	explicit FRequestCluster(UCookOnTheFlyServer& COTFS, TConstArrayView<const ITargetPlatform*> Platforms);
+	FRequestCluster(UCookOnTheFlyServer& COTFS, TArray<FFilePlatformRequest>&& InRequests,
+		bool bInExternalRequestsAreUrgent);
+	FRequestCluster(UCookOnTheFlyServer& COTFS, FPackageDataSet&& InRequests);
+	FRequestCluster(UCookOnTheFlyServer& COTFS, TRingBuffer<FDiscoveryQueueElement>& DiscoveryQueue);
 	FRequestCluster(FRequestCluster&&) = default;
 
 	/**
@@ -63,11 +59,15 @@ public:
 	 */
 	void Process(const FCookerTimer& CookerTimer, bool& bOutComplete);
 
-	TConstArrayView<const ITargetPlatform*> GetPlatforms() const { return Platforms; }
 	/** PackageData container interface: return the number of PackageDatas owned by this container. */
 	int32 NumPackageDatas() const;
 	/** PackageData container interface: remove the PackageData from this container. */
 	void RemovePackageData(FPackageData* PackageData);
+	void OnNewReachablePlatforms(FPackageData* PackageData);
+	void OnPlatformAddedToSession(const ITargetPlatform* TargetPlatform);
+	void OnRemoveSessionPlatform(const ITargetPlatform* TargetPlatform);
+	void RemapTargetPlatforms(TMap<ITargetPlatform*, ITargetPlatform*>& Remap);
+
 	/** PackageData container interface: whether the PackageData is owned by this container. */
 	bool Contains(FPackageData* PackageData) const;
 	/**
@@ -80,14 +80,6 @@ public:
 		TArray<TPair<FPackageData*, ESuppressCookReason>>& OutRequestsToDemote,
 		TMap<FPackageData*, TArray<FPackageData*>>& OutRequestGraph);
 
-	/**
-	 * Create clusters(s) for all the given name or packagedata requests and append them to OutClusters.
-	 * Multiple clusters are necessary if the list of platforms differs between requests.
-	 */
-	static void AddClusters(UCookOnTheFlyServer& COTFS, TArray<FFilePlatformRequest>&& InRequests,
-		bool bRequestsAreUrgent, TRingBuffer<FRequestCluster>& OutClusters);
-	static void AddClusters(UCookOnTheFlyServer& COTFS, FPackageDataSet& UnclusteredRequests,
-		TRingBuffer<FRequestCluster>& OutClusters, FRequestQueue& QueueForReadyRequests);
 	static TConstArrayView<FName> GetLocalizationReferences(FName PackageName, UCookOnTheFlyServer& InCOTFS);
 
 private:
@@ -100,52 +92,40 @@ private:
 		bool bVisited = false;
 	};
 
-	/**
-	 * GraphSearch data for an in-progress package. VertexData is created when a package is discovered
-	 * from the dependencies of a referencer package. It is in-progress while its own dependencies
-	 * are looked up asynchronously. The VertexData is freed after those dependencies have been logged
-	 * and cook flags have been set on the FPackageData for the package.
-	 */
-	struct FVertexData
+	/** Per-platform data in an active query for a vertex's dependencies/previous incremental results. */
+	struct FQueryPlatformData
 	{
-		enum ESkipDependenciesType
-		{
-			ESkipDependencies
-		};
-		enum EAsyncType
-		{
-			EAsync
-		};
-		FVertexData(EAsyncType, int32 NumPlatforms);
-		FVertexData(ESkipDependenciesType, FPackageData& PackageData, FRequestCluster& Cluster);
-		void Reset();
-
 		/** Data looked up about the package's dependencies from the PackageWriter's previous cook of the package */
-		TArray<UE::TargetDomain::FCookAttachments> CookAttachments;
-		/** Copy of PackageData->GetPackageName() */
-		FName PackageName;
+		UE::TargetDomain::FCookAttachments CookAttachments;
 		/**
-		 * Not dereferenceable since it may be removed while in use on async threads.
-		 * Used as an identifier when the vertex is returned to the process thread.
+		 * Platforms present in the FVertexQueryData are recorded as a an array of all possible platforms, with
+		 * bActive=true for the ones that are present.
 		 */
-		UE::Cook::FPackageData* PackageData;
+		bool bActive = false;
+	};
+
+	/** Input/output data for an active query for a vertex's dependencies/previous incrementa results. */
+	struct FVertexQueryData
+	{
+		void Reset();
+
+		FName PackageName;
+		/** Settings and Results for each of the GraphSearch's FetchPlatforms. Element n corresponds to FetchPlatform n. */
+		TArray<FQueryPlatformData> Platforms;
 		/** Number of asynchronous results from PackageWriter still being waited on. */
 		std::atomic<uint32> PendingPlatforms;
-		/** The reason it is not cookable if !bCookable, or ESuppressCookReason::InvalidSuppressCookReason if bCookable. */
-		ESuppressCookReason SuppressCookReason;
-		/**
-		 * Written/Read when VertexData is allocated for a package name.
-		 * It is Set to false if the packagename should not be part of the graph search.
-		 */
-		bool bValid : 1;
-		/** True if the package is from the Cluster's initial requests. Initial requests are always visited. */
-		bool bInitialRequest : 1;
-		/** True if the package is cookable. Non-cookable packages are demoted back to idle. */
-		bool bCookable : 1;
-		/** True if any dependencies should be explored. */
-		bool bExploreDependencies : 1;
-		/** True if soft dependencies should be explored. */
-		bool bExploreSoftDependencies : 1;
+	};
+
+	/**
+	 * GraphSearch data for an package referenced by the cluster. VertexData is created when a package is discovered
+	 * from the dependencies of a referencer package. It remains allocated for the rest of the Cluster's lifetime.
+	 */
+	struct FVertexData
+	{
+		UE::Cook::FPackageData* PackageData = nullptr;
+		/** Non-null if there as a query active for the vertex's dependencies/previous incremental results. */
+		FVertexQueryData* QueryData = nullptr;
+		bool bAnyCookable = true;
 	};
 
 	/**
@@ -164,39 +144,40 @@ private:
 		void RecordCacheResults(FName PackageName, int32 PlatformIndex,
 			UE::TargetDomain::FCookAttachments&& CookAttachments);
 
-		struct FScratch
+		struct FPlatformData
 		{
 			TArray<FName> PackageNames;
 		};
 
-		/** Scratch-space variables that are reused to avoid allocations. */
-		FScratch Scratch;
+		TArray<FPlatformData> PlatformDatas;
 		/**
-		 * Map of the requested vertices by name. The keys in the map are created during Send and are
-		 * read-only afterwards (so the map is multithread-readable). The value for a given 
-		 * key is deleted when the results for the vertex are received and the FVertexData is
-		 * moved to CompletedVertices on the FGraphSearch.
+		 * Map of the requested vertices by name. The map is created during Send and is
+		 * read-only afterwards (so the map is multithread-readable). The Vertices pointed to have their own
+		 * rules for what is accessible from the async work threads.
 		 * */
-		TMap<FName, TUniquePtr<FVertexData>> Vertices;
+		TMap<FName, FVertexData*> Vertices;
 		/** Accessor for the GraphSearch; only thread-safe functions and variables should be accessed. */
 		FGraphSearch& ThreadSafeOnlyVars;
-		/** Copy of NumPlatforms from GraphSearch; read-only */
-		int32 NumPlatforms;
 		/** The number of vertices that are still awaiting results. This batch is closed when PendingVertices == 0. */
 		std::atomic<uint32> PendingVertices;
 	};
 
-	/** A scope to notify GraphSearch that more vertices are coming and it should wait for a full batch to send. */
-	struct FHasPendingVerticesScope
+	/** Platform information that is constant (usually, some events can change it) during the cluster's lifetime. */
+	struct FFetchPlatformData
 	{
-		FHasPendingVerticesScope(FGraphSearch& InGraphSearch);
-		~FHasPendingVerticesScope();
-
-		FGraphSearch& GraphSearch;
+		const ITargetPlatform* Platform = nullptr;
+		ICookedPackageWriter* Writer = nullptr;
+		bool bIsPlatformAgnosticPlatform = false;
+		bool bIsCookerLoadingPlatform = false;
 	};
+	// Platforms are listed in various arrays, always in the same order. Some special case entries exist and are added
+	// at specified indices in the arrays.
+	static constexpr int32 PlatformAgnosticPlatformIndex = 0;
+	static constexpr int32 CookerLoadingPlatformIndex = 1;
+	static constexpr int32 FirstSessionPlatformIndex = 2;
 
 	/**
-	 * Variables and functions that are only used during FetchDependencies. FetchDependencies executes a graph search
+	 * Variables and functions that are only used during PumpExploration. PumpExploration executes a graph search
 	 * over the graph of packages (vertices) and their hard/soft dependencies upon other packages (edges). 
 	 * Finding the dependencies for each package uses previous cook results and is executed asynchronously.
 	 * After the graph is searched, packages are sorted topologically from leaf to root, so that packages are
@@ -210,34 +191,21 @@ private:
 		~FGraphSearch();
 
 		// All public functions are callable only from the process thread
+		/** Skip the entire GraphSearch and just visit the Cluster's current OwnedPackageDatas. */
+		void VisitWithoutDependencies();
+		/** Start a search from the Cluster's current OwnedPackageDatas. */
+		void StartSearch();
+		void RemovePackageData(FPackageData* PackageData);
+		void OnNewReachablePlatforms(FPackageData* PackageData);
+
 		/**
-		 * Called from the Process Thread to consume FVertexData that have finished their async work and are
-		 * ready for the processing that must happen on the process thread
-		 */
-		void Poll(TArray<TUniquePtr<FVertexData>>& OutCompletedVertices, bool& bOutIsDone);
-		/** Sleep (with timeout) until results are ready in Poll */
-		void WaitForPollAvailability(double WaitTimeSeconds);
-		/** Log diagnostic information about the search, e.g. timeout warnings. */
-		void UpdateDisplay();
-		/** AddVertices that need async processing */
-		void AddVertices(TArray<TUniquePtr<FVertexData>>&& Vertices);
-		/**
-		 * Record a vertex's cook flags and dependencies on the PackageData.
-		 * Create and send any new vertex dependencies for async work. Called from processs-thread only.
-		 */
-		void VisitVertex(const FVertexData& VertexData);
-		/**
-		 * Find the PackageData for the given PackageName if it has been visited before.
-		 * If not, construct a new vertex to be async processed and visited.
+		 * Visit newly reachable PackageDatas, queue a fetch of their dependencies, harvest new reachable PackageDatas
+		 * from the results of the fetch.
 		 */
-		FPackageData* FindOrAddVertex(FName PackageName, FPackageData* PackageData, bool bInitialRequest,
-			bool bHardDependency, const FInstigator& InInstigator, TUniquePtr<FVertexData>& OutNewVertex);
-		/** Allocate memory for a new vertex; returned vertex is not yet constructed. */
-		TUniquePtr<FVertexData> AllocateVertex();
-		/** Free an allocated vertex. */ 
-		void FreeVertex(TUniquePtr<FVertexData>&& Batch);
-		/** PackageDatas found during the graph search, including the initial requests. */
-		TArray<FPackageData*>& GetTransitiveRequests();
+		void TickExploration(bool& bOutDone);
+		/** Sleep (with timeout) until work is available in TickExploration */
+		void WaitForAsyncQueue(double WaitTimeSeconds);
+
 		/**
 		 * Edges in the dependency graph found during graph search.
 		 * Only includes PackageDatas that are part of this cluster
@@ -245,26 +213,56 @@ private:
 		TMap<FPackageData*, TArray<FPackageData*>>& GetGraphEdges();
 
 	private:
+		// Scratch data structures used to avoid dynamic allocations; lifetime of each use is only on the stack
+		struct FScratchPlatformDependencyBits
+		{
+			TBitArray<> HasPlatformByIndex;
+			bool bHardDependency = false;
+		};
 		struct FScratch
 		{
-			TArray<FName> HardDependencies;
-			TArray<FName> SoftDependencies;
-			TArray<TUniquePtr<FVertexData>> NewVertices;
+			TArray<FName> HardGameDependencies;
+			TArray<FName> SoftGameDependencies;
+			TArray<FName> CookerLoadingDependencies;
+			TMap<FName, FScratchPlatformDependencyBits> PlatformDependencyMap;
+			TSet<FName> HardDependenciesSet;
 			TSet<FName> SkippedPackages;
 		};
 		friend struct FQueryVertexBatch;
-		friend struct FHasPendingVerticesScope;
+		friend struct FVertexData;
 		
 		// Functions callable only from the Process thread
-		void MarkExplored(FPackageData& PackageData, bool bCookable);
+		/** Log diagnostic information about the search, e.g. timeout warnings. */
+		void UpdateDisplay();
+
+		/** Asynchronously fetch the dependencies and previous incremental results for a vertex */
+		void QueueEdgesFetch(FVertexData& Vertex, TConstArrayView<const ITargetPlatform*> Platforms);
+		/** Calculate and store the vertex's PackageData's cookability for each reachable platform. Kick off edges fetch. */
+		void VisitVertex(FVertexData& VertexData, bool bSkipDependencies=false);
+		/** Calculate and store the vertex's PackageData's cookability for the platform. */
+		void VisitVertexForPlatform(FVertexData& VertexData, const ITargetPlatform* Platform,
+			FPackagePlatformData& PlatformData, ESuppressCookReason& AccumulatedSuppressCookReason);
+		 /** Process the results from async edges fetch and queue the found dependencies-for-visiting. */
+		void ExploreVertexEdges(FVertexData& VertexData);
+
+		/** Find or add a Vertex for PackageName. If PackageData is provided, use it, otherwise look it up. */
+		FVertexData& FindOrAddVertex(FName PackageName);
+		FVertexData& FindOrAddVertex(FName PackageName, FPackageData& PackageData);
+		/** Batched allocation for vertices. */
+		FVertexData* AllocateVertex();
+		/** Batched allocation and free for QueryData. */
+		FVertexQueryData* AllocateQueryData();
+		void FreeQueryData(FVertexQueryData* QueryData);
+		/** Queue a vertex for visiting and dependency traversal */
+		void AddToFrontier(FVertexData& Vertex);
 
 		// Functions that must be called only within the Lock
 		/** Allocate memory for a new batch; returned batch is not yet constructed. */
-		TUniquePtr<FQueryVertexBatch> AllocateBatch();
+		FQueryVertexBatch* AllocateBatch();
 		/** Free an allocated batch. */ 
-		void FreeBatch(TUniquePtr<FQueryVertexBatch>&& Batch);
+		void FreeBatch(FQueryVertexBatch* Batch);
 		/** Pop vertices from VerticesToRead into batches, if there are enough of them. */
-		TArray<FQueryVertexBatch*> CreateAvailableBatches();
+		void CreateAvailableBatches(bool bAllowIncompleteBatch);
 		/** Pop a single batch vertices from VerticesToRead. */
 		FQueryVertexBatch* CreateBatchOfPoppedVertices(int32 BatchSize);
 
@@ -273,59 +271,58 @@ private:
 		void OnBatchCompleted(FQueryVertexBatch* Batch);
 		/** Notify process thread of vertex completion. */
 		void OnVertexCompleted();
-		/**
-		 * Get the number of platforms sent to FetchCookAttachments.
-		 * This includes +1 for null platform that fetches platform-agnostic data.
-		 **/
-		int32 GetNumPlatforms() const;
+
+		/** Total number of platforms known to the cluster, including the special cases. */
+		int32 NumFetchPlatforms() const { return FetchPlatforms.Num(); }
+		/** Total number of non-special-case platforms known to the cluster.Identical to COTFS's session platforms */
+		int32 NumSessionPlatforms() const { return FetchPlatforms.Num() - 2; }
 
 	private:
 		// Variables that are read-only during multithreading
+		TArray<FFetchPlatformData> FetchPlatforms;
 		FRequestCluster& Cluster;
 		bool bCookAttachmentsEnabled = false;
 
 		// Variables that are accessible only from the Process thread
 		/** Scratch-space variables that are reused to avoid allocations. */
 		FScratch Scratch;
-		TArray<FPackageData*> TransitiveRequests;
 		TMap<FPackageData*, TArray<FPackageData*>> GraphEdges;
-		/** Which PackageNames have already been inspected, and either discarded or a Vertex created for them. */
-		TMap<FName, FVisitStatus> Visited;
-		TArray<TUniquePtr<FVertexData>> VertexAllocationPool;
-		TRingBuffer<TUniquePtr<FVertexData>> VerticesToRead;
+		TMap<FName, FVertexData*> Vertices;
+		TSet<FVertexData*> Frontier;
+		TTypedBlockAllocatorFreeList<FVertexData> VertexAllocator;
+		TTypedBlockAllocatorResetList<FVertexQueryData> VertexQueryAllocator;
+		/** Vertices queued for async processing that are not yet numerous enough to fill a batch. */
+		TRingBuffer<FVertexData*> PreAsyncQueue;
 		/** Time-tracker for timeout warnings in Poll */
 		double LastActivityTime;
-		/** A stack-scope variable to control automatic CreateAvailableBatches behavior. */
-		bool bHasPendingVertices = false;
 
 		// Variables that are accessible from multiple threads, guarded by Lock
 		FCriticalSection Lock;
-		TArray<TUniquePtr<FQueryVertexBatch>> BatchAllocationPool;
-		TMap<FQueryVertexBatch*, TUniquePtr<FQueryVertexBatch>> Batches;
+		TTypedBlockAllocatorResetList<FQueryVertexBatch> BatchAllocator;
+		TSet<FQueryVertexBatch*> AsyncQueueBatches;
 
 		// Variables that are accessible from multiple threads, internally threadsafe
-		TMpscQueue<TUniquePtr<FVertexData>> CompletedVertices;
-		FEventRef PollReadyEvent;
+		TMpscQueue<FVertexData*> AsyncQueueResults;
+		FEventRef AsyncResultsReadyEvent;
 	};
 
-	void Initialize(UCookOnTheFlyServer& COTFS);
+private:
+	explicit FRequestCluster(UCookOnTheFlyServer& COTFS);
+	void ReserveInitialRequests(int32 RequestNum);
+	void PullIntoCluster(FPackageData& PackageData);
 	void FetchPackageNames(const FCookerTimer& CookerTimer, bool& bOutComplete);
-	void FetchDependencies(const FCookerTimer& CookerTimer, bool& bOutComplete);
+	void PumpExploration(const FCookerTimer& CookerTimer, bool& bOutComplete);
 	void StartAsync(const FCookerTimer& CookerTimer, bool& bOutComplete);
-	bool TryTakeOwnership(FPackageData& PackageData, bool bUrgent, FCompletionCallback&& CompletionCallback,
-		const FInstigator& InInstigator);
-	bool IsRequestCookable(FName PackageName, FPackageData*& InOutPackageData, ESuppressCookReason& OutReason);
-	static bool IsRequestCookable(FName PackageName, FPackageData*& InOutPackageData,
+	void IsRequestCookable(const ITargetPlatform* TargetPlatform, FName PackageName, FPackageData& PackageData,
+		ESuppressCookReason& OutReason, bool& bOutCookable, bool& bOutExplorable);
+	static void IsRequestCookable(const ITargetPlatform* TargetPlatform, FName PackageName, FPackageData& PackageData,
 		FPackageDatas& InPackageDatas, FPackageTracker& InPackageTracker,
 		FStringView InDLCPath, bool bInErrorOnEngineContentUse, bool bInAllowUncookedAssetReferences,
-		TConstArrayView<const ITargetPlatform*> RequestPlatforms, ESuppressCookReason& OutReason);
-
-	TArray<FFileNameRequest> InRequests;
-	TArray<FPackageData*> Requests;
-	TArray<TPair<FPackageData*, ESuppressCookReason>> RequestsToDemote;
-	TArray<const ITargetPlatform*> Platforms;
-	TArray<ICookedPackageWriter*> PackageWriters;
-	FPackageDataSet OwnedPackageDatas;
+		bool bCanSkipEditorReferencedPackagesWhenCooking,
+		ESuppressCookReason& OutReason, bool& bOutCookable, bool& bOutExplorable);
+
+	TArray<FFilePlatformRequest> FilePlatformRequests;
+	TFastPointerMap<FPackageData*, ESuppressCookReason> OwnedPackageDatas;
 	TMap<FPackageData*, TArray<FPackageData*>> RequestGraph;
 	FString DLCPath;
 	TUniquePtr<FGraphSearch> GraphSearch; // Needs to be dynamic-allocated because of large alignment
@@ -334,7 +331,7 @@ private:
 	IAssetRegistry& AssetRegistry;
 	FPackageTracker& PackageTracker;
 	FBuildDefinitions& BuildDefinitions;
-	int32 NextRequest = 0;
+	bool bExternalRequestsAreUrgent = false;
 	bool bAllowHardDependencies = true;
 	bool bAllowSoftDependencies = true;
 	bool bHybridIterativeEnabled = true;
diff --git a/Engine/Source/Editor/UnrealEd/Private/Cooker/CookTypes.cpp b/Engine/Source/Editor/UnrealEd/Private/Cooker/CookTypes.cpp
--- a/Engine/Source/Editor/UnrealEd/Private/Cooker/CookTypes.cpp
+++ b/Engine/Source/Editor/UnrealEd/Private/Cooker/CookTypes.cpp
@@ -3,6 +3,7 @@
 #include "CookTypes.h"
 
 #include "CompactBinaryTCP.h"
+#include "CookPackageData.h"
 #include "Containers/StringView.h"
 #include "DerivedDataRequest.h"
 #include "HAL/PlatformTLS.h"
@@ -37,7 +38,8 @@ const TCHAR* LexToString(ESuppressCookReason Reason)
 {
 	switch (Reason)
 	{
-	case ESuppressCookReason::InvalidSuppressCookReason: return TEXT("InvalidSuppressCookReason");
+	case ESuppressCookReason::Invalid: return TEXT("Invalid");
+	case ESuppressCookReason::NotSuppressed: return TEXT("NotSuppressed");
 	case ESuppressCookReason::AlreadyCooked: return TEXT("AlreadyCooked");
 	case ESuppressCookReason::NeverCook: return TEXT("NeverCook");
 	case ESuppressCookReason::DoesNotExistInWorkspaceDomain: return TEXT("DoesNotExistInWorkspaceDomain");
@@ -51,7 +53,7 @@ const TCHAR* LexToString(ESuppressCookReason Reason)
 	case ESuppressCookReason::CookCanceled: return TEXT("CookCanceled");
 	case ESuppressCookReason::MultiprocessAssignmentError: return TEXT("MultiprocessAssignmentError");
 	case ESuppressCookReason::RetractedByCookDirector: return TEXT("RetractedByCookDirector");
-	default: return TEXT("InvalidSuppressCookReason");
+	default: return TEXT("Invalid");
 	}
 }
 
@@ -194,6 +196,351 @@ bool IsCookIgnoreTimeouts()
 	return bIsIgnoreCookTimeouts;
 }
 
+FDiscoveredPlatformSet::FDiscoveredPlatformSet(EDiscoveredPlatformSet InSource)
+	: Source(InSource)
+{
+	ConstructUnion();
+}
+
+FDiscoveredPlatformSet::FDiscoveredPlatformSet(TConstArrayView<const ITargetPlatform*> InPlatforms)
+	: FDiscoveredPlatformSet(EDiscoveredPlatformSet::EmbeddedList)
+{
+	Platforms.Append(InPlatforms);
+}
+
+FDiscoveredPlatformSet::FDiscoveredPlatformSet(const TBitArray<>& InOrderedPlatformBits)
+	: FDiscoveredPlatformSet(EDiscoveredPlatformSet::EmbeddedBitField)
+{
+	OrderedPlatformBits = InOrderedPlatformBits;
+}
+
+FDiscoveredPlatformSet::~FDiscoveredPlatformSet()
+{
+	DestructUnion();
+}
+
+FDiscoveredPlatformSet::FDiscoveredPlatformSet(const FDiscoveredPlatformSet& Other)
+	: FDiscoveredPlatformSet(Other.Source)
+{
+	*this = Other;
+}
+
+FDiscoveredPlatformSet::FDiscoveredPlatformSet(FDiscoveredPlatformSet&& Other)
+	: FDiscoveredPlatformSet(Other.Source)
+{
+	*this = MoveTemp(Other);
+}
+
+FDiscoveredPlatformSet& FDiscoveredPlatformSet::operator=(const FDiscoveredPlatformSet& Other)
+{
+	if (Source != Other.Source)
+	{
+		DestructUnion();
+		Source = Other.Source;
+		ConstructUnion();
+	}
+	switch (Source)
+	{
+	case EDiscoveredPlatformSet::EmbeddedBitField:
+		OrderedPlatformBits = Other.OrderedPlatformBits;
+		break;
+	default:
+		Platforms = Other.Platforms;
+		break;
+	}
+	return *this;
+}
+
+FDiscoveredPlatformSet& FDiscoveredPlatformSet::operator=(FDiscoveredPlatformSet&& Other)
+{
+	if (Source != Other.Source)
+	{
+		DestructUnion();
+		Source = Other.Source;
+		ConstructUnion();
+	}
+	switch (Other.Source)
+	{
+	case EDiscoveredPlatformSet::EmbeddedBitField:
+		OrderedPlatformBits = MoveTemp(Other.OrderedPlatformBits);
+		break;
+	default:
+		Platforms = MoveTemp(Other.Platforms);
+		break;
+	}
+	return *this;
+}
+
+void FDiscoveredPlatformSet::ConstructUnion()
+{
+	switch (Source)
+	{
+	case EDiscoveredPlatformSet::EmbeddedBitField:
+		new (&OrderedPlatformBits) TBitArray<>();
+		break;
+	default:
+		new (&Platforms) TArray<const ITargetPlatform*>();
+		break;
+	}
+}
+
+void FDiscoveredPlatformSet::DestructUnion()
+{
+	switch (Source)
+	{
+	case EDiscoveredPlatformSet::EmbeddedBitField:
+		OrderedPlatformBits.~TBitArray<>();
+		break;
+	default:
+		Platforms.~TArray<const ITargetPlatform*>();
+		break;
+	}
+}
+
+void FDiscoveredPlatformSet::RemapTargetPlatforms(const TMap<ITargetPlatform*, ITargetPlatform*>& Remap)
+{
+	if (Source != EDiscoveredPlatformSet::EmbeddedBitField)
+	{
+		for (const ITargetPlatform*& Existing : Platforms)
+		{
+			Existing = Remap[Existing];
+		}
+	}
+}
+
+void FDiscoveredPlatformSet::OnRemoveSessionPlatform(const ITargetPlatform* Platform, int32 RemovedIndex)
+{
+	if (Source == EDiscoveredPlatformSet::EmbeddedBitField)
+	{
+		int32 OldNumPlatforms = OrderedPlatformBits.Num();
+		check(OldNumPlatforms > RemovedIndex);
+		TBitArray<> NewBits(false, OldNumPlatforms - 1);
+		int32 ReadIndex;
+		for (ReadIndex = 0; ReadIndex < RemovedIndex; ++ReadIndex)
+		{
+			NewBits[ReadIndex] = OrderedPlatformBits[ReadIndex];
+		}
+		for (++ReadIndex; ReadIndex < OldNumPlatforms; ++ReadIndex)
+		{
+			NewBits[ReadIndex - 1] = OrderedPlatformBits[ReadIndex];
+		}
+		OrderedPlatformBits = MoveTemp(NewBits);
+	}
+	else
+	{
+		Platforms.Remove(Platform);
+	}
+}
+
+void FDiscoveredPlatformSet::OnPlatformAddedToSession(const ITargetPlatform* Platform)
+{
+	if (Source == EDiscoveredPlatformSet::EmbeddedBitField)
+	{
+		OrderedPlatformBits.Add(false);
+	}
+}
+
+void FDiscoveredPlatformSet::ConvertFromBitfield(TConstArrayView<const ITargetPlatform*> OrderedPlatforms)
+{
+	if (Source == EDiscoveredPlatformSet::EmbeddedBitField)
+	{
+		TArray<const ITargetPlatform*> LocalPlatforms;
+		int32 NumPlatforms = OrderedPlatformBits.Num();
+		check(NumPlatforms == OrderedPlatforms.Num());
+		for (int32 Index = 0; Index < NumPlatforms; ++Index)
+		{
+			if (OrderedPlatformBits[Index])
+			{
+				LocalPlatforms.Add(OrderedPlatforms[Index]);
+			}
+		}
+		DestructUnion();
+		Source = EDiscoveredPlatformSet::EmbeddedList;
+		ConstructUnion();
+		Platforms = MoveTemp(LocalPlatforms);
+	}
+}
+
+void FDiscoveredPlatformSet::ConvertToBitfield(TConstArrayView<const ITargetPlatform*> OrderedPlatforms)
+{
+	if (Source == EDiscoveredPlatformSet::EmbeddedBitField)
+	{
+		check(OrderedPlatformBits.Num() == OrderedPlatforms.Num());
+	}
+	else if (Source == EDiscoveredPlatformSet::EmbeddedList)
+	{
+		TArray<const ITargetPlatform*> LocalPlatforms = MoveTemp(Platforms);
+		DestructUnion();
+		Source = EDiscoveredPlatformSet::EmbeddedBitField;
+		ConstructUnion();
+		int32 NumPlatforms = OrderedPlatforms.Num();
+		OrderedPlatformBits.Init(false, NumPlatforms);
+		for (int32 Index = 0; Index < NumPlatforms; ++Index)
+		{
+			OrderedPlatformBits[Index] = LocalPlatforms.Contains(OrderedPlatforms[Index]);
+		}
+	}
+}
+
+TConstArrayView<const ITargetPlatform*> FDiscoveredPlatformSet::GetPlatforms(UCookOnTheFlyServer& COTFS,
+	FInstigator* Instigator, TConstArrayView<const ITargetPlatform*> OrderedPlatforms,
+	TArray<const ITargetPlatform*, TInlineAllocator<ExpectedMaxNumPlatforms>>* OutBuffer,
+	bool bAllowPartialInstigatorResults)
+{
+	switch (Source)
+	{
+	case EDiscoveredPlatformSet::EmbeddedList:
+		return Platforms;
+	case EDiscoveredPlatformSet::EmbeddedBitField:
+	{
+		check(OutBuffer);
+		OutBuffer->Reset();
+		int32 NumPlatforms = OrderedPlatformBits.Num();
+		check(NumPlatforms == OrderedPlatforms.Num());
+		for (int32 Index = 0; Index < NumPlatforms; ++Index)
+		{
+			if (OrderedPlatformBits[Index])
+			{
+				OutBuffer->Add(OrderedPlatforms[Index]);
+			}
+		}
+		return *OutBuffer;
+	}
+	case EDiscoveredPlatformSet::CopyFromInstigator:
+		Platforms.Reset();
+		check(Instigator);
+		FPackageData::GetReachablePlatformsForInstigator(COTFS, Instigator->Referencer, Platforms,
+			bAllowPartialInstigatorResults);
+		return Platforms;
+	default:
+		checkNoEntry();
+		return Platforms;
+	}
+}
+
+void WriteToCompactBinary(FCbWriter& Writer, const FDiscoveredPlatformSet& Value,
+	TConstArrayView<const ITargetPlatform*> OrderedReplicationPlatforms)
+{
+	Writer.BeginArray();
+	Writer << static_cast<uint8>(Value.Source);
+	switch (Value.Source)
+	{
+	case EDiscoveredPlatformSet::EmbeddedList:
+	{
+		TArray<uint8, TInlineAllocator<ExpectedMaxNumPlatforms>> PlatformIntegers;
+		for (const ITargetPlatform* Platform : Value.Platforms)
+		{
+			if (Platform == CookerLoadingPlatformKey)
+			{
+				PlatformIntegers.Add(MAX_uint8);
+			}
+			else
+			{
+				int32 PlatformIndex = OrderedReplicationPlatforms.IndexOfByKey(Platform);
+				check(0 <= PlatformIndex && PlatformIndex < MAX_uint8);
+				PlatformIntegers.Add(static_cast<uint8>(PlatformIndex));
+			}
+		}
+		Writer << PlatformIntegers;
+		break;
+	}
+	case EDiscoveredPlatformSet::EmbeddedBitField:
+	{
+		int32 NumPlatforms = Value.OrderedPlatformBits.Num();
+		check(OrderedReplicationPlatforms.Num() == NumPlatforms);
+		Writer.BeginArray();
+		for (int32 Index = 0; Index < NumPlatforms; ++Index)
+		{
+			Writer.AddBool(Value.OrderedPlatformBits[Index]);
+		}
+		Writer.EndArray();
+		break;
+	}
+	case EDiscoveredPlatformSet::CopyFromInstigator:
+		break;
+	default:
+		checkNoEntry();
+		break;
+	}
+	Writer.EndArray();
+}
+
+bool LoadFromCompactBinary(FCbFieldView Field, FDiscoveredPlatformSet& OutValue,
+	TConstArrayView<const ITargetPlatform*> OrderedReplicationPlatforms)
+{
+	FCbArrayView FieldAsArray = Field.AsArrayView();
+	if (Field.HasError())
+	{
+		return false;
+	}
+	FCbFieldViewIterator It = FieldAsArray.CreateViewIterator();
+	uint8 SourceAsInt;
+	if (!LoadFromCompactBinary(*It++, SourceAsInt) || SourceAsInt >= static_cast<uint8>(EDiscoveredPlatformSet::Count))
+	{
+		return false;
+	}
+	OutValue = FDiscoveredPlatformSet(static_cast<EDiscoveredPlatformSet>(SourceAsInt));
+
+	bool bOk = true;
+	switch (OutValue.Source)
+	{
+	case EDiscoveredPlatformSet::EmbeddedList:
+	{
+		TArray<uint8, TInlineAllocator<ExpectedMaxNumPlatforms>> PlatformIntegers;
+		if (!LoadFromCompactBinary(*It++, PlatformIntegers))
+		{
+			bOk = false;
+		}
+		OutValue.Platforms.Reserve(PlatformIntegers.Num());
+		for (uint8 PlatformInteger : PlatformIntegers)
+		{
+			if (PlatformInteger == MAX_uint8)
+			{
+				OutValue.Platforms.Add(CookerLoadingPlatformKey);
+			}
+			else
+			{
+				if (PlatformInteger < OrderedReplicationPlatforms.Num())
+				{
+					OutValue.Platforms.Add(OrderedReplicationPlatforms[PlatformInteger]);
+				}
+				else
+				{
+					bOk = false;
+				}
+			}
+		}
+		break;
+	}
+	case EDiscoveredPlatformSet::EmbeddedBitField:
+	{
+		FCbArrayView BitArrayField = (*It++).AsArrayView();
+		int32 NumPlatforms = BitArrayField.Num();
+		if (NumPlatforms != OrderedReplicationPlatforms.Num())
+		{
+			bOk = false;
+			OutValue = FDiscoveredPlatformSet(EDiscoveredPlatformSet::EmbeddedList);
+		}
+		else
+		{
+			OutValue.OrderedPlatformBits.Init(false, NumPlatforms);
+			int32 Index = 0;
+			for (FCbFieldView BoolField : BitArrayField)
+			{
+				OutValue.OrderedPlatformBits[Index++] = BoolField.AsBool(false);
+			}
+		}
+		break;
+	}
+	case EDiscoveredPlatformSet::CopyFromInstigator:
+		break;
+	default:
+		checkNoEntry();
+		break;
+	}
+	return bOk;
+}
+
 }
 
 FCbWriter& operator<<(FCbWriter& Writer, const UE::Cook::FInitializeConfigSettings& Value)
diff --git a/Engine/Source/Editor/UnrealEd/Private/Cooker/CookTypes.h b/Engine/Source/Editor/UnrealEd/Private/Cooker/CookTypes.h
--- a/Engine/Source/Editor/UnrealEd/Private/Cooker/CookTypes.h
+++ b/Engine/Source/Editor/UnrealEd/Private/Cooker/CookTypes.h
@@ -27,6 +27,7 @@ namespace UE::Cook { struct FBeginCookConfigSettings; }
 namespace UE::Cook { struct FCookByTheBookOptions; }
 namespace UE::Cook { struct FCookOnTheFlyOptions; }
 namespace UE::Cook { struct FInitializeConfigSettings; }
+namespace UE::Cook { struct FInstigator; }
 
 FCbWriter& operator<<(FCbWriter& Writer, const UE::Cook::FBeginCookConfigSettings& Value);
 bool LoadFromCompactBinary(FCbFieldView Field, UE::Cook::FBeginCookConfigSettings& OutValue);
@@ -42,6 +43,8 @@ bool LoadFromCompactBinary(FCbFieldView Field, UE::Cook::FInitializeConfigSettin
 
 LLM_DECLARE_TAG(Cooker_CachedPlatformData);
 
+inline constexpr uint32 ExpectedMaxNumPlatforms = 32;
+
 /** A BaseKeyFuncs for Maps and Sets with a quicker hash function for pointers than TDefaultMapKeyFuncs */
 template<typename KeyType>
 struct TFastPointerSetKeyFuncs : public DefaultKeyFuncs<KeyType>
@@ -141,7 +144,8 @@ namespace UE::Cook
 
 	enum class ESuppressCookReason : uint8
 	{
-		InvalidSuppressCookReason,
+		Invalid, // Used by containers for values not in container, not used in cases for passing between containers
+		NotSuppressed,
 		AlreadyCooked,
 		NeverCook,
 		DoesNotExistInWorkspaceDomain,
@@ -451,6 +455,60 @@ struct FCookOnTheFlyOptions
 	friend bool ::LoadFromCompactBinary(FCbFieldView Field, UE::Cook::FCookOnTheFlyOptions& Value);
 };
 
+/** Enum used by FDiscoveredPlatformSet to specify what source it will use for the set of platforms. */
+enum class EDiscoveredPlatformSet
+{
+	EmbeddedList,
+	EmbeddedBitField,
+	CopyFromInstigator,
+	Count,
+};
+
+/**
+ * A provider of a set of platforms to mark reachable for a discovered package. It might be an embedded list or
+ * it might hold instructions for where to get the platforms from other context data.
+ */
+struct FDiscoveredPlatformSet
+{
+	FDiscoveredPlatformSet(EDiscoveredPlatformSet InSource = EDiscoveredPlatformSet::EmbeddedList);
+	explicit FDiscoveredPlatformSet(TConstArrayView<const ITargetPlatform*> InPlatforms);
+	explicit FDiscoveredPlatformSet(const TBitArray<>& InOrderedPlatformBits);
+	~FDiscoveredPlatformSet();
+	FDiscoveredPlatformSet(const FDiscoveredPlatformSet& Other);
+	FDiscoveredPlatformSet(FDiscoveredPlatformSet&& Other);
+	FDiscoveredPlatformSet& operator=(const FDiscoveredPlatformSet& Other);
+	FDiscoveredPlatformSet& operator=(FDiscoveredPlatformSet&& Other);
+
+	EDiscoveredPlatformSet GetSource() const { return Source; }
+
+	void RemapTargetPlatforms(const TMap<ITargetPlatform*, ITargetPlatform*>& Remap);
+	void OnRemoveSessionPlatform(const ITargetPlatform* Platform, int32 RemovedIndex);
+	void OnPlatformAddedToSession(const ITargetPlatform* Platform);
+	TConstArrayView<const ITargetPlatform*> GetPlatforms(UCookOnTheFlyServer& COTFS,
+		FInstigator* Instigator, TConstArrayView<const ITargetPlatform*> OrderedPlatforms,
+		TArray<const ITargetPlatform*, TInlineAllocator<ExpectedMaxNumPlatforms>>* OutBuffer,
+		bool bAllowPartialInstigatorResults=false);
+	/** If the current type is EmbeddedBitField, change it to EmbeddedList. */
+	void ConvertFromBitfield(TConstArrayView<const ITargetPlatform*> OrderedPlatforms);
+	/** If the current type is EmbeddedList, change it to EmbeddedBitfield. Asserts if the type is already EmbeddedBitfield. */
+	void ConvertToBitfield(TConstArrayView<const ITargetPlatform*> OrderedPlatforms);
+
+private:
+	void DestructUnion();
+	void ConstructUnion();
+	friend void WriteToCompactBinary(FCbWriter& Writer, const FDiscoveredPlatformSet& Value,
+		TConstArrayView<const ITargetPlatform*> OrderedSessionPlatforms);
+	friend bool LoadFromCompactBinary(FCbFieldView Field, FDiscoveredPlatformSet& OutValue,
+		TConstArrayView<const ITargetPlatform*> OrderedSessionPlatforms);
+
+	union 
+	{
+		TArray<const ITargetPlatform*> Platforms;
+		TBitArray<> OrderedPlatformBits;
+	};
+	EDiscoveredPlatformSet Source;
+};
+
 }
 
 
@@ -519,7 +577,6 @@ bool LoadFromCompactBinary(FCbFieldView Field, FBeginCookContextForWorker& Value
 void LogCookerMessage(const FString& MessageText, EMessageSeverity::Type Severity);
 LLM_DECLARE_TAG(Cooker);
 
-inline constexpr uint32 ExpectedMaxNumPlatforms = 32;
 #define REMAPPED_PLUGINS TEXTVIEW("RemappedPlugins")
 extern float GCookProgressWarnBusyTime;
 
diff --git a/Engine/Source/Editor/UnrealEd/Private/Cooker/CookWorkerClient.cpp b/Engine/Source/Editor/UnrealEd/Private/Cooker/CookWorkerClient.cpp
--- a/Engine/Source/Editor/UnrealEd/Private/Cooker/CookWorkerClient.cpp
+++ b/Engine/Source/Editor/UnrealEd/Private/Cooker/CookWorkerClient.cpp
@@ -192,7 +192,7 @@ void FCookWorkerClient::ReportPromoteToSaveComplete(FPackageData& PackageData)
 
 	FName PackageName = PackageData.GetPackageName();
 	Result->SetPackageName(PackageName);
-	Result->SetSuppressCookReason(ESuppressCookReason::InvalidSuppressCookReason);
+	Result->SetSuppressCookReason(ESuppressCookReason::NotSuppressed);
 	Result->SetPlatforms(OrderedSessionPlatforms);
 	if (FGeneratorPackage* Generator = PackageData.GetGeneratorPackage(); Generator)
 	{
@@ -204,7 +204,7 @@ void FCookWorkerClient::ReportPromoteToSaveComplete(FPackageData& PackageData)
 	{
 		ITargetPlatform* TargetPlatform = OrderedSessionPlatforms[PlatformIndex];
 		FPackageRemoteResult::FPlatformResult& PlatformResults = Result->GetPlatforms()[PlatformIndex];
-		FPackageData::FPlatformData& PackagePlatformData = PackageData.FindOrAddPlatformData(TargetPlatform);
+		FPackagePlatformData& PackagePlatformData = PackageData.FindOrAddPlatformData(TargetPlatform);
 		PlatformResults.SetCookResults(PackagePlatformData.GetCookResults());
 	}
 
@@ -297,12 +297,15 @@ void FCookWorkerClient::ReportPackageMessage(FName PackageName, TUniquePtr<FPack
 	}
 }
 
-void FCookWorkerClient::ReportDiscoveredPackage(const FPackageData& PackageData, const FInstigator& Instigator)
+void FCookWorkerClient::ReportDiscoveredPackage(const FPackageData& PackageData, const FInstigator& Instigator,
+	FDiscoveredPlatformSet&& ReachablePlatforms)
 {
-	FDiscoveredPackage& Discovered = PendingDiscoveredPackages.Emplace_GetRef();
+	FDiscoveredPackageReplication& Discovered = PendingDiscoveredPackages.Emplace_GetRef();
 	Discovered.PackageName = PackageData.GetPackageName();
 	Discovered.NormalizedFileName = PackageData.GetFileName();
 	Discovered.Instigator = Instigator;
+	Discovered.Platforms = MoveTemp(ReachablePlatforms);
+	Discovered.Platforms.ConvertToBitfield(OrderedSessionAndSpecialPlatforms);
 }
 
 EPollStatus FCookWorkerClient::PollTryConnect(const FDirectorConnectionInfo& ConnectInfo)
@@ -468,6 +471,9 @@ void FCookWorkerClient::PollReceiveConfigMessage()
 	}
 	DirectorCookMode = InitialConfigMessage->GetDirectorCookMode();
 	OrderedSessionPlatforms = InitialConfigMessage->GetOrderedSessionPlatforms();
+	OrderedSessionAndSpecialPlatforms.Reset(OrderedSessionPlatforms.Num() + 1);
+	OrderedSessionAndSpecialPlatforms.Append(OrderedSessionPlatforms);
+	OrderedSessionAndSpecialPlatforms.Add(CookerLoadingPlatformKey);
 	const TArray<ITargetPlatform*>& ActiveTargetPlatforms = GetTargetPlatformManagerRef().GetActiveTargetPlatforms();
 
 	auto GetPlatformDetails = [this, &ActiveTargetPlatforms](TStringBuilder<512>& StringBuilder)
@@ -568,6 +574,7 @@ void FCookWorkerClient::SendPendingResults()
 	if (!PendingDiscoveredPackages.IsEmpty())
 	{
 		FDiscoveredPackagesMessage DiscoveredMessage;
+		DiscoveredMessage.OrderedSessionAndSpecialPlatforms = OrderedSessionAndSpecialPlatforms;
 		DiscoveredMessage.Packages = MoveTemp(PendingDiscoveredPackages);
 		SendMessage(DiscoveredMessage);
 		PendingDiscoveredPackages.Reset();
@@ -628,6 +635,7 @@ void FCookWorkerClient::HandleReceiveMessages(TArray<UE::CompactBinaryTCP::FMars
 			else if (Message.MessageType == FAssignPackagesMessage::MessageType)
 			{
 				FAssignPackagesMessage AssignPackagesMessage;
+				AssignPackagesMessage.OrderedSessionPlatforms = OrderedSessionPlatforms;
 				if (!AssignPackagesMessage.TryRead(Message.Object))
 				{
 					LogInvalidMessage(TEXT("FAssignPackagesMessage"));
@@ -741,36 +749,46 @@ void FCookWorkerClient::LogInvalidMessage(const TCHAR* MessageTypeName)
 
 void FCookWorkerClient::AssignPackages(FAssignPackagesMessage& Message)
 {
+	if (Message.PackageDatas.IsEmpty())
+	{
+		return;
+	}
+
+	TArray<const ITargetPlatform*, TInlineAllocator<ExpectedMaxNumPlatforms>> NeedCookPlatformsBuffer;
 	for (FAssignPackageData& AssignData: Message.PackageDatas)
 	{
 		FPackageData& PackageData = COTFS.PackageDatas->FindOrAddPackageData(AssignData.ConstructData.PackageName,
 			AssignData.ConstructData.NormalizedFileName);
-		// If already InProgress, ignore the duplicate package silently
+		TConstArrayView<const ITargetPlatform*> NeedCookPlatforms = 
+			AssignData.NeedCookPlatforms.GetPlatforms(COTFS, nullptr, OrderedSessionPlatforms, &NeedCookPlatformsBuffer);
 		if (PackageData.IsInProgress())
 		{
-			return;
+			// If already in progress and no new platforms, ignore the duplicate
+			// If there are new platforms, we don't currently handle that; the director should have retracted it first
+			for (const ITargetPlatform* TargetPlatform : NeedCookPlatforms)
+			{
+				check(TargetPlatform != CookerLoadingPlatformKey);
+				checkf(PackageData.FindOrAddPlatformData(TargetPlatform).IsReachable(),
+					TEXT("CookWorker received AssignPackage for package %s which is already in progress, with new platform %s. Adding new platforms to an inprogress package is not yet supported"),
+					*PackageData.GetPackageName().ToString(), *TargetPlatform->PlatformName());
+			}
+			continue;
 		}
 
 		// We do not want CookWorkers to explore dependencies in CookRequestCluster because the Director did it already.
 		// Mark the PackageDatas we get from the Director as already explored.
-		for (const ITargetPlatform* TargetPlatform : OrderedSessionPlatforms)
+		for (const ITargetPlatform* TargetPlatform : NeedCookPlatforms)
 		{
-			FPackageData::FPlatformData& PlatformData = PackageData.FindOrAddPlatformData(TargetPlatform);
-			PlatformData.SetExplored(true);
-			PlatformData.SetCookable(true);
+			PackageData.FindOrAddPlatformData(TargetPlatform).MarkCookableForWorker(*this);
 		}
-
-		PackageData.SetRequestData(OrderedSessionPlatforms , false /* bInIsUrgent */, FCompletionCallback(),
-			FInstigator(AssignData.Instigator));
+		PackageData.FindOrAddPlatformData(CookerLoadingPlatformKey).MarkCookableForWorker(*this);
+		PackageData.SetInstigator(*this, FInstigator(AssignData.Instigator));
 		PackageData.SendToState(EPackageState::Request, ESendFlags::QueueAddAndRemove);
 	}
 
-	if (!Message.PackageDatas.IsEmpty())
-	{
-		// Clear the SoftGC diagnostic ExpectedNeverLoadPackages because we have new assigned packages
-		// that we didn't consider during SoftGC
-		COTFS.PackageTracker->ClearExpectedNeverLoadPackages();
-	}
+	// Clear the SoftGC diagnostic ExpectedNeverLoadPackages because we have new assigned packages
+	// that we didn't consider during SoftGC
+	COTFS.PackageTracker->ClearExpectedNeverLoadPackages();
 }
 
 void FCookWorkerClient::Register(IMPCollector* Collector)
diff --git a/Engine/Source/Editor/UnrealEd/Private/Cooker/CookWorkerClient.h b/Engine/Source/Editor/UnrealEd/Private/Cooker/CookWorkerClient.h
--- a/Engine/Source/Editor/UnrealEd/Private/Cooker/CookWorkerClient.h
+++ b/Engine/Source/Editor/UnrealEd/Private/Cooker/CookWorkerClient.h
@@ -16,7 +16,7 @@ namespace UE::Cook { class IMPCollector; }
 namespace UE::Cook { struct FAbortPackagesMessage; }
 namespace UE::Cook { struct FAssignPackagesMessage; }
 namespace UE::Cook { struct FDirectorConnectionInfo; }
-namespace UE::Cook { struct FDiscoveredPackage; }
+namespace UE::Cook { struct FDiscoveredPackageReplication; }
 namespace UE::Cook { struct FHeartbeatMessage; }
 namespace UE::Cook { struct FInitialConfigMessage; }
 namespace UE::Cook { struct FPackageRemoteResult; }
@@ -63,7 +63,8 @@ public:
 	/** Queue a message to the server that the Package was saved. Will be sent during Tick. */
 	void ReportPromoteToSaveComplete(FPackageData& PackageData);
 	/** Queue a message to the server that a package was discovered as needed in the cook. Will be sent during Tick. */
-	void ReportDiscoveredPackage(const FPackageData& PackageData, const FInstigator& Instigator);
+	void ReportDiscoveredPackage(const FPackageData& PackageData, const FInstigator& Instigator,
+		FDiscoveredPlatformSet&& ReachablePlatforms);
 
 	/** Register a Collector for periodic ticking that sends messages to the Director. */
 	void Register(IMPCollector* Collector);
@@ -141,7 +142,8 @@ private:
 	TSharedPtr<FInternetAddr> DirectorAddr;
 	TUniquePtr<FInitialConfigMessage> InitialConfigMessage;
 	TArray<ITargetPlatform*> OrderedSessionPlatforms;
-	TArray<FDiscoveredPackage> PendingDiscoveredPackages;
+	TArray<ITargetPlatform*> OrderedSessionAndSpecialPlatforms;
+	TArray<FDiscoveredPackageReplication> PendingDiscoveredPackages;
 	TMap<FGuid, TRefCountPtr<IMPCollector>> Collectors;
 	UE::CompactBinaryTCP::FSendBuffer SendBuffer;
 	UE::CompactBinaryTCP::FReceiveBuffer ReceiveBuffer;
diff --git a/Engine/Source/Editor/UnrealEd/Private/Cooker/CookWorkerServer.cpp b/Engine/Source/Editor/UnrealEd/Private/Cooker/CookWorkerServer.cpp
--- a/Engine/Source/Editor/UnrealEd/Private/Cooker/CookWorkerServer.cpp
+++ b/Engine/Source/Editor/UnrealEd/Private/Cooker/CookWorkerServer.cpp
@@ -2,6 +2,7 @@
 
 #include "CookWorkerServer.h"
 
+#include "Algo/Find.h"
 #include "Commandlets/AssetRegistryGenerator.h"
 #include "CompactBinaryTCP.h"
 #include "CookDirector.h"
@@ -393,6 +394,9 @@ bool FCookWorkerServer::TryHandleConnectMessage(FWorkerConnectMessage& Message,
 	HandleReceiveMessagesInternal();
 	const FInitialConfigMessage& InitialConfigMessage = Director.GetInitialConfigMessage();
 	OrderedSessionPlatforms = InitialConfigMessage.GetOrderedSessionPlatforms();
+	OrderedSessionAndSpecialPlatforms.Reset(OrderedSessionPlatforms.Num() + 1);
+	OrderedSessionAndSpecialPlatforms.Append(OrderedSessionPlatforms);
+	OrderedSessionAndSpecialPlatforms.Add(CookerLoadingPlatformKey);
 	SendMessageInLock(InitialConfigMessage);
 	return true;
 }
@@ -602,15 +606,27 @@ void FCookWorkerServer::SendPendingPackages()
 
 	TArray<FAssignPackageData> AssignDatas;
 	AssignDatas.Reserve(PackagesToAssign.Num());
+	TBitArray<> SessionPlatformNeedsCook;
+
 	for (FPackageData* PackageData : PackagesToAssign)
 	{
 		FAssignPackageData& AssignData = AssignDatas.Emplace_GetRef();
 		AssignData.ConstructData = PackageData->CreateConstructData();
 		AssignData.Instigator = PackageData->GetInstigator();
+		SessionPlatformNeedsCook.Init(false, OrderedSessionPlatforms.Num());
+		int32 PlatformIndex = 0;
+		for (const ITargetPlatform* SessionPlatform : OrderedSessionPlatforms)
+		{
+			FPackagePlatformData* PlatformData = PackageData->FindPlatformData(SessionPlatform);
+			SessionPlatformNeedsCook[PlatformIndex++] = PlatformData && PlatformData->NeedsCooking(SessionPlatform);
+		}
+		AssignData.NeedCookPlatforms = FDiscoveredPlatformSet(SessionPlatformNeedsCook);
 	}
 	PendingPackages.Append(PackagesToAssign);
 	PackagesToAssign.Empty();
-	SendMessageInLock(FAssignPackagesMessage(MoveTemp(AssignDatas)));
+	FAssignPackagesMessage AssignPackagesMessage(MoveTemp(AssignDatas));
+	AssignPackagesMessage.OrderedSessionPlatforms = OrderedSessionPlatforms;
+	SendMessageInLock(MoveTemp(AssignPackagesMessage));
 }
 
 void FCookWorkerServer::PumpReceiveMessages()
@@ -683,15 +699,16 @@ void FCookWorkerServer::HandleReceiveMessagesInternal()
 		else if (Message.MessageType == FDiscoveredPackagesMessage::MessageType)
 		{
 			FDiscoveredPackagesMessage DiscoveredMessage;
+			DiscoveredMessage.OrderedSessionAndSpecialPlatforms = OrderedSessionAndSpecialPlatforms;
 			if (!DiscoveredMessage.TryRead(Message.Object))
 			{
 				LogInvalidMessage(TEXT("FDiscoveredPackagesMessage"));
 			}
 			else
 			{
-				for (FDiscoveredPackage& DiscoveredPackage : DiscoveredMessage.Packages)
+				for (FDiscoveredPackageReplication& DiscoveredPackage : DiscoveredMessage.Packages)
 				{
-					AddDiscoveredPackage(MoveTemp(DiscoveredPackage));
+					QueueDiscoveredPackage(MoveTemp(DiscoveredPackage));
 				}
 			}
 		}
@@ -792,7 +809,7 @@ void FCookWorkerServer::RecordResults(FPackageResultsMessage& Message)
 
 		// MPCOOKTODO: Refactor FSaveCookedPackageContext::FinishPlatform and ::FinishPackage so we can call them from here
 		// to reduce duplication
-		if (Result.GetSuppressCookReason() == ESuppressCookReason::InvalidSuppressCookReason)
+		if (Result.GetSuppressCookReason() == ESuppressCookReason::NotSuppressed)
 		{
 			int32 NumPlatforms = OrderedSessionPlatforms.Num();
 			if (Result.GetPlatforms().Num() != NumPlatforms)
@@ -832,25 +849,33 @@ void FCookWorkerServer::LogInvalidMessage(const TCHAR* MessageTypeName)
 		MessageTypeName);
 }
 
-void FCookWorkerServer::AddDiscoveredPackage(FDiscoveredPackage&& DiscoveredPackage)
+void FCookWorkerServer::QueueDiscoveredPackage(FDiscoveredPackageReplication&& DiscoveredPackage)
 {
 	check(TickState.TickThread == ECookDirectorThread::SchedulerThread);
 
-	FPackageData& PackageData = COTFS.PackageDatas->FindOrAddPackageData(DiscoveredPackage.PackageName,
+	FPackageDatas& PackageDatas = *COTFS.PackageDatas;
+	FInstigator& Instigator = DiscoveredPackage.Instigator;
+	FDiscoveredPlatformSet& Platforms = DiscoveredPackage.Platforms;
+	FPackageData& PackageData = PackageDatas.FindOrAddPackageData(DiscoveredPackage.PackageName,
 		DiscoveredPackage.NormalizedFileName);
-	if (PackageData.IsInProgress() || PackageData.HasAnyCookedPlatform())
+
+	TArray<const ITargetPlatform*, TInlineAllocator<ExpectedMaxNumPlatforms>> BufferPlatforms;
+	if (PackageData.HasReachablePlatforms(
+		Platforms.GetPlatforms(COTFS, &Instigator, OrderedSessionAndSpecialPlatforms, &BufferPlatforms,
+			true /* bAllowPartialInstigatorResults */)))
 	{
 		// The CookWorker thought this was a new package, but the Director already knows about it; ignore the report
 		return;
 	}
 
-	if (DiscoveredPackage.Instigator.Category == EInstigator::GeneratedPackage)
+	if (Instigator.Category == EInstigator::GeneratedPackage)
 	{
 		PackageData.SetGenerated(true);
 		PackageData.SetWorkerAssignmentConstraint(GetWorkerId());
 	}
 	Director.ResetFinalIdleHeartbeatFence();
-	COTFS.QueueDiscoveredPackageData(PackageData, MoveTemp(DiscoveredPackage.Instigator));
+	Platforms.ConvertFromBitfield(OrderedSessionAndSpecialPlatforms);
+	COTFS.QueueDiscoveredPackageOnDirector(PackageData, MoveTemp(Instigator), MoveTemp(Platforms));
 }
 
 FCookWorkerServer::FTickState::FTickState()
@@ -884,29 +909,50 @@ FAssignPackagesMessage::FAssignPackagesMessage(TArray<FAssignPackageData>&& InPa
 
 void FAssignPackagesMessage::Write(FCbWriter& Writer) const
 {
-	Writer << "P" << PackageDatas;
+	Writer.BeginArray("P");
+	for (const FAssignPackageData& PackageData : PackageDatas)
+	{
+		WriteToCompactBinary(Writer, PackageData, OrderedSessionPlatforms);
+	}
+	Writer.EndArray();
 }
 
 bool FAssignPackagesMessage::TryRead(FCbObjectView Object)
 {
-	return LoadFromCompactBinary(Object["P"], PackageDatas);
+	bool bOk = true;
+	PackageDatas.Reset();
+	for (FCbFieldView PackageField : Object["P"])
+	{
+		FAssignPackageData& PackageData = PackageDatas.Emplace_GetRef();
+		if (!LoadFromCompactBinary(PackageField, PackageData, OrderedSessionPlatforms))
+		{
+			PackageDatas.Pop();
+			bOk = false;
+		}
+	}
+	return bOk;
 }
 
 FGuid FAssignPackagesMessage::MessageType(TEXT("B7B1542B73254B679319D73F753DB6F8"));
 
-FCbWriter& operator<<(FCbWriter& Writer, const FAssignPackageData& AssignData)
+void WriteToCompactBinary(FCbWriter& Writer, const FAssignPackageData& AssignData, 
+	TConstArrayView<const ITargetPlatform*> OrderedSessionPlatforms)
 {
-	Writer.BeginObject();
-	Writer << "C" << AssignData.ConstructData;
-	Writer << "I" << AssignData.Instigator;
-	Writer.EndObject();
-	return Writer;
+	Writer.BeginArray();
+	Writer << AssignData.ConstructData;
+	Writer << AssignData.Instigator;
+	WriteToCompactBinary(Writer, AssignData.NeedCookPlatforms, OrderedSessionPlatforms);
+	Writer.EndArray();
 }
 
-bool LoadFromCompactBinary(FCbFieldView Field, FAssignPackageData& AssignData)
+bool LoadFromCompactBinary(FCbFieldView Field, FAssignPackageData& AssignData,
+	TConstArrayView<const ITargetPlatform*> OrderedSessionPlatforms)
 {
-	bool bOk = LoadFromCompactBinary(Field["C"], AssignData.ConstructData);
-	bOk = LoadFromCompactBinary(Field["I"], AssignData.Instigator) & bOk;
+	FCbFieldViewIterator It = Field.CreateViewIterator();
+	bool bOk = true;
+	bOk = LoadFromCompactBinary(*It++, AssignData.ConstructData) & bOk;
+	bOk = LoadFromCompactBinary(*It++, AssignData.Instigator) & bOk;
+	bOk = LoadFromCompactBinary(*It++, AssignData.NeedCookPlatforms, OrderedSessionPlatforms) & bOk;
 	return bOk;
 }
 
@@ -1059,18 +1105,21 @@ bool FInitialConfigMessage::TryRead(FCbObjectView Object)
 
 FGuid FInitialConfigMessage::MessageType(TEXT("340CDCB927304CEB9C0A66B5F707FC2B"));
 
-FCbWriter& operator<<(FCbWriter& Writer, const FDiscoveredPackage& Package)
+void WriteToCompactBinary(FCbWriter& Writer, const FDiscoveredPackageReplication& Package,
+	TConstArrayView<const ITargetPlatform*> OrderedSessionAndSpecialPlatforms)
 {
 	Writer.BeginObject();
 	Writer << "PackageName" << Package.PackageName;
 	Writer << "NormalizedFileName" << Package.NormalizedFileName;
 	Writer << "Instigator.Category" << static_cast<uint8>(Package.Instigator.Category);
 	Writer << "Instigator.Referencer" << Package.Instigator.Referencer;
+	Writer.SetName("Platforms");
+	WriteToCompactBinary(Writer, Package.Platforms, OrderedSessionAndSpecialPlatforms);
 	Writer.EndObject();
-	return Writer;
 }
 
-bool LoadFromCompactBinary(FCbFieldView Field, FDiscoveredPackage& OutPackage)
+bool LoadFromCompactBinary(FCbFieldView Field, FDiscoveredPackageReplication& OutPackage,
+	TConstArrayView<const ITargetPlatform*> OrderedSessionAndSpecialPlatforms)
 {
 	bool bOk = LoadFromCompactBinary(Field["PackageName"], OutPackage.PackageName);
 	bOk = LoadFromCompactBinary(Field["NormalizedFileName"], OutPackage.NormalizedFileName) & bOk;
@@ -1086,21 +1135,38 @@ bool LoadFromCompactBinary(FCbFieldView Field, FDiscoveredPackage& OutPackage)
 		bOk = false;
 	}
 	bOk = LoadFromCompactBinary(Field["Instigator.Referencer"], OutPackage.Instigator.Referencer) & bOk;
+	bOk = LoadFromCompactBinary(Field["Platforms"], OutPackage.Platforms, OrderedSessionAndSpecialPlatforms) & bOk;
 	if (!bOk)
 	{
-		OutPackage = FDiscoveredPackage();
+		OutPackage = FDiscoveredPackageReplication();
 	}
 	return bOk;
 }
 
 void FDiscoveredPackagesMessage::Write(FCbWriter& Writer) const
 {
-	Writer << "Packages" << Packages;
+	Writer.BeginArray("Packages");
+	for (const FDiscoveredPackageReplication& Package : Packages)
+	{
+		WriteToCompactBinary(Writer, Package, OrderedSessionAndSpecialPlatforms);
+	}
+	Writer.EndArray();
 }
 
 bool FDiscoveredPackagesMessage::TryRead(FCbObjectView Object)
 {
-	return LoadFromCompactBinary(Object["Packages"], Packages);
+	bool bOk = true;
+	Packages.Reset();
+	for (FCbFieldView PackageField : Object["Packages"])
+	{
+		FDiscoveredPackageReplication& Package = Packages.Emplace_GetRef();
+		if (!LoadFromCompactBinary(PackageField, Package, OrderedSessionAndSpecialPlatforms))
+		{
+			Packages.Pop();
+			bOk = false;
+		}
+	}
+	return bOk;
 }
 
 FGuid FDiscoveredPackagesMessage::MessageType(TEXT("C9F5BC5C11484B06B346B411F1ED3090"));
diff --git a/Engine/Source/Editor/UnrealEd/Private/Cooker/CookWorkerServer.h b/Engine/Source/Editor/UnrealEd/Private/Cooker/CookWorkerServer.h
--- a/Engine/Source/Editor/UnrealEd/Private/Cooker/CookWorkerServer.h
+++ b/Engine/Source/Editor/UnrealEd/Private/Cooker/CookWorkerServer.h
@@ -23,7 +23,7 @@ class ITargetPlatform;
 struct FProcHandle;
 namespace UE::CompactBinaryTCP { struct FMarshalledMessage; }
 namespace UE::Cook { class FCookDirector; }
-namespace UE::Cook { struct FDiscoveredPackage; }
+namespace UE::Cook { struct FDiscoveredPackageReplication; }
 namespace UE::Cook { enum class ECookDirectorThread : uint8; };
 namespace UE::Cook { struct FPackageData; }
 namespace UE::Cook { struct FPackageResultsMessage; }
@@ -167,7 +167,7 @@ private:
 	/** Add results from the client to the local CookOnTheFlyServer. */
 	void RecordResults(FPackageResultsMessage& Message);
 	void LogInvalidMessage(const TCHAR* MessageTypeName);
-	void AddDiscoveredPackage(FDiscoveredPackage&& DiscoveredPackage);
+	void QueueDiscoveredPackage(FDiscoveredPackageReplication&& DiscoveredPackage);
 
 	// Lock guarding access to all data on *this
 	mutable FCriticalSection CommunicationLock;
@@ -176,6 +176,7 @@ private:
 	TArray<FPackageData*> PackagesToAssign;
 	TSet<FPackageData*> PendingPackages;
 	TArray<ITargetPlatform*> OrderedSessionPlatforms;
+	TArray<ITargetPlatform*> OrderedSessionAndSpecialPlatforms;
 	UE::CompactBinaryTCP::FSendBuffer SendBuffer;
 	UE::CompactBinaryTCP::FReceiveBuffer ReceiveBuffer;
 	TRingBuffer<UE::CompactBinaryTCP::FMarshalledMessage> ReceiveMessages;
@@ -201,9 +202,13 @@ struct FAssignPackageData
 {
 	FConstructPackageData ConstructData;
 	FInstigator Instigator;
+	FDiscoveredPlatformSet NeedCookPlatforms;
 };
-FCbWriter& operator<<(FCbWriter& Writer, const FAssignPackageData& AssignData);
-bool LoadFromCompactBinary(FCbFieldView Field, FAssignPackageData& AssignData);
+
+void WriteToCompactBinary(FCbWriter& Writer, const FAssignPackageData& AssignData,
+	TConstArrayView<const ITargetPlatform*> OrderedSessionPlatforms);
+bool LoadFromCompactBinary(FCbFieldView Field, FAssignPackageData& AssignData, 
+	TConstArrayView<const ITargetPlatform*> OrderedSessionPlatforms);
 FCbWriter& operator<<(FCbWriter& Writer, const FInstigator& Instigator);
 bool LoadFromCompactBinary(FCbFieldView Field, FInstigator& Instigator);
 
@@ -220,6 +225,7 @@ public:
 
 public:
 	TArray<FAssignPackageData> PackageDatas;
+	TConstArrayView<const ITargetPlatform*> OrderedSessionPlatforms;
 	static FGuid MessageType;
 };
 
@@ -300,14 +306,18 @@ private:
 };
 
 /** Information about a discovered package sent from a CookWorker to the Director. */
-struct FDiscoveredPackage
+struct FDiscoveredPackageReplication
 {
 	FName PackageName;
 	FName NormalizedFileName;
 	FInstigator Instigator;
+	FDiscoveredPlatformSet Platforms;
 };
-FCbWriter& operator<<(FCbWriter& Writer, const FDiscoveredPackage& Package);
-bool LoadFromCompactBinary(FCbFieldView Field, FDiscoveredPackage& OutPackage);
+
+void WriteToCompactBinary(FCbWriter& Writer, const FDiscoveredPackageReplication& Package,
+	TConstArrayView<const ITargetPlatform*> OrderedSessionAndSpecialPlatforms);
+bool LoadFromCompactBinary(FCbFieldView Field, FDiscoveredPackageReplication& OutPackage,
+	TConstArrayView<const ITargetPlatform*> OrderedSessionAndSpecialPlatforms);
 
 /**
  * Message from CookWorker to Director that reports dependency packages discovered during load/save of
@@ -321,7 +331,8 @@ public:
 	virtual FGuid GetMessageType() const override { return MessageType; }
 
 public:
-	TArray<FDiscoveredPackage> Packages;
+	TArray<FDiscoveredPackageReplication> Packages;
+	TConstArrayView<const ITargetPlatform*> OrderedSessionAndSpecialPlatforms;
 	static FGuid MessageType;
 };
 
diff --git a/Engine/Source/Editor/UnrealEd/Private/Cooker/IWorkerRequests.h b/Engine/Source/Editor/UnrealEd/Private/Cooker/IWorkerRequests.h
--- a/Engine/Source/Editor/UnrealEd/Private/Cooker/IWorkerRequests.h
+++ b/Engine/Source/Editor/UnrealEd/Private/Cooker/IWorkerRequests.h
@@ -49,8 +49,8 @@ public:
 	virtual void DequeueAllExternal(TArray<FSchedulerCallback>& OutCallbacks, TArray<FFilePlatformRequest>& OutCookRequests) = 0;
 
 	// Writing Packages during Cook
-	virtual void AddDiscoveredPackage(const FPackageData& PackageData, const FInstigator& Instigator,
-		bool bLoadReady, bool& bOutShouldAddToQueue) = 0;
+	virtual void QueueDiscoveredPackage(UCookOnTheFlyServer& COTFS, FPackageData& PackageData,
+		FInstigator&& Instigator, FDiscoveredPlatformSet&& ReachablePlatforms) = 0;
 
 	// Writing Packages from CookByTheBook
 	virtual void AddStartCookByTheBookRequest(FFilePlatformRequest&& Request) = 0;
diff --git a/Engine/Source/Editor/UnrealEd/Private/Cooker/PackageResultsMessage.h b/Engine/Source/Editor/UnrealEd/Private/Cooker/PackageResultsMessage.h
--- a/Engine/Source/Editor/UnrealEd/Private/Cooker/PackageResultsMessage.h
+++ b/Engine/Source/Editor/UnrealEd/Private/Cooker/PackageResultsMessage.h
@@ -121,7 +121,7 @@ private:
 	TArray<FAsyncMessage> AsyncMessages;
 	TArray<FName> ExternalActorDependencies;
 	FName PackageName;
-	/** If failure reason is InvalidSuppressCookReason, it was saved. Otherwise, holds the suppression reason */
+	/** If failure reason is NotSuppressed, it was saved. Otherwise, holds the suppression reason */
 	ESuppressCookReason SuppressCookReason;
 	bool bReferencedOnlyByEditorOnlyData = false;
 
diff --git a/Engine/Source/Editor/UnrealEd/Private/Cooker/PackageTracker.cpp b/Engine/Source/Editor/UnrealEd/Private/Cooker/PackageTracker.cpp
--- a/Engine/Source/Editor/UnrealEd/Private/Cooker/PackageTracker.cpp
+++ b/Engine/Source/Editor/UnrealEd/Private/Cooker/PackageTracker.cpp
@@ -105,11 +105,21 @@ TMap<UPackage*, FInstigator> FPackageTracker::GetNewPackages()
 	return Result;
 }
 
-bool FPackageTracker::HasBeenConsumed() const
+void FPackageTracker::MarkLoadedPackagesAsNew()
 {
-	// Tracking is initialized only during GetNewPackages, so we can use bTrackingInitialized to
-	// tell whether GetNewPackages has been called.
-	return bTrackingInitialized;
+	if (!bTrackingInitialized)
+	{
+		return;
+	}
+
+	LLM_SCOPE_BYTAG(Cooker);
+
+	FWriteScopeLock ScopeLock(Lock);
+	NewPackages.Reserve(LoadedPackages.Num());
+	for (UPackage* Package : LoadedPackages)
+	{
+		NewPackages.FindOrAdd(Package, FInstigator(EInstigator::StartupPackage));
+	}
 }
 
 void FPackageTracker::NotifyUObjectCreated(const class UObjectBase* Object, int32 Index)
diff --git a/Engine/Source/Editor/UnrealEd/Private/Cooker/PackageTracker.h b/Engine/Source/Editor/UnrealEd/Private/Cooker/PackageTracker.h
--- a/Engine/Source/Editor/UnrealEd/Private/Cooker/PackageTracker.h
+++ b/Engine/Source/Editor/UnrealEd/Private/Cooker/PackageTracker.h
@@ -179,8 +179,11 @@ public:
 	/** Returns all packages that have been loaded since the last time GetNewPackages was called */
 	TMap<UPackage*, FInstigator> GetNewPackages();
 
-	/** Report whether GetNewPackages has ever been called */
-	bool HasBeenConsumed() const;
+	/**
+	 * Copy all LoadedPackages into NewPackages. Called when reachability of all packages needs to be recalculated
+	 * because a new session platform was added.
+	 */
+	void MarkLoadedPackagesAsNew();
 
 	virtual void NotifyUObjectCreated(const class UObjectBase* Object, int32 Index) override;
 	virtual void NotifyUObjectDeleted(const class UObjectBase* Object, int32 Index) override;
diff --git a/Engine/Source/Editor/UnrealEd/Private/Cooker/WorkerRequestsLocal.cpp b/Engine/Source/Editor/UnrealEd/Private/Cooker/WorkerRequestsLocal.cpp
--- a/Engine/Source/Editor/UnrealEd/Private/Cooker/WorkerRequestsLocal.cpp
+++ b/Engine/Source/Editor/UnrealEd/Private/Cooker/WorkerRequestsLocal.cpp
@@ -42,10 +42,10 @@ void FWorkerRequestsLocal::DequeueAllExternal(TArray<FSchedulerCallback>& OutCal
 	ExternalRequests.DequeueAll(OutCallbacks, OutCookRequests);
 }
 
-void FWorkerRequestsLocal::AddDiscoveredPackage(const FPackageData& PackageData, const FInstigator& Instigator,
-	bool bLoadReady, bool& bOutShouldAddToQueue)
+void FWorkerRequestsLocal::QueueDiscoveredPackage(UCookOnTheFlyServer& COTFS, FPackageData& PackageData,
+	FInstigator&& Instigator, FDiscoveredPlatformSet&& ReachablePlatforms)
 {
-	bOutShouldAddToQueue = true;
+	COTFS.QueueDiscoveredPackageOnDirector(PackageData, MoveTemp(Instigator), MoveTemp(ReachablePlatforms));
 }
 
 void FWorkerRequestsLocal::AddStartCookByTheBookRequest(FFilePlatformRequest&& Request)
diff --git a/Engine/Source/Editor/UnrealEd/Private/Cooker/WorkerRequestsLocal.h b/Engine/Source/Editor/UnrealEd/Private/Cooker/WorkerRequestsLocal.h
--- a/Engine/Source/Editor/UnrealEd/Private/Cooker/WorkerRequestsLocal.h
+++ b/Engine/Source/Editor/UnrealEd/Private/Cooker/WorkerRequestsLocal.h
@@ -32,8 +32,8 @@ public:
 		TArray<FFilePlatformRequest>& OutCookRequests) override;
 	virtual void RemapTargetPlatforms(const TMap<ITargetPlatform*, ITargetPlatform*>& Remap) override;
 	virtual void OnRemoveSessionPlatform(const ITargetPlatform* TargetPlatform) override;
-	virtual void AddDiscoveredPackage(const FPackageData& PackageData, const FInstigator& Instigator,
-		bool bLoadReady, bool& bOutShouldAddToQueue) override;
+	virtual void QueueDiscoveredPackage(UCookOnTheFlyServer& COTFS, FPackageData& PackageData,
+		FInstigator&& Instigator, FDiscoveredPlatformSet&& ReachablePlatforms) override;
 	virtual void AddStartCookByTheBookRequest(FFilePlatformRequest&& Request) override;
 	virtual void InitializeCookOnTheFly() override;
 	virtual void AddCookOnTheFlyRequest(FFilePlatformRequest&& Request) override;
diff --git a/Engine/Source/Editor/UnrealEd/Private/Cooker/WorkerRequestsRemote.cpp b/Engine/Source/Editor/UnrealEd/Private/Cooker/WorkerRequestsRemote.cpp
--- a/Engine/Source/Editor/UnrealEd/Private/Cooker/WorkerRequestsRemote.cpp
+++ b/Engine/Source/Editor/UnrealEd/Private/Cooker/WorkerRequestsRemote.cpp
@@ -53,11 +53,10 @@ void FWorkerRequestsRemote::DequeueAllExternal(TArray<FSchedulerCallback>& OutCa
 	ExternalRequests.DequeueAll(OutCallbacks, OutCookRequests);
 }
 
-void FWorkerRequestsRemote::AddDiscoveredPackage(const FPackageData& PackageData, const FInstigator& Instigator,
-	bool bLoadReady, bool& bOutShouldAddToQueue)
+void FWorkerRequestsRemote::QueueDiscoveredPackage(UCookOnTheFlyServer& COTFS, FPackageData& PackageData,
+	FInstigator&& Instigator, FDiscoveredPlatformSet&& ReachablePlatforms)
 {
-	bOutShouldAddToQueue = false;
-	CookWorkerClient.ReportDiscoveredPackage(PackageData, Instigator);
+	CookWorkerClient.ReportDiscoveredPackage(PackageData, MoveTemp(Instigator), MoveTemp(ReachablePlatforms));
 }
 
 void FWorkerRequestsRemote::AddStartCookByTheBookRequest(FFilePlatformRequest&& Request)
diff --git a/Engine/Source/Editor/UnrealEd/Private/Cooker/WorkerRequestsRemote.h b/Engine/Source/Editor/UnrealEd/Private/Cooker/WorkerRequestsRemote.h
--- a/Engine/Source/Editor/UnrealEd/Private/Cooker/WorkerRequestsRemote.h
+++ b/Engine/Source/Editor/UnrealEd/Private/Cooker/WorkerRequestsRemote.h
@@ -37,8 +37,8 @@ public:
 		TArray<FFilePlatformRequest>& OutCookRequests) override;
 	virtual void RemapTargetPlatforms(const TMap<ITargetPlatform*, ITargetPlatform*>& Remap) override;
 	virtual void OnRemoveSessionPlatform(const ITargetPlatform* TargetPlatform) override;
-	virtual void AddDiscoveredPackage(const FPackageData& PackageData, const FInstigator& Instigator,
-		bool bLoadReady, bool& bOutShouldAddToQueue) override;
+	virtual void QueueDiscoveredPackage(UCookOnTheFlyServer& COTFS, FPackageData& PackageData,
+		FInstigator&& Instigator, FDiscoveredPlatformSet&& ReachablePlatforms) override;
 	virtual void AddStartCookByTheBookRequest(FFilePlatformRequest&& Request) override;
 	virtual void InitializeCookOnTheFly() override;
 	virtual void AddCookOnTheFlyRequest(FFilePlatformRequest&& Request) override;
