DebugSkelMeshComponent: Do not clear the preview instance (or mark it for garbage collection) when clearing anim script instance.

This is a followup fix to CL 23407330. The original bug was caused by calling ClearAnimScriptInstance (which would mark the AnimScriptInstance for garbage collection) when the PreviewInstance is pointing to it. The original fix was to set the PreviewInstance to nullptr after this happened. However, this broke a few other things, namely RBAN simulation and PhysicsAssetEditor preview animation stopped working because the PreviewInstance was expected to exist.

Upon review, thomas.sarkanen suggested that the MarkAsGarbage was completely unnecessary, so just not doing that.

Tested original bug  (UE-171723) was still fixed, PhysicsAssetEditor preview animation works, and RBAN preview simulation still works.

diff --git a/Engine/Source/Editor/UnrealEd/Private/Animation/DebugSkelMeshComponent.cpp b/Engine/Source/Editor/UnrealEd/Private/Animation/DebugSkelMeshComponent.cpp
--- a/Engine/Source/Editor/UnrealEd/Private/Animation/DebugSkelMeshComponent.cpp
+++ b/Engine/Source/Editor/UnrealEd/Private/Animation/DebugSkelMeshComponent.cpp
@@ -570,6 +570,14 @@ void UDebugSkelMeshComponent::SetAnimClass(class UClass* NewClass)
 	UE_LOG(LogAnimation, Warning, TEXT("Attempting to destroy an animation preview actor, skipping."));
 }
 
+void UDebugSkelMeshComponent::OnClearAnimScriptInstance()
+{
+	// call to super not strictly necessary (since it is empty)
+	Super::OnClearAnimScriptInstance();
+	
+	SavedAnimScriptInstance = nullptr;
+}
+
 void UDebugSkelMeshComponent::SetSkeletalMesh(USkeletalMesh* InSkelMesh, bool bReinitPose)
 {
 	Super::SetSkeletalMesh(InSkelMesh, bReinitPose);
