Add missing function prototype which got lost in a cherry-pick/merge conflict

diff --git a/Engine/Source/Editor/UnrealEd/Classes/Animation/DebugSkelMeshComponent.h b/Engine/Source/Editor/UnrealEd/Classes/Animation/DebugSkelMeshComponent.h
--- a/Engine/Source/Editor/UnrealEd/Classes/Animation/DebugSkelMeshComponent.h
+++ b/Engine/Source/Editor/UnrealEd/Classes/Animation/DebugSkelMeshComponent.h
@@ -366,6 +366,7 @@ class UNREALED_API UDebugSkelMeshComponent : public USkeletalMeshComponent
 	virtual void InitAnim(bool bForceReinit) override;
 	virtual bool IsWindEnabled() const override { return true; }
 	virtual void SetAnimClass(class UClass* NewClass) override;
+	virtual void OnClearAnimScriptInstance() override;
 	virtual void SetSkeletalMesh(USkeletalMesh* InSkelMesh, bool bReinitPose = true) override;
 	//~ End SkeletalMeshComponent Interface
 
