Some usability updates to DiffAssetBulkData.

diff --git a/Engine/Source/Editor/UnrealEd/Private/Commandlets/DiffAssetBulkDataCommandlet.cpp b/Engine/Source/Editor/UnrealEd/Private/Commandlets/DiffAssetBulkDataCommandlet.cpp
--- a/Engine/Source/Editor/UnrealEd/Private/Commandlets/DiffAssetBulkDataCommandlet.cpp
+++ b/Engine/Source/Editor/UnrealEd/Private/Commandlets/DiffAssetBulkDataCommandlet.cpp
@@ -92,9 +92,26 @@ int32 UDiffAssetBulkDataCommandlet::Main(const FString& FullCommandLine)
 		UE_LOG(LogDiffAssetBulk, Display, TEXT(""));
 		UE_LOG(LogDiffAssetBulk, Display, TEXT("    -Base=<path/to/file>              Base Development Asset Registry (Required)"));
 		UE_LOG(LogDiffAssetBulk, Display, TEXT("    -Current=<path/to/file>           New Development Asset Registry (Required)"));
+		UE_LOG(LogDiffAssetBulk, Display, TEXT("    -ListMixed                        Show the list of changed packages with assets that have matching"));
+		UE_LOG(LogDiffAssetBulk, Display, TEXT("                                      blame tags, but also assets without."));
+		UE_LOG(LogDiffAssetBulk, Display, TEXT("    -ListDeterminism                  Show the list of changed packages with assets that have matching"));
+		UE_LOG(LogDiffAssetBulk, Display, TEXT("                                      blame tags."));		
+		UE_LOG(LogDiffAssetBulk, Display, TEXT("    -ListBlame=<blame tag>            Show the list of assets that changed due to a specific blame"));
+		UE_LOG(LogDiffAssetBulk, Display, TEXT("                                      tag or \"All\" to list all changed assets with known blame."));
+		UE_LOG(LogDiffAssetBulk, Display, TEXT("    -ListUnrepresented                Show the list of packages where a representative asset couldn't be found.")); 
+		UE_LOG(LogDiffAssetBulk, Display, TEXT("    -ListNoBlame=<class>              Show the list of assets that changed for a specific class, or \"All\""));
 		return 1;
 	}
 
+	bool bListMixed = FParse::Param(CmdLine, TEXT("ListMixed"));
+	bool bListDeterminism = FParse::Param(CmdLine, TEXT("ListDeterminism"));
+	bool bListUnrepresented = FParse::Param(CmdLine, TEXT("ListUnrepresented"));
+	FString ListBlame;
+	FParse::Value(CmdLine, TEXT("ListBlame="), ListBlame);
+	FString ListNoBlame;
+	FParse::Value(CmdLine, TEXT("ListNoBlame="), ListNoBlame);
+
+
 	// Convert the static init help text to a map
 	TMap<FName, const TCHAR*> BuiltinDiffTagHelpMap;
 	for (FBuiltinDiffTagHelp& DiffTagHelp : GBuiltinDiffTagHelp)
@@ -418,16 +435,27 @@ int32 UDiffAssetBulkDataCommandlet::Main(const FString& FullCommandLine)
 			continue;
 		}
 
-		UE_LOG(LogDiffAssetBulk, Display, TEXT("        %-30s: %d"), *ClassPackages.Key.ToString(), ClassPackages.Value.Num());
+		UE_LOG(LogDiffAssetBulk, Display, TEXT("        %-30s: %d  // -ListNoBlame=%s"), *ClassPackages.Key.ToString(), ClassPackages.Value.Num(), *ClassPackages.Key.ToString());
+		if (ListNoBlame.Compare(TEXT("All"), ESearchCase::IgnoreCase) == 0 ||
+			ListNoBlame.Compare(ClassPackages.Key.ToString(), ESearchCase::IgnoreCase) == 0)
+		{
+			for (const FName& PackageName : ClassPackages.Value)
+			{
+				UE_LOG(LogDiffAssetBulk, Display, TEXT("        %s"), *PackageName.ToString());
+			}
+		}
 	}
 
 	if (CantDetermineAssetClassPackages.Num())
 	{
 		Algo::Sort(CantDetermineAssetClassPackages, FNameLexicalLess());
-		UE_LOG(LogDiffAssetBulk, Display, TEXT("    Can't determine asset class:  : %-7d // Couldn't pick a representative asset in the package."), CantDetermineAssetClassPackages.Num());
-		for (const FName& PackageName : CantDetermineAssetClassPackages)
+		UE_LOG(LogDiffAssetBulk, Display, TEXT("    Can't determine asset class:  : %-7d // Couldn't pick a representative asset in the package. -ListUnrepresented"), CantDetermineAssetClassPackages.Num());
+		if (bListUnrepresented)
 		{
-			UE_LOG(LogDiffAssetBulk, Display, TEXT("        %s"), *PackageName.ToString());
+			for (const FName& PackageName : CantDetermineAssetClassPackages)
+			{
+				UE_LOG(LogDiffAssetBulk, Display, TEXT("        %s"), *PackageName.ToString());
+			}
 		}
 	}
 
@@ -440,14 +468,17 @@ int32 UDiffAssetBulkDataCommandlet::Main(const FString& FullCommandLine)
 		}
 		
 
-		UE_LOG(LogDiffAssetBulk, Warning, TEXT("    Can't determine blame:        : %-7d // Assets had blame tags but all matched - check determinism!"), TotalUnassignablePackages);
+		UE_LOG(LogDiffAssetBulk, Warning, TEXT("    Can't determine blame:        : %-7d // Assets had blame tags but all matched - check determinism! -ListDeterminism"), TotalUnassignablePackages);
 		for (TPair<FTopLevelAssetPath, TArray<FName>>& ClassPackages : PackagesWithUnassignableDiffsByAssumedClass)
 		{
 			UE_LOG(LogDiffAssetBulk, Warning, TEXT("        %s : %d"), *ClassPackages.Key.ToString(), ClassPackages.Value.Num());
 			Algo::Sort(ClassPackages.Value, FNameLexicalLess());
-			for (const FName& PackageName : ClassPackages.Value)
+			if (bListDeterminism)
 			{
-				UE_LOG(LogDiffAssetBulk, Warning, TEXT("            %s"), *PackageName.ToString());
+				for (const FName& PackageName : ClassPackages.Value)
+				{
+					UE_LOG(LogDiffAssetBulk, Warning, TEXT("            %s"), *PackageName.ToString());
+				}
 			}
 		}
 	}
@@ -456,12 +487,14 @@ int32 UDiffAssetBulkDataCommandlet::Main(const FString& FullCommandLine)
 	{
 		Algo::Sort(PackagesWithUnassignableDiffsAndUntaggedAssets, FNameLexicalLess());
 		
-		UE_LOG(LogDiffAssetBulk, Display, TEXT("    Potential untagged assets:    : %-7d // Package had assets with blame tags that matched, but also untagged assets. Might be determinism!"), PackagesWithUnassignableDiffsAndUntaggedAssets.Num());
-		for (const FName& PackageName : PackagesWithUnassignableDiffsAndUntaggedAssets)
+		UE_LOG(LogDiffAssetBulk, Display, TEXT("    Potential untagged assets:    : %-7d // Package had assets with blame tags that matched, but also untagged assets. Might be determinism! -ListMixed"), PackagesWithUnassignableDiffsAndUntaggedAssets.Num());
+		if (bListMixed)
 		{
-			UE_LOG(LogDiffAssetBulk, Display, TEXT("        %s"), *PackageName.ToString());
+			for (const FName& PackageName : PackagesWithUnassignableDiffsAndUntaggedAssets)
+			{
+				UE_LOG(LogDiffAssetBulk, Display, TEXT("        %s"), *PackageName.ToString());
+			}
 		}
-
 	}
 
 	if (Results.Num())
@@ -491,16 +524,20 @@ int32 UDiffAssetBulkDataCommandlet::Main(const FString& FullCommandLine)
 
 		for (TPair<FName, TMap<FTopLevelAssetPath, TArray<FDiffResult>>>& TagResults : Results)
 		{
-			UE_LOG(LogDiffAssetBulk, Display, TEXT("        %s"), *TagResults.Key.ToString());
+			UE_LOG(LogDiffAssetBulk, Display, TEXT("        %s  // -ListBlame=%s"), *TagResults.Key.ToString(), *TagResults.Key.ToString());
 
 			for (TPair<FTopLevelAssetPath, TArray<FDiffResult>>& ClassResults : TagResults.Value)
 			{
 				Algo::SortBy(ClassResults.Value, &FDiffResult::ChangedAssetObjectPath);
 				UE_LOG(LogDiffAssetBulk, Display, TEXT("            %s [%d]"), *ClassResults.Key.ToString(), ClassResults.Value.Num());
 
-				for (FDiffResult& Result : ClassResults.Value)
+				if (ListBlame.Compare(TEXT("All"), ESearchCase::IgnoreCase) == 0 ||
+					ListBlame.Compare(ClassResults.Key.ToString(), ESearchCase::IgnoreCase) == 0)
 				{
-					UE_LOG(LogDiffAssetBulk, Display, TEXT("                %s [%s -> %s]"), *Result.ChangedAssetObjectPath, *Result.TagBaseValue, *Result.TagCurrentValue);
+					for (FDiffResult& Result : ClassResults.Value)
+					{
+						UE_LOG(LogDiffAssetBulk, Display, TEXT("                %s [%s -> %s]"), *Result.ChangedAssetObjectPath, *Result.TagBaseValue, *Result.TagCurrentValue);
+					}
 				}
 			}
 		}
