World Partition - Builder: Modified autosubmit logic to avoid multiple separate submits when processing a collection

diff --git a/Engine/Source/Editor/UnrealEd/Classes/Commandlets/WorldPartitionBuilderCommandlet.h b/Engine/Source/Editor/UnrealEd/Classes/Commandlets/WorldPartitionBuilderCommandlet.h
--- a/Engine/Source/Editor/UnrealEd/Classes/Commandlets/WorldPartitionBuilderCommandlet.h
+++ b/Engine/Source/Editor/UnrealEd/Classes/Commandlets/WorldPartitionBuilderCommandlet.h
@@ -21,4 +21,13 @@ class UWorldPartitionBuilderCommandlet : public UCommandlet
 private:
 	TArray<FString> GatherMapsFromCollection(const FString& CollectionName) const;
 	bool RunBuilder(TSubclassOf<UWorldPartitionBuilder> InBuilderClass, const FString& InWorldPackageName);
+
+	bool OnFilesModified(const TArray<FString>& InModifiedFiles, const FString& InChangeDescription);
+	bool AutoSubmitModifiedFiles() const;
+
+private:
+	bool bAutoSubmit = false;
+	FString AutoSubmitTags;
+
+	TMap<FString, TArray<FString>> AutoSubmitFiles;
 };
\ No newline at end of file
diff --git a/Engine/Source/Editor/UnrealEd/Private/Commandlets/WorldPartitionBuilderCommandlet.cpp b/Engine/Source/Editor/UnrealEd/Private/Commandlets/WorldPartitionBuilderCommandlet.cpp
--- a/Engine/Source/Editor/UnrealEd/Private/Commandlets/WorldPartitionBuilderCommandlet.cpp
+++ b/Engine/Source/Editor/UnrealEd/Private/Commandlets/WorldPartitionBuilderCommandlet.cpp
@@ -7,6 +7,7 @@
 #include "EditorWorldUtils.h"
 #include "Logging/LogMacros.h"
 #include "Misc/CommandLine.h"
+#include "Misc/EngineVersion.h"
 #include "HAL/PlatformFileManager.h"
 #include "ProfilingDebugging/ScopedTimers.h"
 #include "WorldPartition/WorldPartition.h"
@@ -19,6 +20,10 @@
 #include "AssetRegistry/AssetData.h"
 #include "AssetRegistry/AssetRegistryModule.h"
 
+#include "ISourceControlModule.h"
+#include "ISourceControlProvider.h"
+#include "SourceControlOperations.h"
+
 DEFINE_LOG_CATEGORY_STATIC(LogWorldPartitionBuilderCommandlet, All, All);
 
 UWorldPartitionBuilderCommandlet::UWorldPartitionBuilderCommandlet(const FObjectInitializer& ObjectInitializer)
@@ -42,6 +47,18 @@ int32 UWorldPartitionBuilderCommandlet::Main(const FString& Params)
 		return 1;
 	}
 
+	bAutoSubmit = Switches.Contains(TEXT("AutoSubmit"));
+	if (bAutoSubmit)
+	{
+		if (!ISourceControlModule::Get().GetProvider().IsEnabled())
+		{
+			UE_LOG(LogWorldPartitionBuilderCommandlet, Error, TEXT("-AutoSubmit requires that a valid revision control provider is enabled, exiting..."));
+			return 0;
+		}
+
+		FParse::Value(*Params, TEXT("AutoSubmitTags="), AutoSubmitTags);
+	}
+
 	if (Switches.Contains(TEXT("Verbose")))
 	{
 		LogWorldPartitionBuilderCommandlet.SetVerbosity(ELogVerbosity::Verbose);
@@ -101,6 +118,12 @@ int32 UWorldPartitionBuilderCommandlet::Main(const FString& Params)
 		}
 	}
 
+	// Autosubmit
+	if (!AutoSubmitModifiedFiles())
+	{
+		return 1;
+	}
+
 	return 0;
 }
 
@@ -185,6 +208,8 @@ bool UWorldPartitionBuilderCommandlet::RunBuilder(TSubclassOf<UWorldPartitionBui
 		return false;
 	}
 
+	Builder->SetModifiedFilesHandler(UWorldPartitionBuilder::FModifiedFilesHandler::CreateUObject(this, &UWorldPartitionBuilderCommandlet::OnFilesModified));
+
 	bool bResult;
 	{
 		FGCObjectScopeGuard BuilderGuard(Builder);
@@ -200,3 +225,50 @@ bool UWorldPartitionBuilderCommandlet::RunBuilder(TSubclassOf<UWorldPartitionBui
 
 	return bResult;
 }
+
+bool UWorldPartitionBuilderCommandlet::OnFilesModified(const TArray<FString>& InModifiedFiles, const FString& InChangeDescription)
+{
+	AutoSubmitFiles.Emplace(InChangeDescription, InModifiedFiles);
+	return true;
+}
+
+bool UWorldPartitionBuilderCommandlet::AutoSubmitModifiedFiles() const
+{
+	bool bSucceeded = true;
+
+	if (bAutoSubmit)
+	{
+		UE_LOG(LogWorldPartitionBuilderCommandlet, Display, TEXT("Submitting changes to revision control..."));
+
+		if (!AutoSubmitFiles.IsEmpty())
+		{
+			FString AllChanges;
+			TArray<FString> AllModifiedFiles;
+			for (const auto& [Description, Files] : AutoSubmitFiles)
+			{
+				AllChanges += Description + TEXT("\n");
+				AllModifiedFiles.Append(Files);
+			}
+
+			FText ChangelistDescription = FText::FromString(FString::Printf(TEXT("%s\nBased on CL %d\n%s"), *AllChanges, FEngineVersion::Current().GetChangelist(), *AutoSubmitTags));
+
+			TSharedRef<FCheckIn, ESPMode::ThreadSafe> CheckInOperation = ISourceControlOperation::Create<FCheckIn>();
+			CheckInOperation->SetDescription(ChangelistDescription);
+			if (ISourceControlModule::Get().GetProvider().Execute(CheckInOperation, AllModifiedFiles) != ECommandResult::Succeeded)
+			{
+				UE_LOG(LogWorldPartitionBuilderCommandlet, Error, TEXT("Failed to submit changes to revision control."));
+				bSucceeded = false;
+			}
+			else
+			{
+				UE_LOG(LogWorldPartitionBuilderCommandlet, Display, TEXT("Submitted changes to revision control"));
+			}
+		}
+		else
+		{
+			UE_LOG(LogWorldPartitionBuilderCommandlet, Display, TEXT("No files to submit!"));
+		}
+	}
+
+	return bSucceeded;
+}
\ No newline at end of file
diff --git a/Engine/Source/Editor/UnrealEd/Private/WorldPartition/WorldPartitionBuilder.cpp b/Engine/Source/Editor/UnrealEd/Private/WorldPartition/WorldPartitionBuilder.cpp
--- a/Engine/Source/Editor/UnrealEd/Private/WorldPartition/WorldPartitionBuilder.cpp
+++ b/Engine/Source/Editor/UnrealEd/Private/WorldPartition/WorldPartitionBuilder.cpp
@@ -8,7 +8,6 @@
 #include "EngineModule.h"
 #include "UObject/Linker.h"
 #include "HAL/PlatformFileManager.h"
-#include "StaticMeshCompiler.h"
 #include "Engine/World.h"
 #include "WorldPartition/WorldPartition.h"
 #include "WorldPartition/WorldPartitionHelpers.h"
@@ -21,10 +20,6 @@
 #include "UObject/SavePackage.h"
 #include "Algo/Transform.h"
 
-#include "ISourceControlModule.h"
-#include "ISourceControlProvider.h"
-#include "SourceControlOperations.h"
-
 DEFINE_LOG_CATEGORY_STATIC(LogWorldPartitionBuilder, All, All);
 
 namespace FWorldPartitionBuilderLog
@@ -53,11 +48,6 @@ FCellInfo::FCellInfo()
 UWorldPartitionBuilder::UWorldPartitionBuilder(const FObjectInitializer& ObjectInitializer)
 	: Super(ObjectInitializer)
 {
-	bAutoSubmit = FParse::Param(FCommandLine::Get(), TEXT("AutoSubmit"));
-	if (bAutoSubmit)
-	{
-		FParse::Value(FCommandLine::Get(), TEXT("AutoSubmitTags="), AutoSubmitTags);
-	}
 }
 
 UWorld::InitializationValues UWorldPartitionBuilder::GetWorldInitializationValues() const
@@ -421,40 +411,34 @@ bool UWorldPartitionBuilder::DeletePackages(const TArray<FString>& PackageNames,
 
 bool UWorldPartitionBuilder::AutoSubmitPackages(const TArray<UPackage*>& InModifiedPackages, const FString& InChangelistDescription) const
 {
-	TArray<FString> ModifiedFiles;
-	Algo::Transform(InModifiedPackages, ModifiedFiles, [](const UPackage* InPackage) { return USourceControlHelpers::PackageFilename(InPackage); });
-	return AutoSubmitFiles(ModifiedFiles, InChangelistDescription);
+	OnPackagesModified(InModifiedPackages, InChangelistDescription);
+	return true;
 }
 
 bool UWorldPartitionBuilder::AutoSubmitFiles(const TArray<FString>& InModifiedFiles, const FString& InChangelistDescription) const
 {
-	bool bSucceeded = true;
-
-	if (bAutoSubmit)
-	{
-		UE_LOG(LogWorldPartitionBuilder, Display, TEXT("Submitting changes to revision control..."));
+	OnFilesModified(InModifiedFiles, InChangelistDescription);
+	return true;
+}
 
-		if (!InModifiedFiles.IsEmpty())
-		{
-			FText ChangelistDescription = FText::FromString(FString::Printf(TEXT("%s\nBased on CL %d\n%s"), *InChangelistDescription, FEngineVersion::Current().GetChangelist(), *AutoSubmitTags));
+bool UWorldPartitionBuilder::OnPackagesModified(const TArray<UPackage*>& InModifiedPackages, const FString& InChangelistDescription) const
+{
+	TArray<FString> ModifiedFiles;
+	Algo::Transform(InModifiedPackages, ModifiedFiles, [](const UPackage* InPackage) { return USourceControlHelpers::PackageFilename(InPackage); });
+	return OnFilesModified(ModifiedFiles, InChangelistDescription);
+}
 
-			TSharedRef<FCheckIn, ESPMode::ThreadSafe> CheckInOperation = ISourceControlOperation::Create<FCheckIn>();
-			CheckInOperation->SetDescription(ChangelistDescription);
-			if (ISourceControlModule::Get().GetProvider().Execute(CheckInOperation, InModifiedFiles) != ECommandResult::Succeeded)
-			{
-				UE_LOG(LogWorldPartitionBuilder, Error, TEXT("Failed to submit changes to revision control."));
-				bSucceeded = false;
-			}
-			else
-			{
-				UE_LOG(LogWorldPartitionBuilder, Display, TEXT("Submitted changes to revision control"));
-			}
-		}
-		else
-		{
-			UE_LOG(LogWorldPartitionBuilder, Display, TEXT("No files to submit!"));
-		}
+bool UWorldPartitionBuilder::OnFilesModified(const TArray<FString>& InModifiedFiles, const FString& InChangelistDescription) const
+{
+	if (ModifiedFilesHandler.IsBound())
+	{
+		return ModifiedFilesHandler.Execute(InModifiedFiles, InChangelistDescription);
 	}
 
-	return bSucceeded;
+	return true;
 }
+
+void UWorldPartitionBuilder::SetModifiedFilesHandler(const UWorldPartitionBuilder::FModifiedFilesHandler& InModifiedFilesHandler)
+{
+	ModifiedFilesHandler = InModifiedFilesHandler;
+}
\ No newline at end of file
diff --git a/Engine/Source/Editor/UnrealEd/Private/WorldPartition/WorldPartitionHLODsBuilder.cpp b/Engine/Source/Editor/UnrealEd/Private/WorldPartition/WorldPartitionHLODsBuilder.cpp
--- a/Engine/Source/Editor/UnrealEd/Private/WorldPartition/WorldPartitionHLODsBuilder.cpp
+++ b/Engine/Source/Editor/UnrealEd/Private/WorldPartition/WorldPartitionHLODsBuilder.cpp
@@ -253,12 +253,6 @@ bool UWorldPartitionHLODsBuilder::ValidateParams() const
 		}
 	}
 
-	if (ShouldRunStep(EHLODBuildStep::HLOD_Finalize) && bAutoSubmit && !ISourceControlModule::Get().GetProvider().IsEnabled())
-	{
-		UE_LOG(LogWorldPartitionHLODsBuilder, Error, TEXT("Submit requires that a valid revision control provider is enabled, exiting..."), *BuildManifest);
-		return false;
-	}
-
 	return true;
 }
 
@@ -322,7 +316,7 @@ bool UWorldPartitionHLODsBuilder::RunInternal(UWorld* World, const FCellInfo& In
 		bRet = DeleteHLODActors();
 	}
 
-	if (bRet && ShouldRunStep(EHLODBuildStep::HLOD_Finalize) && bAutoSubmit)
+	if (bRet && ShouldRunStep(EHLODBuildStep::HLOD_Finalize))
 	{
 		bRet = SubmitHLODActors();
 	}
@@ -585,8 +579,8 @@ bool UWorldPartitionHLODsBuilder::SubmitHLODActors()
 	ModifiedFiles.Append(SourceControlHelper->GetModifiedFiles());
 
 	// Check in all modified files
-	const FString ChangeDescription = FString::Printf(TEXT("Rebuilt HLODs for %s"), *WorldPartition->GetWorld()->GetName());
-	return AutoSubmitFiles(ModifiedFiles.GetAllFiles(), ChangeDescription);
+	const FString ChangeDescription = FString::Printf(TEXT("Rebuilt HLODs for %s"), *WorldPartition->GetWorld()->GetPackage()->GetName());
+	return OnFilesModified(ModifiedFiles.GetAllFiles(), ChangeDescription);
 }
 
 bool UWorldPartitionHLODsBuilder::DumpStats()
diff --git a/Engine/Source/Editor/UnrealEd/Private/WorldPartition/WorldPartitionLandscapeSplineMeshesBuilder.cpp b/Engine/Source/Editor/UnrealEd/Private/WorldPartition/WorldPartitionLandscapeSplineMeshesBuilder.cpp
--- a/Engine/Source/Editor/UnrealEd/Private/WorldPartition/WorldPartitionLandscapeSplineMeshesBuilder.cpp
+++ b/Engine/Source/Editor/UnrealEd/Private/WorldPartition/WorldPartitionLandscapeSplineMeshesBuilder.cpp
@@ -523,7 +523,7 @@ bool UWorldPartitionLandscapeSplineMeshesBuilder::RunInternal(UWorld* InWorld, c
 		Algo::Transform(PackagesToSave, ModifiedFiles, GetPackageFilename);
 
 		const FString ChangeDescription = FString::Printf(TEXT("Rebuilt landscape splines for %s"), *InWorld->GetName());
-		if (!AutoSubmitFiles(ModifiedFiles, ChangeDescription))
+		if (!OnFilesModified(ModifiedFiles, ChangeDescription))
 		{
 			return false;
 		}
diff --git a/Engine/Source/Editor/UnrealEd/Private/WorldPartition/WorldPartitionMiniMapBuilder.cpp b/Engine/Source/Editor/UnrealEd/Private/WorldPartition/WorldPartitionMiniMapBuilder.cpp
--- a/Engine/Source/Editor/UnrealEd/Private/WorldPartition/WorldPartitionMiniMapBuilder.cpp
+++ b/Engine/Source/Editor/UnrealEd/Private/WorldPartition/WorldPartitionMiniMapBuilder.cpp
@@ -240,5 +240,5 @@ bool UWorldPartitionMiniMapBuilder::PostRun(UWorld* World, FPackageSourceControl
 	}
 
 	const FString ChangeDescription = FString::Printf(TEXT("Rebuilt minimap for %s"), *World->GetName());
-	return AutoSubmitFiles({ PackageFileName }, ChangeDescription);
+	return OnFilesModified({ PackageFileName }, ChangeDescription);
 }
diff --git a/Engine/Source/Editor/UnrealEd/Public/WorldPartition/WorldPartitionBuilder.h b/Engine/Source/Editor/UnrealEd/Public/WorldPartition/WorldPartitionBuilder.h
--- a/Engine/Source/Editor/UnrealEd/Public/WorldPartition/WorldPartitionBuilder.h
+++ b/Engine/Source/Editor/UnrealEd/Public/WorldPartition/WorldPartitionBuilder.h
@@ -61,6 +61,9 @@ public:
 	static bool DeletePackages(const TArray<UPackage*>& Packages, FPackageSourceControlHelper& PackageHelper, bool bErrorsAsWarnings = false);
 	static bool DeletePackages(const TArray<FString>& PackageNames, FPackageSourceControlHelper& PackageHelper, bool bErrorsAsWarnings = false);
 
+	DECLARE_DELEGATE_RetVal_TwoParams(bool, FModifiedFilesHandler, const TArray<FString>&, const FString&); /*bool ModifiedFileHander(ModifiedFiles, ChangeDescription) */
+	void SetModifiedFilesHandler(const FModifiedFilesHandler& ModifiedFilesHandler);
+
 protected:
 	/**
 	 * Overridable method for derived classes to perform operations when world builder process starts.
@@ -86,7 +89,12 @@ protected:
 	 */
 	virtual bool ShouldSkipCell(const FWorldBuilderCellCoord& CellCoord) const { return false; }
 
+	bool OnFilesModified(const TArray<FString>& InModifiedFiles, const FString& InChangelistDescription) const;
+	bool OnPackagesModified(const TArray<UPackage*>& InModifiedPackages, const FString& InChangelistDescription) const;
+
+	UE_DEPRECATED(5.3, "Please use OnFilesModified")
 	bool AutoSubmitFiles(const TArray<FString>& InModifiedFiles, const FString& InChangelistDescription) const;
+	UE_DEPRECATED(5.3, "Please use OnPackagesModified")
 	bool AutoSubmitPackages(const TArray<UPackage*>& InModifiedPackages, const FString& InChangelistDescription) const;
 
 	virtual UWorld::InitializationValues GetWorldInitializationValues() const;
@@ -103,6 +111,5 @@ protected:
 	bool bLoadNonDynamicDataLayers = true;
 	bool bLoadInitiallyActiveDataLayers = true;
 
-	bool bAutoSubmit = false;
-	FString AutoSubmitTags;
+	FModifiedFilesHandler ModifiedFilesHandler;
 };
\ No newline at end of file
