Fix for skeletal mesh compilation crashing on gc.

diff --git a/Engine/Source/Editor/UnrealEd/Classes/Commandlets/BaseIteratePackagesCommandlet.h b/Engine/Source/Editor/UnrealEd/Classes/Commandlets/BaseIteratePackagesCommandlet.h
--- a/Engine/Source/Editor/UnrealEd/Classes/Commandlets/BaseIteratePackagesCommandlet.h
+++ b/Engine/Source/Editor/UnrealEd/Classes/Commandlets/BaseIteratePackagesCommandlet.h
@@ -107,6 +107,8 @@ protected:
 	/** Similar to above, but applies to all packages rather than just maps **/
 	//bool bShouldBuildTextureStreamingForAll;
 
+	bool bForceFinishAllCompilationBeforeGC;
+
 	/** only process packages containing materials **/
 	bool bOnlyMaterials;
 
diff --git a/Engine/Source/Editor/UnrealEd/Private/Commandlets/BaseIteratePackagesCommandlet.cpp b/Engine/Source/Editor/UnrealEd/Private/Commandlets/BaseIteratePackagesCommandlet.cpp
--- a/Engine/Source/Editor/UnrealEd/Private/Commandlets/BaseIteratePackagesCommandlet.cpp
+++ b/Engine/Source/Editor/UnrealEd/Private/Commandlets/BaseIteratePackagesCommandlet.cpp
@@ -929,7 +929,7 @@ void UBaseIteratePackagesCommandlet::LoadAndSaveOnePackage(const FString& Filena
 
 		if (!GarbageCollectionFrequency || Counter++ % GarbageCollectionFrequency == 0)
 		{
-			if (FApp::CanEverRender())
+			if (FApp::CanEverRender() || bForceFinishAllCompilationBeforeGC)
 			{
 				FAssetCompilingManager::Get().FinishAllCompilation();
 			}
@@ -1084,6 +1084,8 @@ int32 UBaseIteratePackagesCommandlet::Main( const FString& Params )
 	/** only process packages containing materials */
 	bOnlyMaterials = Switches.Contains(TEXT("onlymaterials"));
 
+	bForceFinishAllCompilationBeforeGC = false;
+
 	bUseWorldPartitionBuilder = false;
 
 	bKeepPackageGUIDOnSave = Switches.Contains(TEXT("KeepPackageGUIDOnSave"));
