Fixed deprecation error
#rnx

diff --git a/Engine/Source/Editor/UnrealEd/Private/WorldPartition/WorldPartitionNavigationDataBuilder.cpp b/Engine/Source/Editor/UnrealEd/Private/WorldPartition/WorldPartitionNavigationDataBuilder.cpp
--- a/Engine/Source/Editor/UnrealEd/Private/WorldPartition/WorldPartitionNavigationDataBuilder.cpp
+++ b/Engine/Source/Editor/UnrealEd/Private/WorldPartition/WorldPartitionNavigationDataBuilder.cpp
@@ -334,7 +334,7 @@ bool UWorldPartitionNavigationDataBuilder::PostRun(UWorld* World, FPackageSource
 	AddedPackagesToSubmit.Empty();
 	AddedPackagesToSubmitMap.Empty();
 
-	if (!AutoSubmitFiles(PackagesToSubmit, ChangeDescription))
+	if (!OnFilesModified(PackagesToSubmit, ChangeDescription))
 	{
 		UE_LOG(LogWorldPartitionNavigationDataBuilder, Error, TEXT("Error auto-submitting packages."));
 		return true;
