Drag&Drop a cubemap onto a mesh in the editor now generates a compiling material

diff --git a/Engine/Source/Editor/UnrealEd/Private/LevelEditorViewport.cpp b/Engine/Source/Editor/UnrealEd/Private/LevelEditorViewport.cpp
--- a/Engine/Source/Editor/UnrealEd/Private/LevelEditorViewport.cpp
+++ b/Engine/Source/Editor/UnrealEd/Private/LevelEditorViewport.cpp
@@ -94,11 +94,13 @@
 #include "ViewportWorldInteraction.h"
 #include "Subsystems/BrushEditingSubsystem.h"
 #include "Engine/VolumeTexture.h"
+#include "Engine/TextureCube.h"
 #include "Materials/MaterialExpressionDivide.h"
 #include "Materials/MaterialExpressionSubtract.h"
 #include "Materials/MaterialExpressionTransformPosition.h"
 #include "Materials/MaterialExpressionCustom.h"
 #include "Materials/MaterialExpressionWorldPosition.h"
+#include "Materials/MaterialExpressionReflectionVectorWS.h"
 #include "LevelEditorDragDropHandler.h"
 #include "UnrealWidget.h"
 #include "EdModeInteractiveToolsContext.h"
@@ -474,6 +476,18 @@ static bool TryAndCreateMaterialInput( UMaterial* UnrealMaterial, EMaterialKind
 		WorldPosExpression->MaterialExpressionEditorX = EditorPosX;
 		WorldPosExpression->MaterialExpressionEditorY = EditorPosY;
 	}
+	else if (UnrealTexture->IsA<UTextureCube>())
+	{
+		// If it's a cube texture, add a reflection vector expression to the UVW input to satisfy the expression input requirement
+		UMaterialExpressionReflectionVectorWS* ReflectionVectorExpression = NewObject<UMaterialExpressionReflectionVectorWS>(UnrealMaterial);
+
+		UnrealMaterial->GetExpressionCollection().AddExpression(ReflectionVectorExpression);
+
+		ReflectionVectorExpression->MaterialExpressionEditorX = UnrealTextureExpression->MaterialExpressionEditorX - 250; // Place node to the left to avoid overlap
+		ReflectionVectorExpression->MaterialExpressionEditorY = UnrealTextureExpression->MaterialExpressionEditorY;
+
+		UnrealTextureExpression->Coordinates.Expression = ReflectionVectorExpression;
+	}
 
 	// If we know for a fact this is a normal map, it can only legally be placed in the normal map slot.
 	// Ignore the Material kind for, but for everything else try and match it to the right slot, fallback
