Renamed Mobile Preview ES3.1 (PIE)to Standalone Game Mobile
Renamed Vulkan Mobile Preview (PIE) to Standalone Game Vulkan Mobile
Renamed Mobile Preview (PIE) to Standalone Game Target Device Mobile
Renamed Preview Rendering Level to Preview Platform
Rename DebugText in EditorViewport from "Feature level:" to "Preview Platform:"

diff --git a/Engine/Source/Editor/LevelEditor/Private/LevelEditorToolBar.cpp b/Engine/Source/Editor/LevelEditor/Private/LevelEditorToolBar.cpp
--- a/Engine/Source/Editor/LevelEditor/Private/LevelEditorToolBar.cpp
+++ b/Engine/Source/Editor/LevelEditor/Private/LevelEditorToolBar.cpp
@@ -1756,8 +1756,8 @@ void FLevelEditorToolBar::RegisterQuickSettingsMenu()
 
 		Section.AddSubMenu(
 			"FeatureLevelPreview",
-			LOCTEXT("FeatureLevelPreviewSubMenu", "Preview Rendering Level"),
-			LOCTEXT("FeatureLevelPreviewSubMenu_ToolTip", "Sets the rendering level used by the main editor"),
+			LOCTEXT("PreviewPlatformSubMenu", "Preview Platform"),
+			LOCTEXT("PreviewPlatformSubMenu_ToolTip", "Sets the preview platform used by the main editor"),
 			FNewToolMenuDelegate::CreateStatic(&MakeShaderModelPreviewMenu));
 	}
 
diff --git a/Engine/Source/Editor/UnrealEd/Private/Kismet2/DebuggerCommands.cpp b/Engine/Source/Editor/UnrealEd/Private/Kismet2/DebuggerCommands.cpp
--- a/Engine/Source/Editor/UnrealEd/Private/Kismet2/DebuggerCommands.cpp
+++ b/Engine/Source/Editor/UnrealEd/Private/Kismet2/DebuggerCommands.cpp
@@ -339,8 +339,8 @@ void FPlayWorldCommands::RegisterCommands()
 	UI_COMMAND(PlayInViewport, "Selected Viewport", "Play this level in the active level editor viewport", EUserInterfaceActionType::Check, FInputChord());
 	UI_COMMAND(PlayInEditorFloating, "New Editor Window (PIE)", "Play this level in a new window", EUserInterfaceActionType::Check, FInputChord());
 	UI_COMMAND(PlayInVR, "VR Preview", "Play this level in VR", EUserInterfaceActionType::Check, FInputChord());
-	UI_COMMAND(PlayInMobilePreview, "Mobile Preview ES3.1 (PIE)", "Play this level as a mobile device preview in ES3.1 mode (runs in its own process)", EUserInterfaceActionType::Check, FInputChord());
-	UI_COMMAND(PlayInVulkanPreview, "Vulkan Mobile Preview (PIE)", "Play this level using mobile Vulkan rendering (runs in its own process)", EUserInterfaceActionType::Check, FInputChord());
+	UI_COMMAND(PlayInMobilePreview, "Standalone Game Mobile", "Play this level as a mobile device preview in ES3.1 mode (runs in its own process)", EUserInterfaceActionType::Check, FInputChord());
+	UI_COMMAND(PlayInVulkanPreview, "Standalone Game Vulkan Mobile", "Play this level using mobile Vulkan rendering (runs in its own process)", EUserInterfaceActionType::Check, FInputChord());
 	UI_COMMAND(PlayInNewProcess, "Standalone Game", "Play this level in a new window that runs in its own process", EUserInterfaceActionType::Check, FInputChord());
 	UI_COMMAND(PlayInCameraLocation, "Current Camera Location", "Spawn the player at the current camera location", EUserInterfaceActionType::RadioButton, FInputChord());
 	UI_COMMAND(PlayInDefaultPlayerStart, "Default Player Start", "Spawn the player at the map's default player start", EUserInterfaceActionType::RadioButton, FInputChord());
@@ -834,22 +834,22 @@ TSharedRef< SWidget > FPlayWorldCommands::GeneratePlayMenuContent(TSharedRef<FUI
 		{
 			FToolMenuSection& Section = Menu->AddSection("LevelEditorPlayModes", LOCTEXT("PlayButtonModesSection", "Modes"));
 			FLocal::AddPlayModeMenuEntry(Section, PlayMode_InViewPort);
-			FLocal::AddPlayModeMenuEntry(Section, PlayMode_InMobilePreview);
+			FLocal::AddPlayModeMenuEntry(Section, PlayMode_InEditorFloating);
+			FLocal::AddPlayModeMenuEntry(Section, PlayMode_InVR);
 
 			if (GetDefault<UEditorExperimentalSettings>()->bMobilePIEPreviewDeviceLaunch)
 			{
 				Section.AddSubMenu(
 					"TargetedMobilePreview",
-					LOCTEXT("TargetedMobilePreviewSubMenu", "Mobile Preview (PIE)"),
+					LOCTEXT("TargetedMobilePreviewSubMenu", "Standalone Game Target Device Mobile"),
 					LOCTEXT("TargetedMobilePreviewSubMenu_ToolTip", "Play this level using a specified mobile device preview (runs in its own process)"),
 					FNewMenuDelegate::CreateStatic(&MakePreviewDeviceMenu), false,
 					FSlateIcon(FAppStyle::GetAppStyleSetName(), "PlayWorld.PlayInMobilePreview")
 				);
 			}
 
+			FLocal::AddPlayModeMenuEntry(Section, PlayMode_InMobilePreview);
 			FLocal::AddPlayModeMenuEntry(Section, PlayMode_InVulkanPreview);
-			FLocal::AddPlayModeMenuEntry(Section, PlayMode_InEditorFloating);
-			FLocal::AddPlayModeMenuEntry(Section, PlayMode_InVR);
 			FLocal::AddPlayModeMenuEntry(Section, PlayMode_InNewProcess);
 			FLocal::AddPlayModeMenuEntry(Section, PlayMode_Simulate);
 		}
diff --git a/Engine/Source/Editor/UnrealEd/Private/SEditorViewport.cpp b/Engine/Source/Editor/UnrealEd/Private/SEditorViewport.cpp
--- a/Engine/Source/Editor/UnrealEd/Private/SEditorViewport.cpp
+++ b/Engine/Source/Editor/UnrealEd/Private/SEditorViewport.cpp
@@ -853,7 +853,7 @@ FText SEditorViewport::GetCurrentFeatureLevelPreviewText(bool bDrawOnlyLabel) co
 
 	if (bDrawOnlyLabel)
 	{
-		LabelName = LOCTEXT("FeatureLevelLabel", "Feature Level:");
+		LabelName = LOCTEXT("PreviewPlatformLabel", "Preview Platform:");
 	}
 	else
 	{
