Added a default MapName to SlowTask and EditorBoot events so we can use filtering by MapName more appropriately in Grafana

diff --git a/Engine/Source/Editor/UnrealEd/Private/FeedbackContextEditor.cpp b/Engine/Source/Editor/UnrealEd/Private/FeedbackContextEditor.cpp
--- a/Engine/Source/Editor/UnrealEd/Private/FeedbackContextEditor.cpp
+++ b/Engine/Source/Editor/UnrealEd/Private/FeedbackContextEditor.cpp
@@ -495,7 +495,7 @@ void FFeedbackContextEditor::FinalizeSlowTask()
 		const double SlowTaskDialogTime = FStudioAnalytics::GetAnalyticSeconds() - SlowTaskStartTime;
 		
 		const FText TaskName = ScopeStack.Num() > 0 ? ScopeStack[0]->DefaultMessage : FText::GetEmpty();
-		FStudioAnalytics::FireEvent_Loading(TEXT("SlowTaskDialog"), SlowTaskDialogTime, { FAnalyticsEventAttribute(TEXT("Task"), TaskName.ToString()) });
+		FStudioAnalytics::FireEvent_Loading(TEXT("SlowTaskDialog"), SlowTaskDialogTime, { FAnalyticsEventAttribute(TEXT("Task"), TaskName.ToString()), FAnalyticsEventAttribute(TEXT("MapName"), TEXT("SlowTask")) });
 
 		Window->SetContent(SNullWidget::NullWidget);
 		Window->RequestDestroyWindow();
diff --git a/Engine/Source/Editor/UnrealEd/Private/UnrealEdMisc.cpp b/Engine/Source/Editor/UnrealEd/Private/UnrealEdMisc.cpp
--- a/Engine/Source/Editor/UnrealEd/Private/UnrealEdMisc.cpp
+++ b/Engine/Source/Editor/UnrealEd/Private/UnrealEdMisc.cpp
@@ -318,7 +318,7 @@ void FUnrealEdMisc::OnInit()
 	const double InitialEditorStartupTime = (FStudioAnalytics::GetAnalyticSeconds() - GStartTime);
 	UE_LOG(LogUnrealEdMisc, Log, TEXT("Loading editor; pre map load, took %.3f"), InitialEditorStartupTime);
 
-	FStudioAnalytics::FireEvent_Loading(TEXT("InitializeEditor"), InitialEditorStartupTime, { FAnalyticsEventAttribute(TEXT("FirstTime"), true )});
+	FStudioAnalytics::FireEvent_Loading(TEXT("InitializeEditor"), InitialEditorStartupTime, { FAnalyticsEventAttribute(TEXT("FirstTime"), true ), FAnalyticsEventAttribute(TEXT("MapName"), TEXT("EditorBoot"))});
 	// Check for automated build/submit option
 	const bool bDoAutomatedMapBuild = FParse::Param( ParsedCmdLine, TEXT("AutomatedMapBuild") );
 
@@ -537,7 +537,7 @@ void FUnrealEdMisc::OnInit()
 
 	TRACE_BOOKMARK(TEXT("Editor Startup"));
 
-	FStudioAnalytics::FireEvent_Loading(TEXT("TotalEditorStartup"), TotalEditorStartupTime, { FAnalyticsEventAttribute(TEXT("FirstTime"), true ) } );
+	FStudioAnalytics::FireEvent_Loading(TEXT("TotalEditorStartup"), TotalEditorStartupTime, { FAnalyticsEventAttribute(TEXT("FirstTime"), true ), FAnalyticsEventAttribute(TEXT("MapName"), TEXT("EditorBoot")) } );
 
 	GShaderCompilingManager->PrintStats(true);
 }
