fix for determinism in output of CookedIniVersion.txt
https://p4-swarm.epicgames.net/reviews/24740476

diff --git a/Engine/Source/Editor/UnrealEd/Private/CookOnTheFlyServer.cpp b/Engine/Source/Editor/UnrealEd/Private/CookOnTheFlyServer.cpp
--- a/Engine/Source/Editor/UnrealEd/Private/CookOnTheFlyServer.cpp
+++ b/Engine/Source/Editor/UnrealEd/Private/CookOnTheFlyServer.cpp
@@ -7450,6 +7450,7 @@ bool UCookOnTheFlyServer::SaveCurrentIniSettings(const ITargetPlatform* TargetPl
 
 	TMap<FString, FString> AdditionalIniSettings;
 	GetAdditionalCurrentIniVersionStrings(this, TargetPlatform, AdditionalIniSettings);
+	AdditionalIniSettings.KeySort(TLess<FString>());
 
 	UE::Cook::FIniSettingContainer CurrentIniSettings;
 	GetCurrentIniVersionStrings(TargetPlatform, CurrentIniSettings);
