Avoid overhead of package dependency tracking in scenarios where we don't use the data it provides.  Specifically if not running the cook commandlet and not supporting cook in the interactive editor process (this has been disabled by default).  Taking the approach of having it enabled during static initialization and then disabling later during initialization inside the Main entrypoint to help ensure that it is capturing package dependencies that may happen early in the initialization sequence.

diff --git a/Engine/Source/Editor/UnrealEd/Private/Cooker/PackageBuildDependencyTracker.cpp b/Engine/Source/Editor/UnrealEd/Private/Cooker/PackageBuildDependencyTracker.cpp
--- a/Engine/Source/Editor/UnrealEd/Private/Cooker/PackageBuildDependencyTracker.cpp
+++ b/Engine/Source/Editor/UnrealEd/Private/Cooker/PackageBuildDependencyTracker.cpp
@@ -13,8 +13,28 @@ DEFINE_LOG_CATEGORY_STATIC(LogPackageBuildDependencyTracker, Log, All);
 
 FPackageBuildDependencyTracker FPackageBuildDependencyTracker::Singleton;
 
+void FPackageBuildDependencyTracker::Disable()
+{
+	if (bEnabled)
+	{
+		UE::CoreUObject::RemoveObjectHandleReadCallback(ObjectHandleReadHandle);
+		ObjectHandleReadHandle = UE::CoreUObject::FObjectHandleTrackingCallbackId {};
+		bEnabled = false;
+	}
+}
+
+bool FPackageBuildDependencyTracker::IsEnabled() const
+{
+	return bEnabled;
+}
+
 void FPackageBuildDependencyTracker::DumpStats() const
 {
+	if (!IsEnabled())
+	{
+		return;
+	}
+
 	FScopeLock RecordsScopeLock(&RecordsLock);
 	uint64 ReferencingPackageCount = 0;
 	uint64 ReferenceCount = 0;
@@ -57,11 +77,12 @@ TArray<FBuildDependencyAccessData> FPackageBuildDependencyTracker::GetAccessData
 FPackageBuildDependencyTracker::FPackageBuildDependencyTracker()
 {
 	ObjectHandleReadHandle = UE::CoreUObject::AddObjectHandleReadCallback(StaticOnObjectHandleRead);
+	bEnabled = true;
 }
 
 FPackageBuildDependencyTracker::~FPackageBuildDependencyTracker()
 {
-	UE::CoreUObject::RemoveObjectHandleReadCallback(ObjectHandleReadHandle);
+	Disable();
 }
 
 void FPackageBuildDependencyTracker::StaticOnObjectHandleRead(TArrayView<const UObject* const> Objects)
diff --git a/Engine/Source/Editor/UnrealEd/Private/Cooker/PackageBuildDependencyTracker.h b/Engine/Source/Editor/UnrealEd/Private/Cooker/PackageBuildDependencyTracker.h
--- a/Engine/Source/Editor/UnrealEd/Private/Cooker/PackageBuildDependencyTracker.h
+++ b/Engine/Source/Editor/UnrealEd/Private/Cooker/PackageBuildDependencyTracker.h
@@ -32,6 +32,8 @@ class FPackageBuildDependencyTracker: public FNoncopyable
 public:
 	static FPackageBuildDependencyTracker& Get() { return Singleton; }
 
+	void Disable();
+	bool IsEnabled() const;
 	void DumpStats() const;
 	TArray<FBuildDependencyAccessData> GetAccessDatas(FName ReferencerPackage) const;
 
@@ -48,6 +50,7 @@ private:
 	FBuildDependencyAccessData LastAccessData{ NAME_None, nullptr };
 	TSet<FBuildDependencyAccessData>* LastReferencerSet = nullptr;
 	UE::CoreUObject::FObjectHandleTrackingCallbackId ObjectHandleReadHandle;
+	bool bEnabled = false;
 	static FPackageBuildDependencyTracker Singleton;
 };
 
diff --git a/Engine/Source/Editor/UnrealEd/Private/TargetDomain/TargetDomainUtils.cpp b/Engine/Source/Editor/UnrealEd/Private/TargetDomain/TargetDomainUtils.cpp
--- a/Engine/Source/Editor/UnrealEd/Private/TargetDomain/TargetDomainUtils.cpp
+++ b/Engine/Source/Editor/UnrealEd/Private/TargetDomain/TargetDomainUtils.cpp
@@ -150,25 +150,34 @@ bool TryCollectKeyAndDependencies(UPackage* Package, const ITargetPlatform* Targ
 		UE::AssetRegistry::EDependencyQuery::Game);
 
 	FPackageBuildDependencyTracker& Tracker = FPackageBuildDependencyTracker::Get();
-	TArray<FBuildDependencyAccessData> AccessDatas = Tracker.GetAccessDatas(PackageName);
 
-	BuildDependencies.Reserve(AccessDatas.Num());
-	for (FBuildDependencyAccessData& AccessData : AccessDatas)
+	if (Tracker.IsEnabled())
 	{
-		if (AccessData.TargetPlatform == TargetPlatform || AccessData.TargetPlatform == nullptr)
+		TArray<FBuildDependencyAccessData> AccessDatas = Tracker.GetAccessDatas(PackageName);
+
+		BuildDependencies.Reserve(AccessDatas.Num());
+		for (FBuildDependencyAccessData& AccessData : AccessDatas)
 		{
-			BuildDependencies.Add(AccessData.ReferencedPackage);
+			if (AccessData.TargetPlatform == TargetPlatform || AccessData.TargetPlatform == nullptr)
+			{
+				BuildDependencies.Add(AccessData.ReferencedPackage);
+			}
 		}
-	}
 
-	RuntimeOnlyDependencies.Reserve(AccessDatas.Num());
-	for (FName DependencyName : AssetDependencies)
-	{
-		if (!BuildDependencies.Contains(DependencyName))
+		RuntimeOnlyDependencies.Reserve(AssetDependencies.Num());
+		for (FName DependencyName : AssetDependencies)
 		{
-			RuntimeOnlyDependencies.Add(DependencyName);
+			if (!BuildDependencies.Contains(DependencyName))
+			{
+				RuntimeOnlyDependencies.Add(DependencyName);
+			}
 		}
 	}
+	else
+	{
+		// Defensively treat all asset dependencies as build dependencies and have zero runtime only dependencies
+		BuildDependencies.Append(AssetDependencies);
+	}
 
 	TArray<FName> SortedBuild;
 	SortedBuild = BuildDependencies.Array();
diff --git a/Engine/Source/Editor/UnrealEd/Private/UnrealEdMisc.cpp b/Engine/Source/Editor/UnrealEd/Private/UnrealEdMisc.cpp
--- a/Engine/Source/Editor/UnrealEd/Private/UnrealEdMisc.cpp
+++ b/Engine/Source/Editor/UnrealEd/Private/UnrealEdMisc.cpp
@@ -92,6 +92,7 @@
 #include "HAL/PlatformTime.h"
 #include "StudioAnalytics.h"
 #include "DeveloperToolSettingsDelegates.h"
+#include "Cooker/PackageBuildDependencyTracker.h"
 
 #define USE_UNIT_TESTS 0
 
@@ -217,6 +218,14 @@ FUnrealEdMisc::FUnrealEdMisc() :
 	PerformanceAnalyticsStats(new FPerformanceAnalyticsStats()),
 	NavigationBuildingNotificationHandler(NULL)
 {
+	//This is an early entry-point into the UnrealEd module to perform some editor-specific configuration
+#if UE_WITH_PACKAGE_ACCESS_TRACKING
+	const bool bBuildDependencyTrackingNeeded = GIsEditor && (IsRunningCookCommandlet() || !GetDefault<UEditorExperimentalSettings>()->bDisableCookInEditor);
+	if (!bBuildDependencyTrackingNeeded)
+	{
+		FPackageBuildDependencyTracker::Get().Disable();
+	}
+#endif
 }
 
 FUnrealEdMisc::~FUnrealEdMisc()
