Fix a crash when coping actors in a editor from before 5.1 and then pasting the actors in a editor newer then 5.1.

diff --git a/Engine/Source/Editor/UnrealEd/Private/EditorObject.cpp b/Engine/Source/Editor/UnrealEd/Private/EditorObject.cpp
--- a/Engine/Source/Editor/UnrealEd/Private/EditorObject.cpp
+++ b/Engine/Source/Editor/UnrealEd/Private/EditorObject.cpp
@@ -502,9 +502,19 @@ static const TCHAR* ImportProperties(
 						FString ObjectPath;
 						if ( FPackageName::ParseExportTextPath(ArchetypeName, &ObjectClass, &ObjectPath) )
 						{
+							UClass* ArchetypeClass = nullptr;
+
 							// find the class
-							check(FPackageName::IsValidObjectPath(ObjectClass));
-							UClass* ArchetypeClass = (UClass*)StaticFindObject(UClass::StaticClass(), nullptr, *ObjectClass);
+							if (FPackageName::IsValidObjectPath(ObjectClass))
+							{
+								ArchetypeClass = (UClass*)StaticFindObject(UClass::StaticClass(), nullptr, *ObjectClass);
+							}
+							else
+							{
+								// The text might be from before the full path of the class was exported (ensure if the name is ambiguous)
+								ArchetypeClass = (UClass*)StaticFindFirstObject(UClass::StaticClass(), *ObjectClass, EFindFirstObjectOptions::NativeFirst | EFindFirstObjectOptions::EnsureIfAmbiguous, ELogVerbosity::Warning);
+							}
+
 							if (ArchetypeClass)
 							{
 								ObjectPath = ObjectPath.TrimQuotes();
