Convert FActorFolders to a use TLazySingleton.
Fix crash on editor exit.

diff --git a/Engine/Source/Editor/SceneOutliner/Private/ActorHierarchy.cpp b/Engine/Source/Editor/SceneOutliner/Private/ActorHierarchy.cpp
--- a/Engine/Source/Editor/SceneOutliner/Private/ActorHierarchy.cpp
+++ b/Engine/Source/Editor/SceneOutliner/Private/ActorHierarchy.cpp
@@ -118,14 +118,10 @@ FActorHierarchy::~FActorHierarchy()
 	FWorldDelegates::LevelAddedToWorld.RemoveAll(this);
 	FWorldDelegates::LevelRemovedFromWorld.RemoveAll(this);
 
-
-	if (FActorFolders::IsAvailable())
-	{
-		auto& Folders = FActorFolders::Get();
-		Folders.OnFolderCreated.RemoveAll(this);
-		Folders.OnFolderMoved.RemoveAll(this);
-		Folders.OnFolderDeleted.RemoveAll(this);
-	}
+	auto& Folders = FActorFolders::Get();
+	Folders.OnFolderCreated.RemoveAll(this);
+	Folders.OnFolderMoved.RemoveAll(this);
+	Folders.OnFolderDeleted.RemoveAll(this);
 }
 
 FSceneOutlinerTreeItemPtr FActorHierarchy::FindOrCreateParentItem(const ISceneOutlinerTreeItem& Item, const TMap<FSceneOutlinerTreeItemID, FSceneOutlinerTreeItemPtr>& Items, bool bCreate)
diff --git a/Engine/Source/Editor/UnrealEd/Private/EditorActorFolders.cpp b/Engine/Source/Editor/UnrealEd/Private/EditorActorFolders.cpp
--- a/Engine/Source/Editor/UnrealEd/Private/EditorActorFolders.cpp
+++ b/Engine/Source/Editor/UnrealEd/Private/EditorActorFolders.cpp
@@ -5,6 +5,7 @@
 #include "GameFramework/Actor.h"
 #include "HAL/FileManager.h"
 #include "Misc/Paths.h"
+#include "Misc/LazySingleton.h"
 #include "Dom/JsonObject.h"
 #include "Serialization/JsonReader.h"
 #include "Serialization/JsonSerializer.h"
@@ -36,7 +37,6 @@ FOnActorFolderDelete	FActorFolders::OnFolderDelete;
 PRAGMA_ENABLE_DEPRECATION_WARNINGS
 //~ End Deprecated
 
-FActorFolders* FActorFolders::Singleton;
 FOnActorFolderCreated	FActorFolders::OnFolderCreated;
 FOnActorFolderMoved		FActorFolders::OnFolderMoved;
 FOnActorFolderDeleted	FActorFolders::OnFolderDeleted;
@@ -51,18 +51,28 @@ FActorFolders::FActorFolders()
 	FEditorDelegates::MapChange.AddRaw(this, &FActorFolders::OnMapChange);
 	FEditorDelegates::PostSaveWorldWithContext.AddRaw(this, &FActorFolders::OnWorldSaved);
 	FEditorDelegates::PostSaveExternalActors.AddRaw(this, &FActorFolders::SaveWorldFoldersState);
+
+	check(GEditor);
+	UActorEditorContextSubsystem::Get()->RegisterClient(this);
 }
 
 FActorFolders::~FActorFolders()
 {
-	check(GEngine);
-	GEngine->OnLevelActorFolderChanged().RemoveAll(this);
-	GEngine->OnLevelActorListChanged().RemoveAll(this);
-	GEngine->OnActorFolderAdded().RemoveAll(this);
+	if (GEngine)
+	{
+		GEngine->OnLevelActorFolderChanged().RemoveAll(this);
+		GEngine->OnLevelActorListChanged().RemoveAll(this);
+		GEngine->OnActorFolderAdded().RemoveAll(this);
+	}
 
 	FEditorDelegates::MapChange.RemoveAll(this);
 	FEditorDelegates::PostSaveWorldWithContext.RemoveAll(this);
 	FEditorDelegates::PostSaveExternalActors.RemoveAll(this);
+
+	if (GEditor)
+	{
+		UActorEditorContextSubsystem::Get()->UnregisterClient(this);
+	}
 }
 
 void FActorFolders::AddReferencedObjects(FReferenceCollector& Collector)
@@ -73,21 +83,7 @@ void FActorFolders::AddReferencedObjects(FReferenceCollector& Collector)
 
 FActorFolders& FActorFolders::Get()
 {
-	check(Singleton);
-	return *Singleton;
-}
-
-void FActorFolders::Init()
-{
-	Singleton = new FActorFolders;
-	UActorEditorContextSubsystem::Get()->RegisterClient(Singleton);
-}
-
-void FActorFolders::Cleanup()
-{
-	UActorEditorContextSubsystem::Get()->UnregisterClient(Singleton);
-	delete Singleton;
-	Singleton = nullptr;
+	return TLazySingleton<FActorFolders>::Get();
 }
 
 void FActorFolders::Housekeeping()
diff --git a/Engine/Source/Editor/UnrealEd/Private/UnrealEdGlobals.cpp b/Engine/Source/Editor/UnrealEd/Private/UnrealEdGlobals.cpp
--- a/Engine/Source/Editor/UnrealEd/Private/UnrealEdGlobals.cpp
+++ b/Engine/Source/Editor/UnrealEd/Private/UnrealEdGlobals.cpp
@@ -127,7 +127,7 @@ int32 EditorInit( IEngineLoop& EngineLoop )
 	SlowTask.EnterProgressFrame(40);
 
 	// Set up the actor folders singleton
-	FActorFolders::Init();
+	FActorFolders::Get();
 
 	// Initialize the misc editor
 	FUnrealEdMisc::Get().OnInit();
@@ -241,8 +241,6 @@ void EditorExit()
 	GEditor->SaveConfig();
 	GLevelEditorModeTools().SaveConfig();
 
-	// Clean up the actor folders singleton
-	FActorFolders::Cleanup();
 
 	// Save out default file directories
 	FEditorDirectories::Get().SaveLastDirectories();
diff --git a/Engine/Source/Editor/UnrealEd/Public/EditorActorFolders.h b/Engine/Source/Editor/UnrealEd/Public/EditorActorFolders.h
--- a/Engine/Source/Editor/UnrealEd/Public/EditorActorFolders.h
+++ b/Engine/Source/Editor/UnrealEd/Public/EditorActorFolders.h
@@ -42,18 +42,9 @@ struct UNREALED_API FActorFolders : public FGCObject, public IActorEditorContext
 	}
 	// End FGCObject Interface
 
-	/** Check whether the singleton is valid */
-	static bool IsAvailable() { return Singleton != nullptr; }
-
-	/** Singleton access - only valid if IsAvailable() */
+	/** Singleton access */
 	static FActorFolders& Get();
 
-	/** Initialize the singleton instance - called on Editor Startup */
-	static void Init();
-
-	/** Clean up the singleton instance - called on Editor Exit */
-	static void Cleanup();
-
 	/** Folder creation and deletion events. Called whenever a folder is created or deleted in a world. */
 	static FOnActorFolderCreated OnFolderCreated;
 	static FOnActorFolderMoved 	OnFolderMoved;
@@ -232,8 +223,5 @@ private:
 	/** Delegate used to notify changes to ActorEditorContextSubsystem */
 	FOnActorEditorContextClientChanged ActorEditorContextClientChanged;
 
-	/** Singleton instance maintained by the editor */
-	static FActorFolders* Singleton;
-
 	friend UWorldFolders;
 };
