Forward input axis events to the Editor World extension collection and give extensions an opportunity to consume input.

diff --git a/Engine/Source/Editor/UnrealEd/Private/LevelEditorViewport.cpp b/Engine/Source/Editor/UnrealEd/Private/LevelEditorViewport.cpp
--- a/Engine/Source/Editor/UnrealEd/Private/LevelEditorViewport.cpp
+++ b/Engine/Source/Editor/UnrealEd/Private/LevelEditorViewport.cpp
@@ -3521,6 +3521,21 @@ bool FLevelEditorViewportClient::InputAxis(FViewport* InViewport, FInputDeviceId
 		return true;
 	}
 
+	// Forward input axis events to the Editor World extension collection and give extensions an opportunity to consume input.
+	IPlatformInputDeviceMapper& DeviceMapper = IPlatformInputDeviceMapper::Get();
+	FPlatformUserId UserId = DeviceMapper.GetUserForInputDevice(DeviceId);
+
+	int32 ControllerId;
+	if (DeviceMapper.RemapUserAndDeviceToControllerId(UserId, ControllerId, DeviceId))
+	{
+		UEditorWorldExtensionCollection& ExtensionCollection = *GEditor->GetEditorWorldExtensionsManager()->GetEditorWorldExtensions(GetWorld());
+		
+		if (ExtensionCollection.InputAxis(this, InViewport, ControllerId, Key, Delta, DeltaTime))
+		{
+			return true;
+		}
+	}
+
 	// @todo Slate: GCurrentLevelEditingViewportClient is switched multiple times per frame and since we draw the border in slate this effectively causes border to always draw on the last viewport
 
 	FScopedSetCurrentViewportClient( this );
