Fix an array out of bound crash in the legacy fbx importer and in interchange fbx translator, when a user import a fbx with bad vertex index data.

#rnx

diff --git a/Engine/Plugins/Interchange/Runtime/Source/Parsers/Fbx/Private/FbxMesh.cpp b/Engine/Plugins/Interchange/Runtime/Source/Parsers/Fbx/Private/FbxMesh.cpp
--- a/Engine/Plugins/Interchange/Runtime/Source/Parsers/Fbx/Private/FbxMesh.cpp
+++ b/Engine/Plugins/Interchange/Runtime/Source/Parsers/Fbx/Private/FbxMesh.cpp
@@ -664,6 +664,7 @@ bool FMeshDescriptionImporter::FillMeshDescriptionFromFbxMesh(FbxMesh* Mesh, TAr
 			TArray<FVertexID> CornerVerticesIDs;
 			TArray<FVector, TInlineAllocator<3>> P;
 
+			bool bCorruptedMsgDone = false;
 			//Polygons
 			for (int32 PolygonIndex = 0; PolygonIndex < PolygonCount; PolygonIndex++)
 			{
@@ -673,12 +674,30 @@ bool FMeshDescriptionImporter::FillMeshDescriptionFromFbxMesh(FbxMesh* Mesh, TAr
 					float ComparisonThreshold = SMALL_NUMBER;
 					P.Reset();
 					P.AddUninitialized(PolygonVertexCount);
+					bool bAllCornerValid = true;
 					for (int32 CornerIndex = 0; CornerIndex < PolygonVertexCount; CornerIndex++)
 					{
 						const int32 ControlPointIndex = Mesh->GetPolygonVertex(PolygonIndex, CornerIndex);
 						const FVertexID VertexID(VertexOffset + ControlPointIndex);
+						if (!MeshDescription->Vertices().IsValid(VertexID))
+						{
+							bAllCornerValid = false;
+							break;
+						}
 						P[CornerIndex] = (FVector)VertexPositions[VertexID];
 					}
+					if (!bAllCornerValid)
+					{
+						if (!bCorruptedMsgDone)
+						{
+							UInterchangeResultMeshError_Generic* Message = AddMessage<UInterchangeResultMeshError_Generic>(Mesh);
+							Message->Text = LOCTEXT("CorruptedFbxPolygon", "Corrupted triangle for the mesh '{MeshName}'.");
+							bCorruptedMsgDone = true;
+						}
+
+						SkippedVertexInstance += PolygonVertexCount;
+						continue;
+					}
 					check(P.Num() > 2); //triangle is the smallest polygon we can have
 					const FVector Normal = ((P[1] - P[2]) ^ (P[0] - P[2])).GetSafeNormal(ComparisonThreshold);
 					//Check for degenerated polygons, avoid NAN
diff --git a/Engine/Source/Editor/UnrealEd/Private/Fbx/FbxStaticMeshImport.cpp b/Engine/Source/Editor/UnrealEd/Private/Fbx/FbxStaticMeshImport.cpp
--- a/Engine/Source/Editor/UnrealEd/Private/Fbx/FbxStaticMeshImport.cpp
+++ b/Engine/Source/Editor/UnrealEd/Private/Fbx/FbxStaticMeshImport.cpp
@@ -803,6 +803,7 @@ bool UnFbx::FFbxImporter::BuildStaticMeshFromGeometry(FbxNode* Node, UStaticMesh
 
 			bool bFaceMaterialIndexInconsistencyErrorDisplayed = false;
 			bool bUnsupportedSmoothingGroupErrorDisplayed = false;
+			bool bCorruptedMsgDone = false;
 			//Polygons
 			for (int32 PolygonIndex = 0; PolygonIndex < PolygonCount; PolygonIndex++)
 			{
@@ -825,12 +826,28 @@ bool UnFbx::FFbxImporter::BuildStaticMeshFromGeometry(FbxNode* Node, UStaticMesh
 					float ComparisonThreshold = GetTriangleAreaThreshold();
 					P.Reset();
 					P.AddUninitialized(PolygonVertexCount);
+					bool bAllCornerValid = true;
 					for (int32 CornerIndex = 0; CornerIndex < PolygonVertexCount; CornerIndex++)
 					{
 						const int32 ControlPointIndex = Mesh->GetPolygonVertex(PolygonIndex, CornerIndex);
 						const FVertexID VertexID(VertexOffset + ControlPointIndex);
+						if (!MeshDescription->Vertices().IsValid(VertexID))
+						{
+							bAllCornerValid = false;
+							break;
+						}
 						P[CornerIndex] = VertexPositions[VertexID];
 					}
+					if (!bAllCornerValid)
+					{
+						if (!bCorruptedMsgDone)
+						{
+							AddTokenizedErrorMessage(FTokenizedMessage::Create(EMessageSeverity::Error, FText::Format(LOCTEXT("Error_CorruptedFbxPolygon", "Corrupted triangle for the mesh '{0}'"), FText::FromString(Mesh->GetName()))), FFbxErrors::StaticMesh_BuildError);
+							bCorruptedMsgDone = true;
+						}
+						SkippedVertexInstance += PolygonVertexCount;
+						continue;
+					}
 					check(P.Num() > 2); //triangle is the smallest polygon we can have
 					const FVector3f Normal = ((P[1] - P[2]) ^ (P[0] - P[2])).GetSafeNormal(ComparisonThreshold);
 					//Check for degenerated polygons, avoid NAN
