Fixed deprecation warnings in compiler-generated copy/move operators by making FReplaceInstancesOfClassParameters non-copyable.

diff --git a/Engine/Source/Editor/UnrealEd/Public/Kismet2/KismetReinstanceUtilities.h b/Engine/Source/Editor/UnrealEd/Public/Kismet2/KismetReinstanceUtilities.h
--- a/Engine/Source/Editor/UnrealEd/Public/Kismet2/KismetReinstanceUtilities.h
+++ b/Engine/Source/Editor/UnrealEd/Public/Kismet2/KismetReinstanceUtilities.h
@@ -39,6 +39,8 @@ public:
 
 struct UNREALED_API FReplaceInstancesOfClassParameters
 {
+	UE_NONCOPYABLE(FReplaceInstancesOfClassParameters)
+
 	FReplaceInstancesOfClassParameters() = default;
 
 	UE_DEPRECATED(5.2, "Please use the default constructor instead.")
