fix TGA export supports G8 also

diff --git a/Engine/Source/Editor/UnrealEd/Private/Factories/EditorFactories.cpp b/Engine/Source/Editor/UnrealEd/Private/Factories/EditorFactories.cpp
--- a/Engine/Source/Editor/UnrealEd/Private/Factories/EditorFactories.cpp
+++ b/Engine/Source/Editor/UnrealEd/Private/Factories/EditorFactories.cpp
@@ -5255,7 +5255,9 @@ UTextureExporterTGA::UTextureExporterTGA(const FObjectInitializer& ObjectInitial
 bool UTextureExporterTGA::SupportsTexture(UTexture* Texture) const
 {
 	ETextureSourceFormat TSF = Texture->Source.GetFormat();
-	return TSF == TSF_BGRA8;
+	
+	// U8 formats only :
+	return TSF == TSF_BGRA8 || TSF == TSF_G8;
 }
 
 /*------------------------------------------------------------------------------
