Fixed WorldPartitionConvertCommandlet writing actor folders external objects at the wrong location when converting from a level source with existing external folders.

diff --git a/Engine/Source/Editor/UnrealEd/Private/Commandlets/WorldPartitionConvertCommandlet.cpp b/Engine/Source/Editor/UnrealEd/Private/Commandlets/WorldPartitionConvertCommandlet.cpp
--- a/Engine/Source/Editor/UnrealEd/Private/Commandlets/WorldPartitionConvertCommandlet.cpp
+++ b/Engine/Source/Editor/UnrealEd/Private/Commandlets/WorldPartitionConvertCommandlet.cpp
@@ -1500,6 +1500,17 @@ int32 UWorldPartitionConvertCommandlet::Main(const FString& Params)
 		FLevelActorFoldersHelper::SetUseActorFolders(MainLevel, true);
 		MainLevel->SetUseExternalActors(true);
 
+		MainLevel->ForEachActorFolder([this](UActorFolder* ActorFolder)
+		{
+			if (ActorFolder->IsPackageExternal())
+			{
+				PackagesToDelete.Add(ActorFolder->GetExternalPackage());
+				ActorFolder->SetPackageExternal(false);
+			}
+			ActorFolder->SetPackageExternal(true);
+			return true;
+		});
+
 		TSet<FGuid> ActorGuids;
 		for(AActor* Actor: ActorList)
 		{
