Fixed crash in WorldPartitionLandscapeSplineMeshesBuilder when accessing invalid component material.

diff --git a/Engine/Source/Editor/UnrealEd/Private/WorldPartition/WorldPartitionLandscapeSplineMeshesBuilder.cpp b/Engine/Source/Editor/UnrealEd/Private/WorldPartition/WorldPartitionLandscapeSplineMeshesBuilder.cpp
--- a/Engine/Source/Editor/UnrealEd/Private/WorldPartition/WorldPartitionLandscapeSplineMeshesBuilder.cpp
+++ b/Engine/Source/Editor/UnrealEd/Private/WorldPartition/WorldPartitionLandscapeSplineMeshesBuilder.cpp
@@ -118,7 +118,10 @@ int32 UWorldPartitionLandscapeSplineMeshesBuilder::HashStaticMeshComponent(const
 	// Hash used materials
 	for (int32 i = 0; i < InComponent->GetNumMaterials(); ++i)
 	{
-		HashValue = HashCombine(HashValue, GetTypeHash(InComponent->GetMaterial(i)->GetFName()));
+		if (UMaterialInterface* Material = InComponent->GetMaterial(i))
+		{
+			HashValue = HashCombine(HashValue, GetTypeHash(Material->GetFName()));
+		}
 	}
 
 	// Hash spline mesh component params
