MPCook: Fix startup-compiled packages being missing on CookWorkers. CookWorkers need to BlockOnAssetRegistry so that systems that register for AssetRegistry.OnFilesLoaded() on startup have their callbacks called before cook starts.
#rnx

diff --git a/Engine/Source/Editor/UnrealEd/Private/CookOnTheFlyServer.cpp b/Engine/Source/Editor/UnrealEd/Private/CookOnTheFlyServer.cpp
--- a/Engine/Source/Editor/UnrealEd/Private/CookOnTheFlyServer.cpp
+++ b/Engine/Source/Editor/UnrealEd/Private/CookOnTheFlyServer.cpp
@@ -1166,7 +1166,7 @@ bool UCookOnTheFlyServer::IsCookingAgainstFixedBase() const
 
 bool UCookOnTheFlyServer::ShouldPopulateFullAssetRegistry() const
 {
-	return !IsCookingDLC() || CookByTheBookOptions->bDlcLoadMainAssetRegistry;
+	return IsCookWorkerMode() || !IsCookingDLC() || CookByTheBookOptions->bDlcLoadMainAssetRegistry;
 }
 
 FString UCookOnTheFlyServer::GetBaseDirectoryForDLC() const
@@ -7823,7 +7823,10 @@ void UCookOnTheFlyServer::BlockOnAssetRegistry()
 	}
 #endif
 
-	FAssetRegistryGenerator::UpdateAssetManagerDatabase();
+	if (!IsCookWorkerMode())
+	{
+		FAssetRegistryGenerator::UpdateAssetManagerDatabase();
+	}
 	AssetRegistry->ClearGathererCache();
 }
 
@@ -10142,7 +10145,7 @@ void UCookOnTheFlyServer::StartCookAsCookWorker()
 	// Functions in this section are ordered and can depend on the functions before them
 	FBeginCookContext BeginContext = CreateCookWorkerContext();
 	// MPCOOKTODO: Load serialized AssetRegistry from Director
-	UE::Cook::FPackageDatas::OnAssetRegistryGenerated(*AssetRegistry);
+	BlockOnAssetRegistry();
 	CreateSandboxFile(BeginContext);
 	LoadBeginCookConfigSettings(BeginContext);
 	SelectSessionPlatforms(BeginContext);
