Add some missing TEXT() macros to ContentCommandlets

#rnx

diff --git a/Engine/Source/Editor/UnrealEd/Private/Commandlets/ContentCommandlets.cpp b/Engine/Source/Editor/UnrealEd/Private/Commandlets/ContentCommandlets.cpp
--- a/Engine/Source/Editor/UnrealEd/Private/Commandlets/ContentCommandlets.cpp
+++ b/Engine/Source/Editor/UnrealEd/Private/Commandlets/ContentCommandlets.cpp
@@ -289,7 +289,7 @@ int32 UResavePackagesCommandlet::InitializeResaveParameters( const TArray<FStrin
 	}
 
 	// ... if not, load in all packages
-	if( !bExplicitPackages || Switches.Contains("SaveAll"))
+	if( !bExplicitPackages || Switches.Contains(TEXT("SaveAll")))
 	{
 		UE_LOG( LogContentCommandlet, Display, TEXT( "No maps found to save when building HLODs, checking Project Settings for Directory or Asset Path(s)" ) );
 
@@ -1172,13 +1172,13 @@ int32 UResavePackagesCommandlet::Main( const FString& Params )
 	{
 		FString HLODOptions;
 		FParse::Value(*Params, TEXT("BuildOptions="), HLODOptions);
-		bGenerateClusters = HLODOptions.Contains("Clusters");
-		bGenerateMeshProxies = HLODOptions.Contains("Proxies");
-		bForceClusterGeneration = HLODOptions.Contains("ForceClusters");
-		bForceProxyGeneration = HLODOptions.Contains("ForceProxies");
-		bForceSingleClusterForLevel = HLODOptions.Contains("ForceSingleCluster");
-		bSkipSubLevels = HLODOptions.Contains("SkipSubLevels");
-		bHLODMapCleanup = HLODOptions.Contains("MapCleanup");
+		bGenerateClusters = HLODOptions.Contains(TEXT("Clusters"));
+		bGenerateMeshProxies = HLODOptions.Contains(TEXT("Proxies"));
+		bForceClusterGeneration = HLODOptions.Contains(TEXT("ForceClusters"));
+		bForceProxyGeneration = HLODOptions.Contains(TEXT("ForceProxies"));
+		bForceSingleClusterForLevel = HLODOptions.Contains(TEXT("ForceSingleCluster"));
+		bSkipSubLevels = HLODOptions.Contains(TEXT("SkipSubLevels"));
+		bHLODMapCleanup = HLODOptions.Contains(TEXT("MapCleanup"));
 
 		ForceHLODSetupAsset = FString();
 		FParse::Value(*Params, TEXT("ForceHLODSetupAsset="), ForceHLODSetupAsset);
@@ -2093,7 +2093,7 @@ void UResavePackagesCommandlet::PerformAdditionalOperations(class UWorld* World,
 					Builder.BuildMeshesForLODActors(bForceProxyGeneration);
 				}
 
-				FHierarchicalLODUtilitiesModule& Module = FModuleManager::LoadModuleChecked<FHierarchicalLODUtilitiesModule>("HierarchicalLODUtilities");
+				FHierarchicalLODUtilitiesModule& Module = FModuleManager::LoadModuleChecked<FHierarchicalLODUtilitiesModule>(FName("HierarchicalLODUtilities"));
 				FHierarchicalLODProxyProcessor* Processor = Module.GetProxyProcessor();
 
 				while (Processor->IsProxyGenerationRunning())
