Fix fbx importer to ensure the section order is the same has the fbx node material order.

#rnx

diff --git a/Engine/Plugins/Interchange/Runtime/Source/Parsers/Fbx/Private/FbxMesh.cpp b/Engine/Plugins/Interchange/Runtime/Source/Parsers/Fbx/Private/FbxMesh.cpp
--- a/Engine/Plugins/Interchange/Runtime/Source/Parsers/Fbx/Private/FbxMesh.cpp
+++ b/Engine/Plugins/Interchange/Runtime/Source/Parsers/Fbx/Private/FbxMesh.cpp
@@ -540,6 +540,30 @@ bool FMeshDescriptionImporter::FillMeshDescriptionFromFbxMesh(FbxMesh* Mesh, TAr
 
 		TMap<int32, FPolygonGroupID> PolygonGroupMapping;
 
+		//Ensure the polygon groups are create in the same order has the imported materials order
+		for (int32 FbxMeshMaterialIndex = 0; FbxMeshMaterialIndex < MaterialNames.Num(); ++FbxMeshMaterialIndex)
+		{
+			FName MaterialName = MaterialNames[FbxMeshMaterialIndex];
+			if (!PolygonGroupMapping.Contains(FbxMeshMaterialIndex))
+			{
+				FPolygonGroupID ExistingPolygonGroup = INDEX_NONE;
+				for (const FPolygonGroupID PolygonGroupID : MeshDescription->PolygonGroups().GetElementIDs())
+				{
+					if (PolygonGroupImportedMaterialSlotNames[PolygonGroupID] == MaterialName)
+					{
+						ExistingPolygonGroup = PolygonGroupID;
+						break;
+					}
+				}
+				if (ExistingPolygonGroup == INDEX_NONE)
+				{
+					ExistingPolygonGroup = MeshDescription->CreatePolygonGroup();
+					PolygonGroupImportedMaterialSlotNames[ExistingPolygonGroup] = MaterialName;
+				}
+				PolygonGroupMapping.Add(FbxMeshMaterialIndex, ExistingPolygonGroup);
+			}
+		}
+
 		// When importing multiple mesh pieces to the same static mesh.  Ensure each mesh piece has the same number of Uv's
 		int32 ExistingUVCount = VertexInstanceUVs.GetNumChannels();
 
@@ -1127,6 +1151,25 @@ bool FMeshDescriptionImporter::FillMeshDescriptionFromFbxMesh(FbxMesh* Mesh, TAr
 			}
 		}
 	}
+
+	TArray<FPolygonGroupID> EmptyPolygonGroups;
+	for (const FPolygonGroupID PolygonGroupID : MeshDescription->PolygonGroups().GetElementIDs())
+	{
+		if (MeshDescription->GetNumPolygonGroupTriangles(PolygonGroupID) == 0)
+		{
+			EmptyPolygonGroups.Add(PolygonGroupID);
+		}
+	}
+	if (EmptyPolygonGroups.Num() > 0)
+	{
+		for (const FPolygonGroupID PolygonGroupID : EmptyPolygonGroups)
+		{
+			MeshDescription->DeletePolygonGroup(PolygonGroupID);
+		}
+		FElementIDRemappings OutRemappings;
+		MeshDescription->Compact(OutRemappings);
+	}
+	
 	// needed?
 	FBXUVs.Cleanup();
 
diff --git a/Engine/Source/Editor/UnrealEd/Private/Fbx/FbxStaticMeshImport.cpp b/Engine/Source/Editor/UnrealEd/Private/Fbx/FbxStaticMeshImport.cpp
--- a/Engine/Source/Editor/UnrealEd/Private/Fbx/FbxStaticMeshImport.cpp
+++ b/Engine/Source/Editor/UnrealEd/Private/Fbx/FbxStaticMeshImport.cpp
@@ -413,13 +413,14 @@ bool UnFbx::FFbxImporter::BuildStaticMeshFromGeometry(FbxNode* Node, UStaticMesh
 	//
 	int32 MaterialCount = 0;
 	int32 MaterialIndexOffset = 0;
+	TArray<UMaterialInterface*> FbxMeshMaterials;
 	{
 		TRACE_CPUPROFILER_EVENT_SCOPE(CreateMaterials);
 		
-		TArray<UMaterialInterface*> Materials;
+		
 		const bool bForSkeletalMesh = false;
 		
-		FindOrImportMaterialsFromNode(Node, Materials, FBXUVs.UVSets, bForSkeletalMesh);
+		FindOrImportMaterialsFromNode(Node, FbxMeshMaterials, FBXUVs.UVSets, bForSkeletalMesh);
 		if (!ImportOptions->bImportMaterials && ImportOptions->bImportTextures)
 		{
 			//If we are not importing any new material, we might still want to import new textures.
@@ -427,7 +428,7 @@ bool UnFbx::FFbxImporter::BuildStaticMeshFromGeometry(FbxNode* Node, UStaticMesh
 		}
 
 		MaterialCount = Node->GetMaterialCount();
-		check(Materials.Num() == MaterialCount);
+		check(FbxMeshMaterials.Num() == MaterialCount);
 	
 		// Used later to offset the material indices on the raw triangle data
 		MaterialIndexOffset = MeshMaterials.Num();
@@ -438,9 +439,9 @@ bool UnFbx::FFbxImporter::BuildStaticMeshFromGeometry(FbxNode* Node, UStaticMesh
 			FbxSurfaceMaterial *FbxMaterial = Node->GetMaterial(MaterialIndex);
 			NewMaterial->FbxMaterial = FbxMaterial;
 			
-			if (Materials[MaterialIndex])
+			if (FbxMeshMaterials[MaterialIndex])
 			{
-				NewMaterial->Material = Materials[MaterialIndex];
+				NewMaterial->Material = FbxMeshMaterials[MaterialIndex];
 			}
 			else
 			{
@@ -676,6 +677,33 @@ bool UnFbx::FFbxImporter::BuildStaticMeshFromGeometry(FbxNode* Node, UStaticMesh
 		check(PolygonOffset == MeshDescription->Polygons().GetArraySize());
 
 		TMap<int32, FPolygonGroupID> PolygonGroupMapping;
+		
+		//Ensure the polygon groups are create in the same order has the imported materials order
+		for (int32 FbxMeshMaterialIndex = 0; FbxMeshMaterialIndex < FbxMeshMaterials.Num(); ++FbxMeshMaterialIndex)
+		{
+			int32 RealMaterialIndex = FbxMeshMaterialIndex + MaterialIndexOffset;
+			if (!PolygonGroupMapping.Contains(RealMaterialIndex))
+			{
+				UMaterialInterface* Material = MeshMaterials.IsValidIndex(RealMaterialIndex) ? MeshMaterials[RealMaterialIndex].Material : UMaterial::GetDefaultMaterial(MD_Surface);
+				FName ImportedMaterialSlotName = MeshMaterials.IsValidIndex(RealMaterialIndex) ? FName(*MeshMaterials[RealMaterialIndex].GetName()) : (Material != nullptr ? FName(*Material->GetName()) : NAME_None);
+				FPolygonGroupID ExistingPolygonGroup = INDEX_NONE;
+				for (const FPolygonGroupID PolygonGroupID : MeshDescription->PolygonGroups().GetElementIDs())
+				{
+					if (PolygonGroupImportedMaterialSlotNames[PolygonGroupID] == ImportedMaterialSlotName)
+					{
+						ExistingPolygonGroup = PolygonGroupID;
+						break;
+					}
+				}
+				if (ExistingPolygonGroup == INDEX_NONE)
+				{
+					ExistingPolygonGroup = MeshDescription->CreatePolygonGroup();
+					PolygonGroupImportedMaterialSlotNames[ExistingPolygonGroup] = ImportedMaterialSlotName;
+				}
+				PolygonGroupMapping.Add(RealMaterialIndex, ExistingPolygonGroup);
+			}
+		}
+
 
 		// When importing multiple mesh pieces to the same static mesh.  Ensure each mesh piece has the same number of Uv's
 		int32 ExistingUVCount = VertexInstanceUVs.GetNumChannels();
@@ -1199,6 +1227,24 @@ bool UnFbx::FFbxImporter::BuildStaticMeshFromGeometry(FbxNode* Node, UStaticMesh
 		}
 	}
 
+	TArray<FPolygonGroupID> EmptyPolygonGroups;
+	for (const FPolygonGroupID PolygonGroupID : MeshDescription->PolygonGroups().GetElementIDs())
+	{
+		if (MeshDescription->GetNumPolygonGroupTriangles(PolygonGroupID) == 0)
+		{
+			EmptyPolygonGroups.Add(PolygonGroupID);
+		}
+	}
+	if (EmptyPolygonGroups.Num() > 0)
+	{
+		for (const FPolygonGroupID PolygonGroupID : EmptyPolygonGroups)
+		{
+			MeshDescription->DeletePolygonGroup(PolygonGroupID);
+		}
+		FElementIDRemappings OutRemappings;
+		MeshDescription->Compact(OutRemappings);
+	}
+
 	// needed?
 	FBXUVs.Cleanup();
 
