Change UMaterialExpressionReroute CastChecked to Cast as it fires an assert when certain material expressions are undone then recreated. The if statement handles a returned nullptr, and if the expression somehow passes incorrectly, the subsequent CastChecked for UMaterialExpressionComposite should catch it, so the assert is not needed.

diff --git a/Engine/Source/Editor/UnrealEd/Private/MaterialGraph.cpp b/Engine/Source/Editor/UnrealEd/Private/MaterialGraph.cpp
--- a/Engine/Source/Editor/UnrealEd/Private/MaterialGraph.cpp
+++ b/Engine/Source/Editor/UnrealEd/Private/MaterialGraph.cpp
@@ -460,7 +460,7 @@ void UMaterialGraph::LinkGraphNodesFromMaterial()
 				{
 					InputPin->MakeLinkTo(GraphNode->GetOutputPin(GetValidOutputIndex(&ExpressionInput)));
 				}
-				else if (UMaterialExpressionReroute* CompositeReroute = CastChecked<UMaterialExpressionReroute>(ExpressionInput.Expression))
+				else if (UMaterialExpressionReroute* CompositeReroute = Cast<UMaterialExpressionReroute>(ExpressionInput.Expression))
 				{
 					// This is an unseen composite reroute expression, find the actual expression output to connect to.
 					UMaterialExpressionComposite* OwningComposite = CastChecked<UMaterialExpressionComposite>(CompositeReroute->SubgraphExpression);
