Added Editor BuildMenu support for custom build types registered to EditorBuildUtils
#rnx

diff --git a/Engine/Source/Editor/LevelEditor/Private/LevelEditor.cpp b/Engine/Source/Editor/LevelEditor/Private/LevelEditor.cpp
--- a/Engine/Source/Editor/LevelEditor/Private/LevelEditor.cpp
+++ b/Engine/Source/Editor/LevelEditor/Private/LevelEditor.cpp
@@ -1514,6 +1514,14 @@ void FLevelEditorModule::BindGlobalLevelEditorCommands()
 	ActionList.MapAction(Commands.BuildAllLandscape,
 		FExecuteAction::CreateStatic(&FLevelEditorActionCallbacks::BuildAllLandscape_Execute));
 
+	for (int32 Index = 0; Index < FLevelEditorCommands::MaxExternalBuildTypes; ++Index)
+	{
+		ActionList.MapAction(
+			Commands.ExternalBuildTypeCommands[Index],
+			FExecuteAction::CreateStatic(&FLevelEditorActionCallbacks::BuildExternalType_Execute, Index),
+			FCanExecuteAction::CreateStatic(&FLevelEditorActionCallbacks::BuildExternalType_CanExecute, Index));
+	}
+
 	ActionList.MapAction( 
 		Commands.LightingQuality_Production, 
 		FExecuteAction::CreateStatic( &FLevelEditorActionCallbacks::SetLightingQuality, (ELightingBuildQuality)Quality_Production ),
diff --git a/Engine/Source/Editor/LevelEditor/Private/LevelEditorActions.cpp b/Engine/Source/Editor/LevelEditor/Private/LevelEditorActions.cpp
--- a/Engine/Source/Editor/LevelEditor/Private/LevelEditorActions.cpp
+++ b/Engine/Source/Editor/LevelEditor/Private/LevelEditorActions.cpp
@@ -894,6 +894,30 @@ void FLevelEditorActionCallbacks::BuildAllLandscape_Execute()
 	FEditorBuildUtils::EditorBuild(GetWorld(), FBuildOptions::BuildAllLandscape);
 }
 
+bool FLevelEditorActionCallbacks::BuildExternalType_CanExecute(const int32 Index)
+{
+	TArray<FName> BuildTypeNames;
+	FEditorBuildUtils::GetBuildTypes(BuildTypeNames);
+
+	if (BuildTypeNames.IsValidIndex(Index))
+	{
+		return FEditorBuildUtils::EditorCanBuild(GetWorld(), BuildTypeNames[Index]);
+	}
+
+	return false;
+}
+
+void FLevelEditorActionCallbacks::BuildExternalType_Execute(const int32 Index)
+{
+	TArray<FName> BuildTypeNames;
+	FEditorBuildUtils::GetBuildTypes(BuildTypeNames);
+
+	if (BuildTypeNames.IsValidIndex(Index))
+	{
+		FEditorBuildUtils::EditorBuild(GetWorld(), BuildTypeNames[Index]);
+	}
+}
+
 bool FLevelEditorActionCallbacks::IsLightingQualityChecked( ELightingBuildQuality TestQuality )
 {
 	int32 CurrentQualityLevel;
@@ -3464,6 +3488,7 @@ namespace
 {
 	const FName OpenRecentFileBundle = "OpenRecentFile";
 	const FName OpenFavoriteFileBundle = "OpenFavoriteFile";
+	const FName ExternalBuildTypesBundle = "ExternalBuilds";
 }
 
 FLevelEditorCommands::FLevelEditorCommands()
@@ -3477,6 +3502,7 @@ FLevelEditorCommands::FLevelEditorCommands()
 {
 	AddBundle(OpenRecentFileBundle, NSLOCTEXT("LevelEditorCommands", "OpenRecentFileBundle", "Open Recent File"));
 	AddBundle(OpenFavoriteFileBundle, NSLOCTEXT("LevelEditorCommands", "OpenFavoriteFileBundle", "Open Favorite File"));
+	AddBundle(ExternalBuildTypesBundle, NSLOCTEXT("LevelEditorCommands", "ExternalBuildTypesBundle", "Build External Type"));
 }
 
 /** UI_COMMAND takes long for the compile to optimize */
@@ -3528,6 +3554,22 @@ void FLevelEditorCommands::RegisterCommands()
 	UI_COMMAND( ExportAll, "Export All...", "Exports the entire level to a file on disk (multiple formats are supported.)", EUserInterfaceActionType::Button, FInputChord() );
 	UI_COMMAND( ExportSelected, "Export Selected...", "Exports currently-selected objects to a file on disk (multiple formats are supported.)", EUserInterfaceActionType::Button, FInputChord() );
 
+	// External Build commands (inspired from RecentFiles/FavoriteFiles)
+	for (int32 Index = 0; Index < FLevelEditorCommands::MaxExternalBuildTypes; ++Index)
+	{
+		// NOTE: The actual label and tool-tip will be overridden at runtime when the command is bound to a menu item, however
+		// we still need to set one here so that the key bindings UI can function properly
+		ExternalBuildTypeCommands.Add(
+			FUICommandInfoDecl(
+				this->AsShared(),
+				FName(*FString::Printf(TEXT("ExternalBuildType %i"), Index)),
+				FText::Format(NSLOCTEXT("LevelEditorCommands", "ExternalBuildType", "Build Type {0}"), FText::AsNumber(Index)),
+				/*Description*/NSLOCTEXT("LevelEditorCommands", "ExternalBuildToolTip", "Builds an external type"),
+				ExternalBuildTypesBundle)
+			.UserInterfaceType(EUserInterfaceActionType::Button)
+			.DefaultChord(FInputChord()));
+	}
+
 	UI_COMMAND( Build, "Build All Levels", "Builds all levels (precomputes lighting data and visibility data, generates navigation networks and updates brush models.)\nThis action is not available while Play in Editor is active, static lighting is disabled in the project settings, or when previewing less than Shader Model 5", EUserInterfaceActionType::Button, FInputChord() );
 	UI_COMMAND( BuildAndSubmitToSourceControl, "Build and Submit...", "Displays a window that allows you to build all levels and submit them to revision control", EUserInterfaceActionType::Button, FInputChord() );
 	UI_COMMAND( BuildLightingOnly, "Build Lighting", "Only precomputes lighting (all levels.)\nThis action is not available while Play in Editor is active, static lighting is disabled in the project settings, or when previewing less than Shader Model 5", EUserInterfaceActionType::Button, FInputChord(EModifierKey::Control|EModifierKey::Shift, EKeys::Semicolon) );
diff --git a/Engine/Source/Editor/LevelEditor/Private/LevelEditorMenu.cpp b/Engine/Source/Editor/LevelEditor/Private/LevelEditorMenu.cpp
--- a/Engine/Source/Editor/LevelEditor/Private/LevelEditorMenu.cpp
+++ b/Engine/Source/Editor/LevelEditor/Private/LevelEditorMenu.cpp
@@ -24,6 +24,7 @@
 #include "ITranslationEditor.h"
 #include "ILocalizationDashboardModule.h"
 #include "AssetSelection.h"
+#include "EditorBuildUtils.h"
 #include "EditorViewportCommands.h"
 
 #define LOCTEXT_NAMESPACE "LevelEditorMenu"
@@ -609,6 +610,75 @@ void FLevelEditorMenu::RegisterBuildMenu()
 		Section.AddMenuEntry(FLevelEditorCommands::Get().BuildPathsOnly);
 	}
 
+	// Add section for external build types
+	{
+		Menu->AddDynamicSection(NAME_None, FNewToolMenuDelegate::CreateLambda([](UToolMenu* InMenu)
+		{
+			TArray<FName> BuildTypeNames;
+			FEditorBuildUtils::GetBuildTypes(BuildTypeNames);
+			const int32 NumTypes = BuildTypeNames.Num();
+			const int32 AllowedTypes = FMath::Min(NumTypes, FLevelEditorCommands::Get().ExternalBuildTypeCommands.Num());
+			
+			if (AllowedTypes > 0)
+			{
+				TArray<FText> BuildTypeLocalizedNames;
+				TArray<FText> BuildTypeLocalizedSubmenuNames;
+				FEditorBuildUtils::GetBuildTypesLocalizedLabels(BuildTypeLocalizedNames, BuildTypeLocalizedSubmenuNames);
+				check(BuildTypeNames.Num() == BuildTypeLocalizedNames.Num());
+				check(BuildTypeNames.Num() == BuildTypeLocalizedSubmenuNames.Num());
+
+				static FName ExternalBuildTypeSectionName("LevelEditorExternalBuildTypes");
+				FToolMenuSection* BuildTypeSection = nullptr;
+				FToolMenuSection* ExternalBuildTypeSection = nullptr;
+				
+				for (int32 Index = 0; Index < AllowedTypes; ++Index)
+				{
+					const FText& LocalizedSectionName = BuildTypeLocalizedSubmenuNames[Index];
+					// Create dedicated section if provided
+					if (!LocalizedSectionName.IsEmpty())
+					{
+						const FName SectionName(LocalizedSectionName.ToString()); 
+						BuildTypeSection = InMenu->FindSection(SectionName);
+						if (BuildTypeSection == nullptr)
+						{
+							BuildTypeSection = &InMenu->AddSection(SectionName, LocalizedSectionName);	
+						}
+					}
+					// Otherwise create generic section for all external types if not already created
+					else if (ExternalBuildTypeSection == nullptr)
+					{
+						ExternalBuildTypeSection = InMenu->FindSection(ExternalBuildTypeSectionName);
+						if (ExternalBuildTypeSection == nullptr)
+						{
+							ExternalBuildTypeSection = &InMenu->AddSection(ExternalBuildTypeSectionName, LOCTEXT("ExternalBuildTypesHeading", "External Types")); 
+						}
+						BuildTypeSection = ExternalBuildTypeSection;
+					}
+					// Reused existing generic section
+					else
+					{
+						BuildTypeSection = ExternalBuildTypeSection;
+					}
+					
+					const TSharedPtr<FUICommandInfo>& CommandInfo = FLevelEditorCommands::Get().ExternalBuildTypeCommands[Index];
+
+					// Create tooltip from localized entry label if provided, use generic version otherwise
+					const FText ToolTip = BuildTypeLocalizedNames[Index].IsEmpty()
+						? FText::Format(LOCTEXT("ExternalBuildTypeToolTip", "Builds external type: {0}"), FText::FromName(BuildTypeNames[Index]))
+						: BuildTypeLocalizedNames[Index];
+
+					// Create label from localized entry label if provided, use generic version otherwise
+					const FText Label = BuildTypeLocalizedNames[Index].IsEmpty()
+						? FText::Format(LOCTEXT("ExternalBuildTypeLabel", "Build {0}"), FText::FromName(BuildTypeNames[Index]))
+						: BuildTypeLocalizedNames[Index];
+
+					check(BuildTypeSection != nullptr);
+					BuildTypeSection->AddMenuEntry(CommandInfo, Label, ToolTip).Name = NAME_None;
+				}
+			}
+		}));
+	}
+
 	{
 		Menu->AddDynamicSection(NAME_None, FNewToolMenuDelegate::CreateLambda([](UToolMenu* InMenu)
 		{
diff --git a/Engine/Source/Editor/LevelEditor/Public/LevelEditorActions.h b/Engine/Source/Editor/LevelEditor/Public/LevelEditorActions.h
--- a/Engine/Source/Editor/LevelEditor/Public/LevelEditorActions.h
+++ b/Engine/Source/Editor/LevelEditor/Public/LevelEditorActions.h
@@ -69,6 +69,9 @@ public:
 
 	
 	/** Build commands */
+	static constexpr int32 MaxExternalBuildTypes = 10;
+	TArray<TSharedPtr< FUICommandInfo >> ExternalBuildTypeCommands;
+
 	TSharedPtr< FUICommandInfo > Build;
 	TSharedPtr< FUICommandInfo > BuildAndSubmitToSourceControl;
 	TSharedPtr< FUICommandInfo > BuildLightingOnly;
@@ -765,6 +768,8 @@ public:
 	static void BuildTextureStreamingOnly_Execute();
 	static void BuildVirtualTextureOnly_Execute();
 	static void BuildAllLandscape_Execute();
+	static bool BuildExternalType_CanExecute( int32 Index );
+	static void BuildExternalType_Execute( int32 Index );
 	static void SetLightingQuality( ELightingBuildQuality NewQuality );
 	static bool IsLightingQualityChecked( ELightingBuildQuality TestQuality );
 	static float GetLightingDensityIdeal();
diff --git a/Engine/Source/Editor/UnrealEd/Private/EditorBuildUtils.cpp b/Engine/Source/Editor/UnrealEd/Private/EditorBuildUtils.cpp
--- a/Engine/Source/Editor/UnrealEd/Private/EditorBuildUtils.cpp
+++ b/Engine/Source/Editor/UnrealEd/Private/EditorBuildUtils.cpp
@@ -256,6 +256,23 @@ static bool IsBuildCancelled()
 	return GEditor->GetMapBuildCancelled();
 }
 
+
+bool FEditorBuildUtils::EditorCanBuild( UWorld* InWorld, FName Id )
+{
+	// Only process custom types with conditional execution.
+	// Preserving legacy behavior otherwise.
+	if (CustomBuildTypes.Contains(Id))
+	{
+		const FCustomBuildType& CustomBuild = CustomBuildTypes.FindChecked(Id);
+		if (CustomBuild.CanDoBuild.IsBound())
+		{
+			return CustomBuild.CanDoBuild.Execute(InWorld, Id);
+		}
+	}
+
+	return true;
+}
+
 /**
  * Perform an editor build with behavior dependent upon the specified id
  *
@@ -471,7 +488,7 @@ bool FEditorBuildUtils::EditorBuild( UWorld* InWorld, FName Id, const bool bAllo
 				bDirtyPersistentLevel = false;
 			}
 
-			TriggerHierarchicalLODBuilder(InWorld, Id);
+			TriggerHierarchicalLODBuilder(InWorld);
 		}
 	}
 	else if (Id == FBuildOptions::BuildMinimap)
@@ -488,7 +505,7 @@ bool FEditorBuildUtils::EditorBuild( UWorld* InWorld, FName Id, const bool bAllo
 			bShouldMapCheck = false;
 			bDirtyPersistentLevel = false;
 
-			TriggerMinimapBuilder(InWorld, Id);
+			TriggerMinimapBuilder(InWorld);
 		}
 	}
 	else if (Id == FBuildOptions::BuildLandscapeSplineMeshes)
@@ -958,7 +975,7 @@ bool FEditorBuildUtils::WorldPartitionBuildNavigation(const FString& InLongPacka
 	DefaultBuildNavigationOptions->bVerbose = false;
 	DefaultBuildNavigationOptions->bCleanPackages = false;
 
-	TSharedPtr<SWindow> DlgWindow =
+	const TSharedPtr<SWindow> DlgWindow =
 		SNew(SWindow)
 		.Title(LOCTEXT("BuildNavigationWindowTitle", "Build Navigation Settings"))
 		.ClientSize(SWorldPartitionBuildNavigationDialog::DEFAULT_WINDOW_SIZE)
@@ -966,106 +983,131 @@ bool FEditorBuildUtils::WorldPartitionBuildNavigation(const FString& InLongPacka
 		.SupportsMaximize(false)
 		.SizingRule(ESizingRule::FixedSize);
 
-	TSharedRef<SWorldPartitionBuildNavigationDialog> Dialog =
+	const TSharedRef<SWorldPartitionBuildNavigationDialog> Dialog =
 		SNew(SWorldPartitionBuildNavigationDialog)
 		.ParentWindow(DlgWindow)
 		.BuildNavigationOptions(DefaultBuildNavigationOptions);
 
 	DlgWindow->SetContent(Dialog);
 
-	IMainFrameModule& MainFrameModule = FModuleManager::LoadModuleChecked<IMainFrameModule>(TEXT("MainFrame"));
+	const IMainFrameModule& MainFrameModule = FModuleManager::LoadModuleChecked<IMainFrameModule>(TEXT("MainFrame"));
 	FSlateApplication::Get().AddModalWindow(DlgWindow.ToSharedRef(), MainFrameModule.GetParentWindow());
 
 	if (Dialog->ClickedOk())
 	{
-		// Ask user to save dirty packages
-		if (!FEditorFileUtils::SaveDirtyPackages(/*bPromptUserToSave=*/true, /*bSaveMapPackages=*/true, /*bSaveContentPackages=*/false))
-		{
-			return false;
-		}
+		// Try to provide complete Path, if we can't try with project name
+		const FString ProjectPath = FPaths::IsProjectFilePathSet() ? FPaths::GetProjectFilePath() : FApp::GetProjectName();
 
-		// Unload any loaded map
-		if (!UEditorLoadingAndSavingUtils::NewBlankMap(/*bSaveExistingMap*/false))
-		{
-			return false;
-		}
-	
-		FProcHandle ProcessHandle;
-		bool bCancelled = false;
+		const ISourceControlProvider& SCCProvider = ISourceControlModule::Get().GetProvider();
 
-		// Task scope
-		{
-			FScopedSlowTask SlowTask(0, LOCTEXT("WorldPartitionBuildNavigationProgress", "Building navigation..."));
-			SlowTask.MakeDialog(true);
-
-			const FString CurrentExecutableName = FPlatformProcess::ExecutablePath();
-
-			// Try to provide complete Path, if we can't try with project name
-			const FString ProjectPath = FPaths::IsProjectFilePathSet() ? FPaths::GetProjectFilePath() : FApp::GetProjectName();
-
-			uint32 ProcessID;
-
-			ISourceControlProvider& SCCProvider = ISourceControlModule::Get().GetProvider();
-
-			const FString Arguments = FString::Printf(TEXT("\"%s\" -run=WorldPartitionBuilderCommandlet %s %s -SCCProvider=%s %s %s"),
+		const FString Arguments = FString::Printf(TEXT("\"%s\" -run=WorldPartitionBuilderCommandlet %s %s -SCCProvider=%s %s %s"),
 				*ProjectPath,
 				*InLongPackageName,
 				TEXT(" -AllowCommandletRendering -Builder=WorldPartitionNavigationDataBuilder -log=WPNavigationBuilderLog.txt"),
 				*SCCProvider.GetName().ToString(),
 				DefaultBuildNavigationOptions->bVerbose ? TEXT("-Verbose") : TEXT(""),
 				DefaultBuildNavigationOptions->bCleanPackages ? TEXT("-CleanPackages") : TEXT(""));
-				
-			ProcessHandle = FPlatformProcess::CreateProc(*CurrentExecutableName, *Arguments, true, false, false, &ProcessID, 0, nullptr, nullptr);
 		
-			while (FPlatformProcess::IsProcRunning(ProcessHandle))
-			{
-				if (SlowTask.ShouldCancel())
-				{
-					bCancelled = true;
-					FPlatformProcess::TerminateProc(ProcessHandle);
-					break;
-				}
+		RunWorldPartitionBuilder(InLongPackageName,
+			LOCTEXT("WorldPartitionBuildNavigationProgress", "Building navigation..."),
+			LOCTEXT("WorldPartitionBuildNavigationCancelled", "Building navigation cancelled!"),
+			LOCTEXT("WorldPartitionBuildNavigationFailed", "Errors occured during the build process, please refer to the logs ('WPNavigationBuilderLog.txt')."),
+			Arguments);
+	}
+	
+	return false;
+}
 
-				SlowTask.EnterProgressFrame(0);
-				FPlatformProcess::Sleep(0.1);
-			}
-		}
+bool FEditorBuildUtils::RunWorldPartitionBuilder(
+	const FString& MapToLoad,
+	const FText& ProgressText,
+	const FText& CancelledText,
+	const FText& FailureText,
+	const FString& CommandLineArguments
+	)
+{
+	// Ask user to save dirty packages
+	if (!FEditorFileUtils::SaveDirtyPackages(/*bPromptUserToSave=*/true, /*bSaveMapPackages=*/true, /*bSaveContentPackages=*/false))
+	{
+		return false;
+	}
+
+	// Unload any loaded map
+	if (!UEditorLoadingAndSavingUtils::NewBlankMap(/*bSaveExistingMap*/false))
+	{
+		return false;
+	}
+
+	FProcHandle ProcessHandle;
+	bool bCancelled = false;
 
-		int32 Result = 0;
-		if (!bCancelled && FPlatformProcess::GetProcReturnCode(ProcessHandle, &Result))
-		{	
-			// Force a directory watcher tick for the asset registry to get notified of the changes
-			FDirectoryWatcherModule& DirectoryWatcherModule = FModuleManager::Get().LoadModuleChecked<FDirectoryWatcherModule>(TEXT("DirectoryWatcher"));
-			DirectoryWatcherModule.Get()->Tick(-1.0f);
-			
-			// Unload any loaded map
-			if (!UEditorLoadingAndSavingUtils::NewBlankMap(/*bSaveExistingMap*/false))
+	// Task scope
+	{
+		FScopedSlowTask SlowTask(0, ProgressText);
+		SlowTask.MakeDialog(true);
+
+		const FString CurrentExecutableName = FPlatformProcess::ExecutablePath();
+
+		uint32 ProcessID;
+		ProcessHandle = FPlatformProcess::CreateProc(
+			*CurrentExecutableName,
+			*CommandLineArguments,
+			/*bLaunchedDetached*/true,
+			/*bLaunchedHidden*/false,
+			/*bLaunchedReallyHidden*/false,
+			&ProcessID,
+			/*PriorityModifier*/0,
+			/*OptionalWorkingDirectory*/nullptr,
+			/*PipeWriteChild*/nullptr);
+
+		while (FPlatformProcess::IsProcRunning(ProcessHandle))
+		{
+			if (SlowTask.ShouldCancel())
 			{
-				return false;
+				bCancelled = true;
+				FPlatformProcess::TerminateProc(ProcessHandle);
+				break;
 			}
-		
-			// Force registry update before loading converted map
-			const FAssetRegistryModule& AssetRegistryModule = FModuleManager::LoadModuleChecked<FAssetRegistryModule>("AssetRegistry");
-			IAssetRegistry& AssetRegistry = AssetRegistryModule.Get();
-										
-			FString MapToLoad = InLongPackageName;
-		
-			AssetRegistry.ScanModifiedAssetFiles({ MapToLoad });
-			AssetRegistry.ScanPathsSynchronous(ULevel::GetExternalObjectsPaths(MapToLoad), true);
-		
-			FEditorFileUtils::LoadMap(MapToLoad);
-		}
-		else if (bCancelled)
-		{
-			FMessageDialog::Open(EAppMsgType::Ok, LOCTEXT("WorldPartitionBuildNavigationCancelled", "Building navigation cancelled!"));
+
+			SlowTask.EnterProgressFrame(0);
+			FPlatformProcess::Sleep(0.1);
 		}
-	
-		if (Result != 0)
+	}
+
+	int32 Result = 0;
+	if (!bCancelled && FPlatformProcess::GetProcReturnCode(ProcessHandle, &Result))
+	{
+		// Force a directory watcher tick for the asset registry to get notified of the changes
+		FDirectoryWatcherModule& DirectoryWatcherModule = FModuleManager::Get().LoadModuleChecked<FDirectoryWatcherModule>(TEXT("DirectoryWatcher"));
+		DirectoryWatcherModule.Get()->Tick(-1.0f);
+
+		// Unload any loaded map
+		if (!UEditorLoadingAndSavingUtils::NewBlankMap(/*bSaveExistingMap*/false))
 		{
-			FMessageDialog::Open(EAppMsgType::Ok, LOCTEXT("WorldPartitionBuildNavigationFailed", "Errors occured during the build process, please refer to the logs ('WPNavigationBuilderLog.txt')."));
+			return false;
 		}
+
+		// Force registry update before loading converted map
+		const FAssetRegistryModule& AssetRegistryModule = FModuleManager::LoadModuleChecked<FAssetRegistryModule>("AssetRegistry");
+		IAssetRegistry& AssetRegistry = AssetRegistryModule.Get();
+
+		AssetRegistry.ScanModifiedAssetFiles({MapToLoad});
+		AssetRegistry.ScanPathsSynchronous(ULevel::GetExternalObjectsPaths(MapToLoad), true);
+
+		FEditorFileUtils::LoadMap(MapToLoad);
+		return true;
 	}
 	
+	if (bCancelled)
+	{
+		FMessageDialog::Open(EAppMsgType::Ok, CancelledText);
+	}
+
+	if (Result != 0)
+	{
+		FMessageDialog::Open(EAppMsgType::Ok, FailureText);
+	}
+
 	return false;
 }
 
@@ -1096,14 +1138,36 @@ bool FEditorBuildUtils::IsBuildCurrentlyRunning()
 
 /**
  * Register a custom build type.
- * @param Id The identifier to use for this build type.
- * @param DoBuild The delegate to execute to run this build.
- * @param BuildAllExtensionPoint If a valid name, run this build *before* running the build with this id when performing a Build All.
  */
-void FEditorBuildUtils::RegisterCustomBuildType(FName Id, const FDoEditorBuildDelegate& DoBuild, FName BuildAllExtensionPoint)
+void FEditorBuildUtils::RegisterCustomBuildType(
+	const FName Id,
+	const FDoEditorBuildDelegate& DoBuild,
+	const FName BuildAllExtensionPoint,
+	const FText& MenuEntryLabel,
+	const FText& MenuSectionLabel)
 {
 	check(!CustomBuildTypes.Contains(Id));
-	CustomBuildTypes.Add(Id, FCustomBuildType(DoBuild, BuildAllExtensionPoint));
+	CustomBuildTypes.Add(Id, FCustomBuildType(DoBuild, BuildAllExtensionPoint, MenuEntryLabel, MenuSectionLabel));
+
+	if (BuildAllExtensionPoint != NAME_None)
+	{
+		FBuildAllHandler::Get().AddCustomBuildStep(Id, BuildAllExtensionPoint);
+	}
+}
+
+/**
+ * Register a custom build type.
+ */
+void FEditorBuildUtils::RegisterCustomBuildType(
+	const FName Id,
+	const FCanDoEditorBuildDelegate& CanDoBuild,
+	const FDoEditorBuildDelegate& DoBuild,
+	const FName BuildAllExtensionPoint,
+	const FText& MenuEntryLabel,
+	const FText& MenuSectionLabel)
+{
+	check(!CustomBuildTypes.Contains(Id));
+	CustomBuildTypes.Add(Id, FCustomBuildType(CanDoBuild, DoBuild, BuildAllExtensionPoint, MenuEntryLabel, MenuSectionLabel));
 
 	if (BuildAllExtensionPoint != NAME_None)
 	{
@@ -1121,6 +1185,20 @@ void FEditorBuildUtils::UnregisterCustomBuildType(FName Id)
 	FBuildAllHandler::Get().RemoveCustomBuildStep(Id);
 }
 
+void FEditorBuildUtils::GetBuildTypes(TArray<FName>& Types)
+{
+	CustomBuildTypes.GetKeys(Types);
+}
+
+void FEditorBuildUtils::GetBuildTypesLocalizedLabels(TArray<FText>& RegisteredBuildTypesEntryLabels, TArray<FText>& RegisteredBuildTypesSectionLabels)
+{
+	for (auto It = CustomBuildTypes.CreateConstIterator(); It; ++It)
+	{
+		RegisteredBuildTypesEntryLabels.Add(It->Value.MenuEntryLabel);
+		RegisteredBuildTypesSectionLabels.Add(It->Value.MenuSectionLabel);
+	}
+}
+
 /**
  * Initialise Build All handler.
  */
@@ -1224,12 +1302,12 @@ void FBuildAllHandler::ProcessBuild(const TWeakPtr<SBuildProgressWidget>& BuildP
 		else if (StepId == FBuildOptions::BuildHierarchicalLOD)
 		{
 			BuildProgressWidget.Pin()->SetBuildType(SBuildProgressWidget::BUILDTYPE_HLODs);
-			FEditorBuildUtils::TriggerHierarchicalLODBuilder(CurrentWorld, CurrentBuildId);
+			FEditorBuildUtils::TriggerHierarchicalLODBuilder(CurrentWorld);
 		}
 		else if (StepId == FBuildOptions::BuildMinimap)
 		{
 			BuildProgressWidget.Pin()->SetBuildType(SBuildProgressWidget::BUILDTYPE_Minimap);
-			FEditorBuildUtils::TriggerMinimapBuilder(CurrentWorld, CurrentBuildId);
+			FEditorBuildUtils::TriggerMinimapBuilder(CurrentWorld);
 		}
 		else if (StepId == FBuildOptions::BuildLandscapeSplineMeshes)
 		{
@@ -1287,17 +1365,20 @@ void FBuildAllHandler::ProcessBuild(const TWeakPtr<SBuildProgressWidget>& BuildP
 		}
 		else
 		{
-			auto& CustomBuildType = FEditorBuildUtils::CustomBuildTypes[StepId];
-			auto Result = CustomBuildType.DoBuild.Execute(CurrentWorld, CurrentBuildId);
-
-			if (Result == EEditorBuildResult::InProgress)
+			FEditorBuildUtils::FCustomBuildType& CustomBuildType = FEditorBuildUtils::CustomBuildTypes[StepId];
+			if (CustomBuildType.CanDoBuild.IsBound() == false || CustomBuildType.CanDoBuild.Execute(CurrentWorld, CurrentBuildId))
 			{
-				// Build & Submit builds must be synchronous.
-				check(CurrentBuildId != FBuildOptions::BuildAllSubmit);
+				const EEditorBuildResult Result = CustomBuildType.DoBuild.Execute(CurrentWorld, CurrentBuildId);
 
-				// Build step is running asynchronously, so let it run.
-				FEditorBuildUtils::InProgressBuildId = CurrentBuildId;
-				break;
+				if (Result == EEditorBuildResult::InProgress)
+				{
+					// Build & Submit builds must be synchronous.
+					check(CurrentBuildId != FBuildOptions::BuildAllSubmit);
+
+					// Build step is running asynchronously, so let it run.
+					FEditorBuildUtils::InProgressBuildId = CurrentBuildId;
+					break;
+				}
 			}
 		}
 
@@ -1316,7 +1397,7 @@ void FBuildAllHandler::BuildFinished()
 	CurrentBuildId = NAME_None;
 }
 
-void FEditorBuildUtils::TriggerHierarchicalLODBuilder(UWorld* InWorld, FName Id)
+void FEditorBuildUtils::TriggerHierarchicalLODBuilder(UWorld* InWorld)
 {
 	if (InWorld->IsPartitionedWorld())
 	{
@@ -1331,7 +1412,7 @@ void FEditorBuildUtils::TriggerHierarchicalLODBuilder(UWorld* InWorld, FName Id)
 	}
 }
 
-void FEditorBuildUtils::TriggerMinimapBuilder(UWorld* InWorld, FName Id)
+void FEditorBuildUtils::TriggerMinimapBuilder(UWorld* InWorld)
 {
 	if (InWorld->IsPartitionedWorld())
 	{
diff --git a/Engine/Source/Editor/UnrealEd/Public/EditorBuildUtils.h b/Engine/Source/Editor/UnrealEd/Public/EditorBuildUtils.h
--- a/Engine/Source/Editor/UnrealEd/Public/EditorBuildUtils.h
+++ b/Engine/Source/Editor/UnrealEd/Public/EditorBuildUtils.h
@@ -46,13 +46,21 @@ struct FBuildOptions
 /**
  * Result of a custom editor build.
  */
-enum class EEditorBuildResult
+enum class EEditorBuildResult : uint8
 {
 	Success,		// The build step completed successfully
 	Skipped,		// The build step was skipped for some reason (e.g. cancelled)
 	InProgress,		// The build step is running asynchronously
 };
 
+/**
+ * Delegate to validate if a custom editor build can be executed.
+ * @param UWorld* The world to run the build on.
+ * @param FName The Id of the build being run (either the registered build Id, or one of the BuildAll types).
+ * @return Whether the build step could run or not.
+ */
+DECLARE_DELEGATE_RetVal_TwoParams(bool, FCanDoEditorBuildDelegate, const UWorld*, FName);
+
 /**
  * Delegate for performing a custom editor build.
  * @param UWorld* The world to run the build on.
@@ -121,6 +129,14 @@ public:
 	 */
 	static UNREALED_API bool EditorAutomatedBuildAndSubmit( const FEditorAutomatedBuildSettings& BuildSettings, FText& OutErrorMessages );
 
+	/**
+	 * Validate if the editor build matching the specified id can be executed.
+	 * @param	InWorld				WorldContext
+	 * @param	Id					Action Id specifying what kind of build is requested
+	 * @return	true if the build can be executed for the provided world
+	 */
+	static UNREALED_API bool EditorCanBuild( UWorld* InWorld, FName Id );
+
 	/**
 	 * Perform an editor build with behavior dependent upon the specified id
 	 *
@@ -201,14 +217,71 @@ public:
 	 * @param Id The identifier to use for this build type.
 	 * @param DoBuild The delegate to execute to run this build.
 	 * @param BuildAllExtensionPoint If a valid name, run this build *before* running the build with this id when performing a Build All.
+	 * @param MenuEntryLabel If non empty, will be used as label for the command in the menu. Otherwise `Build {FText::FromName(Id)}` will be used.
+	 * @param MenuSectionLabel If non empty, will be used a label for a new submenu for this build type. Otherwise the entry will be created under `External Types`.
+	 */
+	static UNREALED_API void RegisterCustomBuildType(
+		const FName Id,
+		const FDoEditorBuildDelegate& DoBuild,
+		const FName BuildAllExtensionPoint,
+		const FText& MenuEntryLabel = FText::GetEmpty(),
+		const FText& MenuSectionLabel = FText::GetEmpty());
+	
+	/**
+	 * Register a custom build type.
+	 * @param Id The identifier to use for this build type.
+	 * @param CanDoBuild The delegate to validate if a given build could be executed.
+	 * @param DoBuild The delegate to execute to run this build.
+	 * @param BuildAllExtensionPoint If a valid name, run this build *before* running the build with this id when performing a Build All.
+	 * @param MenuEntryLabel If non empty, will be used as label for the command in the menu. Otherwise `Build {FText::FromName(Id)}` will be used.
+	 * @param MenuSectionLabel If non empty, will be used a label for a new submenu for this build type. Otherwise the entry will be created under `External Types`.
 	 */
-	static UNREALED_API void RegisterCustomBuildType(FName Id, const FDoEditorBuildDelegate& DoBuild, FName BuildAllExtensionPoint);
+	static UNREALED_API void RegisterCustomBuildType(
+		const FName Id,
+		const FCanDoEditorBuildDelegate& CanDoBuild,
+		const FDoEditorBuildDelegate& DoBuild,
+		const FName BuildAllExtensionPoint,
+		const FText& MenuEntryLabel = FText::GetEmpty(),
+		const FText& MenuSectionLabel = FText::GetEmpty()
+		);
 
 	/**
 	 * Unregister a custom build type.
 	 * @param Id The identifier of the build type to unregister.
 	 */
 	static UNREALED_API void UnregisterCustomBuildType(FName Id);
+
+	/**
+	 * Fills `RegisteredBuildTypes` with the names of all build types registered by 'RegisterCustomBuildType'.
+	 * @param RegisteredBuildTypes Names of all registered build types.
+	 */
+	static UNREALED_API void GetBuildTypes(TArray<FName>& RegisteredBuildTypes);
+
+	/**
+	 * Fills provided arrays with localized menu entry and section labels.
+	 * @param RegisteredBuildTypesEntryLabels Labels for each registered build types. Can be empty.
+	 * @param RegisteredBuildTypesSectionLabels Labels for each registered build types section where the entry should reside. Can be empty.
+	 */
+	static UNREALED_API void GetBuildTypesLocalizedLabels(TArray<FText>& RegisteredBuildTypesEntryLabels, TArray<FText>& RegisteredBuildTypesSectionLabels);
+	
+	/**
+	 * Runs another instance of the current executable with the provided command line arguments.
+	 * On success the current level will be unloaded. Using `MapToLoad` allows to reload that specific level to reflect the changes.
+	 * @param MapToLoad Map to load after successfully run the process
+	 * @param ProgressText Message used by the progress dialog
+	 * @param CancelledText Message shown to the user after cancelling the task 
+	 * @param FailureText Message shown to the user if the process returned an error code
+	 * @param CommandLineArguments Arguments passed to the other instance commandline
+	 * 
+	 * @return True if process completed without returning any error code.  
+	 */
+	static UNREALED_API bool RunWorldPartitionBuilder(
+		const FString& MapToLoad,
+		const FText& ProgressText,
+		const FText& CancelledText,
+		const FText& FailureText,
+		const FString& CommandLineArguments
+		);
 private:
 
 	/**
@@ -259,7 +332,7 @@ private:
 	 * Trigger navigation builder to (re)generate NavMesh 
 	 *
 	 * @param	InOutWorld		WorldContext
-	 * @param	BuildSettings	Build settings that will be used for the editor build
+	 * @param	Id	Build options requested
 	 */
 	static void TriggerNavigationBuilder(UWorld*& InOutWorld, FName Id);
 
@@ -269,17 +342,15 @@ private:
 	 * Trigger HLOD builder to (re)generate HLOD actors
 	 *
 	 * @param	InWorld			WorldContext
-	 * @param	BuildSettings	Build settings that will be used for the editor build
 	 */
-	static void TriggerHierarchicalLODBuilder(UWorld* InWorld, FName Id);
+	static void TriggerHierarchicalLODBuilder(UWorld* InWorld);
 
 	/** 
 	 * Trigger minimap builder to (re)generate minimap
 	 *
 	 * @param	InWorld			WorldContext
-	 * @param	BuildSettings	Build settings that will be used for the editor build
 	 */
-	static void TriggerMinimapBuilder(UWorld* InWorld, FName Id);
+	static void TriggerMinimapBuilder(UWorld* InWorld);
 
 	/**
 	 * Trigger LandscapeSplineMeshes builder to (re)generate landscape spline meshes actors
@@ -305,13 +376,33 @@ private:
 	 */
 	struct FCustomBuildType
 	{
+		FCanDoEditorBuildDelegate CanDoBuild;
 		FDoEditorBuildDelegate DoBuild;
 		FName BuildAllExtensionPoint;
-
-		FCustomBuildType(const FDoEditorBuildDelegate& InDoBuild, FName InBuildAllExtensionPoint)
+		const FText MenuEntryLabel;
+		const FText MenuSectionLabel;
+
+		FCustomBuildType(
+			const FDoEditorBuildDelegate& InDoBuild,
+			const FName InBuildAllExtensionPoint,
+			const FText& InMenuEntryLabel,
+			const FText& InMenuSectionLabel)
 			: DoBuild(InDoBuild)
 			, BuildAllExtensionPoint(InBuildAllExtensionPoint)
+			, MenuEntryLabel(InMenuEntryLabel)
+			, MenuSectionLabel(InMenuSectionLabel)
 		{}
+
+		FCustomBuildType(
+			const FCanDoEditorBuildDelegate& InCanDoBuild,
+			const FDoEditorBuildDelegate& InDoBuild,
+			const FName InBuildAllExtensionPoint,
+			const FText& InMenuEntryLabel,
+			const FText& InMenuSectionLabel)
+			: FCustomBuildType(InDoBuild, InBuildAllExtensionPoint, InMenuEntryLabel, InMenuSectionLabel)
+		{
+			CanDoBuild = InCanDoBuild;
+		}
 	};
 
 	/** Map of custom build types registered with us. */
