MPCook: Double the WaitForConnect timeout; some preflights are hitting the timeout due to AssetRegistry scanning taking more than 10 minutes.
#rnx

diff --git a/Engine/Source/Editor/UnrealEd/Private/Cooker/CookWorkerServer.cpp b/Engine/Source/Editor/UnrealEd/Private/Cooker/CookWorkerServer.cpp
--- a/Engine/Source/Editor/UnrealEd/Private/Cooker/CookWorkerServer.cpp
+++ b/Engine/Source/Editor/UnrealEd/Private/Cooker/CookWorkerServer.cpp
@@ -523,7 +523,7 @@ void FCookWorkerServer::LaunchProcess()
 void FCookWorkerServer::TickWaitForConnect()
 {
 	constexpr float TestProcessExistencePeriod = 1.f;
-	constexpr float WaitForConnectTimeout = 60.f * 10;
+	constexpr float WaitForConnectTimeout = 60.f * 20;
 
 	check(!Socket); // When the Socket is assigned we leave the WaitForConnect state, and we set it to null before entering
 
