Allow asset factories to be registered at runtime

Also expand usage of placement system to ensure these factories are used in more places

diff --git a/Engine/Source/Editor/EditorFramework/Private/Subsystems/PlacementSubsystem.cpp b/Engine/Source/Editor/EditorFramework/Private/Subsystems/PlacementSubsystem.cpp
--- a/Engine/Source/Editor/EditorFramework/Private/Subsystems/PlacementSubsystem.cpp
+++ b/Engine/Source/Editor/EditorFramework/Private/Subsystems/PlacementSubsystem.cpp
@@ -137,6 +137,16 @@ TScriptInterface<IAssetFactoryInterface> UPlacementSubsystem::GetAssetFactoryFro
 	return nullptr;
 }
 
+void UPlacementSubsystem::RegisterAssetFactory(TScriptInterface<IAssetFactoryInterface> AssetFactory)
+{
+	AssetFactories.Add(AssetFactory);
+}
+
+void UPlacementSubsystem::UnregisterAssetFactory(TScriptInterface<IAssetFactoryInterface> AssetFactory)
+{
+	AssetFactories.Remove(AssetFactory);
+}
+
 void UPlacementSubsystem::UnregisterPlacementFactories()
 {
 	AssetFactories.Empty();
diff --git a/Engine/Source/Editor/EditorFramework/Public/Subsystems/PlacementSubsystem.h b/Engine/Source/Editor/EditorFramework/Public/Subsystems/PlacementSubsystem.h
--- a/Engine/Source/Editor/EditorFramework/Public/Subsystems/PlacementSubsystem.h
+++ b/Engine/Source/Editor/EditorFramework/Public/Subsystems/PlacementSubsystem.h
@@ -126,7 +126,11 @@ public:
 	FSimpleMulticastDelegate& OnPlacementFactoriesRegistered();
 
 	TScriptInterface<IAssetFactoryInterface> GetAssetFactoryFromFactoryClass(TSubclassOf<UClass> InFactoryInterfaceClass) const;
+	
+	void RegisterAssetFactory(TScriptInterface<IAssetFactoryInterface> AssetFactory);
 
+	void UnregisterAssetFactory(TScriptInterface<IAssetFactoryInterface> AssetFactory);
+	
 private:
 	void RegisterPlacementFactories();
 	void UnregisterPlacementFactories();
diff --git a/Engine/Source/Editor/UnrealEd/Private/AssetSelection.cpp b/Engine/Source/Editor/UnrealEd/Private/AssetSelection.cpp
--- a/Engine/Source/Editor/UnrealEd/Private/AssetSelection.cpp
+++ b/Engine/Source/Editor/UnrealEd/Private/AssetSelection.cpp
@@ -843,25 +843,12 @@ AActor* FActorFactoryAssetProxy::AddActorForAsset( UObject* AssetObj, bool Selec
 				Result = PrivateAddActor( AssetObj, FactoryToUse, SelectActor, ObjectFlags, Name );
 			}
 		}
-		// If no specific factory has been provided, find the highest priority one that is valid for the asset and use
-		// it to create the actor
+		// If no specific factory has been provided, use the placement subsystem to find the appropriate factory and place
 		else
 		{
-			const TArray<UActorFactory*>& ActorFactories = GEditor->ActorFactories;
-			for ( int32 FactoryIdx = 0; FactoryIdx < ActorFactories.Num(); FactoryIdx++ )
-			{
-				UActorFactory* ActorFactory = ActorFactories[FactoryIdx];
-
-				// Check if the actor can be created using this factory, making sure to check for an asset to be assigned from the selector
-				if ( ActorFactory->CanCreateActorFrom( AssetData, UnusedErrorMessage ) )
-				{
-					Result = PrivateAddActor(AssetObj, ActorFactory, SelectActor, ObjectFlags, Name);
-					if ( Result != NULL )
-					{
-						break;
-					}
-				}
-			}
+			UPlacementSubsystem* PlacementSubsystem = GEditor->GetEditorSubsystem<UPlacementSubsystem>();
+			TScriptInterface<IAssetFactoryInterface> AssetFactory = PlacementSubsystem->FindAssetFactoryFromAssetData(AssetData);
+			Result = PrivateAddActor(AssetObj, Cast<UActorFactory>(AssetFactory.GetObject()), SelectActor, ObjectFlags, Name);
 		}
 	}
 
diff --git a/Engine/Source/Editor/UnrealEd/Private/LevelEditorViewport.cpp b/Engine/Source/Editor/UnrealEd/Private/LevelEditorViewport.cpp
--- a/Engine/Source/Editor/UnrealEd/Private/LevelEditorViewport.cpp
+++ b/Engine/Source/Editor/UnrealEd/Private/LevelEditorViewport.cpp
@@ -107,6 +107,7 @@
 #include "GenericPlatform/GenericPlatformInputDeviceMapper.h"
 #include "Rendering/StaticLightingSystemInterface.h"
 #include "TextureResource.h"
+#include "Subsystems/PlacementSubsystem.h"
 
 DEFINE_LOG_CATEGORY(LogEditorViewport);
 
@@ -1442,8 +1443,10 @@ FDropQuery FLevelEditorViewportClient::CanDropObjectsAtCoordinates(int32 MouseX,
 	UObject* AssetObj = AssetData.GetAsset();
 	UClass* ClassObj = Cast<UClass>(AssetObj);
 
+	UPlacementSubsystem* PlacementSubsystem = GEditor->GetEditorSubsystem<UPlacementSubsystem>();
 	// Check if the asset has an actor factory
-	bool bHasActorFactory = FActorFactoryAssetProxy::GetFactoryForAsset(AssetData) != nullptr;
+	bool bHasActorFactory = PlacementSubsystem && PlacementSubsystem->FindAssetFactoryFromAssetData(AssetData) != nullptr;
+	bHasActorFactory = bHasActorFactory || FActorFactoryAssetProxy::GetFactoryForAsset(AssetData) != nullptr;
 
 	if (ClassObj)
 	{
