Updated UnrealEdSharedPCH.h to include EngineSharedPCH.h and removed redundant includes.

diff --git a/Engine/Source/Editor/UnrealEd/Public/UnrealEdSharedPCH.h b/Engine/Source/Editor/UnrealEd/Public/UnrealEdSharedPCH.h
--- a/Engine/Source/Editor/UnrealEd/Public/UnrealEdSharedPCH.h
+++ b/Engine/Source/Editor/UnrealEd/Public/UnrealEdSharedPCH.h
@@ -2,585 +2,50 @@
 
 #pragma once
 
+#include "EngineSharedPCH.h"
+
 // From Core:
-#include "CoreTypes.h"
-#include "Misc/Exec.h"
-#include "Misc/AssertionMacros.h"
-#include "HAL/PlatformMisc.h"
-#include "GenericPlatform/GenericPlatformMisc.h"
-#include "CoreFwd.h"
-#include "Containers/ContainersFwd.h"
-#include "UObject/UObjectHierarchyFwd.h"
-#include "HAL/PlatformCrt.h"
-#include "Containers/Array.h"
-#include "HAL/UnrealMemory.h"
-#include "Templates/IsPointer.h"
-#include "HAL/PlatformMemory.h"
-#include "GenericPlatform/GenericPlatformMemory.h"
-#include "HAL/MemoryBase.h"
-#include "Misc/OutputDevice.h"
-#include "Logging/LogVerbosity.h"
-#include "Misc/VarArgs.h"
-#include "HAL/PlatformAtomics.h"
-#include "GenericPlatform/GenericPlatformAtomics.h"
-#include "Templates/UnrealTypeTraits.h"
-#include "Templates/AndOrNot.h"
-#include "Templates/IsArithmetic.h"
-#include "Templates/IsPODType.h"
-#include "Templates/IsTriviallyCopyConstructible.h"
-#include "Templates/UnrealTemplate.h"
-#include "Templates/EnableIf.h"
-#include "Templates/RemoveReference.h"
-#include "Templates/TypeCompatibleBytes.h"
-#include "Templates/ChooseClass.h"
-#include "Templates/IntegralConstant.h"
-#include "Templates/IsClass.h"
-#include "Traits/IsContiguousContainer.h"
-#include "Containers/ContainerAllocationPolicies.h"
-#include "HAL/PlatformMath.h"
-#include "GenericPlatform/GenericPlatformMath.h"
-#include "Templates/MemoryOps.h"
-#include "Templates/IsTriviallyCopyAssignable.h"
-#include "Templates/IsTriviallyDestructible.h"
-#include "Math/NumericLimits.h"
-#include "Serialization/Archive.h"
-#include "Templates/IsEnumClass.h"
-#include "HAL/PlatformProperties.h"
-#include "GenericPlatform/GenericPlatformProperties.h"
-#include "Misc/Compression.h"
-#include "Misc/EngineVersionBase.h"
-#include "Internationalization/TextNamespaceFwd.h"
-#include "Templates/Less.h"
-#include "Templates/Sorting.h"
-#include "Containers/UnrealString.h"
-#include "Misc/CString.h"
-#include "Misc/Char.h"
-#include "HAL/PlatformString.h"
-#include "GenericPlatform/GenericPlatformStricmp.h"
-#include "GenericPlatform/GenericPlatformString.h"
-#include "Misc/Crc.h"
-#include "Math/UnrealMathUtility.h"
-#include "Containers/Map.h"
-#include "Misc/StructBuilder.h"
-#include "Templates/AlignmentTemplates.h"
-#include "Templates/Function.h"
-#include "Templates/Decay.h"
-#include "Templates/Invoke.h"
-#include "Templates/PointerIsConvertibleFromTo.h"
-#include "Containers/Set.h"
-#include "Templates/TypeHash.h"
-#include "Containers/SparseArray.h"
-#include "Containers/ScriptArray.h"
-#include "Containers/BitArray.h"
-#include "Algo/Reverse.h"
-#include "Math/IntPoint.h"
-#include "Logging/LogMacros.h"
-#include "Logging/LogCategory.h"
-#include "UObject/NameTypes.h"
-#include "HAL/CriticalSection.h"
-#include "Misc/Timespan.h"
-#include "Containers/StringConv.h"
-#include "UObject/UnrealNames.h"
-#include "Templates/SharedPointer.h"
-#include "CoreGlobals.h"
-#include "HAL/PlatformTLS.h"
-#include "GenericPlatform/GenericPlatformTLS.h"
-#include "Delegates/Delegate.h"
-#include "UObject/WeakObjectPtrTemplates.h"
-#include "Delegates/MulticastDelegateBase.h"
-#include "Delegates/IDelegateInstance.h"
-#include "Delegates/DelegateSettings.h"
-#include "Delegates/DelegateBase.h"
-#include "Delegates/IntegerSequence.h"
-#include "Templates/Tuple.h"
-#include "UObject/ScriptDelegates.h"
-#include "Misc/Optional.h"
-#include "CoreMinimal.h"
-#include "Math/UnrealMath.h"
-#include "Templates/IsFloatingPoint.h"
-#include "Templates/IsIntegral.h"
-#include "Templates/IsSigned.h"
-#include "Templates/Greater.h"
-#include "Math/ColorList.h"
-#include "Math/Color.h"
-#include "Misc/Parse.h"
-#include "Math/IntRect.h"
-#include "Math/Vector2D.h"
-#include "Math/Edge.h"
-#include "Math/Vector.h"
-#include "Misc/ByteSwap.h"
-#include "Internationalization/Text.h"
-#include "Containers/EnumAsByte.h"
-#include "Internationalization/CulturePointer.h"
-#include "Internationalization/TextLocalizationManager.h"
-#include "Templates/UniquePtr.h"
-#include "Templates/IsArray.h"
-#include "Templates/RemoveExtent.h"
-#include "Internationalization/Internationalization.h"
-#include "Templates/UniqueObj.h"
-#include "Math/IntVector.h"
-#include "Math/CapsuleShape.h"
-#include "Math/RangeSet.h"
-#include "Math/Range.h"
-#include "Misc/DateTime.h"
-#include "Math/RangeBound.h"
-#include "Math/Box2D.h"
-#include "Math/BoxSphereBounds.h"
-#include "Math/Sphere.h"
-#include "Math/Box.h"
-#include "Math/OrientedBox.h"
-#include "Math/Interval.h"
-#include "Math/RotationAboutPointMatrix.h"
-#include "Math/Matrix.h"
-#include "Math/Vector4.h"
-#include "Math/Plane.h"
-#include "UObject/ObjectVersion.h"
-#include "Math/Rotator.h"
-#include "Math/VectorRegister.h"
-#include "Math/Axis.h"
-#include "Math/RotationTranslationMatrix.h"
-#include "Math/ScaleRotationTranslationMatrix.h"
-#include "Math/RotationMatrix.h"
-#include "Math/PerspectiveMatrix.h"
-#include "Math/OrthoMatrix.h"
-#include "Math/TranslationMatrix.h"
-#include "Math/QuatRotationTranslationMatrix.h"
-#include "Math/Quat.h"
-#include "Math/InverseRotationMatrix.h"
-#include "Math/ScaleMatrix.h"
-#include "Math/MirrorMatrix.h"
-#include "Math/ClipProjectionMatrix.h"
-#include "Math/InterpCurve.h"
-#include "Math/TwoVectors.h"
-#include "Math/InterpCurvePoint.h"
-#include "Math/CurveEdInterface.h"
-#include "Math/Float16Color.h"
-#include "Math/Float16.h"
-#include "Math/Float32.h"
-#include "Math/Vector2DHalf.h"
-#include "Math/Transform.h"
-#include "Math/ConvexHull2d.h"
-#include "HAL/ThreadSingleton.h"
-#include "HAL/TlsAutoCleanup.h"
-#include "HAL/ThreadSafeCounter.h"
-#include "Misc/EnumClassFlags.h"
-#include "HAL/PlatformTime.h"
-#include "GenericPlatform/GenericPlatformTime.h"
-#include "Stats/Stats.h"
-#include "Containers/LockFreeList.h"
-#include "Misc/NoopCounter.h"
-#include "Containers/ChunkedArray.h"
-#include "Containers/IndirectArray.h"
-#include "Misc/OutputDeviceRedirector.h"
-#include "ProfilingDebugging/ResourceSize.h"
-#include "Misc/Guid.h"
-#include "Math/RandomStream.h"
-#include "Misc/Attribute.h"
-#include "Misc/CoreMisc.h"
-#include "GenericPlatform/GenericWindow.h"
-#include "GenericPlatform/GenericApplication.h"
-#include "GenericPlatform/GenericApplicationMessageHandler.h"
-#include "GenericPlatform/GenericWindowDefinition.h"
-#include "GenericPlatform/ICursor.h"
-#include "Containers/LockFreeFixedSizeAllocator.h"
-#include "Misc/MemStack.h"
-#include "Math/TransformCalculus.h"
-#include "Math/TransformCalculus2D.h"
-#include "Containers/List.h"
-#include "HAL/IConsoleManager.h"
-#include "Modules/ModuleInterface.h"
-#include "Serialization/BitReader.h"
-#include "Misc/NetworkGuid.h"
-#include "Serialization/BitWriter.h"
-#include "HAL/Event.h"
-#include "Templates/RefCounting.h"
-#include "Async/TaskGraphInterfaces.h"
-#include "Misc/Paths.h"
-#include "Serialization/CustomVersion.h"
-#include "Misc/OutputDeviceError.h"
-#include "Misc/ObjectThumbnail.h"
-#include "Containers/StaticArray.h"
-#include "Templates/ScopedCallback.h"
-#include "Misc/CoreStats.h"
-#include "HAL/PlatformProcess.h"
-#include "GenericPlatform/GenericPlatformProcess.h"
-#include "Misc/EngineVersion.h"
-#include "Internationalization/GatherableTextData.h"
-#include "Internationalization/InternationalizationMetadata.h"
-#include "Misc/IQueuedWork.h"
-#include "Serialization/BufferReader.h"
-#include "Misc/QueuedThreadPool.h"
-#include "GenericPlatform/GenericPlatformAffinity.h"
-#include "Async/AsyncWork.h"
-#include "Misc/SecureHash.h"
-#include "Modules/ModuleManager.h"
-#include "Modules/Boilerplate/ModuleBoilerplate.h"
-#include "Serialization/ArchiveProxy.h"
-#include "UObject/DebugSerializationFlags.h"
-#include "Containers/ResourceArray.h"
-#include "Containers/DynamicRHIResourceArray.h"
-#include "Async/Future.h"
-#include "Misc/BufferedOutputDevice.h"
-#include "Misc/ScopeLock.h"
-#include "Math/SHMath.h"
-#include "Misc/ScopedEvent.h"
-#include "Containers/Ticker.h"
-#include "UObject/PropertyPortFlags.h"
-#include "Misc/ConfigCacheIni.h"
-#include "Templates/ValueOrError.h"
-#include "Features/IModularFeature.h"
-#include "Misc/ExpressionParserTypes.h"
-#include "GenericPlatform/GenericPlatformFile.h"
 #include "Math/BasicMathExpressionEvaluator.h"
-#include "Misc/IFilter.h"
+#include "Misc/ExpressionParserTypes.h"
 #include "Misc/FilterCollection.h"
-#include "Misc/CommandLine.h"
-#include "Features/IModularFeatures.h"
-#include "HAL/ThreadSafeBool.h"
-#include "Misc/CoreDelegates.h"
-#include "Misc/App.h"
-#include "Misc/MessageDialog.h"
-#include "Containers/ArrayView.h"
-#include "Misc/SlowTask.h"
-#include "HAL/FileManager.h"
-
-// From Json:
-#include "Policies/JsonPrintPolicy.h"
-#include "Policies/PrettyJsonPrintPolicy.h"
+#include "Misc/IFilter.h"
 
 // From CoreUObject:
-#include "UObject/Script.h"
-#include "UObject/ObjectMacros.h"
-#include "UObject/UObjectGlobals.h"
-#include "UObject/UObjectBase.h"
-#include "UObject/UObjectArray.h"
-#include "UObject/Object.h"
-#include "UObject/UObjectBaseUtility.h"
-#include "UObject/UObjectMarks.h"
-#include "Serialization/ArchiveUObject.h"
-#include "UObject/Class.h"
-#include "UObject/GarbageCollection.h"
-#include "UObject/GeneratedCppIncludes.h"
-#include "UObject/CoreNative.h"
-#include "UObject/WeakObjectPtr.h"
-#include "UObject/Interface.h"
-#include "Templates/Casts.h"
-#include "UObject/CoreNetTypes.h"
-#include "UObject/LazyObjectPtr.h"
-#include "UObject/SoftObjectPtr.h"
-#include "UObject/GCObject.h"
-#include "Templates/SubclassOf.h"
-#include "Misc/WorldCompositionUtility.h"
-#include "UObject/UnrealType.h"
-#include "UObject/ScriptInterface.h"
-#include "UObject/PropertyTag.h"
-#include "Serialization/SerializedPropertyScope.h"
-#include "UObject/CoreNet.h"
-#include "UObject/ScriptMacros.h"
-#include "UObject/Stack.h"
-#include "UObject/Package.h"
-#include "Misc/PackageName.h"
-#include "UObject/LinkerLoad.h"
-#include "UObject/ObjectResource.h"
-#include "UObject/Linker.h"
-#include "UObject/PackageFileSummary.h"
-#include "UObject/ObjectRedirector.h"
-#include "Serialization/BulkData.h"
-#include "UObject/UObjectAnnotation.h"
-#include "UObject/StructOnScope.h"
 #include "Misc/NotifyHook.h"
-#include "UObject/UObjectIterator.h"
-#include "UObject/UObjectHash.h"
 #include "UObject/ObjectKey.h"
-#include "UObject/TextProperty.h"
-
-// From InputCore:
-#include "InputCoreTypes.h"
 
 // From SlateCore:
-#include "Types/SlateEnums.h"
-#include "Styling/SlateColor.h"
-#include "Styling/WidgetStyle.h"
-#include "Layout/Visibility.h"
-#include "Layout/SlateRect.h"
-#include "Rendering/SlateLayoutTransform.h"
-#include "Layout/Geometry.h"
-#include "Layout/PaintGeometry.h"
-#include "Rendering/SlateRenderTransform.h"
-#include "Input/Events.h"
-#include "Input/PopupMethodReply.h"
-#include "Input/ReplyBase.h"
-#include "Widgets/SWidget.h"
-#include "Input/CursorReply.h"
-#include "Input/Reply.h"
-#include "Input/DragAndDrop.h"
-#include "Input/NavigationReply.h"
-#include "Types/ISlateMetaData.h"
-#include "Layout/ArrangedWidget.h"
-#include "Types/WidgetActiveTimerDelegate.h"
-#include "Layout/LayoutGeometry.h"
-#include "Layout/Margin.h"
-#include "Widgets/SNullWidget.h"
-#include "Widgets/DeclarativeSyntaxSupport.h"
-#include "SlotBase.h"
-#include "Layout/Children.h"
-#include "Widgets/SPanel.h"
-#include "Fonts/SlateFontInfo.h"
-#include "Fonts/CompositeFont.h"
-#include "Styling/SlateBrush.h"
-#include "Sound/SlateSound.h"
-#include "Brushes/SlateNoResource.h"
-#include "Styling/ISlateStyle.h"
-#include "Styling/StyleDefaults.h"
-#include "Widgets/SOverlay.h"
-#include "Widgets/SCompoundWidget.h"
-#include "Styling/SlateWidgetStyle.h"
-#include "Styling/SlateTypes.h"
-#include "Styling/CoreStyle.h"
-#include "Styling/SlateWidgetStyleAsset.h"
-#include "Styling/SlateWidgetStyleContainerBase.h"
-#include "Styling/SlateWidgetStyleContainerInterface.h"
-#include "Types/SlateStructs.h"
-#include "Widgets/SBoxPanel.h"
-#include "Rendering/RenderingCommon.h"
-#include "Animation/CurveSequence.h"
-#include "Animation/CurveHandle.h"
-#include "Textures/SlateIcon.h"
-#include "Widgets/SWindow.h"
-#include "Widgets/SLeafWidget.h"
-#include "Types/PaintArgs.h"
-#include "SlateGlobals.h"
-#include "Fonts/ShapedTextFwd.h"
-#include "Textures/SlateShaderResource.h"
-#include "Textures/SlateTextureData.h"
-#include "Rendering/DrawElements.h"
-#include "Rendering/ShaderResourceManager.h"
-#include "Widgets/IToolTip.h"
-#include "Layout/ArrangedChildren.h"
-#include "Brushes/SlateDynamicImageBrush.h"
-#include "Application/SlateWindowHelper.h"
-#include "Application/ThrottleManager.h"
-#include "Rendering/SlateRenderer.h"
-#include "Application/SlateApplicationBase.h"
-#include "Layout/WidgetPath.h"
-#include "Types/SlateConstants.h"
 #include "Layout/LayoutUtils.h"
 
 // From Slate:
-#include "SlateFwd.h"
-#include "Framework/SlateDelegates.h"
-#include "Framework/Commands/InputChord.h"
-#include "Framework/Commands/UIAction.h"
-#include "Framework/Commands/UICommandInfo.h"
-#include "Framework/Docking/TabManager.h"
-#include "Framework/Docking/WorkspaceItem.h"
-#include "Widgets/Layout/SBorder.h"
-#include "Framework/Commands/UICommandList.h"
-#include "Framework/MultiBox/MultiBoxExtender.h"
-#include "Framework/Text/TextLayout.h"
-#include "Framework/Text/TextRunRenderer.h"
-#include "Framework/Text/TextLineHighlight.h"
-#include "Framework/Text/IRun.h"
-#include "Framework/Text/ShapedTextCacheFwd.h"
-#include "Widgets/Text/STextBlock.h"
-#include "Framework/Application/IMenu.h"
-#include "Widgets/Layout/SBox.h"
-#include "Framework/Views/ITypedTableView.h"
 #include "Framework/Docking/LayoutService.h"
-#include "Framework/Application/SlateApplication.h"
-#include "Framework/Application/MenuStack.h"
-#include "Widgets/Views/SHeaderRow.h"
-#include "Framework/Layout/InertialScrollManager.h"
-#include "Widgets/Layout/SScrollBar.h"
-#include "Widgets/Views/SExpanderArrow.h"
-#include "Widgets/Views/STableRow.h"
-#include "Widgets/Views/STableViewBase.h"
-#include "Framework/Layout/IScrollableWidget.h"
-#include "Framework/Views/TableViewTypeTraits.h"
-#include "Framework/Layout/Overscroll.h"
-#include "Framework/Commands/InputBindingManager.h"
-#include "Framework/Commands/Commands.h"
-#include "Framework/MultiBox/MultiBoxDefs.h"
-#include "Framework/MultiBox/MultiBoxBuilder.h"
-#include "Widgets/Views/SListView.h"
-#include "Widgets/Images/SImage.h"
-#include "Widgets/Input/SButton.h"
-#include "Widgets/Input/IVirtualKeyboardEntry.h"
-#include "Widgets/Layout/SSpacer.h"
-#include "Widgets/Input/SMenuAnchor.h"
-#include "Widgets/Input/SComboButton.h"
-#include "Widgets/Text/ISlateEditableTextWidget.h"
-#include "Widgets/Input/SEditableText.h"
-#include "Widgets/Input/SCheckBox.h"
-#include "Widgets/SToolTip.h"
+#include "Framework/MarqueeRect.h"
 #include "Widgets/Input/NumericTypeInterface.h"
-#include "Widgets/Views/STreeView.h"
-#include "Widgets/Input/SEditableTextBox.h"
-#include "Widgets/Layout/SSplitter.h"
 #include "Widgets/Notifications/SErrorText.h"
-#include "Framework/MarqueeRect.h"
-#include "Widgets/Docking/SDockTab.h"
-
-// From RHI:
-#include "RHIDefinitions.h"
-#include "RHI.h"
-#include "RHIStaticStates.h"
-
-// From RenderCore:
-#include "RenderCommandFence.h"
-#include "RenderResource.h"
-#if UE_ENABLE_INCLUDE_ORDER_DEPRECATED_IN_5_2
-#include "RenderCore.h"
-#endif
-#include "RenderingThread.h"
-#include "UniformBuffer.h"
-#include "PackedNormal.h"
-#include "RenderUtils.h"
-#include "ShaderParameters.h"
-#include "Shader.h"
-#include "ShaderCore.h"
-#include "VertexFactory.h"
+#include "Widgets/SToolTip.h"
 
 // From Messaging:
 #include "IMessageContext.h"
 
 // From AssetRegistry:
-#include "AssetRegistry/AssetData.h"
 #include "AssetRegistry/ARFilter.h"
 
 // From Engine:
-#include "EngineLogs.h"
-#include "Engine/EngineTypes.h"
-#include "Engine/NetSerialization.h"
-#include "Engine/EngineBaseTypes.h"
-#include "EdGraph/EdGraphNode.h"
-#include "EdGraph/EdGraphPin.h"
-#include "Interfaces/Interface_AssetUserData.h"
-#include "Components/ActorComponent.h"
-#include "Components/SceneComponent.h"
-#include "Engine/MaterialMerging.h"
-#include "Engine/Level.h"
-#include "GameFramework/Actor.h"
-#include "EngineDefines.h"
-#include "NavFilters/NavigationQueryFilter.h"
-#include "AI/Navigation/NavigationTypes.h"
-#include "Engine/LatentActionManager.h"
-#include "Engine/Blueprint.h"
-#include "Engine/BlueprintCore.h"
-#include "CollisionQueryParams.h"
-#include "PixelFormat.h"
-#include "AI/Navigation/NavAgentInterface.h"
-#include "GameFramework/Pawn.h"
-#include "WorldCollision.h"
-#include "Engine/World.h"
-#include "Engine/PendingNetGame.h"
-#include "Engine/GameInstance.h"
-#include "SceneTypes.h"
-#include "HitProxies.h"
-#include "Engine/BlendableInterface.h"
-#include "Components.h"
-#include "UnrealClient.h"
-#include "Tickable.h"
-#include "Interfaces/Interface_CollisionDataProvider.h"
-#include "ShowFlags.h"
-#include "Engine/GameViewportClient.h"
-#include "Engine/ScriptViewportClient.h"
-#include "Engine/ViewportSplitScreen.h"
-#include "Engine/TitleSafeZone.h"
-#include "Engine/GameViewportDelegates.h"
-#include "Engine/DebugDisplayProperty.h"
-#include "Engine/Scene.h"
-#include "Curves/KeyHandle.h"
-#include "Curves/IndexedCurve.h"
-#include "Curves/RichCurve.h"
-#include "Engine/Engine.h"
-#include "Materials/MaterialInterface.h"
-#include "SceneView.h"
-#include "ConvexVolume.h"
-#include "SceneInterface.h"
-#include "FinalPostProcessSettings.h"
-#include "BlendableManager.h"
-#include "DebugViewModeHelpers.h"
-#include "LocalVertexFactory.h"
-#include "RawIndexBuffer.h"
-#include "PrimitiveUniformShaderParameters.h"
-#include "Engine/Brush.h"
-#include "TimerManager.h"
-#include "Engine/MeshMerging.h"
-#include "Model.h"
-#include "Engine/StaticMesh.h"
-#include "PhysxUserData.h"
-#include "MaterialShared.h"
-#include "StaticParameterSet.h"
-#include "PhysicsEngine/BodyInstance.h"
-#include "Engine/TextureDefines.h"
-#include "Engine/TextureStreamingTypes.h"
-#include "Components/PrimitiveComponent.h"
-#include "AI/Navigation/NavRelevantInterface.h"
-#include "Engine/Texture.h"
-#include "TextureResource.h"
-#include "PrimitiveSceneProxy.h"
-#include "PrimitiveViewRelevance.h"
-#include "Engine/Texture2D.h"
-#include "BlueprintUtilities.h"
-#include "Components/MeshComponent.h"
-#include "BatchedElements.h"
-#include "StaticBoundShaderState.h"
-#include "EdGraph/EdGraphSchema.h"
-#include "ReferenceSkeleton.h"
 #include "BoneIndices.h"
-#include "Camera/CameraTypes.h"
-#include "MeshBatch.h"
-#include "LatentActions.h"
-#include "Engine/LevelStreaming.h"
-#include "Engine/TextureLightProfile.h"
-#include "GPUSkinPublicDefs.h"
-#include "SceneManagement.h"
-#include "SceneUtils.h"
-#include "Curves/CurveOwnerInterface.h"
-#include "Curves/CurveBase.h"
-#include "Animation/AnimTypes.h"
-#include "Animation/AnimLinkableElement.h"
-#include "Animation/PreviewAssetAttachComponent.h"
-#include "BoneContainer.h"
-#include "EngineGlobals.h"
-#include "Engine/SkeletalMesh.h"
-#include "Animation/SmartName.h"
-#include "Animation/Skeleton.h"
-#include "EdGraph/EdGraph.h"
-#include "Animation/AnimationAsset.h"
-#include "AnimInterpFilter.h"
-#include "Animation/AnimCurveTypes.h"
-#include "Curves/CurveFloat.h"
-#include "PreviewScene.h"
-#include "Components/SkinnedMeshComponent.h"
-#include "VisualLogger/VisualLoggerTypes.h"
-#include "Engine/DataAsset.h"
-#include "Components/SkeletalMeshComponent.h"
-#include "ClothSimData.h"
-#include "SingleAnimationPlayData.h"
-#include "EdGraph/EdGraphNodeUtils.h"
-#include "Engine/BlueprintGeneratedClass.h"
-#include "EngineStats.h"
-#include "VisualLogger/VisualLogger.h"
-#include "Engine/DataTable.h"
 #include "DataTableUtils.h"
-#include "GameFramework/Controller.h"
-#include "Animation/AnimSequenceBase.h"
-#include "MaterialExpressionIO.h"
-#include "Materials/MaterialExpression.h"
-#include "Materials/Material.h"
-#include "Materials/MaterialExpressionMaterialFunctionCall.h"
-#include "Materials/MaterialFunction.h"
-#include "PhysicsSettingsEnums.h"
-#include "Sound/SoundAttenuation.h"
+#include "EdGraph/EdGraphNodeUtils.h"
 #include "Engine/CurveTable.h"
-#include "Audio.h"
-#include "IAudioExtensionPlugin.h"
+#include "Engine/DataAsset.h"
+#include "Engine/DataTable.h"
+#include "Engine/LevelStreaming.h"
+#include "LatentActions.h"
+#include "NavFilters/NavigationQueryFilter.h"
 #include "PhysicalMaterials/PhysicalMaterial.h"
+#include "PhysicsSettingsEnums.h"
+#include "PreviewScene.h"
 #include "Vehicles/TireType.h"
-#include "Animation/AnimSequence.h"
-#include "Kismet/BlueprintFunctionLibrary.h"
+#include "VisualLogger/VisualLogger.h"
+#include "VisualLogger/VisualLoggerTypes.h"
 
 // From BlueprintGraph:
 #include "BlueprintNodeSignature.h"
@@ -594,27 +59,26 @@
 #include "GameplayTask.h"
 
 // From UnrealEd:
-#include "Toolkits/IToolkit.h"
-#include "Toolkits/IToolkitHost.h"
-#include "TickableEditorObject.h"
-#include "Editor/UnrealEdTypes.h"
-#include "Viewports.h"
-#include "Settings/LevelEditorViewportSettings.h"
-#include "Settings/LevelEditorPlaySettings.h"
-#include "Editor/EditorEngine.h"
-
 #include "AssetThumbnail.h"
-#include "Toolkits/BaseToolkit.h"
-#include "Toolkits/AssetEditorToolkit.h"
 #include "Editor.h"
-#include "EditorUndoClient.h"
-#include "UnrealWidgetFwd.h"
-#include "ScopedTransaction.h"
-#include "GraphEditor.h"
+#include "Editor/EditorEngine.h"
+#include "Editor/UnrealEdTypes.h"
 #include "EditorComponents.h"
+#include "EditorUndoClient.h"
 #include "EditorViewportClient.h"
 #include "Factories/Factory.h"
+#include "GraphEditor.h"
+#include "ScopedTransaction.h"
 #include "Settings/EditorLoadingSavingSettings.h"
+#include "Settings/LevelEditorPlaySettings.h"
+#include "Settings/LevelEditorViewportSettings.h"
+#include "TickableEditorObject.h"
+#include "Toolkits/AssetEditorToolkit.h"
+#include "Toolkits/BaseToolkit.h"
+#include "Toolkits/IToolkit.h"
+#include "Toolkits/IToolkitHost.h"
+#include "UnrealWidgetFwd.h"
+#include "Viewports.h"
 
 // From ToolMenus
 #include "ToolMenus.h"
\ No newline at end of file
diff --git a/Engine/Source/Runtime/CoreUObject/Public/CoreUObjectSharedPCH.h b/Engine/Source/Runtime/CoreUObject/Public/CoreUObjectSharedPCH.h
--- a/Engine/Source/Runtime/CoreUObject/Public/CoreUObjectSharedPCH.h
+++ b/Engine/Source/Runtime/CoreUObject/Public/CoreUObjectSharedPCH.h
@@ -4,6 +4,7 @@
 
 #include "CoreSharedPCH.h"
 
+// From Core:
 #include "HAL/PlatformAffinity.h"
 #include "HAL/RunnableThread.h"
 #include "HAL/ThreadSafeBool.h"
diff --git a/Engine/Source/Runtime/Slate/Public/SlateSharedPCH.h b/Engine/Source/Runtime/Slate/Public/SlateSharedPCH.h
--- a/Engine/Source/Runtime/Slate/Public/SlateSharedPCH.h
+++ b/Engine/Source/Runtime/Slate/Public/SlateSharedPCH.h
@@ -20,7 +20,6 @@
 #include "Misc/MessageDialog.h"
 #include "Misc/ObjectThumbnail.h"
 #include "Misc/OutputDeviceError.h"
-#include "ProfilingDebugging/ResourceSize.h"
 #include "Serialization/CustomVersion.h"
 
 // From Json:
