Add "noerrorsummary" option to cooker to hide the display of errors and warnings on completion

diff --git a/Engine/Source/Editor/UnrealEd/Private/Commandlets/CookCommandlet.cpp b/Engine/Source/Editor/UnrealEd/Private/Commandlets/CookCommandlet.cpp
--- a/Engine/Source/Editor/UnrealEd/Private/Commandlets/CookCommandlet.cpp
+++ b/Engine/Source/Editor/UnrealEd/Private/Commandlets/CookCommandlet.cpp
@@ -181,7 +181,7 @@ int32 UCookCommandlet::Main(const FString& CmdLineParams)
 	bCookSinglePackage = bCookSinglePackage || bKeepSinglePackageRefs;
 	bVerboseCookerWarnings = Switches.Contains(TEXT("verbosecookerwarnings"));
 	bPartialGC = Switches.Contains(TEXT("Partialgc"));
-	ShowErrorCount = !Switches.Contains(TEXT("DIFFONLY"));
+	ShowErrorCount = !Switches.Contains(TEXT("DIFFONLY")) && !Switches.Contains(TEXT("NoErrorSummary"));
 	ShowProgress = !Switches.Contains(TEXT("DIFFONLY"));
 	bIgnoreIniSettingsOutOfDate = Switches.Contains(TEXT("IgnoreIniSettingsOutOfDate"));
 	bFastCook = Switches.Contains(TEXT("FastCook"));
