Modified CleanupAfterSuccessfulDelete() to check if deleting a package is supported by source control first and only if that is not the case mark it for local deletion.
This fixes an issue where deleting a file checked out by someone else (when working offline) and choosing to only save locally resulted in the file not being actually deleted on the local file system.

diff --git a/Engine/Source/Editor/UnrealEd/Private/ObjectTools.cpp b/Engine/Source/Editor/UnrealEd/Private/ObjectTools.cpp
--- a/Engine/Source/Editor/UnrealEd/Private/ObjectTools.cpp
+++ b/Engine/Source/Editor/UnrealEd/Private/ObjectTools.cpp
@@ -2444,26 +2444,23 @@ namespace ObjectTools
 						SCCFilesToRevert.Add(FullPackageFilename);
 					}
 
-					if ( bIsAdded )
+					if (bIsAdded)
 					{
 						// The file was open for add and reverted, this leaves the file on disk so here we delete it
 						IFileManager::Get().Delete(*PackageFilename);
 					}
-					else if (!bIsCheckedOut && SourceControlProvider.UsesLocalReadOnlyState() && !IFileManager::Get().IsReadOnly(*PackageFilename))
+					else if (SourceControlState->CanDelete())
+					{
+						// Batch this file for deletion so that we only send one deletion request to the source control module.
+						SCCFilesToDelete.Add(FullPackageFilename);
+					}
+					else if (!bIsCheckedOut && !IFileManager::Get().IsReadOnly(*PackageFilename))
 					{
 						bDeletedFileLocallyWritable = true;
 					}
 					else
 					{
-						// Batch this file for deletion so that we only send one deletion request to the source control module.
-						if (SourceControlState->CanDelete())
-						{
-							SCCFilesToDelete.Add(FullPackageFilename);
-						}
-						else
-						{
-							UE_LOG(LogObjectTools, Warning, TEXT("SCC failed to open '%s' for deletion."), *PackageFilename);
-						}
+						UE_LOG(LogObjectTools, Warning, TEXT("SCC failed to open '%s' for deletion."), *PackageFilename);
 					}
 				}
 				else
