WorldPartition - HLODsBuilder: Running -DumpStats on a collection will output the stats in a single CSV file
* Replaced the old stats code with the new one found in the HLOD Subsystem

diff --git a/Engine/Source/Editor/UnrealEd/Private/WorldPartition/WorldPartitionHLODsBuilder.cpp b/Engine/Source/Editor/UnrealEd/Private/WorldPartition/WorldPartitionHLODsBuilder.cpp
--- a/Engine/Source/Editor/UnrealEd/Private/WorldPartition/WorldPartitionHLODsBuilder.cpp
+++ b/Engine/Source/Editor/UnrealEd/Private/WorldPartition/WorldPartitionHLODsBuilder.cpp
@@ -5,6 +5,7 @@
 #include "CoreMinimal.h"
 #include "HAL/FileManager.h"
 #include "Logging/LogMacros.h"
+#include "Misc/CommandLine.h"
 #include "Misc/ConfigCacheIni.h"
 #include "Misc/EngineVersion.h"
 #include "Misc/FileHelper.h"
@@ -13,6 +14,7 @@
 
 #include "ActorFolder.h"
 #include "Engine/Engine.h"
+#include "Engine/LevelStreamingGCHelper.h"
 #include "EngineUtils.h"
 #include "SourceControlHelpers.h"
 #include "ISourceControlModule.h"
@@ -27,16 +29,7 @@
 #include "WorldPartition/HLOD/HLODActor.h"
 #include "WorldPartition/HLOD/HLODActorDesc.h"
 #include "WorldPartition/HLOD/HLODLayer.h"
-
-#include "Engine/StaticMesh.h"
-#include "Engine/LevelStreamingGCHelper.h"
-#include "Engine/Texture.h"
-#include "Materials/Material.h"
-#include "Materials/MaterialInstance.h"
-
-#include "Components/InstancedStaticMeshComponent.h"
-
-#include "AssetCompilingManager.h"
+#include "WorldPartition/HLOD/HLODSubsystem.h"
 
 DEFINE_LOG_CATEGORY_STATIC(LogWorldPartitionHLODsBuilder, All, All);
 
@@ -585,146 +578,8 @@ bool UWorldPartitionHLODsBuilder::SubmitHLODActors()
 
 bool UWorldPartitionHLODsBuilder::DumpStats()
 {
-	struct FHLODStat
-	{
-		FString				Name;
-		FName				RuntimeGrid;
-		bool				bIsSpatiallyLoaded = false;
-		uint64				GridLocationZ = 0;
-		uint64				GridLocationX = 0;
-		uint64				GridLocationY = 0;
-		FString				DataLayers;
-		EHLODLayerType		HLODType;
-
-		uint64				InstanceCount = 0;
-		uint64				NaniteTriangleCount = 0;
-		uint64				NaniteVertexCount = 0;
-		uint64				TriangleCount = 0;
-		uint64				VertexCount = 0;
-		uint64				UVChannelCount = 0;
-
-		float				BaseColorTextureSize = 0.0f;
-		float				NormalTextureSize = 0.0f;
-		float				MRSTextureSize = 0.0f;
-		float				EmissiveTextureSize = 0.0f;
-
-		uint64				MeshResourceSize = 0;
-		uint64				TexturesResourceSize = 0;
-
-		uint64				DiskSize = 0;
-
-		static FString GetHeaderCSVString()
-		{
-			return TEXT("Name, RuntimeGrid, SpatiallyLoaded, GridLocationZ, GridLocationX, GridLocationY, DataLayers, HLODType, InstanceCount, NaniteTriangleCount, NaniteVertexCount, TriangleCount, VertexCount, UVChannelCount, BaseColorTextureSize, NormalTextureSize, MRSTextureSize, EmissiveTextureSize, MeshResourceSize, TexturesResourceSize, DiskSize");
-		}
-
-		FString ToCSVString() const
-		{
-			return FString::Printf(TEXT("%s, %s, %s, %d, %d, %d, %s, %s, %d, %d, %d, %d, %d, %d, %.0f, %.0f, %.0f, %.0f, %d, %d, %d"), 
-				*Name, 
-				*RuntimeGrid.ToString(),
-				bIsSpatiallyLoaded ? TEXT("true") : TEXT("false"),
-				GridLocationZ, GridLocationX, GridLocationY,
-				*DataLayers,
-				*StaticEnum<EHLODLayerType>()->GetNameStringByValue((int64)HLODType),
-				InstanceCount, NaniteTriangleCount, NaniteVertexCount, TriangleCount, VertexCount, UVChannelCount,
-				BaseColorTextureSize, NormalTextureSize, MRSTextureSize, EmissiveTextureSize, 
-				MeshResourceSize, TexturesResourceSize,
-				DiskSize);
-		}
-	};
-
-	TArray<FHLODStat> HLODStats;
-	for (FActorDescContainerCollection::TIterator<AWorldPartitionHLOD> HLODIterator(WorldPartition); HLODIterator; ++HLODIterator)
-	{
-		FWorldPartitionReference HLODActorRef(WorldPartition, HLODIterator->GetGuid());
-		AWorldPartitionHLOD* HLODActor = CastChecked<AWorldPartitionHLOD>(HLODActorRef->GetActor());
-
-		FAssetCompilingManager::Get().FinishAllCompilation();
-
-		FHLODStat& HLODStat = HLODStats.Emplace_GetRef();
-
-		HLODStat.Name = HLODActor->GetActorLabel();
-		HLODStat.RuntimeGrid = HLODActor->GetRuntimeGrid();
-		HLODStat.bIsSpatiallyLoaded = HLODActor->GetIsSpatiallyLoaded();
-		HLODStat.DataLayers = *FString::JoinBy(HLODActor->GetDataLayerInstances(), TEXT("| "), [](const UDataLayerInstance* DataLayer) { return DataLayer->GetDataLayerShortName(); });
-		HLODStat.HLODType = HLODActor->GetSubActorsHLODLayer()->GetLayerType();
-
-		HLODActor->ForEachComponent<UInstancedStaticMeshComponent>(false, [&](const UInstancedStaticMeshComponent* ISMC)
-		{
-			HLODStat.InstanceCount += ISMC->GetInstanceCount();
-		});
-
-		ForEachObjectWithPackage(HLODActor->GetPackage(), [&HLODStat](UObject* Object)
-		{
-			if (UStaticMesh* StaticMesh = Cast<UStaticMesh>(Object))
-			{
-				HLODStat.MeshResourceSize += StaticMesh->GetResourceSizeBytes(EResourceSizeMode::Exclusive);
-
-				HLODStat.TriangleCount += StaticMesh->GetNumTriangles(0);
-				HLODStat.VertexCount += StaticMesh->GetNumVertices(0);
-				HLODStat.UVChannelCount += StaticMesh->GetNumTexCoords(0);
-
-				HLODStat.NaniteTriangleCount += StaticMesh->GetNumNaniteTriangles();
-				HLODStat.NaniteVertexCount += StaticMesh->GetNumNaniteVertices();
-			}
-
-			if (UMaterialInstance* MaterialInstance = Cast<UMaterialInstance>(Object))
-			{
-				auto GetTexture = [MaterialInstance, &HLODStat](const FName TextureParamName)
-				{
-					UTexture* Texture = nullptr;
-					MaterialInstance->GetTextureParameterValue(TextureParamName, Texture, true);
-					HLODStat.TexturesResourceSize += Texture ? Texture->GetResourceSizeBytes(EResourceSizeMode::Exclusive) : 0;
-					return Texture;
-				};
-
-				UTexture* BaseColorTex = GetTexture("BaseColorTexture");
-				UTexture* NormalTex = GetTexture("NormalTexture");
-				UTexture* MetallicTex = GetTexture("MetallicTexture");
-				UTexture* RoughnessTex = GetTexture("RoughnessTexture");
-				UTexture* SpecularTex = GetTexture("SpecularTexture");
-				UTexture* EmissiveTex = GetTexture("EmissiveTexture");
-				UTexture* EmissiveHDRTex = GetTexture("EmissiveHDRTexture");
-				UTexture* PackedMRSTex = GetTexture("PackedTexture");
-				
-				HLODStat.BaseColorTextureSize = BaseColorTex ? BaseColorTex->GetSurfaceWidth() : 0;
-				HLODStat.NormalTextureSize = NormalTex ? NormalTex->GetSurfaceWidth() : 0;
-				HLODStat.MRSTextureSize = FMath::Max(HLODStat.MRSTextureSize, MetallicTex ? MetallicTex->GetSurfaceWidth() : 0);
-				HLODStat.MRSTextureSize = FMath::Max(HLODStat.MRSTextureSize, RoughnessTex ? RoughnessTex->GetSurfaceWidth() : 0);
-				HLODStat.MRSTextureSize = FMath::Max(HLODStat.MRSTextureSize, SpecularTex ? SpecularTex->GetSurfaceWidth() : 0);
-				HLODStat.MRSTextureSize = FMath::Max(HLODStat.MRSTextureSize, PackedMRSTex ? PackedMRSTex->GetSurfaceWidth() : 0);
-				HLODStat.EmissiveTextureSize = FMath::Max(HLODStat.EmissiveTextureSize, EmissiveTex ? EmissiveTex->GetSurfaceWidth() : 0);
-				HLODStat.EmissiveTextureSize = FMath::Max(HLODStat.EmissiveTextureSize, EmissiveHDRTex ? EmissiveHDRTex->GetSurfaceWidth() : 0);
-			}
-
-			return true;
-		}, false);
-
-		HLODStat.DiskSize = IFileManager::Get().FileSize(*HLODActor->GetPackage()->GetLoadedPath().GetLocalFullPath());
-	}
-
-	Algo::Sort(HLODStats, [](const FHLODStat& A, const FHLODStat& B) { return A.Name < B.Name; });
-
-
-	const FString Filename = FString::Printf(TEXT("%s/HLODStats/HLODStats-%s.csv"), *FPaths::ProjectLogDir(), *FDateTime::Now().ToString());
-	FArchive* CSVFile = IFileManager::Get().CreateFileWriter(*Filename, FILEWRITE_AllowRead);
-
-	FTCHARToUTF8 CSVHeader = FHLODStat::GetHeaderCSVString();
-	FString Newline = "\n";
-
-	CSVFile->Serialize((UTF8CHAR*)CSVHeader.Get(), CSVHeader.Length() * sizeof(UTF8CHAR));
-	*CSVFile << Newline;
-	for (const FHLODStat& HLODStat : HLODStats)
-	{
-		FTCHARToUTF8 CSVString = HLODStat.ToCSVString();
-		CSVFile->Serialize((UTF8CHAR*)CSVString.Get(), CSVString.Length() * sizeof(UTF8CHAR));
-		*CSVFile << Newline;
-	}
-
-	delete CSVFile;
-
-	return true;
+	const FString HLODStatsOutputFilename = FPaths::ProjectSavedDir() / TEXT("WorldPartition") / FString::Printf(TEXT("HLODStats-%08x.csv"), FPlatformProcess::GetCurrentProcessId());
+	return UHLODSubsystem::WriteHLODStatsCSV(WorldPartition->GetWorld(), HLODStatsOutputFilename);
 }
 
 bool UWorldPartitionHLODsBuilder::GetHLODActorsToBuild(TArray<FGuid>& HLODActorsToBuild) const
diff --git a/Engine/Source/Runtime/Engine/Private/WorldPartition/HLOD/HLODSubsystem.cpp b/Engine/Source/Runtime/Engine/Private/WorldPartition/HLOD/HLODSubsystem.cpp
--- a/Engine/Source/Runtime/Engine/Private/WorldPartition/HLOD/HLODSubsystem.cpp
+++ b/Engine/Source/Runtime/Engine/Private/WorldPartition/HLOD/HLODSubsystem.cpp
@@ -819,12 +819,12 @@ void UHLODSubsystem::OnBeginRenderViews(const FSceneViewFamily& InViewFamily)
 
 #if WITH_EDITOR
 
-static void DumpHLODStats(UWorld* World)
+bool UHLODSubsystem::WriteHLODStatsCSV(UWorld* InWorld, const FString& InFilename)
 {
-	UWorldPartition* WorldPartition = World ? World->GetWorldPartition() : nullptr;
+	UWorldPartition* WorldPartition = InWorld ? InWorld->GetWorldPartition() : nullptr;
 	if (!WorldPartition)
 	{
-		return;
+		return false;
 	}
 	
 	typedef TFunction<FString(const FHLODActorDesc&)> FGetStatFunc;
@@ -846,7 +846,8 @@ static void DumpHLODStats(UWorld* World)
 	};
 
 	TArray<TPair<FName, FGetStatFunc>> StatsToWrite =
-	{
+	{	
+		{ "WorldPackage",		[InWorld](const FHLODActorDesc& InActorDesc) { return InWorld->GetPackage()->GetName(); } },
 		{ "Name",				[](const FHLODActorDesc& InActorDesc) { return InActorDesc.GetActorLabel().ToString(); } },
 		{ "HLODLayer",			[](const FHLODActorDesc& InActorDesc) { return InActorDesc.GetSourceHLODLayerName().ToString(); } },
 		{ "SpatiallyLoaded",	[](const FHLODActorDesc& InActorDesc) { return InActorDesc.GetIsSpatiallyLoaded() ? TEXT("true") : TEXT("false"); } },
@@ -879,10 +880,14 @@ static void DumpHLODStats(UWorld* World)
 		GetHLODStat(FWorldPartitionHLODStats::BuildTimeTotalMilliseconds)				
 	};
 
-	// Write header
 	FStringOutputDevice Output;
-	const FString StatHeader = FString::JoinBy(StatsToWrite, TEXT(","), [](const TPair<FName, FGetStatFunc>& Pair) { return Pair.Key.ToString(); });
-	Output.Logf(TEXT("%s" LINE_TERMINATOR_ANSI), *StatHeader);
+
+	// Write header if file doesn't exist
+	if (!IFileManager::Get().FileExists(*InFilename))
+	{
+		const FString StatHeader = FString::JoinBy(StatsToWrite, TEXT(","), [](const TPair<FName, FGetStatFunc>& Pair) { return Pair.Key.ToString(); });
+		Output.Logf(TEXT("%s" LINE_TERMINATOR_ANSI), *StatHeader);
+	}
 
 	// Write one line per HLOD actor desc
 	for (FActorDescContainerCollection::TIterator<AWorldPartitionHLOD> HLODIterator(WorldPartition); HLODIterator; ++HLODIterator)
@@ -892,8 +897,7 @@ static void DumpHLODStats(UWorld* World)
 	}
 
 	// Write to file
-	const FString HLODStatsOutputFilename = FPaths::ProjectSavedDir() / TEXT("WorldPartition") / FString::Printf(TEXT("HLODStats-%s-%08x-%s.csv"), *World->GetName(), FPlatformProcess::GetCurrentProcessId(), *FDateTime::Now().ToString());
-	FFileHelper::SaveStringToFile(Output, *HLODStatsOutputFilename);
+	return FFileHelper::SaveStringToFile(Output, *InFilename, FFileHelper::EEncodingOptions::AutoDetect, &IFileManager::Get(), EFileWrite::FILEWRITE_Append);
 }
 
 FAutoConsoleCommand HLODDumpStats(
@@ -901,11 +905,13 @@ FAutoConsoleCommand HLODDumpStats(
 	TEXT("Write various HLOD stats to a CSV formatted file."),
 	FConsoleCommandWithArgsDelegate::CreateLambda([](const TArray<FString>& Args)
 	{
+		const FString HLODStatsOutputFilename = FPaths::ProjectLogDir() / TEXT("WorldPartition") / FString::Printf(TEXT("HLODStats-%08x-%s.csv"), FPlatformProcess::GetCurrentProcessId(), *FDateTime::Now().ToString());
+
 		for (const FWorldContext& Context : GEngine->GetWorldContexts())
 		{
 			if (UWorld* World = Context.World())
 			{
-				DumpHLODStats(World);
+				UHLODSubsystem::WriteHLODStatsCSV(World, HLODStatsOutputFilename);
 			}
 		}
 	})
diff --git a/Engine/Source/Runtime/Engine/Public/WorldPartition/HLOD/HLODSubsystem.h b/Engine/Source/Runtime/Engine/Public/WorldPartition/HLOD/HLODSubsystem.h
--- a/Engine/Source/Runtime/Engine/Public/WorldPartition/HLOD/HLODSubsystem.h
+++ b/Engine/Source/Runtime/Engine/Public/WorldPartition/HLOD/HLODSubsystem.h
@@ -67,6 +67,10 @@ public:
 	void SetHLODAlwaysLoadedCullDistance(int32 InCullDistance);
 
 	void OnCVarsChanged();
+
+#ifdef WITH_EDITOR
+	static bool WriteHLODStatsCSV(UWorld* InWorld, const FString& InFilename);
+#endif
 	
 private:
 	struct FCellData
