Error on creating an animation sequence with Python
#fix when creating an anim sequence now seed it with a single frame setup and notify it as populated
#misc expose anim sequence factory + members to scripting

diff --git a/Engine/Source/Editor/UnrealEd/Classes/Factories/AnimSequenceFactory.h b/Engine/Source/Editor/UnrealEd/Classes/Factories/AnimSequenceFactory.h
--- a/Engine/Source/Editor/UnrealEd/Classes/Factories/AnimSequenceFactory.h
+++ b/Engine/Source/Editor/UnrealEd/Classes/Factories/AnimSequenceFactory.h
@@ -14,16 +14,16 @@
 struct FAssetData;
 class SWindow;
 
-UCLASS(HideCategories=Object)
+UCLASS(HideCategories=Object, BlueprintType)
 class UNREALED_API UAnimSequenceFactory : public UFactory
 {
 	GENERATED_UCLASS_BODY()
 
-	UPROPERTY()
+	UPROPERTY(EditAnywhere, BlueprintReadWrite, Category = Animation)
 	TObjectPtr<class USkeleton> TargetSkeleton;
 
 	/** The preview mesh to use with this animation */
-	UPROPERTY()
+	UPROPERTY(EditAnywhere, BlueprintReadWrite, Category = Animation)
 	TObjectPtr<class USkeletalMesh> PreviewSkeletalMesh;
 
 	//~ Begin UFactory Interface
diff --git a/Engine/Source/Editor/UnrealEd/Private/Factories/AnimSequenceFactory.cpp b/Engine/Source/Editor/UnrealEd/Private/Factories/AnimSequenceFactory.cpp
--- a/Engine/Source/Editor/UnrealEd/Private/Factories/AnimSequenceFactory.cpp
+++ b/Engine/Source/Editor/UnrealEd/Private/Factories/AnimSequenceFactory.cpp
@@ -74,11 +74,12 @@ UObject* UAnimSequenceFactory::FactoryCreateNew(UClass* Class, UObject* InParent
 	}
 
 	AnimSequence->GetController().InitializeModel();
-	
+	AnimSequence->GetController().SetNumberOfFrames(1);
 	if (PreviewSkeletalMesh)
 	{
 		AnimSequence->SetPreviewMesh(PreviewSkeletalMesh);
 	}
+	AnimSequence->GetController().NotifyPopulated();
 
 	return AnimSequence;
 }
