Asset Guideline: flush the config cache so that it writes out to the ini files in case a restart is requested. It's possible for the new editor process to boot and read DefaultEngine.ini before the old editor process had time to save DefaultEngine.ini during its exit process.

diff --git a/Engine/Source/Editor/UnrealEd/Private/AssetGuideline.cpp b/Engine/Source/Editor/UnrealEd/Private/AssetGuideline.cpp
--- a/Engine/Source/Editor/UnrealEd/Private/AssetGuideline.cpp
+++ b/Engine/Source/Editor/UnrealEd/Private/AssetGuideline.cpp
@@ -278,6 +278,7 @@ void UAssetGuideline::EnableMissingGuidelines(TArray<FString> IncorrectPlugins,
 			}
 		}
 
+		TSet<FString> ConfigFilesToFlush;
 		for (const FIniStringValue& IncorrectProjectSetting : IncorrectProjectSettings)
 		{
 			// Only fails if file DNE
@@ -289,6 +290,7 @@ void UAssetGuideline::EnableMissingGuidelines(TArray<FString> IncorrectPlugins,
 				if (!FPlatformFileManager::Get().GetPlatformFile().IsReadOnly(*FilenamePath))
 				{
 					GConfig->SetString(*IncorrectProjectSetting.Section, *IncorrectProjectSetting.Key, *IncorrectProjectSetting.Value, FilenamePath);
+					ConfigFilesToFlush.Add(MoveTemp(FilenamePath));
 				}
 				else
 				{
@@ -303,6 +305,12 @@ void UAssetGuideline::EnableMissingGuidelines(TArray<FString> IncorrectPlugins,
 			}
 		}
 
+		for (const FString& ConfigFileToFlush : ConfigFilesToFlush)
+		{
+			constexpr bool bRemoveFromCache = false;
+			GConfig->Flush(bRemoveFromCache, ConfigFileToFlush);
+		}
+
 		if (bSuccess)
 		{
 			auto ShowRestartPrompt = []()
