Allow to copy-paste content bundle actors without requiring to activate the destination content bundle

diff --git a/Engine/Source/Editor/UnrealEd/Private/EditorExporters.cpp b/Engine/Source/Editor/UnrealEd/Private/EditorExporters.cpp
--- a/Engine/Source/Editor/UnrealEd/Private/EditorExporters.cpp
+++ b/Engine/Source/Editor/UnrealEd/Private/EditorExporters.cpp
@@ -562,9 +562,10 @@ bool ULevelExporterT3D::ExportText( const FExportObjectInnerContext* Context, UO
 				FString GroupFolder = (Actor->GroupActor ? FString::Printf(TEXT(" GroupFolder=%s"), *Actor->GroupActor->GetFolderPath().ToString() ) : TEXT(""));
 				FString ActorFolderPath = (Level->IsUsingActorFolders() ? FString::Printf(TEXT(" ActorFolderPath=%s"), *Actor->GetFolderPath().ToString()) : TEXT(""));
 				FString CopyPasteId = (Actor->CopyPasteId != INDEX_NONE) ? FString::Printf(TEXT(" CopyPasteId=%d"), Actor->CopyPasteId ) : TEXT("");
-				Ar.Logf( TEXT("%sBegin Actor Class=%s Name=%s Archetype=%s'%s'%s%s%s%s%s%s") LINE_TERMINATOR, 
+				FString ContentBundleGuid = (Actor->GetContentBundleGuid().IsValid() ? FString::Printf(TEXT(" ActorContentBundleGuid=%s"), *Actor->GetContentBundleGuid().ToString()) : TEXT(""));
+				Ar.Logf( TEXT("%sBegin Actor Class=%s Name=%s Archetype=%s'%s'%s%s%s%s%s%s%s") LINE_TERMINATOR, 
 					FCString::Spc(TextIndent), *Actor->GetClass()->GetPathName(), *Actor->GetName(),
-					*Actor->GetArchetype()->GetClass()->GetPathName(), *Actor->GetArchetype()->GetPathName(), *ParentActorString, *SocketNameString, *GroupActor, *GroupFolder, *ActorFolderPath, *CopyPasteId );
+					*Actor->GetArchetype()->GetClass()->GetPathName(), *Actor->GetArchetype()->GetPathName(), *ParentActorString, *SocketNameString, *GroupActor, *GroupFolder, *ActorFolderPath, *CopyPasteId, *ContentBundleGuid );
 
 				ExportRootScope = Actor;
 				ExportObjectInner( Context, Actor, Ar, PortFlags | PPF_ExportsNotFullyQualified );
diff --git a/Engine/Source/Editor/UnrealEd/Private/Factories/EditorFactories.cpp b/Engine/Source/Editor/UnrealEd/Private/Factories/EditorFactories.cpp
--- a/Engine/Source/Editor/UnrealEd/Private/Factories/EditorFactories.cpp
+++ b/Engine/Source/Editor/UnrealEd/Private/Factories/EditorFactories.cpp
@@ -196,6 +196,7 @@
 #include "SSkeletonWidget.h"
 #include "AssetToolsModule.h"
 #include "IAssetTools.h"
+#include "WorldPartition/ContentBundle/ContentBundleActivationScope.h"
 
 #include "DDSFile.h"
 #include "IESConverter.h"
@@ -905,6 +906,11 @@ UObject* ULevelFactory::FactoryCreateText
 					}
 				}
 
+
+				FGuid ActorContentBundleGuid;
+				FParse::Value(Str, TEXT("ActorContentBundleGuid="), ActorContentBundleGuid);
+				FContentBundleActivationScope Scope(ActorContentBundleGuid);
+
 				// If we're pasting from a class that belongs to a map we need to duplicate the class and use that instead
 				if (FBlueprintEditorUtils::IsAnonymousBlueprintClass(TempClass))
 				{
diff --git a/Engine/Source/Editor/WorldPartitionEditor/Private/WorldPartition/ContentBundle/ContentBundleEditorSubsystem.cpp b/Engine/Source/Editor/WorldPartitionEditor/Private/WorldPartition/ContentBundle/ContentBundleEditorSubsystem.cpp
--- a/Engine/Source/Editor/WorldPartitionEditor/Private/WorldPartition/ContentBundle/ContentBundleEditorSubsystem.cpp
+++ b/Engine/Source/Editor/WorldPartitionEditor/Private/WorldPartition/ContentBundle/ContentBundleEditorSubsystem.cpp
@@ -379,6 +379,16 @@ bool UContentBundleEditorSubsystem::IsContentBundleEditingActivated(TSharedPtr<F
 	return ContentBundleEditor.IsValid() && IsEditingContentBundle(ContentBundleEditor->GetDescriptor()->GetGuid());
 }
 
+void UContentBundleEditorSubsystem::PushContentBundleEditing()
+{
+	ContentBundleEditingSubModule->PushContentBundleEditing();
+}
+
+void UContentBundleEditorSubsystem::PopContentBundleEditing()
+{
+	ContentBundleEditingSubModule->PopContentBundleEditing();
+}
+
 void UContentBundleEditorSubsystem::SelectActorsInternal(FContentBundleEditor& EditorContentBundle, bool bSelect)
 {
 	GEditor->GetSelectedActors()->BeginBatchSelectOperation();
diff --git a/Engine/Source/Editor/WorldPartitionEditor/Public/WorldPartition/ContentBundle/ContentBundleEditorSubsystem.h b/Engine/Source/Editor/WorldPartitionEditor/Public/WorldPartition/ContentBundle/ContentBundleEditorSubsystem.h
--- a/Engine/Source/Editor/WorldPartitionEditor/Public/WorldPartition/ContentBundle/ContentBundleEditorSubsystem.h
+++ b/Engine/Source/Editor/WorldPartitionEditor/Public/WorldPartition/ContentBundle/ContentBundleEditorSubsystem.h
@@ -44,6 +44,8 @@ class UContentBundleEditingSubmodule : public UContentBundleEditorSubsystemModul
 {
 	GENERATED_BODY()
 
+	friend class UContentBundleEditorSubsystem;
+
 public:
 	//~ Begin IActorEditorContextClient interface
 	virtual void OnExecuteActorEditorContextAction(UWorld* InWorld, const EActorEditorContextAction& InType, AActor* InActor = nullptr) override;
@@ -132,6 +134,9 @@ public:
 	bool DeactivateContentBundleEditing(TSharedPtr<FContentBundleEditor>& ContentBundleEditor) const override;
 	bool IsContentBundleEditingActivated(TSharedPtr<FContentBundleEditor>& ContentBundleEditor) const;
 
+	virtual void PushContentBundleEditing() override;
+	virtual void PopContentBundleEditing() override;
+
 	DECLARE_EVENT_OneParam(UContentBundleEditorSubsystem, FOnContentBundleChanged, const FContentBundleEditor*);
 	FOnContentBundleChanged& OnContentBundleChanged() { return ContentBundleChanged; }
 
diff --git a/Engine/Source/Runtime/Engine/Public/WorldPartition/ContentBundle/ContentBundleActivationScope.h b/Engine/Source/Runtime/Engine/Public/WorldPartition/ContentBundle/ContentBundleActivationScope.h
--- a/Engine/Source/Runtime/Engine/Public/WorldPartition/ContentBundle/ContentBundleActivationScope.h
+++ b/Engine/Source/Runtime/Engine/Public/WorldPartition/ContentBundle/ContentBundleActivationScope.h
@@ -15,13 +15,14 @@ public:
 	{
 		if (InContentBundleGuid.IsValid())
 		{
+			ContentBundleGuid = InContentBundleGuid;
+
 			IContentBundleEditorSubsystemInterface* ContentBundleEditorSubsystem = IContentBundleEditorSubsystemInterface::Get();
-			if (TSharedPtr<FContentBundleEditor> ContentBundleEditor = ContentBundleEditorSubsystem->GetEditorContentBundle(InContentBundleGuid))
-			{
-				check(!ContentBundleEditorSubsystem->IsEditingContentBundle());
+			ContentBundleEditorSubsystem->PushContentBundleEditing();
 
+			if (TSharedPtr<FContentBundleEditor> ContentBundleEditor = ContentBundleEditorSubsystem->GetEditorContentBundle(ContentBundleGuid))
+			{
 				ContentBundleEditorSubsystem->ActivateContentBundleEditing(ContentBundleEditor);
-				ContentBundleGuid = InContentBundleGuid;
 			}
 		}
 	}
@@ -34,6 +35,8 @@ public:
 			{
 				ContentBundleEditorSubsystem->DeactivateContentBundleEditing(ContentBundleEditor);
 			}
+
+			ContentBundleEditorSubsystem->PopContentBundleEditing();
 		}
 	}
 
diff --git a/Engine/Source/Runtime/Engine/Public/WorldPartition/ContentBundle/ContentBundleEditorSubsystemInterface.h b/Engine/Source/Runtime/Engine/Public/WorldPartition/ContentBundle/ContentBundleEditorSubsystemInterface.h
--- a/Engine/Source/Runtime/Engine/Public/WorldPartition/ContentBundle/ContentBundleEditorSubsystemInterface.h
+++ b/Engine/Source/Runtime/Engine/Public/WorldPartition/ContentBundle/ContentBundleEditorSubsystemInterface.h
@@ -22,6 +22,8 @@ public:
 	virtual bool IsEditingContentBundle() const = 0;
 	virtual bool ActivateContentBundleEditing(TSharedPtr<FContentBundleEditor>& ContentBundleEditor) const = 0;
 	virtual bool DeactivateContentBundleEditing(TSharedPtr<FContentBundleEditor>& ContentBundleEditor) const = 0;
+	virtual void PushContentBundleEditing() = 0;
+	virtual void PopContentBundleEditing() = 0;
 
 protected:
 	static void SetInstance(IContentBundleEditorSubsystemInterface* InInstance);
