World Partition - Builder: Added better logging between maps when processing collections

diff --git a/Engine/Source/Editor/UnrealEd/Private/Commandlets/WorldPartitionBuilderCommandlet.cpp b/Engine/Source/Editor/UnrealEd/Private/Commandlets/WorldPartitionBuilderCommandlet.cpp
--- a/Engine/Source/Editor/UnrealEd/Private/Commandlets/WorldPartitionBuilderCommandlet.cpp
+++ b/Engine/Source/Editor/UnrealEd/Private/Commandlets/WorldPartitionBuilderCommandlet.cpp
@@ -110,8 +110,16 @@ int32 UWorldPartitionBuilderCommandlet::Main(const FString& Params)
 	}	
 
 	// Run the builder on the provided map(s)
+	uint32 PackageIndex = 0;
+	const uint32 PackageCount = MapPackagesNames.Num();
 	for (const FString& MapPackageName : MapPackagesNames)
 	{
+		if (PackageCount > 1)
+		{
+			UE_LOG(LogWorldPartitionBuilderCommandlet, Display, TEXT("##################################################"));
+			UE_LOG(LogWorldPartitionBuilderCommandlet, Display, TEXT("[%d / %d] Executing %s on map %s..."), ++PackageIndex, PackageCount, *BuilderClassName, *MapPackageName);
+		}
+
 		if (!RunBuilder(BuilderClass, MapPackageName))
 		{
 			return 1;
