Added an invalid handle constant for PluginEditorExtensionHandle

diff --git a/Engine/Source/Editor/UnrealEd/Public/Features/IPluginsEditorFeature.h b/Engine/Source/Editor/UnrealEd/Public/Features/IPluginsEditorFeature.h
--- a/Engine/Source/Editor/UnrealEd/Public/Features/IPluginsEditorFeature.h
+++ b/Engine/Source/Editor/UnrealEd/Public/Features/IPluginsEditorFeature.h
@@ -96,6 +96,7 @@ struct FPluginEditorExtension : public TSharedFromThis<FPluginEditorExtension>
 DECLARE_DELEGATE_RetVal_TwoParams(TSharedPtr<FPluginEditorExtension>, FOnPluginBeingEdited, FPluginEditingContext& /*InPluginContext*/, IDetailLayoutBuilder& /*DetailBuilder*/);
 
 using FPluginEditorExtensionHandle = int32;
+const FPluginEditorExtensionHandle PluginEditorExtensionInvalidHandle = 0;
 
 /**
  * Feature interface for a Plugins management UI
