Level Instance: actors instance guid is now preserved when replacing actors.

diff --git a/Engine/Source/Editor/UnrealEd/Private/Kismet2/KismetReinstanceUtilities.cpp b/Engine/Source/Editor/UnrealEd/Private/Kismet2/KismetReinstanceUtilities.cpp
--- a/Engine/Source/Editor/UnrealEd/Private/Kismet2/KismetReinstanceUtilities.cpp
+++ b/Engine/Source/Editor/UnrealEd/Private/Kismet2/KismetReinstanceUtilities.cpp
@@ -1535,6 +1535,15 @@ namespace InstancedPropertyUtils
 	};
 }
 
+class FReplaceActorHelperSetActorInstanceGuid
+{
+public:
+	FReplaceActorHelperSetActorInstanceGuid(AActor* InActor, const FGuid& InActorInstanceGuid)
+	{
+		FSetActorInstanceGuid SetActorInstanceGuid(InActor, InActorInstanceGuid);
+	}
+};
+
 // @todo_deprecated - Remove in a future release.
 void FBlueprintCompileReinstancer::ReplaceInstancesOfClass(UClass* OldClass, UClass* NewClass, UObject*	OriginalCDO, TSet<UObject*>* ObjectsThatShouldUseOldStuff, bool bClassObjectReplaced, bool bPreserveRootComponent)
 {
@@ -2089,6 +2098,9 @@ static void ReplaceActorHelper(AActor* OldActor, UClass* OldClass, UObject*& New
 
 	check(NewActor != nullptr);
 
+	// Set the actor instance guid before registering components
+	FReplaceActorHelperSetActorInstanceGuid SetActorInstanceGuid(NewActor, OldActor->GetActorInstanceGuid());
+
 	// When Spawning an actor that has an external package the package can be PKG_NewlyCreated. 
 	// We need to remove this flag if the package didn't have that flag prior to the SpawnActor. 
 	// This means we are reinstancing an actor that was already saved on disk.
diff --git a/Engine/Source/Runtime/Engine/Classes/GameFramework/Actor.h b/Engine/Source/Runtime/Engine/Classes/GameFramework/Actor.h
--- a/Engine/Source/Runtime/Engine/Classes/GameFramework/Actor.h
+++ b/Engine/Source/Runtime/Engine/Classes/GameFramework/Actor.h
@@ -911,7 +911,7 @@ protected:
 	 * @see		ActorGuid
 	 * @note	This is not guaranteed to be valid during PostLoad in all situations, but safe to access from RegisterAllComponents.
 	 */
-	UPROPERTY(BluePrintReadOnly, AdvancedDisplay, Category=Actor, Transient, NonTransactional)
+	UPROPERTY(BluePrintReadOnly, AdvancedDisplay, Category=Actor, Transient, NonPIEDuplicateTransient, TextExportTransient, NonTransactional)
 	FGuid ActorInstanceGuid;
 
 	/**
@@ -4253,9 +4253,15 @@ private:
 	FSetActorInstanceGuid(AActor* InActor, const FGuid& InActorInstanceGuid)
 	{
 		InActor->ActorInstanceGuid = InActorInstanceGuid;
+		if (InActorInstanceGuid == InActor->ActorGuid)
+		{
+			InActor->ActorInstanceGuid.Invalidate();
+		}
 	}
+	friend class UEngine;
 	friend class ULevelStreamingLevelInstance;
 	friend class FWorldPartitionLevelHelper;
+	friend class FReplaceActorHelperSetActorInstanceGuid;
 };
 
 struct FSetActorFolderPath
diff --git a/Engine/Source/Runtime/Engine/Private/UnrealEngine.cpp b/Engine/Source/Runtime/Engine/Private/UnrealEngine.cpp
--- a/Engine/Source/Runtime/Engine/Private/UnrealEngine.cpp
+++ b/Engine/Source/Runtime/Engine/Private/UnrealEngine.cpp
@@ -16950,6 +16950,7 @@ void UEngine::CopyPropertiesForUnrelatedObjects(UObject* OldObject, UObject* New
 
 #if WITH_EDITOR
 	FGuid NewActorGuid;
+	FGuid NewActorInstanceGuid;
 #endif
 
 	// Bad idea to write data to an actor while its components are registered
@@ -16963,6 +16964,7 @@ void UEngine::CopyPropertiesForUnrelatedObjects(UObject* OldObject, UObject* New
 
 #if WITH_EDITOR
 		NewActorGuid = NewActor->GetActorGuid();
+		NewActorInstanceGuid = NewActor->GetActorInstanceGuid();
 #endif
 	}
 
@@ -17173,6 +17175,11 @@ void UEngine::CopyPropertiesForUnrelatedObjects(UObject* OldObject, UObject* New
 		{
 			FSetActorGuid SetActorGuid(NewActor, NewActorGuid);
 		}
+
+		if (NewActorInstanceGuid.IsValid())
+		{
+			FSetActorInstanceGuid SetActorInstanceGuid(NewActor, NewActorInstanceGuid);
+		}
 #endif
 	}
 
