Force all viewports to have focus so Auto LOD updating works again.  This is to fix issues with CL#22566736

- This will re-introduce LOD flicker between viewports.

diff --git a/Engine/Source/Editor/UnrealEd/Private/EditorEngine.cpp b/Engine/Source/Editor/UnrealEd/Private/EditorEngine.cpp
--- a/Engine/Source/Editor/UnrealEd/Private/EditorEngine.cpp
+++ b/Engine/Source/Editor/UnrealEd/Private/EditorEngine.cpp
@@ -2119,7 +2119,7 @@ void UEditorEngine::Tick( float DeltaSeconds, bool bIdleMode )
 					{
 						//if we haven't drawn a non-realtime viewport OR not one of the main viewports
 						bool bAllowNonRealtimeViewports = (!bEditorFrameNonRealtimeViewportDrawn) || !(ViewportClient->IsLevelEditorClient());
-						ViewportClient->SetIsCurrentLevelEditingFocus(false);
+						ViewportClient->SetIsCurrentLevelEditingFocus(true);
 						bool bWasNonRealtimeViewportDrawn = UpdateSingleViewportClient(ViewportClient, bAllowNonRealtimeViewports, bUpdateLinkedOrthoViewports);
 						if (ViewportClient->IsLevelEditorClient())
 						{
