Don't create default UCS graph or event graph nodes for new Blueprint assets when created inside a restricted editor mode context.

Notes:
- The end result is that no Function Entry or Call Site node objects will be instanced, and thus will not be exported by the Blueprint asset package inside of a restricted editor mode context.
- The ubergraph context is still created along with the event graph, but no default nodes are added.
- If the asset is later opened inside an unrestricted editor mode context, the UCS function graph will be added at that point, but at this time, the default event node set will not also be added.

diff --git a/Engine/Source/Editor/BlueprintGraph/Private/BlueprintEditorSettings.cpp b/Engine/Source/Editor/BlueprintGraph/Private/BlueprintEditorSettings.cpp
--- a/Engine/Source/Editor/BlueprintGraph/Private/BlueprintEditorSettings.cpp
+++ b/Engine/Source/Editor/BlueprintGraph/Private/BlueprintEditorSettings.cpp
@@ -270,4 +270,30 @@ bool UBlueprintEditorSettings::IsClassPathAllowedOnPin(const FTopLevelAssetPath&
 		}
 	}
 	return true;
+}
+
+bool UBlueprintEditorSettings::IsFunctionAllowed(const UBlueprint* InBlueprint, const FName FunctionName) const
+{
+	if (!FunctionPermissions.HasFiltering())
+	{
+		return true;
+	}
+
+	if (FunctionPermissions.PassesFilter(FunctionName))
+	{
+		return true;
+	}
+
+	if (InBlueprint && InBlueprint->ParentClass)
+	{
+		if (const UFunction* ParentFunction = InBlueprint->ParentClass->FindFunctionByName(FunctionName))
+		{
+			if (FunctionPermissions.PassesFilter(ParentFunction->GetPathName()))
+			{
+				return true;
+			}
+		}
+	}
+
+	return false;
 }
\ No newline at end of file
diff --git a/Engine/Source/Editor/BlueprintGraph/Public/BlueprintEditorSettings.h b/Engine/Source/Editor/BlueprintGraph/Public/BlueprintEditorSettings.h
--- a/Engine/Source/Editor/BlueprintGraph/Public/BlueprintEditorSettings.h
+++ b/Engine/Source/Editor/BlueprintGraph/Public/BlueprintEditorSettings.h
@@ -354,6 +354,8 @@ public:
 	void UnregisterIsClassPathAllowedOnPinDelegate(const FName OwnerName);
 	bool IsClassPathAllowedOnPin(const FTopLevelAssetPath& InClassPath) const;
 	bool HasClassPathOnPinFiltering() const { return IsClassPathAllowedOnPinDelegates.Num() > 0; }
+
+	bool IsFunctionAllowed(const UBlueprint* InBlueprint, const FName FunctionName) const;
 	
 private:
 	/** All function permissions */
diff --git a/Engine/Source/Editor/Kismet/Private/BlueprintEditor.cpp b/Engine/Source/Editor/Kismet/Private/BlueprintEditor.cpp
--- a/Engine/Source/Editor/Kismet/Private/BlueprintEditor.cpp
+++ b/Engine/Source/Editor/Kismet/Private/BlueprintEditor.cpp
@@ -1854,7 +1854,7 @@ void FBlueprintEditor::EnsureBlueprintIsUpToDate(UBlueprint* BlueprintObj)
 		{
 			ExistingUCS->bAllowDeletion = false;
 		}
-		else
+		else if(GetDefault<UBlueprintEditorSettings>()->IsFunctionAllowed(BlueprintObj, UEdGraphSchema_K2::FN_UserConstructionScript))
 		{
 			UEdGraph* UCSGraph = FBlueprintEditorUtils::CreateNewGraph(BlueprintObj, UEdGraphSchema_K2::FN_UserConstructionScript, UEdGraph::StaticClass(), UEdGraphSchema_K2::StaticClass());
 			FBlueprintEditorUtils::AddFunctionGraph(BlueprintObj, UCSGraph, /*bIsUserCreated=*/ false, AActor::StaticClass());
diff --git a/Engine/Source/Editor/UnrealEd/Private/Kismet2/Kismet2.cpp b/Engine/Source/Editor/UnrealEd/Private/Kismet2/Kismet2.cpp
--- a/Engine/Source/Editor/UnrealEd/Private/Kismet2/Kismet2.cpp
+++ b/Engine/Source/Editor/UnrealEd/Private/Kismet2/Kismet2.cpp
@@ -425,6 +425,9 @@ UBlueprint* FKismetEditorUtilities::CreateBlueprint(UClass* ParentClass, UObject
 	{
 		BlueprintType = BPTYPE_Const;
 	}
+
+	const UBlueprintEditorSettings* Settings = GetDefault<UBlueprintEditorSettings>();
+	check(Settings);
 	
 	// Create new UBlueprint object
 	UBlueprint* NewBP = NewObject<UBlueprint>(Outer, *BlueprintClassType, NewBPName, RF_Public | RF_Standalone | RF_Transactional | RF_LoadCompleted);
@@ -453,10 +456,14 @@ UBlueprint* FKismetEditorUtilities::CreateBlueprint(UClass* ParentClass, UObject
 		NewBP->SimpleConstructionScript->SetFlags(RF_Transactional);
 		NewBP->LastEditedDocuments.Add(ToRawPtr(NewBP->SimpleConstructionScript));
 
-		UEdGraph* UCSGraph = FKismetEditorUtilities::CreateUserConstructionScript(NewBP);
+		// Note: UCS graph creation may be restricted due to editor permissions.
+		if (Settings->IsFunctionAllowed(NewBP, UEdGraphSchema_K2::FN_UserConstructionScript))
+		{
+			UEdGraph* UCSGraph = FKismetEditorUtilities::CreateUserConstructionScript(NewBP);
 
-		NewBP->LastEditedDocuments.Add(UCSGraph);
-		UCSGraph->bAllowDeletion = false;
+			NewBP->LastEditedDocuments.Add(UCSGraph);
+			UCSGraph->bAllowDeletion = false;
+		}
 	}
 
 	// Create default event graph(s)
@@ -508,8 +515,7 @@ UBlueprint* FKismetEditorUtilities::CreateBlueprint(UClass* ParentClass, UObject
 	// Mark the BP as being regenerated, so it will not be confused as needing to be loaded and regenerated when a referenced BP loads.
 	NewBP->bHasBeenRegenerated = true;
 
-	UBlueprintEditorSettings* Settings = GetMutableDefault<UBlueprintEditorSettings>();
-	if(Settings && Settings->bSpawnDefaultBlueprintNodes)
+	if(Settings->bSpawnDefaultBlueprintNodes)
 	{
 		// Only add default events if there is an ubergraph and they are supported
 		if(NewBP->UbergraphPages.Num() && FBlueprintEditorUtils::DoesSupportEventGraphs(NewBP))
@@ -650,7 +656,12 @@ UK2Node_Event* FKismetEditorUtilities::AddDefaultEventNode(UBlueprint* InBluepri
 	EventReference.SetExternalMember(InEventName, InEventClass);
 
 	// Prevent events that are hidden in the Blueprint's class from being auto-generated.
-	if(!FObjectEditorUtils::IsFunctionHiddenFromClass(EventReference.ResolveMember<UFunction>(InBlueprint), InBlueprint->ParentClass))
+	const bool bIsFunctionVisible = !FObjectEditorUtils::IsFunctionHiddenFromClass(EventReference.ResolveMember<UFunction>(InBlueprint), InBlueprint->ParentClass);
+	
+	// Prevent events that are not allowed in the editor (due to permission settings).
+	const bool bIsFunctionAllowed = GetDefault<UBlueprintEditorSettings>()->IsFunctionAllowed(InBlueprint, InEventName);
+
+	if(bIsFunctionVisible && bIsFunctionAllowed)
 	{
 		const UEdGraphSchema_K2* Schema = GetDefault<UEdGraphSchema_K2>();
 
