[UnrealEd] Removed public dependency on AddContentDialog to break circular dependency "AddContentDialog -> UnrealEd -> AddContentDialog"

diff --git a/Engine/Source/Editor/UnrealEd/UnrealEd.Build.cs b/Engine/Source/Editor/UnrealEd/UnrealEd.Build.cs
--- a/Engine/Source/Editor/UnrealEd/UnrealEd.Build.cs
+++ b/Engine/Source/Editor/UnrealEd/UnrealEd.Build.cs
@@ -322,7 +322,6 @@ public class UnrealEd : ModuleRules
 				"CollectionManager",
 				"ContentBrowser",
 				"BlueprintGraph",
-				"AddContentDialog",
 				"MeshUtilities",
 				"KismetCompiler",
 				"NavigationSystem",
