Make sure to assign a unique label to new actors that are created without a factory
- In OFPA, actors without a factory had their name uniqueness tested using the object FName (with the appended UID part...) , so the name was always considered unique
- Added an explicit call to FActorLabelUtilities::SetActorLabelUnique(), which does the correct tests

diff --git a/Engine/Source/Editor/UnrealEd/Private/EditorEngine.cpp b/Engine/Source/Editor/UnrealEd/Private/EditorEngine.cpp
--- a/Engine/Source/Editor/UnrealEd/Private/EditorEngine.cpp
+++ b/Engine/Source/Editor/UnrealEd/Private/EditorEngine.cpp
@@ -6517,6 +6517,8 @@ AActor* UEditorEngine::AddActor(ULevel* InLevel, UClass* Class, const FTransform
 
 		if (Actor)
 		{
+			FActorLabelUtilities::SetActorLabelUnique(Actor, Actor->GetDefaultActorLabel());
+
 			if (bSelectActor)
 			{
 				SelectActor(Actor, 1, 0);
