Fix a memory leak because of a cicular dependency.
This also fix a general perf degradation of the editor when deleting assets because of the leaking assets pickers.

diff --git a/Engine/Source/Editor/UnrealEd/Private/Dialogs/SDeleteAssetsDialog.cpp b/Engine/Source/Editor/UnrealEd/Private/Dialogs/SDeleteAssetsDialog.cpp
--- a/Engine/Source/Editor/UnrealEd/Private/Dialogs/SDeleteAssetsDialog.cpp
+++ b/Engine/Source/Editor/UnrealEd/Private/Dialogs/SDeleteAssetsDialog.cpp
@@ -582,7 +582,10 @@ FReply SDeleteAssetsDialog::OnKeyDown( const FGeometry& MyGeometry, const FKeyEv
 {
 	if( InKeyEvent.GetKey() == EKeys::Escape )
 	{
-		ParentWindow.Get()->RequestDestroyWindow();
+		if (TSharedPtr<SWindow> OwningWindow = ParentWindow.Pin())
+		{
+			OwningWindow->RequestDestroyWindow();
+		}
 		return FReply::Handled();
 	}
 
@@ -680,7 +683,10 @@ void SDeleteAssetsDialog::DeleteRelevantSourceContent()
 
 FReply SDeleteAssetsDialog::Delete()
 {
-	ParentWindow.Get()->RequestDestroyWindow();
+	if (TSharedPtr<SWindow> OwningWindow = ParentWindow.Pin())
+	{
+		OwningWindow->RequestDestroyWindow();
+	}
 
 	if (DeleteModel->IsAnythingReferencedInMemoryByUndo())
 	{
@@ -695,14 +701,20 @@ FReply SDeleteAssetsDialog::Delete()
 
 FReply SDeleteAssetsDialog::Cancel()
 {
-	ParentWindow.Get()->RequestDestroyWindow();
+	if (TSharedPtr<SWindow> OwningWindow = ParentWindow.Pin())
+	{
+		OwningWindow->RequestDestroyWindow();
+	}
 
 	return FReply::Handled();
 }
 
 FReply SDeleteAssetsDialog::ForceDelete()
 {
-	ParentWindow.Get()->RequestDestroyWindow();
+	if (TSharedPtr<SWindow> OwningWindow = ParentWindow.Pin())
+	{
+		OwningWindow->RequestDestroyWindow();
+	}
 
 	if( DeleteModel->IsAnythingReferencedInMemoryByUndo() )
 	{
@@ -739,7 +751,10 @@ FReply SDeleteAssetsDialog::ReplaceReferences()
 
 	if ( EAppReturnType::Ok == FMessageDialog::Open( EAppMsgType::OkCancel, Message, &Title ) )
 	{
-		ParentWindow.Get()->RequestDestroyWindow();
+		if (TSharedPtr<SWindow> OwningWindow = ParentWindow.Pin())
+		{
+			OwningWindow->RequestDestroyWindow();
+		}
 		DeleteRelevantSourceContent();
 		DeleteModel->DoReplaceReferences( ConsolidationAsset );
 	}
@@ -812,7 +827,10 @@ void SDeleteAssetsDialog::OnAssetsActivated(const TArray<FAssetData>& ActivatedA
 	// Open a simple asset editor for all assets which do not have asset type actions if activating with enter or double click
 	if ( ActivationMethod == EAssetTypeActivationMethod::DoubleClicked || ActivationMethod == EAssetTypeActivationMethod::Opened )
 	{
-		ParentWindow.Get()->RequestDestroyWindow();
+		if (TSharedPtr<SWindow> OwningWindow = ParentWindow.Pin())
+		{
+			OwningWindow->RequestDestroyWindow();
+		}
 
 		for(const FAssetData& ActivatedAsset : ActivatedAssets)
 		{
diff --git a/Engine/Source/Editor/UnrealEd/Public/Dialogs/SDeleteAssetsDialog.h b/Engine/Source/Editor/UnrealEd/Public/Dialogs/SDeleteAssetsDialog.h
--- a/Engine/Source/Editor/UnrealEd/Public/Dialogs/SDeleteAssetsDialog.h
+++ b/Engine/Source/Editor/UnrealEd/Public/Dialogs/SDeleteAssetsDialog.h
@@ -39,7 +39,7 @@ public:
 		SLATE_ARGUMENT(FName, Style)
 
 		// The parent window hosting this dialog
-		SLATE_ATTRIBUTE(TSharedPtr<SWindow>, ParentWindow)
+		SLATE_ARGUMENT(TSharedPtr<SWindow>, ParentWindow)
 
 		/** When specified, the path box will request this fixed size. */
 		SLATE_ATTRIBUTE(FOptionalSize, WidthOverride)
@@ -147,8 +147,7 @@ private:
 	/** The model used for deleting assets */
 	TSharedPtr<FAssetDeleteModel> DeleteModel;
 
-	// Attributes
-	TAttribute< TSharedPtr< SWindow > > ParentWindow;
+	TWeakPtr<SWindow> ParentWindow;
 
 	// Widgets
 	TSharedPtr< SBorder > RootContainer;
