[Feature] AsyncTreeDifferences
This is an improvement on how we diff object details so we can diff in real time and unify several algorithms into one
AsyncTreeDifferences is a generic structure that can diff any tree while AsyncDetailViewDiff is a specialization of that for details panels
Docs: https://docs.google.com/document/d/1rwY-FtTxq2BVptLaMdpKNpzb4hOSn8efm-Z-cF0fBhE/edit?usp=sharing

diff --git a/Engine/Source/Editor/PropertyEditor/Internal/AsyncDetailViewDiff.h b/Engine/Source/Editor/PropertyEditor/Internal/AsyncDetailViewDiff.h
new file mode 100644
--- /dev/null
+++ b/Engine/Source/Editor/PropertyEditor/Internal/AsyncDetailViewDiff.h
@@ -0,0 +1,43 @@
+// Copyright Epic Games, Inc. All Rights Reserved.
+#pragma once
+#include "AsyncTreeDifferences.h"
+#include "DiffUtils.h"
+#include "IDetailsView.h"
+
+
+namespace TreeDiffSpecification
+{
+	template<>
+	bool PROPERTYEDITOR_API AreValuesEqual(const TWeakPtr<FDetailTreeNode>& TreeNodeA, const TWeakPtr<FDetailTreeNode>& TreeNodeB);
+	
+	template<>
+	bool PROPERTYEDITOR_API AreMatching(const TWeakPtr<FDetailTreeNode>& TreeNodeA, const TWeakPtr<FDetailTreeNode>& TreeNodeB);
+	
+	template<>
+	void PROPERTYEDITOR_API GetChildren(const TWeakPtr<FDetailTreeNode>& InParent, TArray<TWeakPtr<FDetailTreeNode>>& OutChildren);
+
+	template<>
+	bool PROPERTYEDITOR_API ShouldMatchByValue(const TWeakPtr<FDetailTreeNode>& TreeNodeA);
+};
+
+// asynchronously update a difference tree with changes to details views
+// note: users are expected to call TAsyncTreeDifferences::Tick or TAsyncTreeDifferences::FlushQueue to get accurate data
+class PROPERTYEDITOR_API FAsyncDetailViewDiff : public TAsyncTreeDifferences<TWeakPtr<FDetailTreeNode>>
+{
+public:
+	FAsyncDetailViewDiff(TSharedRef<IDetailsView> LeftView, TSharedRef<IDetailsView> RightView);
+
+	// Note: by default this list is not necessarily accurate or exhaustive.
+	// if you need a perfectly accurate list, call FlushQueue() first
+	void GetPropertyDifferences(TArray<FSingleObjectDiffEntry>& OutDiffEntries) const;
+
+	// execute a method for each node associated with a row in at least one details view
+	// method provides current left and right row number
+	// returns the number of rows in each details view
+	TPair<int32, int32> ForEachRow(const TFunction<ETreeTraverseControl(const TUniquePtr<DiffNodeType>&, int32, int32)>& Method) const;
+
+private:
+	static TAttribute<TArray<TWeakPtr<FDetailTreeNode>>> RootNodesAttribute(TWeakPtr<IDetailsView> DetailsView);
+	TWeakPtr<IDetailsView> LeftView;
+	TWeakPtr<IDetailsView> RightView;
+};
diff --git a/Engine/Source/Editor/PropertyEditor/Private/AsyncDetailViewDiff.cpp b/Engine/Source/Editor/PropertyEditor/Private/AsyncDetailViewDiff.cpp
new file mode 100644
--- /dev/null
+++ b/Engine/Source/Editor/PropertyEditor/Private/AsyncDetailViewDiff.cpp
@@ -0,0 +1,312 @@
+// Copyright Epic Games, Inc. All Rights Reserved.
+#include "AsyncDetailViewDiff.h"
+
+#include "DetailTreeNode.h"
+#include "IDetailsViewPrivate.h"
+#include "DiffUtils.h"
+
+static FResolvedProperty GetResolvedProperty(const TSharedRef<FDetailTreeNode>& TreeNode)
+{
+	const FPropertyPath PropertyPath = TreeNode->GetPropertyPath();
+	if (PropertyPath.IsValid())
+	{
+		if (const IDetailsViewPrivate* DetailsView = TreeNode->GetDetailsView())
+		{
+			if (const UObject* Object = DetailsView->GetSelectedObjects()[0].Get())
+			{
+				return FPropertySoftPath(PropertyPath).Resolve(Object);
+			}
+		}
+	}
+	return FResolvedProperty();
+}
+
+template<>
+bool TreeDiffSpecification::AreValuesEqual<TWeakPtr<FDetailTreeNode>>(const TWeakPtr<FDetailTreeNode>& TreeNodeA,  const TWeakPtr<FDetailTreeNode>& TreeNodeB)
+{
+	const TSharedPtr<FDetailTreeNode> PinnedTreeNodeA = TreeNodeA.Pin();
+	const TSharedPtr<FDetailTreeNode> PinnedTreeNodeB = TreeNodeB.Pin();
+	if (!PinnedTreeNodeA || !PinnedTreeNodeB)
+	{
+		return PinnedTreeNodeA == PinnedTreeNodeB;
+	}
+
+	const FResolvedProperty ResolvedA = GetResolvedProperty(PinnedTreeNodeA.ToSharedRef());
+    const FResolvedProperty ResolvedB = GetResolvedProperty(PinnedTreeNodeB.ToSharedRef());
+	if (ResolvedA.Property && ResolvedB.Property)
+	{
+		// property nodes
+		if (!ResolvedA.Property->SameType(ResolvedB.Property))
+		{
+			return false;
+		}
+	    const void* DataA = ResolvedA.Property->ContainerPtrToValuePtr<void*>(ResolvedA.Object);
+	    const void* DataB = ResolvedA.Property->ContainerPtrToValuePtr<void*>(ResolvedB.Object);
+		return ResolvedA.Property->Identical(DataA, DataB, PPF_DeepComparison);
+	}
+	if (!ResolvedA.Property && !ResolvedB.Property)
+	{
+		// category nodes
+		return PinnedTreeNodeA->GetNodeName() == PinnedTreeNodeB->GetNodeName();
+	}
+	return ensure(false); // AreMatching(...) should've stopped this from happening
+}
+
+static bool MapKeysMatch(const TSharedPtr<FDetailTreeNode>& TreeNodeA, const TSharedPtr<FDetailTreeNode>& TreeNodeB, int32 KeyIndexA, int32 KeyIndexB)
+{
+	const FMapProperty* MapPropertyA = CastField<FMapProperty>(TreeNodeA->GetPropertyNode()->GetProperty());
+	const FMapProperty* MapPropertyB = CastField<FMapProperty>(TreeNodeB->GetPropertyNode()->GetProperty());
+	if (MapPropertyA && MapPropertyB)
+	{
+		const FResolvedProperty ResolvedMapA = GetResolvedProperty(TreeNodeA.ToSharedRef());
+		const FResolvedProperty ResolvedMapB = GetResolvedProperty(TreeNodeB.ToSharedRef());
+		FScriptMapHelper MapHelperA(MapPropertyA, MapPropertyA->ContainerPtrToValuePtr<UObject*>(ResolvedMapA.Object));
+		FScriptMapHelper MapHelperB(MapPropertyB, MapPropertyB->ContainerPtrToValuePtr<UObject*>(ResolvedMapB.Object));
+
+		const void* KeyA = MapHelperA.GetKeyPtr(KeyIndexA);
+		const void* KeyB = MapHelperB.GetKeyPtr(KeyIndexB);
+		if (MapPropertyA->KeyProp->SameType(MapPropertyB->KeyProp))
+		{
+			return MapPropertyA->KeyProp->Identical(KeyA, KeyB, PPF_DeepComparison);
+		}
+	}
+	return false;
+}
+
+static bool SetKeysMatch(const TSharedPtr<FDetailTreeNode>& TreeNodeA, const TSharedPtr<FDetailTreeNode>& TreeNodeB, int32 KeyIndexA, int32 KeyIndexB)
+{
+	const FSetProperty* SetPropertyA = CastField<FSetProperty>(TreeNodeA->GetPropertyNode()->GetProperty());
+	const FSetProperty* SetPropertyB = CastField<FSetProperty>(TreeNodeB->GetPropertyNode()->GetProperty());
+	if (SetPropertyA && SetPropertyB)
+	{
+		const FResolvedProperty ResolvedSetA = GetResolvedProperty(TreeNodeA.ToSharedRef());
+		const FResolvedProperty ResolvedSetB = GetResolvedProperty(TreeNodeB.ToSharedRef());
+		FScriptSetHelper SetHelperA(SetPropertyA, SetPropertyA->ContainerPtrToValuePtr<UObject*>(ResolvedSetA.Object));
+		FScriptSetHelper SetHelperB(SetPropertyB, SetPropertyB->ContainerPtrToValuePtr<UObject*>(ResolvedSetB.Object));
+		const void* KeyA = SetHelperA.GetElementPtr(KeyIndexA);
+		const void* KeyB = SetHelperB.GetElementPtr(KeyIndexB);
+		if (SetPropertyA->ElementProp->SameType(SetPropertyB->ElementProp))
+		{
+			return SetPropertyA->ElementProp->Identical(KeyA, KeyB, PPF_DeepComparison);
+		}
+	}
+	return false;
+}
+
+template<>
+bool TreeDiffSpecification::AreMatching<TWeakPtr<FDetailTreeNode>>(const TWeakPtr<FDetailTreeNode>& TreeNodeA, const TWeakPtr<FDetailTreeNode>& TreeNodeB)
+{
+	const TSharedPtr<FDetailTreeNode> PinnedTreeNodeA = TreeNodeA.Pin();
+	const TSharedPtr<FDetailTreeNode> PinnedTreeNodeB = TreeNodeB.Pin();
+	if (!PinnedTreeNodeA || !PinnedTreeNodeB)
+	{
+		return PinnedTreeNodeA == PinnedTreeNodeB;
+	}
+	
+	const TSharedPtr<FPropertyNode> PropertyNodeA = PinnedTreeNodeA->GetPropertyNode();
+	const TSharedPtr<FPropertyNode> PropertyNodeB =PinnedTreeNodeB->GetPropertyNode();
+	if (PropertyNodeA && PropertyNodeB)
+	{
+		// property nodes
+		const int32 ArrayIndexA = PropertyNodeA->GetArrayIndex();
+		const int32 ArrayIndexB = PropertyNodeB->GetArrayIndex();
+		const FProperty* PropertyA = PropertyNodeA->GetProperty();
+		const FProperty* PropertyB = PropertyNodeB->GetProperty();
+		
+		if (ArrayIndexA != INDEX_NONE && ArrayIndexB != INDEX_NONE)
+		{
+			const FProperty* ParentPropertyA = PropertyNodeA->GetParentNode()->GetProperty();
+			const FProperty* ParentPropertyB = PropertyNodeB->GetParentNode()->GetProperty();
+			
+			// sets and maps are stored by index in the property tree so we need to dig their keys out of the data
+			// and compare those instead
+			if (CastField<FMapProperty>(ParentPropertyA) || CastField<FMapProperty>(ParentPropertyB))
+			{
+				return MapKeysMatch(
+					PinnedTreeNodeA->GetParentNode().Pin(),
+					PinnedTreeNodeB->GetParentNode().Pin(),
+					ArrayIndexA,
+					ArrayIndexB
+				);
+			}
+			if (ParentPropertyA->IsA<FSetProperty>() || ParentPropertyB->IsA<FSetProperty>())
+			{
+				return SetKeysMatch(
+					PinnedTreeNodeA->GetParentNode().Pin(),
+					PinnedTreeNodeB->GetParentNode().Pin(),
+					ArrayIndexA,
+					ArrayIndexB
+				);
+			}
+		}
+		
+		if (ArrayIndexA != ArrayIndexB)
+		{
+			return false;
+		}
+		const FName PropertyNameA = PropertyA ? PropertyA->GetFName() : NAME_None;
+		const FName PropertyNameB = PropertyB ? PropertyB->GetFName() : NAME_None;
+		
+		return PropertyNameA == PropertyNameB;
+	}
+	if (!PropertyNodeA && !PropertyNodeB)
+	{
+		// category nodes
+		return PinnedTreeNodeA->GetNodeName() == PinnedTreeNodeB->GetNodeName();
+	}
+	// node type mismatch
+	return false;
+}
+
+template<>
+void TreeDiffSpecification::GetChildren<TWeakPtr<FDetailTreeNode>>(const TWeakPtr<FDetailTreeNode>& InParent, TArray<TWeakPtr<FDetailTreeNode>>& OutChildren)
+{
+	const TSharedPtr<FDetailTreeNode> PinnedParent = InParent.Pin();
+	if (PinnedParent)
+	{
+		TArray<TSharedRef<FDetailTreeNode>> Children;
+		PinnedParent->GetChildren(Children);
+		for (TSharedRef<FDetailTreeNode> Child : Children)
+		{
+			OutChildren.Add(Child);
+		}
+	}
+}
+
+template<>
+bool TreeDiffSpecification::ShouldMatchByValue<TWeakPtr<FDetailTreeNode>>(const TWeakPtr<FDetailTreeNode>& TreeNode)
+{
+	const TSharedPtr<FDetailTreeNode> PinnedTreeNode = TreeNode.Pin();
+	if (!PinnedTreeNode)
+	{
+		return false;
+	}
+	const TSharedPtr<FPropertyNode> PropertyNodeA = PinnedTreeNode->GetPropertyNode();
+	if (!PropertyNodeA || !PropertyNodeA->GetParentNode())
+	{
+		return false;
+	}
+	
+	const int32 ArrayIndex = PropertyNodeA->GetArrayIndex();
+	const FArrayProperty* ParentArrayProperty = CastField<FArrayProperty>(PropertyNodeA->GetParentNode()->GetProperty());
+	
+	// match array elements by value rather than by index
+	return ParentArrayProperty && ArrayIndex != INDEX_NONE;
+}
+
+FAsyncDetailViewDiff::FAsyncDetailViewDiff(TSharedRef<IDetailsView> InLeftView, TSharedRef<IDetailsView> InRightView)
+	: TAsyncTreeDifferences(RootNodesAttribute(InLeftView), RootNodesAttribute(InRightView))
+	, LeftView(InLeftView)
+	, RightView(InRightView)
+{}
+
+void FAsyncDetailViewDiff::GetPropertyDifferences(TArray<FSingleObjectDiffEntry>& OutDiffEntries) const
+{
+	ForEach(ETreeTraverseOrder::PreOrder, [&](const TUniquePtr<DiffNodeType>& Node)->ETreeTraverseControl
+	{
+		FPropertyPath PropertyPath;
+		FPropertyPath RightPropertyPath;
+		if (const TSharedPtr<FDetailTreeNode> LeftTreeNode = Node->ValueA.Pin())
+		{
+			PropertyPath = LeftTreeNode->GetPropertyPath();
+		}
+		else if (const TSharedPtr<FDetailTreeNode> RightTreeNode = Node->ValueB.Pin())
+		{
+			PropertyPath = RightTreeNode->GetPropertyPath();
+		}
+
+		// only include tree nodes with properties
+		if (!PropertyPath.IsValid())
+		{
+			return ETreeTraverseControl::Continue;
+		}
+		
+		EPropertyDiffType::Type PropertyDiffType;
+        switch(Node->DiffResult)
+        {
+        case ETreeDiffResult::MissingFromTree1: 
+            PropertyDiffType = EPropertyDiffType::PropertyAddedToB;
+            break;
+        case ETreeDiffResult::MissingFromTree2: 
+            PropertyDiffType = EPropertyDiffType::PropertyAddedToA;
+            break;
+        case ETreeDiffResult::DifferentValues: 
+            PropertyDiffType = EPropertyDiffType::PropertyValueChanged;
+            break;
+        default:
+        	// only include changes
+            return ETreeTraverseControl::Continue;
+        }
+
+		OutDiffEntries.Add(FSingleObjectDiffEntry(PropertyPath, PropertyDiffType));
+        return ETreeTraverseControl::SkipChildren; // only include top-most properties
+	});
+}
+
+TPair<int32, int32> FAsyncDetailViewDiff::ForEachRow(const TFunction<ETreeTraverseControl(const TUniquePtr<DiffNodeType>&, int32, int32)>& Method) const
+{
+	const TSharedPtr<IDetailsView> LeftDetailsView = LeftView.Pin();
+	const TSharedPtr<IDetailsView> RightDetailsView = RightView.Pin();
+	if (!LeftDetailsView || !RightDetailsView)
+	{
+		return {0,0};
+	}
+	
+	int32 LeftRowNum = 0;
+	int32 RightRowNum = 0;
+	ForEach(
+		ETreeTraverseOrder::PreOrder,
+		[&LeftDetailsView,&RightDetailsView,&Method,&LeftRowNum,&RightRowNum](const TUniquePtr<DiffNodeType>& DiffNode)->ETreeTraverseControl
+		{
+			bool bFoundLeftRow = false;
+			if (const TSharedPtr<FDetailTreeNode> LeftTreeNode = DiffNode->ValueA.Pin())
+			{
+				if (!LeftDetailsView->IsAncestorCollapsed(LeftTreeNode.ToSharedRef()))
+				{
+					bFoundLeftRow = true;
+				}
+			}
+			
+			bool bFoundRightRow = false;
+			if (const TSharedPtr<FDetailTreeNode> RightTreeNode = DiffNode->ValueB.Pin())
+			{
+				if (!RightDetailsView->IsAncestorCollapsed(RightTreeNode.ToSharedRef()))
+				{
+					bFoundRightRow = true;
+				}
+			}
+
+			ETreeTraverseControl Control = ETreeTraverseControl::SkipChildren;
+			if (bFoundRightRow || bFoundLeftRow)
+			{
+				Control = Method(DiffNode, LeftRowNum, RightRowNum);
+			}
+			
+			if (bFoundLeftRow)
+			{
+				++LeftRowNum;
+			}
+			if (bFoundRightRow)
+			{
+				++RightRowNum;
+			}
+			
+			return Control;
+		}
+	);
+	return {LeftRowNum, RightRowNum};
+}
+
+TAttribute<TArray<TWeakPtr<FDetailTreeNode>>> FAsyncDetailViewDiff::RootNodesAttribute(TWeakPtr<IDetailsView> DetailsView)
+{
+	return TAttribute<TArray<TWeakPtr<FDetailTreeNode>>>::CreateLambda([DetailsView]()
+	{
+		TArray<TWeakPtr<FDetailTreeNode>> Result;
+		if (const TSharedPtr<IDetailsViewPrivate> Details = StaticCastSharedPtr<IDetailsViewPrivate>(DetailsView.Pin()))
+		{
+			Details->GetHeadNodes(Result);
+		}
+		return Result;
+	});
+}
+
diff --git a/Engine/Source/Editor/PropertyEditor/Private/IDetailsViewPrivate.h b/Engine/Source/Editor/PropertyEditor/Private/IDetailsViewPrivate.h
--- a/Engine/Source/Editor/PropertyEditor/Private/IDetailsViewPrivate.h
+++ b/Engine/Source/Editor/PropertyEditor/Private/IDetailsViewPrivate.h
@@ -167,4 +167,7 @@ public:
 
 	/** Set the given group's favorite status. */
 	virtual void SetCustomBuilderFavorite(FStringView GroupPath, bool IsFavorite) = 0;
+	
+	/** Retrieve a list of top-most detail tree nodes. */
+	virtual void GetHeadNodes(TArray<TWeakPtr<FDetailTreeNode>>& OutNodes) {}
 };
diff --git a/Engine/Source/Editor/PropertyEditor/Private/SDetailsViewBase.cpp b/Engine/Source/Editor/PropertyEditor/Private/SDetailsViewBase.cpp
--- a/Engine/Source/Editor/PropertyEditor/Private/SDetailsViewBase.cpp
+++ b/Engine/Source/Editor/PropertyEditor/Private/SDetailsViewBase.cpp
@@ -329,6 +329,19 @@ void SDetailsViewBase::HighlightProperty(const FPropertyPath& Property)
 	CurrentlyHighlightedNode = TreeNode;
 }
 
+bool SDetailsViewBase::IsAncestorCollapsed(const TSharedRef<IDetailTreeNode>& Node) const
+{
+	// in order for a node to be expanded, all it's parents have to be expanded
+	bool bIsCollapsed = false;
+	TSharedPtr<FDetailTreeNode> Ancestor = StaticCastSharedRef<FDetailTreeNode>(Node)->GetParentNode().Pin();
+	while (!bIsCollapsed && Ancestor)
+	{
+		bIsCollapsed = !DetailTree->IsItemExpanded(Ancestor.ToSharedRef());
+		Ancestor = Ancestor->GetParentNode().Pin();
+	}
+	return bIsCollapsed;
+}
+
 void SDetailsViewBase::ShowAllAdvancedProperties()
 {
 	CurrentFilter.bShowAllAdvanced = true;
@@ -339,6 +352,14 @@ void SDetailsViewBase::SetOnDisplayedPropertiesChanged(FOnDisplayedPropertiesCha
 	OnDisplayedPropertiesChangedDelegate = InOnDisplayedPropertiesChangedDelegate;
 }
 
+void SDetailsViewBase::GetHeadNodes(TArray<TWeakPtr<FDetailTreeNode>>& OutNodes)
+{
+	for (TSharedRef<FDetailTreeNode>& Node : RootTreeNodes)
+	{
+		OutNodes.Add(Node.ToWeakPtr());
+	}
+}
+
 void SDetailsViewBase::SetRightColumnMinWidth(float InMinWidth)
 {
 	ColumnSizeData.SetRightColumnMinWidth(InMinWidth);
diff --git a/Engine/Source/Editor/PropertyEditor/Private/SDetailsViewBase.h b/Engine/Source/Editor/PropertyEditor/Private/SDetailsViewBase.h
--- a/Engine/Source/Editor/PropertyEditor/Private/SDetailsViewBase.h
+++ b/Engine/Source/Editor/PropertyEditor/Private/SDetailsViewBase.h
@@ -88,6 +88,7 @@ public:
 	virtual TArray<TPair<int32, FPropertyPath>> GetPropertyRowNumbers() const override;
 	virtual int32 CountRows() const override;
 	virtual void HighlightProperty(const FPropertyPath& Property) override;
+	virtual bool IsAncestorCollapsed(const TSharedRef<IDetailTreeNode>& Node) const override;
 	virtual void ShowAllAdvancedProperties() override;
 	virtual void SetOnDisplayedPropertiesChanged(FOnDisplayedPropertiesChanged InOnDisplayedPropertiesChangedDelegate) override;
 	virtual FOnDisplayedPropertiesChanged& GetOnDisplayedPropertiesChanged() override { return OnDisplayedPropertiesChangedDelegate; }
@@ -148,6 +149,7 @@ public:
 	virtual bool IsFavoritingEnabled() const override { return DetailsViewArgs.bAllowFavoriteSystem; }
 	virtual bool IsConnected() const = 0;
 	virtual FRootPropertyNodeList& GetRootNodes() = 0;
+	virtual void GetHeadNodes(TArray<TWeakPtr<FDetailTreeNode>>&) override;
 
 	/**
 	 * Called when the open color picker window associated with this details view is closed
diff --git a/Engine/Source/Editor/PropertyEditor/Public/IDetailsView.h b/Engine/Source/Editor/PropertyEditor/Public/IDetailsView.h
--- a/Engine/Source/Editor/PropertyEditor/Public/IDetailsView.h
+++ b/Engine/Source/Editor/PropertyEditor/Public/IDetailsView.h
@@ -8,6 +8,7 @@
 #include "Framework/Commands/UICommandList.h"
 #include "DetailsViewArgs.h"
 
+class FDetailTreeNode;
 class AActor;
 class FNotifyHook;
 class FPropertyPath;
@@ -194,6 +195,11 @@ public:
 	 * row numbers. If you need a list of all properties independent of expansion state, use GetPropertiesInOrderDisplayed
 	 */
 	virtual TArray<TPair<int32, FPropertyPath>> GetPropertyRowNumbers() const  = 0;
+
+	/**
+	 * Returns true if the provided detail tree node is part of a collapsed/hidden subsection of the details panel
+	 */
+	virtual bool IsAncestorCollapsed(const TSharedRef<IDetailTreeNode>& Node) const { return false; }
 	
 	/**
 	 * returns the number of rows in the underlying TableView
diff --git a/Engine/Source/Editor/UnrealEd/Internal/AsyncTreeDifferences.h b/Engine/Source/Editor/UnrealEd/Internal/AsyncTreeDifferences.h
new file mode 100644
--- /dev/null
+++ b/Engine/Source/Editor/UnrealEd/Internal/AsyncTreeDifferences.h
@@ -0,0 +1,500 @@
+// Copyright Epic Games, Inc. All Rights Reserved.
+#pragma once
+#include "Templates/UniquePtr.h"
+#include "Engine/EngineTypes.h"
+#include "Misc/Attribute.h"
+
+enum class ETreeDiffResult
+{
+	Invalid,
+	MissingFromTree1,
+	MissingFromTree2,
+	DifferentValues,
+	Identical
+};
+
+enum class ETreeTraverseOrder
+{
+	PreOrder, // parent than children
+	PostOrder // children then parent
+};
+
+enum class ETreeTraverseControl
+{
+	Continue, // continue traversing
+	Break, // stop
+
+	// PreOrder traversal only:
+	SkipChildren, // don't iterate the children of this node
+};
+
+/// To use TAsyncTreeDifferences, define the following template specializations for your type
+////////////////////////////////////////////////////////////////////////////////////////////////////////////////////
+namespace TreeDiffSpecification
+{
+	/**
+	 * determine whether the values stored in two nodes are equal.
+	 * @param TreeNodeA node from the first user provided tree (guaranteed not to be null)
+	 * @param TreeNodeB node from the second user provided tree (guaranteed not to be null)
+	 */
+	template<typename InNodeType>
+	bool AreValuesEqual(const InNodeType& TreeNodeA, const InNodeType& TreeNodeB);
+
+	/**
+	 * determine whether two nodes occupy the same space in their trees
+	 * for example if you have a tree key/value pairs, AreMatching should compare the keys while AreValuesEqual
+	 * should compare the values
+	 * @param TreeNodeA node from the first user provided tree (guaranteed not to be null)
+	 * @param TreeNodeB node from the second user provided tree (guaranteed not to be null)
+	 */
+	template<typename InNodeType>
+	bool AreMatching(const InNodeType& TreeNodeA, const InNodeType& TreeNodeB);
+
+	/**
+	 * retrieves an array of children nodes from the parent node
+	 * @param InParent node from one of the two user provided trees (guaranteed not to be null)
+	 * @param[out] OutChildren to be filled with the children of parent
+	 */
+	template<typename InNodeType>
+	void GetChildren(const InNodeType& InParent, TArray<InNodeType>& OutChildren);
+
+	/**
+	 * return true for nodes that match using AreValuesEqual first, and pair up by position second
+	 * this is useful for arrays since we often want to keep elements with the same data paired while diffing other elements in order
+	 * @param TreeNode node from one of the two user provided trees (guaranteed not to be null)
+	*/
+	template<typename InNodeType>
+	bool ShouldMatchByValue(const InNodeType& TreeNode);
+
+};
+
+
+// the TAsyncTreeDifferences structure's goal is to build and maintain an updated tree made of these nodes.
+// where TDiffNode::ValueA and TDiffNode::ValueB are matching nodes from two diffed trees. Note that these values can
+// be both "matching" and not equal which would set TDiffNode::DiffResult to DifferentValues
+template<typename InNodeType>
+struct TDiffNode
+{
+	// because we've got a bit of node-type inception going on, we'll refer to the nodes in the user provided trees as
+	// values rather than nodes. (but it's still sometimes helpful to remember that they are nodes in trees)
+	using ValueType = InNodeType;
+	
+	ValueType ValueA = NullValue;
+	ValueType ValueB = NullValue;
+	
+	TArray<TUniquePtr<TDiffNode>> Children = {};
+	const TDiffNode* Parent = nullptr;
+	ETreeDiffResult DiffResult = ETreeDiffResult::Invalid;
+
+	TDiffNode() = default;
+	TDiffNode(TDiffNode &&Other) = default;
+	
+	TDiffNode(const ValueType& InValueA, const ValueType& InValueB, const TDiffNode* InParent)
+		: ValueA(InValueA), ValueB(InValueB), Children(), Parent(InParent)
+	{
+		SetDiffType();
+	}
+
+	void SetDiffType();
+
+	bool operator==(const TDiffNode& Other) const;
+	
+	static inline ValueType NullValue{}; // assume default constructor to signify a "null" node
+};
+
+template<typename InNodeType>
+class TAsyncTreeDifferences
+{
+public:
+	using ValueType = InNodeType;
+	using DiffNodeType = TDiffNode<ValueType>;
+	
+	TAsyncTreeDifferences(const TAttribute<TArray<ValueType>>& InRootValuesA, const TAttribute<TArray<ValueType>>& InRootValuesB);
+	
+	// update the differences list over time
+	void Tick(float MaxAllottedTimeMs = 1.f);
+
+	// process all tasks synchronously until all data is up to date. (slow)
+	void FlushQueue();
+
+	// calls Method(DiffNode) on each diff node in the tree.
+	void ForEach(ETreeTraverseOrder TraversalOrder, const TFunction<ETreeTraverseControl(const TUniquePtr<DiffNodeType>&)>& Method) const;
+
+	// head of the main diff tree (note: not for user modification. Tick() will overwrite user changes)
+	TUniquePtr<DiffNodeType> Head;
+private:
+
+	// process one element from the top of the queue
+	void ProcessTopOfQueue();
+	
+	void QueueParallelNodeLists(const TArray<ValueType>& ValuesA, const TArray<ValueType>& ValuesB, DiffNodeType* ParentNode);
+	
+	static TArray<TArray<int32>> CalculateLCSTableForMatchingValues(const TArray<ValueType>& ValuesA, const TArray<ValueType>& ValuesB);
+	static TArray<TArray<int32>> CalculateLCSTableForDiffNodes(const TArray<DiffNodeType>& FoundNodes, const TArray<TUniquePtr<DiffNodeType>>& ExpectedNodes);
+
+	static bool AreMatching(const ValueType& ValueA, const ValueType& ValueB);
+	
+	static bool PreOrderRecursive(const TUniquePtr<DiffNodeType>& Node, const TFunction<ETreeTraverseControl(const TUniquePtr<DiffNodeType>&)>& Method);
+	static bool PostOrderRecursive(const TUniquePtr<DiffNodeType>& Node, const TFunction<ETreeTraverseControl(const TUniquePtr<DiffNodeType>&)>& Method);
+
+	TArray<DiffNodeType*> UpdateQueue;
+	TAttribute<TArray<ValueType>> RootValuesA;
+	TAttribute<TArray<ValueType>> RootValuesB;
+};
+
+template <typename InNodeType>
+void TDiffNode<InNodeType>::SetDiffType()
+{
+	if (ValueA != NullValue && ValueB != NullValue)
+	{
+		if (TreeDiffSpecification::AreValuesEqual<ValueType>(ValueA, ValueB))
+		{
+			DiffResult = ETreeDiffResult::Identical;
+			return;
+		}
+		DiffResult = ETreeDiffResult::DifferentValues;
+		return;
+	}
+	if (ValueA != NullValue && ValueB == NullValue)
+	{
+		DiffResult = ETreeDiffResult::MissingFromTree2;
+		return;
+	}
+	if (ValueA == NullValue && ValueB != NullValue)
+	{
+		DiffResult = ETreeDiffResult::MissingFromTree1;
+		return;
+	}
+	
+	check(false);
+	DiffResult = ETreeDiffResult::Invalid;
+}
+
+template <typename InNodeType>
+bool TDiffNode<InNodeType>::operator==(const TDiffNode& Other) const
+{
+	return Other.ValueA == ValueA && Other.ValueB == ValueB;
+}
+
+template <typename InNodeType>
+TAsyncTreeDifferences<InNodeType>::TAsyncTreeDifferences(const TAttribute<TArray<ValueType>>& InRootValuesA, const TAttribute<TArray<ValueType>>& InRootValuesB)
+	: Head(MakeUnique<DiffNodeType>())
+	, RootValuesA(InRootValuesA)
+	, RootValuesB(InRootValuesB)
+{
+	QueueParallelNodeLists(RootValuesA.Get(), RootValuesB.Get(), Head.Get());
+}
+
+template <typename InNodeType>
+void TAsyncTreeDifferences<InNodeType>::Tick(float MaxAllottedTimeMs)
+{
+	if (UpdateQueue.IsEmpty())
+	{
+		QueueParallelNodeLists(RootValuesA.Get(), RootValuesB.Get(), Head.Get());
+	}
+	
+	const double StartTime = FPlatformTime::Seconds();
+	while(!UpdateQueue.IsEmpty() && FPlatformTime::Seconds() - StartTime < MaxAllottedTimeMs)
+	{
+		ProcessTopOfQueue();
+	}
+}
+
+template <typename InNodeType>
+void TAsyncTreeDifferences<InNodeType>::FlushQueue()
+{
+	while(!UpdateQueue.IsEmpty())
+	{
+		ProcessTopOfQueue();
+	}
+}
+
+template <typename InNodeType>
+void TAsyncTreeDifferences<InNodeType>::ForEach(ETreeTraverseOrder TraversalOrder, const TFunction<ETreeTraverseControl(const TUniquePtr<DiffNodeType>&)>& Method) const
+{
+	switch(TraversalOrder)
+	{
+	case ETreeTraverseOrder::PreOrder:
+		PreOrderRecursive(Head, Method);
+		break;
+	case ETreeTraverseOrder::PostOrder:
+		PostOrderRecursive(Head, Method);
+		break;
+	default: check(false);
+	}
+}
+
+template <typename InNodeType>
+bool TAsyncTreeDifferences<InNodeType>::PreOrderRecursive(const TUniquePtr<DiffNodeType>& Node, const TFunction<ETreeTraverseControl(const TUniquePtr<DiffNodeType>&)>& Method)
+{
+	if (Node->DiffResult != ETreeDiffResult::Invalid)
+	{
+		switch (Method(Node))
+		{
+			case ETreeTraverseControl::Break: return false;
+			case ETreeTraverseControl::SkipChildren: return true;
+		}
+	}
+	
+	for (const TUniquePtr<DiffNodeType>& Child : Node->Children)
+	{
+		if (!PreOrderRecursive(Child, Method))
+		{
+			return false;
+		}
+	}
+
+	return true;
+}
+
+template <typename InNodeType>
+bool TAsyncTreeDifferences<InNodeType>::PostOrderRecursive(const TUniquePtr<DiffNodeType>& Node, const TFunction<ETreeTraverseControl(const TUniquePtr<DiffNodeType>&)>& Method)
+{
+	for (const TUniquePtr<DiffNodeType>& Child : Node->Children)
+	{
+		if (!PostOrderRecursive(Child, Method))
+		{
+			return false;
+		}
+	}
+	
+	if (Node->DiffResult != ETreeDiffResult::Invalid)
+	{
+		switch (Method(Node))
+		{
+			case ETreeTraverseControl::Break: return false;
+		}
+	}
+
+	return true;
+}
+
+template <typename InNodeType>
+void TAsyncTreeDifferences<InNodeType>::ProcessTopOfQueue()
+{
+	DiffNodeType* DiffNode = UpdateQueue.Pop();
+	
+	TArray<ValueType> ChildrenA;
+	TArray<ValueType> ChildrenB;
+	if (DiffNode->ValueA != DiffNodeType::NullValue)
+	{
+		TreeDiffSpecification::GetChildren<ValueType>(DiffNode->ValueA, ChildrenA);
+	}
+	if (DiffNode->ValueB != DiffNodeType::NullValue)
+	{
+		TreeDiffSpecification::GetChildren<ValueType>(DiffNode->ValueB, ChildrenB);
+	}
+	
+	if (!ChildrenA.IsEmpty() || !ChildrenB.IsEmpty())
+	{
+		QueueParallelNodeLists(ChildrenA, ChildrenB, DiffNode);
+	}
+}
+
+
+// this is where the magic happens :)
+template <typename InNodeType>
+void TAsyncTreeDifferences<InNodeType>::QueueParallelNodeLists(const TArray<ValueType>& ValuesA, const TArray<ValueType>& ValuesB, DiffNodeType* ParentNode)
+{
+	// calculate the diff nodes that should be children of ParentNode
+	TArray<DiffNodeType> FoundChildren;
+	{
+		// reversing data before calculating LCS puts "add" diff results towards the end which is a friendlier format
+		// (don't worry we'll return the data to it's original order)
+		Algo::Reverse(const_cast<TArray<ValueType>&>(ValuesA));
+		Algo::Reverse(const_cast<TArray<ValueType>&>(ValuesB));
+		
+		// Find the Longest Common Subsequence between node lists
+        const TArray<TArray<int32>> LCS = CalculateLCSTableForMatchingValues(ValuesA, ValuesB);
+        
+        // Using an LCS Table, we can determine which tree nodes match one another in each tree and queue them up to update together
+        int32 IndexA = ValuesA.Num();
+        int32 IndexB = ValuesB.Num();
+        while (IndexA > 0 || IndexB > 0)
+        {
+        	if (IndexA == 0)
+        	{
+        		// ValueB doesn't match. push it on it's own
+        		FoundChildren.Add(DiffNodeType(DiffNodeType::NullValue, ValuesB[--IndexB], ParentNode));
+        	}
+        	else if (IndexB == 0)
+        	{
+        		// ValueA doesn't match. push it on it's own
+        		FoundChildren.Add(DiffNodeType(ValuesA[--IndexA], DiffNodeType::NullValue, ParentNode));
+        	}
+        	else if (AreMatching(ValuesA[IndexA - 1], ValuesB[IndexB - 1]))
+        	{
+        		// found a match between both nodes
+        		FoundChildren.Add(DiffNodeType(ValuesA[--IndexA], ValuesB[--IndexB], ParentNode));
+        	}
+        	else if (LCS[IndexA - 1][IndexB] <= LCS[IndexA][IndexB - 1])
+        	{
+        		// ValueB doesn't match. push it on it's own
+        		FoundChildren.Add(DiffNodeType(DiffNodeType::NullValue, ValuesB[--IndexB], ParentNode));
+        	}
+        	else
+        	{
+        		// ValueA doesn't match. push it on it's own
+        		FoundChildren.Add(DiffNodeType(ValuesA[--IndexA], DiffNodeType::NullValue, ParentNode));
+        	}
+        }
+
+		// undo the temp reverse from above
+		Algo::Reverse(const_cast<TArray<ValueType>&>(ValuesA));
+		Algo::Reverse(const_cast<TArray<ValueType>&>(ValuesB));
+	}
+	
+	// compress consecutive value matched entries
+	TArray<DiffNodeType> AllFoundChildren = MoveTemp(FoundChildren);
+	FoundChildren = TArray<DiffNodeType>();
+
+	int32 CompressIndex = INDEX_NONE;
+	int32 CompressCount = 0;
+	for (DiffNodeType& FoundChild : AllFoundChildren)
+	{
+		if (FoundChild.DiffResult == ETreeDiffResult::MissingFromTree1 &&
+			TreeDiffSpecification::ShouldMatchByValue<ValueType>(FoundChild.ValueB))
+		{
+			if (CompressCount == 0)
+			{
+				CompressIndex = FoundChildren.Num();
+			}
+			++CompressCount;
+			FoundChildren.Add(MoveTemp(FoundChild));
+			continue;
+		}
+		if (CompressCount > 0 &&
+            FoundChild.DiffResult == ETreeDiffResult::MissingFromTree2 &&
+			TreeDiffSpecification::ShouldMatchByValue<ValueType>(FoundChild.ValueA))
+		{
+			FoundChildren[CompressIndex].ValueA = FoundChild.ValueA;
+			FoundChildren[CompressIndex].SetDiffType();
+			++CompressIndex;
+			--CompressCount;
+			continue;
+		}
+		
+		FoundChildren.Add(MoveTemp(FoundChild));
+		CompressIndex = INDEX_NONE;
+		CompressCount = 0;
+	}
+	
+	// compare/replace results with what's *actually* in the parent node to avoid modifying nodes that didn't change
+	{
+		TArray<TUniquePtr<DiffNodeType>> ExpectedChildren = MoveTemp(ParentNode->Children);
+		ParentNode->Children = TArray<TUniquePtr<DiffNodeType>>();
+		const TArray<TArray<int32>> LCS = CalculateLCSTableForDiffNodes(FoundChildren, ExpectedChildren);
+		int32 FoundChildrenIndex = FoundChildren.Num();
+		int32 ExpectedChildrenIndex = ExpectedChildren.Num();
+		while (FoundChildrenIndex > 0)
+		{
+			if (ExpectedChildrenIndex == 0)
+			{
+				// a new child was found. move it into the children
+				ParentNode->Children.Add(MakeUnique<DiffNodeType>(MoveTemp(FoundChildren[--FoundChildrenIndex])));
+			}
+			else if (FoundChildren[FoundChildrenIndex - 1] == *ExpectedChildren[ExpectedChildrenIndex - 1])
+			{
+				// found a match between both nodes. Preserve the old data by copying it into the new child array
+				ParentNode->Children.Add(MoveTemp(ExpectedChildren[--ExpectedChildrenIndex]));
+				// update the diff type to the latest result
+				ParentNode->Children.Last()->DiffResult = FoundChildren[--FoundChildrenIndex].DiffResult;
+				
+			}
+			else if (LCS[FoundChildrenIndex - 1][ExpectedChildrenIndex] <= LCS[FoundChildrenIndex][ExpectedChildrenIndex - 1])
+			{
+				// a child was removed. ignore the old data
+				--ExpectedChildrenIndex;
+			}
+			else
+			{
+				// a new child was found. move it into the children
+				ParentNode->Children.Add(MakeUnique<DiffNodeType>(MoveTemp(FoundChildren[--FoundChildrenIndex])));
+			}
+		}
+	}
+	Algo::Reverse(ParentNode->Children);
+
+	// queue all the children to update
+	for (const TUniquePtr<DiffNodeType>& Child : ParentNode->Children)
+	{
+		UpdateQueue.Add(Child.Get());
+	}
+}
+
+template <typename InNodeType>
+TArray<TArray<int32>> TAsyncTreeDifferences<InNodeType>::CalculateLCSTableForMatchingValues(
+	const TArray<ValueType>& ValuesA, const TArray<ValueType>& ValuesB)
+{
+	TArray<TArray<int32>> LCS;
+	LCS.SetNum(ValuesA.Num() + 1);
+	for (int32 I = 0; I <= ValuesA.Num(); I++)
+	{
+		LCS[I].SetNum(ValuesB.Num() + 1);
+		if (I == 0)
+		{
+			continue;
+		}
+	
+		for (int32 J = 1; J <= ValuesB.Num(); J++)
+		{
+			if (AreMatching(ValuesA[I - 1], ValuesB[J - 1]))
+			{
+				LCS[I][J] = LCS[I - 1][J - 1] + 1;
+			}
+			else
+			{
+				LCS[I][J] = FMath::Max(LCS[I - 1][J], LCS[I][J - 1]);
+			}
+		}
+	}
+	return LCS;
+}
+
+template <typename InNodeType>
+TArray<TArray<int32>> TAsyncTreeDifferences<InNodeType>::CalculateLCSTableForDiffNodes(const TArray<DiffNodeType>& FoundNodes,
+	const TArray<TUniquePtr<DiffNodeType>>& ExpectedNodes)
+{
+	TArray<TArray<int32>> LCS;
+	LCS.SetNum(FoundNodes.Num() + 1);
+	for (int32 I = 0; I <= FoundNodes.Num(); I++)
+	{
+		LCS[I].SetNum(ExpectedNodes.Num() + 1);
+		if (I == 0)
+		{
+			continue;
+		}
+	
+		for (int32 J = 1; J <= ExpectedNodes.Num(); J++)
+		{
+			if (FoundNodes[I - 1] == *ExpectedNodes[J - 1])
+			{
+				LCS[I][J] = LCS[I - 1][J - 1] + 1;
+			}
+			else
+			{
+				LCS[I][J] = FMath::Max(LCS[I - 1][J], LCS[I][J - 1]);
+			}
+		}
+	}
+	return LCS;
+}
+
+template <typename InNodeType>
+bool TAsyncTreeDifferences<InNodeType>::AreMatching(const ValueType& ValueA, const ValueType& ValueB)
+{
+	const bool bMatchValueAByValue = TreeDiffSpecification::ShouldMatchByValue<ValueType>(ValueA);
+	const bool bMatchValueBByValue = TreeDiffSpecification::ShouldMatchByValue<ValueType>(ValueA);
+	if (bMatchValueAByValue && bMatchValueBByValue)
+	{
+		return TreeDiffSpecification::AreValuesEqual<ValueType>(ValueA, ValueB);
+	}
+	if (!bMatchValueAByValue && !bMatchValueBByValue)
+	{
+		return TreeDiffSpecification::AreMatching<ValueType>(ValueA, ValueB);
+	}
+
+	// a node that should match by value will never match a node that shouldn't match by value
+	return false;
+}
+
