Level Instance: Fix undo/redo of Level Instance loading since CL 24360365 introduced trans buffer reset to UnloadPackage
- Keep the reset by default
- Override with value in EditorLevelUtils.

#rnx

diff --git a/Engine/Source/Editor/UnrealEd/Private/EditorLevelUtils.cpp b/Engine/Source/Editor/UnrealEd/Private/EditorLevelUtils.cpp
--- a/Engine/Source/Editor/UnrealEd/Private/EditorLevelUtils.cpp
+++ b/Engine/Source/Editor/UnrealEd/Private/EditorLevelUtils.cpp
@@ -868,9 +868,12 @@ bool UEditorLevelUtils::RemoveLevelsFromWorld(TArray<ULevel*> InLevels, bool bCl
 		GEditor->Trans->Reset(TransResetText);
 	}
 
-	if (!UPackageTools::UnloadPackages(PackagesToUnload))
+	UPackageTools::FUnloadPackageParams UnloadParams(PackagesToUnload);
+	UnloadParams.bResetTransBuffer = bResetTransBuffer;
+	const bool bUnloadResult = UPackageTools::UnloadPackages(UnloadParams);
+	if (!bUnloadResult && !UnloadParams.OutErrorMessage.IsEmpty())
 	{
-		FMessageDialog::Open(EAppMsgType::Ok, LOCTEXT("UnloadPackagesFail", "Unable to unload packages."));
+		FMessageDialog::Open(EAppMsgType::Ok, UnloadParams.OutErrorMessage);
 	}
 
 	// Redraw the main editor viewports.
diff --git a/Engine/Source/Editor/UnrealEd/Private/PackageTools.cpp b/Engine/Source/Editor/UnrealEd/Private/PackageTools.cpp
--- a/Engine/Source/Editor/UnrealEd/Private/PackageTools.cpp
+++ b/Engine/Source/Editor/UnrealEd/Private/PackageTools.cpp
@@ -290,7 +290,7 @@ UPackageTools::UPackageTools(const FObjectInitializer& ObjectInitializer)
 	bool UPackageTools::UnloadPackages( const TArray<UPackage*>& TopLevelPackages )
 	{
 		FText ErrorMessage;
-		bool bResult = UnloadPackages(TopLevelPackages, ErrorMessage);
+		const bool bResult = UnloadPackages(TopLevelPackages, ErrorMessage);
 		if(!ErrorMessage.IsEmpty())
 		{
 			FMessageDialog::Open( EAppMsgType::Ok, ErrorMessage );
@@ -300,9 +300,18 @@ UPackageTools::UPackageTools(const FObjectInitializer& ObjectInitializer)
 	}
 
 	bool UPackageTools::UnloadPackages(const TArray<UPackage*>& TopLevelPackages, FText& OutErrorMessage, bool bUnloadDirtyPackages)
+	{
+		FUnloadPackageParams Params(TopLevelPackages);
+		Params.bUnloadDirtyPackages = bUnloadDirtyPackages;
+		const bool bResult = UnloadPackages(Params);
+		OutErrorMessage = Params.OutErrorMessage;
+		return bResult;
+	}
+
+	bool UPackageTools::UnloadPackages(FUnloadPackageParams& Params)
 	{
 		// Early out if no package is provided
-		if (TopLevelPackages.IsEmpty())
+		if (Params.Packages.IsEmpty())
 		{
 			return true;
 		}
@@ -315,11 +324,11 @@ UPackageTools::UPackageTools(const FObjectInitializer& ObjectInitializer)
 		// Split the set of selected top level packages into packages which are dirty (and thus cannot be unloaded)
 		// and packages that are not dirty (and thus can be unloaded).
 		TArray<UPackage*> DirtyPackages;
-		for (UPackage* TopLevelPackage : TopLevelPackages)
+		for (UPackage* TopLevelPackage : Params.Packages)
 		{
 			if (TopLevelPackage)
 			{
-				if (!bUnloadDirtyPackages && TopLevelPackage->IsDirty())
+				if (!Params.bUnloadDirtyPackages && TopLevelPackage->IsDirty())
 				{
 					DirtyPackages.Add(TopLevelPackage);
 				}
@@ -347,7 +356,7 @@ UPackageTools::UPackageTools(const FObjectInitializer& ObjectInitializer)
 			FFormatNamedArguments Args;
 			Args.Add(TEXT("DirtyPackages"), FText::FromString(DirtyPackagesList));
 
-			OutErrorMessage = FText::Format( NSLOCTEXT("UnrealEd", "UnloadDirtyPackagesList", "The following assets have been modified and cannot be unloaded:{DirtyPackages}\nSaving these assets will allow them to be unloaded."), Args );
+			Params.OutErrorMessage = FText::Format( NSLOCTEXT("UnrealEd", "UnloadDirtyPackagesList", "The following assets have been modified and cannot be unloaded:{DirtyPackages}\nSaving these assets will allow them to be unloaded."), Args );
 		}
 		if (GEditor)
 		{
@@ -393,7 +402,7 @@ UPackageTools::UPackageTools(const FObjectInitializer& ObjectInitializer)
 			GEditor->GetSelectedObjects()->GetElementSelectionSet()->ClearSelection(FTypedElementSelectionOptions());
 
 			// Clear undo history because transaction records can hold onto assets we want to unload
-			if (GEditor->Trans)
+			if (GEditor->Trans && Params.bResetTransBuffer)
 			{
 				GEditor->Trans->Reset(NSLOCTEXT("UnrealEd", "UnloadPackagesResetUndo", "Unload Assets"));
 			}
diff --git a/Engine/Source/Editor/UnrealEd/Public/PackageTools.h b/Engine/Source/Editor/UnrealEd/Public/PackageTools.h
--- a/Engine/Source/Editor/UnrealEd/Public/PackageTools.h
+++ b/Engine/Source/Editor/UnrealEd/Public/PackageTools.h
@@ -93,6 +93,33 @@ public:
 	 */
 	static bool UnloadPackages( const TArray<UPackage*>& PackagesToUnload, FText& OutErrorMessage, bool bUnloadDirtyPackages = false);
 
+	/**
+	 * struct containing list of packages to unload and other params.
+	 */
+	struct FUnloadPackageParams
+	{
+		FUnloadPackageParams(const TArray<UPackage*>& PackagesToUnload) : Packages(PackagesToUnload) {}
+
+		/** Packages to unload */
+		const TArray<UPackage*> Packages;
+		/** Whether to unload packages that are dirty (that need to be saved) */
+		bool bUnloadDirtyPackages = false;
+		/** Whether to reset the transaction buffer to allow unloading packages that might be referenced by it */
+		bool bResetTransBuffer = true;
+		
+		/** Result: An error message specifying any problems with unloading packages */
+		FText OutErrorMessage;
+	};
+
+	/**
+	 * Helper function that attemps to unload the specified top-level packages.
+	 * 
+	 * @param struct containg list of packages that should be unloaded and other behavior params.
+	 * 
+	 * @return true if the set of loaded packages was changed
+	 */
+	static bool UnloadPackages(FUnloadPackageParams& Params);
+
 	/**
 	 * Helper function that attempts to reload the specified top-level packages.
 	 *
