Remove check statement when we encounter a package assigned to a disconnected remote worker during cook.  This is handled with an error severity message about the fact that the package cannot be cooked.

diff --git a/Engine/Source/Editor/UnrealEd/Private/Cooker/CookDirector.cpp b/Engine/Source/Editor/UnrealEd/Private/Cooker/CookDirector.cpp
--- a/Engine/Source/Editor/UnrealEd/Private/Cooker/CookDirector.cpp
+++ b/Engine/Source/Editor/UnrealEd/Private/Cooker/CookDirector.cpp
@@ -382,7 +382,6 @@ void FCookDirector::AssignRequests(TArray<FWorkerId>&& InWorkers, TArray<TRefCou
 		FWorkerId WorkerIdConstraint = Requests[RequestIndex]->GetWorkerAssignmentConstraint();
 		if (WorkerIdConstraint.IsValid())
 		{
-			check(InWorkers.Contains(WorkerIdConstraint));
 			WorkerId = WorkerIdConstraint;
 		}
 
