Fix Cooking with DiffOnly while diffing both cooked and optional data
Added cmdline arg to compare or not optional data

diff --git a/Engine/Source/Editor/UnrealEd/Private/Cooker/DiffPackageWriter.cpp b/Engine/Source/Editor/UnrealEd/Private/Cooker/DiffPackageWriter.cpp
--- a/Engine/Source/Editor/UnrealEd/Private/Cooker/DiffPackageWriter.cpp
+++ b/Engine/Source/Editor/UnrealEd/Private/Cooker/DiffPackageWriter.cpp
@@ -32,6 +32,7 @@ FDiffPackageWriter::FDiffPackageWriter(TUniquePtr<ICookedPackageWriter>&& InInne
 	FParse::Value(FCommandLine::Get(), TEXT("MaxDiffstoLog="), MaxDiffsToLog);
 
 	bSaveForDiff = FParse::Param(FCommandLine::Get(), TEXT("SaveForDiff"));
+	bDiffOptional = FParse::Param(FCommandLine::Get(), TEXT("DiffOptional"));
 
 	GConfig->GetBool(TEXT("CookSettings"), TEXT("IgnoreHeaderDiffs"), bIgnoreHeaderDiffs, GEditorIni);
 	// Command line override for IgnoreHeaderDiffs
@@ -116,7 +117,8 @@ void FDiffPackageWriter::BeginPackage(const FBeginPackageInfo& Info)
 	bIsDifferent = false;
 	bDiffCallstack = false;
 	bHasStartedSecondSave = false;
-	DiffMap.Reset();
+	DiffMap[0].Reset();
+	DiffMap[1].Reset();
 
 	BeginInfo = Info;
 	ConditionallyDumpObjList();
@@ -171,22 +173,24 @@ void FDiffPackageWriter::WritePackageData(const FPackageInfo& Info, FLargeMemory
 	else
 	{
 		Writer.GetCallstacks().Append(*ExportsCallstacks, Info.HeaderSize);
-		ExportsDiffMapOffset = Info.HeaderSize;
+		check(Info.MultiOutputIndex < 2);
+		ExportsDiffMapOffset[Info.MultiOutputIndex] = Info.HeaderSize;
 
-		bIsDifferent = !Writer.GenerateDiffMap(PreviousPackageData, Info.HeaderSize, MaxDiffsToLog, DiffMap);
+		bIsDifferent = !Writer.GenerateDiffMap(PreviousPackageData, Info.HeaderSize, MaxDiffsToLog, DiffMap[Info.MultiOutputIndex]);
 	}
 
 	Inner->WritePackageData(Info, ExportsArchive, FileRegions);
 }
 
-TUniquePtr<FLargeMemoryWriter> FDiffPackageWriter::CreateLinkerArchive(FName PackageName, UObject* Asset)
+TUniquePtr<FLargeMemoryWriter> FDiffPackageWriter::CreateLinkerArchive(FName PackageName, UObject* Asset, uint16 MultiOutputIndex)
 {
 	// The entire package will be serialized to memory and then compared against package on disk.
 	if (bDiffCallstack)
 	{
+		check(MultiOutputIndex < 2);
 		// Each difference will be logged with its Serialize call stack trace
 		return TUniquePtr<FLargeMemoryWriter>(new FArchiveStackTrace(Asset, *PackageName.ToString(),
-			true /* bInCollectCallstacks */, &DiffMap));
+			true /* bInCollectCallstacks */, &DiffMap[MultiOutputIndex]));
 	}
 	else
 	{
@@ -195,7 +199,7 @@ TUniquePtr<FLargeMemoryWriter> FDiffPackageWriter::CreateLinkerArchive(FName Pac
 	}
 }
 
-TUniquePtr<FLargeMemoryWriter> FDiffPackageWriter::CreateLinkerExportsArchive(FName PackageName, UObject* Asset)
+TUniquePtr<FLargeMemoryWriter> FDiffPackageWriter::CreateLinkerExportsArchive(FName PackageName, UObject* Asset, uint16 MultiOutputIndex)
 {
 	// When cooking, exports are serialized into a separate archive. The serialization callstack offsets
 	// and stack traces are collected into a separate callstack collection and appended to the overall
@@ -210,12 +214,13 @@ TUniquePtr<FLargeMemoryWriter> FDiffPackageWriter::CreateLinkerExportsArchive(FN
 	const int64 PreAllocateBytes = 0;
 	const bool bIsPersistent = true;
 
+	check(MultiOutputIndex < 2);
 	if (bDiffCallstack)
 	{
 		return MakeUnique<FArchiveStackTraceMemoryWriter>(
 			*ExportsCallstacks,
-			&DiffMap,
-			ExportsDiffMapOffset,
+			&DiffMap[MultiOutputIndex],
+			ExportsDiffMapOffset[MultiOutputIndex],
 			PreAllocateBytes,
 			bIsPersistent,
 			*PackageName.ToString());
@@ -232,6 +237,21 @@ TUniquePtr<FLargeMemoryWriter> FDiffPackageWriter::CreateLinkerExportsArchive(FN
 	}
 }
 
+void FDiffPackageWriter::UpdateSaveArguments(FSavePackageArgs& SaveArgs)
+{
+	// if we are diffing optional data, add it to the save args, otherwise strip it
+	if (bDiffOptional)
+	{
+		SaveArgs.SaveFlags |= SAVE_Optional;
+	}
+	else
+	{
+		SaveArgs.SaveFlags &= ~SAVE_Optional;
+	}
+	Inner->UpdateSaveArguments(SaveArgs);
+}
+
+
 bool FDiffPackageWriter::IsAnotherSaveNeeded(FSavePackageResultStruct& PreviousResult, FSavePackageArgs& SaveArgs)
 {
 	checkf(!Inner->IsAnotherSaveNeeded(PreviousResult, SaveArgs),
diff --git a/Engine/Source/Editor/UnrealEd/Private/Cooker/DiffPackageWriter.h b/Engine/Source/Editor/UnrealEd/Private/Cooker/DiffPackageWriter.h
--- a/Engine/Source/Editor/UnrealEd/Private/Cooker/DiffPackageWriter.h
+++ b/Engine/Source/Editor/UnrealEd/Private/Cooker/DiffPackageWriter.h
@@ -45,8 +45,8 @@ public:
 	{
 		return Inner->GetExportsFooterSize();
 	}
-	virtual TUniquePtr<FLargeMemoryWriter> CreateLinkerArchive(FName PackageName, UObject* Asset) override;
-	virtual TUniquePtr<FLargeMemoryWriter> CreateLinkerExportsArchive(FName PackageName, UObject* Asset) override;
+	virtual TUniquePtr<FLargeMemoryWriter> CreateLinkerArchive(FName PackageName, UObject* Asset, uint16 MultiOutputIndex) override;
+	virtual TUniquePtr<FLargeMemoryWriter> CreateLinkerExportsArchive(FName PackageName, UObject* Asset, uint16 MultiOutputIndex) override;
 	virtual bool IsPreSaveCompleted() const override
 	{
 		return bDiffCallstack;
@@ -95,10 +95,7 @@ public:
 	{
 		Inner->MarkPackagesUpToDate(UpToDatePackages);
 	}
-	virtual void UpdateSaveArguments(FSavePackageArgs& SaveArgs) override
-	{
-		Inner->UpdateSaveArguments(SaveArgs);
-	}
+	virtual void UpdateSaveArguments(FSavePackageArgs& SaveArgs) override;
 	virtual bool IsAnotherSaveNeeded(FSavePackageResultStruct& PreviousResult, FSavePackageArgs& SaveArgs) override;
 	virtual TFuture<FCbObject> WriteMPCookMessageForPackage(FName PackageName) override
 	{
@@ -121,15 +118,16 @@ private:
 	void ConditionallyDumpObjList();
 	void ConditionallyDumpObjects();
 
-	FArchiveDiffMap DiffMap;
+	FArchiveDiffMap DiffMap[2];
 	TUniquePtr<FArchiveCallstacks> ExportsCallstacks;
 	FBeginPackageInfo BeginInfo;
 	TUniquePtr<ICookedPackageWriter> Inner;
 	FString DumpObjListParams;
 	FString PackageFilter;
-	int64 ExportsDiffMapOffset = 0;
+	int64 ExportsDiffMapOffset[2] = {0, 0};
 	int32 MaxDiffsToLog = 5;
 	bool bSaveForDiff = false;
+	bool bDiffOptional = false;
 	bool bIgnoreHeaderDiffs = false;
 	bool bIsDifferent = false;
 	bool bDiffCallstack = false;
@@ -182,9 +180,9 @@ public:
 	{
 		return Inner->GetExportsFooterSize();
 	}
-	virtual TUniquePtr<FLargeMemoryWriter> CreateLinkerArchive(FName PackageName, UObject* Asset) override
+	virtual TUniquePtr<FLargeMemoryWriter> CreateLinkerArchive(FName PackageName, UObject* Asset, uint16 MultiOutputIndex) override
 	{
-		return Inner->CreateLinkerArchive(PackageName, Asset);
+		return Inner->CreateLinkerArchive(PackageName, Asset, MultiOutputIndex);
 	}
 	virtual bool IsPreSaveCompleted() const override
 	{
diff --git a/Engine/Source/Runtime/Core/Public/Serialization/PackageWriter.h b/Engine/Source/Runtime/Core/Public/Serialization/PackageWriter.h
--- a/Engine/Source/Runtime/Core/Public/Serialization/PackageWriter.h
+++ b/Engine/Source/Runtime/Core/Public/Serialization/PackageWriter.h
@@ -173,10 +173,10 @@ public:
 	virtual void WritePackageTrailer(const FPackageTrailerInfo& Info, const FIoBuffer& Data) = 0;
 
 	/** Create the FLargeMemoryWriter to which the Header and Exports are written during the save. */
-	COREUOBJECT_API virtual TUniquePtr<FLargeMemoryWriter> CreateLinkerArchive(FName PackageName, UObject* Asset);
+	COREUOBJECT_API virtual TUniquePtr<FLargeMemoryWriter> CreateLinkerArchive(FName PackageName, UObject* Asset, uint16 MultiOutputIndex);
 
 	/** Returns an archive to be used when serializing exports. */
-	COREUOBJECT_API virtual TUniquePtr<FLargeMemoryWriter> CreateLinkerExportsArchive(FName PackageName, UObject* Asset);
+	COREUOBJECT_API virtual TUniquePtr<FLargeMemoryWriter> CreateLinkerExportsArchive(FName PackageName, UObject* Asset, uint16 MultiOutputIndex);
 
 	/** Report whether PreSave was already called by the PackageWriter before the current UPackage::Save call. */
 	virtual bool IsPreSaveCompleted() const
diff --git a/Engine/Source/Runtime/CoreUObject/Private/Serialization/PackageWriterToSharedBuffer.cpp b/Engine/Source/Runtime/CoreUObject/Private/Serialization/PackageWriterToSharedBuffer.cpp
--- a/Engine/Source/Runtime/CoreUObject/Private/Serialization/PackageWriterToSharedBuffer.cpp
+++ b/Engine/Source/Runtime/CoreUObject/Private/Serialization/PackageWriterToSharedBuffer.cpp
@@ -7,14 +7,14 @@
 
 DEFINE_LOG_CATEGORY_STATIC(LogPackageWriter, Display, All);
 
-TUniquePtr<FLargeMemoryWriter> IPackageWriter::CreateLinkerArchive(FName PackageName, UObject* Asset)
+TUniquePtr<FLargeMemoryWriter> IPackageWriter::CreateLinkerArchive(FName PackageName, UObject* Asset, uint16 /*MultiOutputIndex*/)
 {
 	// The LargeMemoryWriter does not need to be persistent; the LinkerSave wraps it and reports Persistent=true
 	bool bIsPersistent = false; 
 	return TUniquePtr<FLargeMemoryWriter>(new FLargeMemoryWriter(0, bIsPersistent, *PackageName.ToString()));
 }
 
-TUniquePtr<FLargeMemoryWriter> IPackageWriter::CreateLinkerExportsArchive(FName PackageName, UObject* Asset)
+TUniquePtr<FLargeMemoryWriter> IPackageWriter::CreateLinkerExportsArchive(FName PackageName, UObject* Asset, uint16 /*MultiOutputIndex*/)
 {
 	const bool bPersistent = true;
 	return MakeUnique<FLargeMemoryWriter>(0, bPersistent, *PackageName.ToString());
diff --git a/Engine/Source/Runtime/CoreUObject/Private/UObject/SavePackage/SaveContext.h b/Engine/Source/Runtime/CoreUObject/Private/UObject/SavePackage/SaveContext.h
--- a/Engine/Source/Runtime/CoreUObject/Private/UObject/SavePackage/SaveContext.h
+++ b/Engine/Source/Runtime/CoreUObject/Private/UObject/SavePackage/SaveContext.h
@@ -429,6 +429,9 @@ public:
 		check(!IsCooking() || WITH_EDITOR);
 		checkf(!IsCooking() || PackageWriter, TEXT("Cook saves require an IPackageWriter"));
 
+		// Store initial state
+		PreviousPackageFlags = Package->GetPackageFlags();
+
 		SaveArgs.TopLevelFlags = UE::SavePackageUtilities::NormalizeTopLevelFlags(SaveArgs.TopLevelFlags, IsCooking());
 		if (PackageWriter)
 		{
@@ -472,6 +475,8 @@ public:
 		{
 			UE::SavePackageUtilities::CallPostSaveRoot(Asset, ObjectSaveContext, bNeedPreSaveCleanup);
 		}
+		// restore initial state
+		Package->SetPackageFlagsTo(PreviousPackageFlags);
 	}
 
 	const FSavePackageArgs& GetSaveArgs() const
@@ -1086,6 +1091,9 @@ private:
 	bool bIgnoreHeaderDiffs = false;
 	bool bIsSaveAutoOptional = false;
 
+	// Mutated package state
+	uint32 PreviousPackageFlags;
+
 	// Config classes shared with the old Save
 	FCanSkipEditorReferencedPackagesWhenCooking SkipEditorRefCookingSetting;
 
diff --git a/Engine/Source/Runtime/CoreUObject/Private/UObject/SavePackage2.cpp b/Engine/Source/Runtime/CoreUObject/Private/UObject/SavePackage2.cpp
--- a/Engine/Source/Runtime/CoreUObject/Private/UObject/SavePackage2.cpp
+++ b/Engine/Source/Runtime/CoreUObject/Private/UObject/SavePackage2.cpp
@@ -935,9 +935,11 @@ ESavePackageResult CreateLinker(FSaveContext& SaveContext)
 		IPackageWriter* PackageWriter = SaveContext.GetPackageWriter();
 		if (PackageWriter || SaveContext.IsSaveToMemory())
 		{
+			const bool bIsOptionalRealm = SaveContext.GetCurrentHarvestingRealm() == ESaveRealm::Optional;
+
 			// Allocate the linker with a memory writer, forcing byte swapping if wanted.
 			TUniquePtr<FLargeMemoryWriter> ExportsArchive = PackageWriter ?
-				PackageWriter->CreateLinkerArchive(SaveContext.GetPackage()->GetFName(), SaveContext.GetAsset()) :
+				PackageWriter->CreateLinkerArchive(SaveContext.GetPackage()->GetFName(), SaveContext.GetAsset(), bIsOptionalRealm ? 1 : 0) :
 				// The LargeMemoryWriter does not need to be persistent; the LinkerSave wraps it and reports Persistent=true
 				TUniquePtr<FLargeMemoryWriter>(new FLargeMemoryWriter(0, false /* bIsPersistent */, *SaveContext.GetPackage()->GetName()));
 			SaveContext.SetLinker(MakePimpl<FLinkerSave>(SaveContext.GetPackage(), ExportsArchive.Release(),
@@ -2735,8 +2737,9 @@ ESavePackageResult SaveHarvestedRealms(FSaveContext& SaveContext, ESaveRealm Har
 		if (Linker.IsCooking())
 		{
 			// Write the exports into a separate archive
+			const bool bIsOptionalRealm = SaveContext.GetCurrentHarvestingRealm() == ESaveRealm::Optional;
 			TUniquePtr<FLargeMemoryWriter> ExportsArchive = SaveContext.GetPackageWriter()->CreateLinkerExportsArchive(
-				SaveContext.GetPackage()->GetFName(), SaveContext.GetAsset());
+				SaveContext.GetPackage()->GetFName(), SaveContext.GetAsset(), bIsOptionalRealm ? 1 : 0);
 			ExportsArchive->SetSerializeContext(SaveContext.GetSerializeContext());
 			SaveContext.Result = WriteCookedExports(*ExportsArchive, SaveContext);
 
