Fix for rare crash double clicking a landscape spline point tangent box.
Copying similar fix for crash in single click handling (https://p4-swarm.epicgames.net/changes/15366417)

diff --git a/Engine/Source/Editor/UnrealEd/Private/EditorViewportClient.cpp b/Engine/Source/Editor/UnrealEd/Private/EditorViewportClient.cpp
--- a/Engine/Source/Editor/UnrealEd/Private/EditorViewportClient.cpp
+++ b/Engine/Source/Editor/UnrealEd/Private/EditorViewportClient.cpp
@@ -3333,7 +3333,7 @@ void FEditorViewportClient::ProcessDoubleClickInViewport( const struct FInputEve
 	MouseDeltaTracker->StartTracking( this, HitX, HitY, InputState );
 	bIsTracking = true;
 	GEditor->MouseMovement = FVector::ZeroVector;
-	HHitProxy*	HitProxy = InputStateViewport->GetHitProxy(HitX,HitY);
+	TRefCountPtr<HHitProxy> HitProxy = InputStateViewport->GetHitProxy(HitX,HitY);
 	ProcessClick(View,HitProxy,Key,Event,HitX,HitY);
 	MouseDeltaTracker->EndTracking( this );
 	bIsTracking = false;
