MPCook: Fix crash on CookDirector when a CookWorker crashes that was handling the cook of a WorldPartition package: CookDirector::AssignRequests needs to robustly handle a WorkerIDConstraint that is not in the set of remaining workers.
#rnx

diff --git a/Engine/Source/Editor/UnrealEd/Private/Cooker/CookDirector.cpp b/Engine/Source/Editor/UnrealEd/Private/Cooker/CookDirector.cpp
--- a/Engine/Source/Editor/UnrealEd/Private/Cooker/CookDirector.cpp
+++ b/Engine/Source/Editor/UnrealEd/Private/Cooker/CookDirector.cpp
@@ -388,8 +388,7 @@ void FCookDirector::AssignRequests(TArray<FWorkerId>&& InWorkers, TArray<TRefCou
 		if (!WorkerId.IsLocal())
 		{
 			uint8 RemoteIndex = WorkerId.GetRemoteIndex();
-			check(RemoteIndex < RemoteBatches.Num());
-			if (!RemoteIndexIsValid[RemoteIndex])
+			if (RemoteIndex >= RemoteBatches.Num() || !RemoteIndexIsValid[RemoteIndex])
 			{
 				UE_LOG(LogCook, Error, TEXT("Package %s can only be cooked by a now-disconnected CookWorker. The package can not be cooked."),
 					*Requests[RequestIndex]->GetPackageName().ToString());
