Fixed undo/redo not allowing unregistering of undo clients during Broadcast. Since undo redo objects can be destroyed, the temporary copy would maintain the destroyed objects, pointing at garbage, causing crashes

diff --git a/Engine/Source/Editor/UnrealEd/Classes/Editor/EditorEngine.h b/Engine/Source/Editor/UnrealEd/Classes/Editor/EditorEngine.h
--- a/Engine/Source/Editor/UnrealEd/Classes/Editor/EditorEngine.h
+++ b/Engine/Source/Editor/UnrealEd/Classes/Editor/EditorEngine.h
@@ -2867,6 +2867,10 @@ private:
 	/** List of editors who want to receive undo/redo events */
 	TSet< class FEditorUndoClient* > UndoClients;
 
+	/** A temporary copy used to add resilience during the undo/redo broadcast. Used to allow unregistering of clients while the broadcast is happening.
+	 * DO NOT add to this set while a broadcast is happening. */
+	TSet<class FEditorUndoClient*> InflightUndoClients;
+	
 	/** List of actors that were selected before Undo/redo */
 	TArray<AActor*> OldSelectedActors;
 
diff --git a/Engine/Source/Editor/UnrealEd/Private/EditorServer.cpp b/Engine/Source/Editor/UnrealEd/Private/EditorServer.cpp
--- a/Engine/Source/Editor/UnrealEd/Private/EditorServer.cpp
+++ b/Engine/Source/Editor/UnrealEd/Private/EditorServer.cpp
@@ -5391,11 +5391,13 @@ void UEditorEngine::BroadcastPostUndoRedo(const FTransactionContext& UndoContext
 	// This sanitization code can be removed once blueprint ::Conform(ImplementedEvents/ImplementedInterfaces) 
 	// functions have been fixed. For the time being it improves editor stability, though:
 	UEdGraphPin::SanitizePinsPostUndoRedo();
-
+	
+	check(InflightUndoClients.IsEmpty());
+	
 	// Note that we use a copy here as clients can register/unregister with the undo system while in PostUndo()/PostRedo()
 	// which modifies UndoClients during the loop. This can cause an infinite loop where the iterator never finishes.
-	const TSet<FEditorUndoClient*> UndoClientsCopy = UndoClients;
-	for (auto UndoIt = UndoClientsCopy.CreateConstIterator(); UndoIt; ++UndoIt)
+	InflightUndoClients = UndoClients;
+	for (auto UndoIt = InflightUndoClients.CreateConstIterator(); UndoIt; ++UndoIt)
 	{
 		FEditorUndoClient* Client = *UndoIt;
 		if (Client && Client->MatchesContext(UndoContext, CurrentUndoRedoContext->TransactionObjects))
@@ -5414,6 +5416,7 @@ void UEditorEngine::BroadcastPostUndoRedo(const FTransactionContext& UndoContext
 	// Invalidate all viewports
 	InvalidateAllViewportsAndHitProxies();
 
+	InflightUndoClients.Empty();
 	FEditorDelegates::PostUndoRedo.Broadcast();
 }
 
@@ -6892,6 +6895,7 @@ void UEditorEngine::UnregisterForUndo( FEditorUndoClient* Client)
 	if (Client)
 	{
 		UndoClients.Remove(Client);
+		InflightUndoClients.Remove(Client);
 	}
 }
 
