MPCook: Remove erroneous reset line that stripped all edges from the traversal of the dependency graph.
#rnx

diff --git a/Engine/Source/Editor/UnrealEd/Private/Cooker/CookRequestCluster.cpp b/Engine/Source/Editor/UnrealEd/Private/Cooker/CookRequestCluster.cpp
--- a/Engine/Source/Editor/UnrealEd/Private/Cooker/CookRequestCluster.cpp
+++ b/Engine/Source/Editor/UnrealEd/Private/Cooker/CookRequestCluster.cpp
@@ -1019,7 +1019,6 @@ void FRequestCluster::FGraphSearch::ExploreVertexEdges(FVertexData& Vertex)
 
 			Cluster.AssetRegistry.GetDependencies(PackageName, CookerLoadingDependencies, EDependencyCategory::Package,
 				EDependencyQuery::Hard);
-			CookerLoadingDependencies.Reset();
 
 			// ITERATIVECOOK_TODO: Build dependencies need to be stored and used to mark package loads as expected
 			// But we can't use them to explore packages that will be loaded during cook because they might not be;
