World Partition - Minimap: Use TMGS_SimpleAverage for the mips settings of the minimap texture

diff --git a/Engine/Source/Editor/UnrealEd/Private/WorldPartition/WorldPartitionMiniMapBuilder.cpp b/Engine/Source/Editor/UnrealEd/Private/WorldPartition/WorldPartitionMiniMapBuilder.cpp
--- a/Engine/Source/Editor/UnrealEd/Private/WorldPartition/WorldPartitionMiniMapBuilder.cpp
+++ b/Engine/Source/Editor/UnrealEd/Private/WorldPartition/WorldPartitionMiniMapBuilder.cpp
@@ -65,6 +65,7 @@ bool UWorldPartitionMiniMapBuilder::PreRun(UWorld* World, FPackageSourceControlH
 		TStrongObjectPtr<UTextureFactory> Factory(NewObject<UTextureFactory>());
 		WorldMiniMap->MiniMapTexture = Factory->CreateTexture2D(WorldMiniMap, TEXT("MinimapTexture"), RF_TextExportTransient);
 		WorldMiniMap->MiniMapTexture->Source.Init(MinimapImageSizeX, MinimapImageSizeY, 1, 1, TSF_BGRA8);
+		WorldMiniMap->MiniMapTexture->MipGenSettings = TMGS_SimpleAverage;
 		WorldMiniMap->MiniMapWorldBounds = IterativeWorldBounds;
 	}
 
