Fix anim BP compilation crash in automated tests

Sparse class data was not getting recreated as when event graphs are restricted, FBlueprintEditorUtils::DoesSupportEventGraphs was now preventing their compilation, skipping anim graph compilation entirely.

diff --git a/Engine/Source/Editor/UnrealEd/Private/Kismet2/BlueprintEditorUtils.cpp b/Engine/Source/Editor/UnrealEd/Private/Kismet2/BlueprintEditorUtils.cpp
--- a/Engine/Source/Editor/UnrealEd/Private/Kismet2/BlueprintEditorUtils.cpp
+++ b/Engine/Source/Editor/UnrealEd/Private/Kismet2/BlueprintEditorUtils.cpp
@@ -3622,9 +3622,8 @@ bool FBlueprintEditorUtils::DoesSupportTimelines(const UBlueprint* Blueprint)
 
 bool FBlueprintEditorUtils::DoesSupportEventGraphs(const UBlueprint* Blueprint)
 {
-	return Blueprint->SupportsEventGraphs()
-		&& (Blueprint->BlueprintType == BPTYPE_Normal
-			|| Blueprint->BlueprintType == BPTYPE_LevelScript);
+	return Blueprint->BlueprintType == BPTYPE_Normal
+			|| Blueprint->BlueprintType == BPTYPE_LevelScript;
 }
 
 /** Returns whether or not the blueprint supports implementing interfaces */
