Tentative fix to Crash in Thumbnail for Fonts asset. Where the Resource can be null.

diff --git a/Engine/Source/Editor/UnrealEd/Private/ThumbnailRendering/FontThumbnailRenderer.cpp b/Engine/Source/Editor/UnrealEd/Private/ThumbnailRendering/FontThumbnailRenderer.cpp
--- a/Engine/Source/Editor/UnrealEd/Private/ThumbnailRendering/FontThumbnailRenderer.cpp
+++ b/Engine/Source/Editor/UnrealEd/Private/ThumbnailRendering/FontThumbnailRenderer.cpp
@@ -59,7 +59,7 @@ void UFontThumbnailRenderer::Draw(UObject* Object, int32 X, int32 Y, uint32 Widt
 		switch(Font->FontCacheType)
 		{
 		case EFontCacheType::Offline:
-			if(Font->Textures.Num() > 0 && Font->Textures[0])
+			if (Font->Textures.Num() > 0 && Font->Textures[0] && Font->Textures[0]->GetResource())
 			{
 				FCanvasTileItem TileItem(FVector2D(X, Y), Font->Textures[0]->GetResource(), FLinearColor::White);
 				TileItem.BlendMode = (Font->ImportOptions.bUseDistanceFieldAlpha) ? SE_BLEND_TranslucentDistanceField : SE_BLEND_Translucent;
@@ -103,4 +103,4 @@ void UFontThumbnailRenderer::Draw(UObject* Object, int32 X, int32 Y, uint32 Widt
 			break;
 		};
 	}
-}
+}
\ No newline at end of file
