Sanitize non-permitted node types and/or event/function graphs prior to saving Blueprint assets when operating within a restricted editor environment.

Notes:
- The code path is bypassed in an unrestricted environment, so there is no additional overhead cost.

diff --git a/Engine/Source/Editor/BlueprintGraph/Private/BlueprintEditorSettings.cpp b/Engine/Source/Editor/BlueprintGraph/Private/BlueprintEditorSettings.cpp
--- a/Engine/Source/Editor/BlueprintGraph/Private/BlueprintEditorSettings.cpp
+++ b/Engine/Source/Editor/BlueprintGraph/Private/BlueprintEditorSettings.cpp
@@ -14,6 +14,7 @@
 #include "EdGraphSchema_K2.h"
 #include "Editor/EditorPerProjectUserSettings.h"
 #include "Engine/Blueprint.h"
+#include "Kismet2/BlueprintEditorUtils.h"
 #include "Misc/AssertionMacros.h"
 #include "Misc/ConfigCacheIni.h"
 #include "Modules/ModuleManager.h"
@@ -284,11 +285,12 @@ bool UBlueprintEditorSettings::IsFunctionAllowed(const UBlueprint* InBlueprint,
 		return true;
 	}
 
-	if (InBlueprint && InBlueprint->ParentClass)
+	const UClass* NativeParentClass = FBlueprintEditorUtils::FindFirstNativeClass(InBlueprint->ParentClass);
+	if (InBlueprint && NativeParentClass)
 	{
-		if (const UFunction* ParentFunction = InBlueprint->ParentClass->FindFunctionByName(FunctionName))
+		if (const UFunction* NativeParentFunction = NativeParentClass->FindFunctionByName(FunctionName))
 		{
-			if (FunctionPermissions.PassesFilter(ParentFunction->GetPathName()))
+			if (FunctionPermissions.PassesFilter(NativeParentFunction->GetPathName()))
 			{
 				return true;
 			}
diff --git a/Engine/Source/Editor/UnrealEd/Private/Kismet2/BlueprintEditorUtils.cpp b/Engine/Source/Editor/UnrealEd/Private/Kismet2/BlueprintEditorUtils.cpp
--- a/Engine/Source/Editor/UnrealEd/Private/Kismet2/BlueprintEditorUtils.cpp
+++ b/Engine/Source/Editor/UnrealEd/Private/Kismet2/BlueprintEditorUtils.cpp
@@ -4,6 +4,7 @@
 #include "Algo/AnyOf.h"
 #include "Algo/Copy.h"
 #include "Algo/RemoveIf.h"
+#include "Algo/Sort.h"
 #include "Algo/Transform.h"
 #include "BlueprintCompilationManager.h"
 #include "UObject/Interface.h"
@@ -10640,6 +10641,112 @@ bool FBlueprintEditorUtils::HasFunctionBlueprintThreadSafeMetaData(const UFuncti
 	return false;
 }
 
+bool FBlueprintEditorUtils::HasRestrictedNodes(const UBlueprint* BP, TArray<UEdGraphNode*>* OutRestrictedNodes /* = nullptr */)
+{
+	if (!BP)
+	{
+		return false;
+	}
+
+	UBlueprintEditorSettings* BlueprintEditorSettings = GetMutableDefault<UBlueprintEditorSettings>();
+	check(BlueprintEditorSettings);
+
+	// No need to continue if we're not restricting anything.
+	if (!BlueprintEditorSettings->HasClassFiltering()
+		&& !BlueprintEditorSettings->GetFunctionPermissions().HasFiltering())
+	{
+		return false;
+	}
+
+	if (OutRestrictedNodes)
+	{
+		OutRestrictedNodes->Reset();
+	}
+
+	// Check for any disallowed node types.
+	auto IsRestrictedNodeType = [BlueprintEditorSettings](const UEdGraphNode* InNode) -> bool
+	{
+		return BlueprintEditorSettings->IsClassAllowed(InNode->GetClass());
+	};
+
+	// Check for any disallowed event graph implementations. The event node is considered restricted in this case.
+	auto IsRestrictedEventGraph = [BP, BlueprintEditorSettings](const UEdGraphNode* InNode) -> bool
+	{
+		const UK2Node_Event* EventNode = Cast<UK2Node_Event>(InNode);
+		return EventNode && !BlueprintEditorSettings->IsFunctionAllowed(BP, EventNode->EventReference.GetMemberName());
+	};
+
+	// Check for any disallowed function override graphs. Its entry node is considered to be restricted in this case.
+	auto IsRestrictedFunctionGraph = [BP, BlueprintEditorSettings](const UEdGraphNode* InNode) -> bool
+	{
+		const UK2Node_FunctionEntry* EntryNode = Cast<UK2Node_FunctionEntry>(InNode);
+		return EntryNode && !BlueprintEditorSettings->IsFunctionAllowed(BP, EntryNode->FunctionReference.GetMemberName());
+	};
+
+	TArray<UEdGraphNode*> AllNodes;
+	GetAllNodesOfClass(BP, AllNodes);
+	for (UEdGraphNode* Node : AllNodes)
+	{
+		if (!Node)
+		{
+			continue;
+		}
+
+		if (IsRestrictedNodeType(Node)
+			|| IsRestrictedEventGraph(Node)
+			|| IsRestrictedFunctionGraph(Node))
+		{
+			if (OutRestrictedNodes)
+			{
+				OutRestrictedNodes->Add(Node);
+			}
+			else
+			{
+				// We found a restricted node and/or graph entry point, so we can return immediately since we're not populating the output list.
+				return true;
+			}
+		}
+	}
+
+	if (OutRestrictedNodes && OutRestrictedNodes->Num() > 0)
+	{
+		return true;
+	}
+
+	return false;
+}
+
+void FBlueprintEditorUtils::SanitizeRestrictedContentOnSave(UBlueprint* BP)
+{
+	TArray<UEdGraphNode*> RestrictedNodes;
+	if (HasRestrictedNodes(BP, &RestrictedNodes))
+	{
+		// Front-load entry nodes so that their graphs can be removed first.
+		Algo::Sort(RestrictedNodes, [](const UEdGraphNode* A, const UEdGraphNode* B)
+		{
+			return A->IsA<UK2Node_Event>() || A->IsA<UK2Node_FunctionEntry>();
+		});
+
+		for (UEdGraphNode* Node : RestrictedNodes)
+		{
+			// Remove the entire graph if this node represents an entry point for an event/function.
+			if (Node->IsA<UK2Node_Event>()
+				|| Node->IsA<UK2Node_FunctionEntry>())
+			{
+				RemoveGraph(BP, Node->GetGraph(), EGraphRemoveFlags::MarkTransient);
+			}
+			else
+			{
+				// If the top-level graph is transient/removed, there's no need to remove the node here.
+				if (!GetTopLevelGraph(Node->GetGraph())->HasAnyFlags(RF_Transient))
+				{
+					RemoveNode(BP, Node, /*bDontRecompile =*/ true);
+				}
+			}
+		}
+	}
+}
+
 bool FBlueprintEditorUtils::ShouldOpenWithDataOnlyEditor(const UBlueprint* Blueprint)
 {
 	return FBlueprintEditorUtils::IsDataOnlyBlueprint(Blueprint)
diff --git a/Engine/Source/Editor/UnrealEd/Public/Kismet2/BlueprintEditorUtils.h b/Engine/Source/Editor/UnrealEd/Public/Kismet2/BlueprintEditorUtils.h
--- a/Engine/Source/Editor/UnrealEd/Public/Kismet2/BlueprintEditorUtils.h
+++ b/Engine/Source/Editor/UnrealEd/Public/Kismet2/BlueprintEditorUtils.h
@@ -1925,6 +1925,19 @@ public:
 	 */
 	static bool HasFunctionBlueprintThreadSafeMetaData(const UFunction* InFunction);
 
+	/**
+	 * Returns true if the given Blueprint is found to contain one or more nodes that are disallowed (restricted) within
+	 * the current editor mode context.
+	 * 
+	 * @param OutRestrictedNodes	(Optional) If non-NULL, will be populated with references to any nodes that are restricted.
+	 */
+	static bool HasRestrictedNodes(const UBlueprint* BP, TArray<UEdGraphNode*>* OutRestrictedNodes = nullptr);
+
+	/**
+	 * Checks the given Blueprint for any restricted content within the current editor context and sanitizes it prior to saving.
+	 */
+	static void SanitizeRestrictedContentOnSave(UBlueprint* BP);
+
 public:
 	static bool ShouldOpenWithDataOnlyEditor(const UBlueprint* Blueprint);
 };
diff --git a/Engine/Source/Runtime/Engine/Private/Blueprint.cpp b/Engine/Source/Runtime/Engine/Private/Blueprint.cpp
--- a/Engine/Source/Runtime/Engine/Private/Blueprint.cpp
+++ b/Engine/Source/Runtime/Engine/Private/Blueprint.cpp
@@ -410,6 +410,11 @@ void UBlueprint::PreSave(FObjectPreSaveContext ObjectSaveContext)
 		// Cache the BP for use (immediate, since we're about to save)
 		FFindInBlueprintSearchManager::Get().AddOrUpdateBlueprintSearchMetadata(this, Flags, OverrideVersion);
 	}
+
+#if WITH_EDITOR
+	// Sanitize the Blueprint by clearing out any objects/types that are restricted within the current editor context.
+	FBlueprintEditorUtils::SanitizeRestrictedContentOnSave(this);
+#endif	// WITH_EDITOR
 }
 
 void UBlueprint::GetPreloadDependencies(TArray<UObject*>& OutDeps)
