Fix fbx scene re-import

diff --git a/Engine/Source/Editor/UnrealEd/Classes/Factories/FbxSceneImportData.h b/Engine/Source/Editor/UnrealEd/Classes/Factories/FbxSceneImportData.h
--- a/Engine/Source/Editor/UnrealEd/Classes/Factories/FbxSceneImportData.h
+++ b/Engine/Source/Editor/UnrealEd/Classes/Factories/FbxSceneImportData.h
@@ -60,6 +60,8 @@ public:
 	/** Attempt to parse an asset import structure from the specified json string. */
 	void FromJson(FString InJsonString);
 
+	virtual void GetAssetRegistryTags(TArray<FAssetRegistryTag>& OutTags) const override;
+
 protected:
 	/** Overridden serialize function to write out the underlying data as json */
 	virtual void Serialize(FArchive& Ar) override;
diff --git a/Engine/Source/Editor/UnrealEd/Private/Fbx/FbxSceneImportData.cpp b/Engine/Source/Editor/UnrealEd/Private/Fbx/FbxSceneImportData.cpp
--- a/Engine/Source/Editor/UnrealEd/Private/Fbx/FbxSceneImportData.cpp
+++ b/Engine/Source/Editor/UnrealEd/Private/Fbx/FbxSceneImportData.cpp
@@ -1,11 +1,12 @@
 // Copyright Epic Games, Inc. All Rights Reserved.
 
 #include "Factories/FbxSceneImportData.h"
-#include "Serialization/JsonReader.h"
-#include "Serialization/JsonSerializer.h"
+
+#include "EditorFramework/AssetImportData.h"
 #include "Factories/FbxSceneImportOptions.h"
 #include "FbxImporter.h"
-
+#include "Serialization/JsonReader.h"
+#include "Serialization/JsonSerializer.h"
 
 UFbxSceneImportData::UFbxSceneImportData(const FObjectInitializer& ObjectInitializer)
 	: Super(ObjectInitializer)
@@ -563,6 +564,15 @@ void UFbxSceneImportData::FromJson(FString InJsonString)
 	}
 }
 
+void UFbxSceneImportData::GetAssetRegistryTags(TArray<FAssetRegistryTag>& OutTags) const
+{
+	FAssetImportInfo AssetImportInfo;
+	FAssetImportInfo::FSourceFile SourceFile(SourceFbxFile);
+	AssetImportInfo.Insert(SourceFile);
+	OutTags.Add(FAssetRegistryTag(SourceFileTagName(), AssetImportInfo.ToJson(), FAssetRegistryTag::TT_Hidden));
+	Super::GetAssetRegistryTags(OutTags);
+}
+
 void UFbxSceneImportData::Serialize(FArchive& Ar)
 {
 	if (!Ar.IsFilterEditorOnly())
