fix Interchange crash on reimport if source file is no longer valid

#rnx

diff --git a/Engine/Source/Editor/UnrealEd/Private/Editor.cpp b/Engine/Source/Editor/UnrealEd/Private/Editor.cpp
--- a/Engine/Source/Editor/UnrealEd/Private/Editor.cpp
+++ b/Engine/Source/Editor/UnrealEd/Private/Editor.cpp
@@ -417,31 +417,34 @@ UE::Interchange::FAssetImportResultRef FReimportManager::ReimportAsync(UObject*
 				{
 					// Make sure SourceFilenames reflects the source filenames in Obj
 					SourceFilenames.Empty();
-					CanReimportHandler->CanReimport(Obj, SourceFilenames);
-
-					int32 RealSourceFileIndex = SourceFileIndex == INDEX_NONE ? 0 : SourceFileIndex;
-					int32 RealValidSourceFileIndex = SourceFilenames.IsValidIndex(RealSourceFileIndex) ? RealSourceFileIndex : 0;
-					UE::Interchange::FScopedSourceData ScopedSourceData(SourceFilenames[RealValidSourceFileIndex]);
-					CanReimportHandler->SetReimportSourceIndex(Obj, SourceFileIndex);
-					if (InterchangeManager.CanTranslateSourceData(ScopedSourceData.GetSourceData()))
+					if ( CanReimportHandler->CanReimport(Obj, SourceFilenames) )
 					{
-						FImportAssetParameters ImportAssetParameters;
-						ImportAssetParameters.bIsAutomated = GIsAutomationTesting || FApp::IsUnattended() || IsRunningCommandlet() || GIsRunningUnattendedScript;
-						ImportAssetParameters.ReimportAsset = Obj;
-						ImportAssetParameters.ReimportSourceIndex = SourceFileIndex;
-						UE::Interchange::FAssetImportResultRef ImportResult = InterchangeManager.ImportAssetAsync(FString(), ScopedSourceData.GetSourceData(), ImportAssetParameters);
-
-						TFunction<void(UE::Interchange::FImportResult&)> AppendAndBroadcastImportResultIfNeeded =
-							[](UE::Interchange::FImportResult& Result)
+						check( SourceFilenames.Num() > 0 );
+
+						int32 RealSourceFileIndex = SourceFileIndex == INDEX_NONE ? 0 : SourceFileIndex;
+						int32 RealValidSourceFileIndex = SourceFilenames.IsValidIndex(RealSourceFileIndex) ? RealSourceFileIndex : 0;
+						UE::Interchange::FScopedSourceData ScopedSourceData(SourceFilenames[RealValidSourceFileIndex]);
+						CanReimportHandler->SetReimportSourceIndex(Obj, SourceFileIndex);
+						if (InterchangeManager.CanTranslateSourceData(ScopedSourceData.GetSourceData()))
 						{
-							UInterchangeManager& InterchangeManager = UInterchangeManager::GetInterchangeManager();
-							TStrongObjectPtr<UInterchangeResultsContainer> ResultsContainer(Result.GetResults());
-							InterchangeManager.OnBatchImportComplete.Broadcast(ResultsContainer);
-						};
+							FImportAssetParameters ImportAssetParameters;
+							ImportAssetParameters.bIsAutomated = GIsAutomationTesting || FApp::IsUnattended() || IsRunningCommandlet() || GIsRunningUnattendedScript;
+							ImportAssetParameters.ReimportAsset = Obj;
+							ImportAssetParameters.ReimportSourceIndex = SourceFileIndex;
+							UE::Interchange::FAssetImportResultRef ImportResult = InterchangeManager.ImportAssetAsync(FString(), ScopedSourceData.GetSourceData(), ImportAssetParameters);
+
+							TFunction<void(UE::Interchange::FImportResult&)> AppendAndBroadcastImportResultIfNeeded =
+								[](UE::Interchange::FImportResult& Result)
+							{
+								UInterchangeManager& InterchangeManager = UInterchangeManager::GetInterchangeManager();
+								TStrongObjectPtr<UInterchangeResultsContainer> ResultsContainer(Result.GetResults());
+								InterchangeManager.OnBatchImportComplete.Broadcast(ResultsContainer);
+							};
 
-						ImportResult->OnDone(AppendAndBroadcastImportResultIfNeeded);
+							ImportResult->OnDone(AppendAndBroadcastImportResultIfNeeded);
 
-						return ImportResult;
+							return ImportResult;
+						}
 					}
 				}
 
