Optimization to speed up rebuilding level actor list in editor
#RNX

diff --git a/Engine/Source/Editor/UnrealEd/Private/EditorActorFolders.cpp b/Engine/Source/Editor/UnrealEd/Private/EditorActorFolders.cpp
--- a/Engine/Source/Editor/UnrealEd/Private/EditorActorFolders.cpp
+++ b/Engine/Source/Editor/UnrealEd/Private/EditorActorFolders.cpp
@@ -54,6 +54,7 @@ FActorFolders::FActorFolders()
 
 	check(GEditor);
 	UActorEditorContextSubsystem::Get()->RegisterClient(this);
+	bAnyLevelsChanged = false;
 }
 
 FActorFolders::~FActorFolders()
@@ -135,6 +136,17 @@ void FActorFolders::BroadcastOnActorFolderMoved(UWorld& InWorld, const FFolder&
 
 void FActorFolders::OnLevelActorListChanged()
 {
+	bAnyLevelsChanged = true;
+}
+
+void FActorFolders::OnAllLevelsChanged()
+{
+	if (!bAnyLevelsChanged)
+	{
+		return;
+	}
+	QUICK_SCOPE_CYCLE_COUNTER(FActorFolders_OnAllLevelsChanged);
+	bAnyLevelsChanged = false;
 	Housekeeping();
 
 	check(GEngine);
@@ -239,6 +251,7 @@ UWorldFolders& FActorFolders::CreateWorldFolders(UWorld& InWorld)
 	Housekeeping();
 
 	InWorld.OnLevelsChanged().AddRaw(this, &FActorFolders::OnLevelActorListChanged);
+	InWorld.OnAllLevelsChanged().AddRaw(this, &FActorFolders::OnAllLevelsChanged);
 
 	// We intentionally don't pass RF_Transactional to ConstructObject so that we don't record the creation of the object into the undo buffer
 	// (to stop it getting deleted on undo as we manage its lifetime), but we still want it to be RF_Transactional so we can record any changes later
diff --git a/Engine/Source/Editor/UnrealEd/Public/EditorActorFolders.h b/Engine/Source/Editor/UnrealEd/Public/EditorActorFolders.h
--- a/Engine/Source/Editor/UnrealEd/Public/EditorActorFolders.h
+++ b/Engine/Source/Editor/UnrealEd/Public/EditorActorFolders.h
@@ -193,6 +193,9 @@ private:
 	/** Called when the actor list of the current world has changed */
 	void OnLevelActorListChanged();
 
+	/** Called when all the levels have been changed */
+	void OnAllLevelsChanged();
+
 	/** Called when an actor folder is added */
 	void OnActorFolderAdded(UActorFolder* InActorFolder);
 
@@ -223,5 +226,7 @@ private:
 	/** Delegate used to notify changes to ActorEditorContextSubsystem */
 	FOnActorEditorContextClientChanged ActorEditorContextClientChanged;
 
+	bool bAnyLevelsChanged;
+
 	friend UWorldFolders;
 };
diff --git a/Engine/Source/Runtime/Engine/Classes/Engine/World.h b/Engine/Source/Runtime/Engine/Classes/Engine/World.h
--- a/Engine/Source/Runtime/Engine/Classes/Engine/World.h
+++ b/Engine/Source/Runtime/Engine/Classes/Engine/World.h
@@ -1660,6 +1660,12 @@ private:
 	/** Broadcasts whenever the number of levels changes */
 	FOnLevelsChangedEvent LevelsChangedEvent;
 
+	/** Called when all the levels have changed. */
+	DECLARE_EVENT(UWorld, FOnAllLevelsChangedEvent);
+
+	/** Broadcasts whenever all the levels change */
+	FOnAllLevelsChangedEvent AllLevelsChangedEvent;
+
 	DECLARE_EVENT(UWorld, FOnBeginTearingDownEvent);
 
 	/** Broadcasted on UWorld::BeginTearingDown */
@@ -3760,6 +3766,9 @@ public:
 	/** Returns the LevelsChangedEvent member. */
 	FOnLevelsChangedEvent& OnLevelsChanged() { return LevelsChangedEvent; }
 
+	/** Returns the AllLevelsChangedEvent member. */
+	FOnAllLevelsChangedEvent& OnAllLevelsChanged() { return AllLevelsChangedEvent; }
+
 	/** Returns the BeginTearingDownEvent member. */
 	UE_DEPRECATED(4.26, "OnBeginTearingDown has been replaced by FWorldDelegates::OnWorldBeginTearDown")
 	FOnBeginTearingDownEvent& OnBeginTearingDown() { return BeginTearingDownEvent; }
diff --git a/Engine/Source/Runtime/Engine/Private/World.cpp b/Engine/Source/Runtime/Engine/Private/World.cpp
--- a/Engine/Source/Runtime/Engine/Private/World.cpp
+++ b/Engine/Source/Runtime/Engine/Private/World.cpp
@@ -4221,6 +4221,7 @@ void UWorld::UpdateLevelStreaming()
 		}
 	}
 
+	AllLevelsChangedEvent.Broadcast();
 	StreamingLevelsToConsider.EndConsideration();
 
 	const int32 CurrentNumLevelsPendingPurge = FLevelStreamingGCHelper::GetNumLevelsPendingPurge();
