Fix crash caused by ULevelEditorPlaySettings::UpdateCustomSafeZones() accessing Slate in a commandline app

diff --git a/Engine/Source/Editor/UnrealEd/Private/Settings/SettingsClasses.cpp b/Engine/Source/Editor/UnrealEd/Private/Settings/SettingsClasses.cpp
--- a/Engine/Source/Editor/UnrealEd/Private/Settings/SettingsClasses.cpp
+++ b/Engine/Source/Editor/UnrealEd/Private/Settings/SettingsClasses.cpp
@@ -636,10 +636,11 @@ bool ULevelEditorPlaySettings::CanEditChange(const FProperty* InProperty) const
 #if WITH_EDITOR
 void ULevelEditorPlaySettings::UpdateCustomSafeZones()
 {
+	const bool bResetCustomSafeZone = FDisplayMetrics::GetDebugTitleSafeZoneRatio() < 1.f;
+
 	// Prefer to use r.DebugSafeZone.TitleRatio if it is set
-	if (FDisplayMetrics::GetDebugTitleSafeZoneRatio() < 1.f)
+	if (bResetCustomSafeZone)
 	{
-		FSlateApplication::Get().ResetCustomSafeZone();
 		PIESafeZoneOverride = FMargin();
 	}
 	else
@@ -647,12 +648,22 @@ void ULevelEditorPlaySettings::UpdateCustomSafeZones()
 		PIESafeZoneOverride = CalculateCustomUnsafeZones(CustomUnsafeZoneStarts, CustomUnsafeZoneDimensions, DeviceToEmulate, FVector2D(NewWindowWidth, NewWindowHeight));
 	}
 
-	FMargin SafeZoneRatio = PIESafeZoneOverride;
-	SafeZoneRatio.Left /= (NewWindowWidth / 2.0f);
-	SafeZoneRatio.Right /= (NewWindowWidth / 2.0f);
-	SafeZoneRatio.Bottom /= (NewWindowHeight / 2.0f);
-	SafeZoneRatio.Top /= (NewWindowHeight / 2.0f);
-	FSlateApplication::Get().OnDebugSafeZoneChanged.Broadcast(SafeZoneRatio, true);
+	if (FSlateApplication::IsInitialized())
+	{
+		if (bResetCustomSafeZone)
+		{
+			FSlateApplication::Get().ResetCustomSafeZone();
+		}
+		else
+		{
+			FMargin SafeZoneRatio = PIESafeZoneOverride;
+			SafeZoneRatio.Left /= (NewWindowWidth / 2.0f);
+			SafeZoneRatio.Right /= (NewWindowWidth / 2.0f);
+			SafeZoneRatio.Bottom /= (NewWindowHeight / 2.0f);
+			SafeZoneRatio.Top /= (NewWindowHeight / 2.0f);
+			FSlateApplication::Get().OnDebugSafeZoneChanged.Broadcast(SafeZoneRatio, true);
+		}
+	}
 }
 #endif
 
