Wrap deprecated structs in pragmas to prevent internal deprecation warnings.

diff --git a/Engine/Source/Editor/UnrealEd/Public/Commandlets/EditorCommandlets.h b/Engine/Source/Editor/UnrealEd/Public/Commandlets/EditorCommandlets.h
--- a/Engine/Source/Editor/UnrealEd/Public/Commandlets/EditorCommandlets.h
+++ b/Engine/Source/Editor/UnrealEd/Public/Commandlets/EditorCommandlets.h
@@ -188,6 +188,7 @@ struct FResourceDiskSize
 /**
  * Contains an object and the object's path name.
  */
+PRAGMA_DISABLE_DEPRECATION_WARNINGS
 struct UE_DEPRECATED(5.3, "DiffPackagesCommandlet has been deleted") FObjectReference
 {
 	UObject* Object;
@@ -202,11 +203,13 @@ struct UE_DEPRECATED(5.3, "DiffPackagesCommandlet has been deleted") FObjectRefe
 		}
 	}
 };
+PRAGMA_ENABLE_DEPRECATION_WARNINGS
 
 
 /**
  * Represents a single top-level object along with all its subobjects.
  */
+PRAGMA_DISABLE_DEPRECATION_WARNINGS
 struct UE_DEPRECATED(5.3, "DiffPackagesCommandlet has been deleted.") FObjectGraph
 {
 	/**
@@ -228,10 +231,12 @@ struct UE_DEPRECATED(5.3, "DiffPackagesCommandlet has been deleted.") FObjectGra
 	 */
 	UObject* GetRootObject() const;
 };
+PRAGMA_ENABLE_DEPRECATION_WARNINGS
 
 /**
  * Contains the natively serialized property data for a single UObject.
  */
+PRAGMA_DISABLE_DEPRECATION_WARNINGS
 struct UE_DEPRECATED(5.3, "DiffPackagesCommandlet has been deleted.") FNativePropertyData
 {
 	/** the object that this property data is for */
@@ -268,3 +273,4 @@ struct UE_DEPRECATED(5.3, "DiffPackagesCommandlet has been deleted.") FNativePro
 		return PropertyData.Num() || PropertyText.Num();
 	}
 };
+PRAGMA_ENABLE_DEPRECATION_WARNINGS
