Add a "Recent Assets" menu to the level editor and all asset editors

diff --git a/Engine/Source/Editor/ContentBrowser/Private/ContentBrowserModule.cpp b/Engine/Source/Editor/ContentBrowser/Private/ContentBrowserModule.cpp
--- a/Engine/Source/Editor/ContentBrowser/Private/ContentBrowserModule.cpp
+++ b/Engine/Source/Editor/ContentBrowser/Private/ContentBrowserModule.cpp
@@ -15,6 +15,8 @@
 #include "Modules/ModuleManager.h"
 #include "Settings/ContentBrowserSettings.h"
 #include "UObject/UObjectGlobals.h"
+#include "Editor.h"
+#include "Subsystems/AssetEditorSubsystem.h"
 
 IMPLEMENT_MODULE( FContentBrowserModule, ContentBrowser );
 DEFINE_LOG_CATEGORY(LogContentBrowser);
@@ -27,9 +29,6 @@ void FContentBrowserModule::StartupModule()
 
 	ContentBrowserSingleton = new FContentBrowserSingleton();
 	
-	RecentlyOpenedAssets = MakeUnique<FMainMRUFavoritesList>(TEXT("ContentBrowserRecent"), GetDefault<UContentBrowserSettings>()->NumObjectsInRecentList);
-	RecentlyOpenedAssets->ReadFromINI();
-
 	UContentBrowserSettings::OnSettingChanged().AddRaw(this, &FContentBrowserModule::ContentBrowserSettingChanged);
 }
 
@@ -41,7 +40,6 @@ void FContentBrowserModule::ShutdownModule()
 		ContentBrowserSingleton = NULL;
 	}
 	UContentBrowserSettings::OnSettingChanged().RemoveAll(this);
-	RecentlyOpenedAssets.Reset();
 }
 
 IContentBrowserSingleton& FContentBrowserModule::Get() const
@@ -50,6 +48,11 @@ IContentBrowserSingleton& FContentBrowserModule::Get() const
 	return *ContentBrowserSingleton;
 }
 
+FMainMRUFavoritesList* FContentBrowserModule::GetRecentlyOpenedAssets() const
+{
+	return GEditor->GetEditorSubsystem<UAssetEditorSubsystem>()->GetRecentlyOpenedAssets();
+}
+
 FDelegateHandle FContentBrowserModule::AddAssetViewExtraStateGenerator(const FAssetViewExtraStateGenerator& Generator)
 {
 	AssetViewExtraStateGenerators.Add(Generator);
@@ -70,13 +73,5 @@ void FContentBrowserModule::ContentBrowserSettingChanged(FName InName)
 
 	ContentBrowserSingleton->SetPrivateContentPermissionListDirty();
 
-	// Resize the recently opened asset list
-	if (InName == NumberOfRecentAssetsName)
-	{
-		RecentlyOpenedAssets->WriteToINI();
-		RecentlyOpenedAssets = MakeUnique<FMainMRUFavoritesList>(TEXT("ContentBrowserRecent"), GetDefault<UContentBrowserSettings>()->NumObjectsInRecentList);
-		RecentlyOpenedAssets->ReadFromINI();
-	}
-
 	OnContentBrowserSettingChanged.Broadcast(InName);
 }
diff --git a/Engine/Source/Editor/ContentBrowser/Public/ContentBrowserModule.h b/Engine/Source/Editor/ContentBrowser/Public/ContentBrowserModule.h
--- a/Engine/Source/Editor/ContentBrowser/Public/ContentBrowserModule.h
+++ b/Engine/Source/Editor/ContentBrowser/Public/ContentBrowserModule.h
@@ -135,10 +135,7 @@ public:
 	/** Override list of paths to expand by default */
 	FDefaultPathsToExpandDelegate& GetDefaultPathsToExpandDelegate() { return DefaultPathsToExpandDelegate; }
 
-	FMainMRUFavoritesList* GetRecentlyOpenedAssets() const
-	{
-		return RecentlyOpenedAssets.Get();
-	};
+	FMainMRUFavoritesList* GetRecentlyOpenedAssets() const;
 
 	static const FName NumberOfRecentAssetsName;
 
@@ -189,8 +186,6 @@ private:
 	/** All delegates that extend available path view plugin filters */
 	TArray<FAddPathViewPluginFilters> PathViewPluginFilters;
 
-	TUniquePtr<FMainMRUFavoritesList> RecentlyOpenedAssets;
-
 	FOnFilterChanged OnFilterChanged;
 	FOnSearchBoxChanged OnSearchBoxChanged;
 	FOnAssetSelectionChanged OnAssetSelectionChanged;
diff --git a/Engine/Source/Editor/EditorStyle/Private/StarshipStyle.cpp b/Engine/Source/Editor/EditorStyle/Private/StarshipStyle.cpp
--- a/Engine/Source/Editor/EditorStyle/Private/StarshipStyle.cpp
+++ b/Engine/Source/Editor/EditorStyle/Private/StarshipStyle.cpp
@@ -454,6 +454,7 @@ void FStarshipEditorStyle::FStyle::SetupGeneralStyles()
 		Set("Icons.Clipboard", new IMAGE_BRUSH_SVG("Starship/Common/Clipboard", Icon16x16));
 		Set("Icons.HollowHeart", new IMAGE_BRUSH_SVG("Starship/Common/HollowHeart", Icon16x16));
 		Set("Icons.Heart", new IMAGE_BRUSH_SVG("Starship/Common/Heart", Icon16x16));
+		Set("Icons.RecentAssets", new IMAGE_BRUSH_SVG("Starship/Common/RecentAssets", Icon16x16));
 
 
 		Set("Icons.Toolbar.Play", new IMAGE_BRUSH_SVG("Starship/Common/play", Icon20x20));
diff --git a/Engine/Source/Editor/LevelEditor/Private/LevelEditorMenu.cpp b/Engine/Source/Editor/LevelEditor/Private/LevelEditorMenu.cpp
--- a/Engine/Source/Editor/LevelEditor/Private/LevelEditorMenu.cpp
+++ b/Engine/Source/Editor/LevelEditor/Private/LevelEditorMenu.cpp
@@ -47,6 +47,10 @@ void FLevelEditorMenu::RegisterLevelEditorMenus()
 			// Open Level
 			OpenSection.AddMenuEntry( FLevelEditorCommands::Get().OpenLevel ).InsertPosition = InsertPos;
 
+			FToolMenuSection& AssetSection = Menu->FindOrAddSection("FileAsset");
+			
+			AssetSection.AddSeparator("FileAssetSeparator").InsertPosition = FToolMenuInsert(NAME_None, EToolMenuInsertType::First);
+
 			// Open Asset
 			//@TODO: Doesn't work when summoned from here: Section.AddMenuEntry( FGlobalEditorCommonCommands::Get().SummonOpenAssetDialog );
 
diff --git a/Engine/Source/Editor/MainFrame/Private/Menus/MainMenu.cpp b/Engine/Source/Editor/MainFrame/Private/Menus/MainMenu.cpp
--- a/Engine/Source/Editor/MainFrame/Private/Menus/MainMenu.cpp
+++ b/Engine/Source/Editor/MainFrame/Private/Menus/MainMenu.cpp
@@ -40,14 +40,17 @@ void FMainMenu::RegisterFileMenu()
 	UToolMenus* ToolMenus = UToolMenus::Get();
 	UToolMenu* FileMenu = ToolMenus->RegisterMenu("MainFrame.MainMenu.File");
 
-	FToolMenuSection& FileLoadAndSaveSection = FileMenu->AddSection("FileOpen", LOCTEXT("FileOpenHeading", "Open"), FToolMenuInsert(NAME_None, EToolMenuInsertType::First));
+	FileMenu->AddSection("FileOpen", LOCTEXT("FileOpenHeading", "Open"), FToolMenuInsert(NAME_None, EToolMenuInsertType::First));
+
+	FToolMenuSection& FileAssetSection = FileMenu->FindOrAddSection("FileAsset");;
 	{
+		FileAssetSection.InsertPosition = FToolMenuInsert("FileOpen", EToolMenuInsertType::After);
 		// Open Asset...
-		FileLoadAndSaveSection.AddMenuEntry(FGlobalEditorCommonCommands::Get().SummonOpenAssetDialog);
+		FileAssetSection.AddMenuEntry(FGlobalEditorCommonCommands::Get().SummonOpenAssetDialog);
 	}
 
 
-	FToolMenuSection& FileSaveSection = FileMenu->AddSection("FileSave", LOCTEXT("FileSaveHeading", "Save"), FToolMenuInsert("FileOpen", EToolMenuInsertType::After));
+	FToolMenuSection& FileSaveSection = FileMenu->AddSection("FileSave", LOCTEXT("FileSaveHeading", "Save"), FToolMenuInsert("FileAsset", EToolMenuInsertType::After));
 	{
 
 		// Save All
diff --git a/Engine/Source/Editor/UnrealEd/Classes/Settings/ContentBrowserSettings.h b/Engine/Source/Editor/UnrealEd/Classes/Settings/ContentBrowserSettings.h
--- a/Engine/Source/Editor/UnrealEd/Classes/Settings/ContentBrowserSettings.h
+++ b/Engine/Source/Editor/UnrealEd/Classes/Settings/ContentBrowserSettings.h
@@ -58,8 +58,8 @@ public:
 	bool bDisplayFriendlyNameForPluginFolders = true;
 
 	/** The number of objects to keep in the Content Browser Recently Opened filter */
-	UPROPERTY(EditAnywhere, config, Category = ContentBrowser, meta = (DisplayName = "Number of Assets to Keep in the Recently Opened Filter", ClampMin = "1", ClampMax = "30"))
-	int32 NumObjectsInRecentList;
+	UPROPERTY(meta = (DeprecatedProperty, DeprecationMessage = "Use NumRecentAssets in UEditorLoadingSavingSettings instead"))
+	int32 NumObjectsInRecentList_DEPRECATED;
 
 	/** Enables the rendering of Material Instance thumbnail previews */
 	UPROPERTY(EditAnywhere, config, Category = ContentBrowser)
diff --git a/Engine/Source/Editor/UnrealEd/Classes/Settings/EditorLoadingSavingSettings.h b/Engine/Source/Editor/UnrealEd/Classes/Settings/EditorLoadingSavingSettings.h
--- a/Engine/Source/Editor/UnrealEd/Classes/Settings/EditorLoadingSavingSettings.h
+++ b/Engine/Source/Editor/UnrealEd/Classes/Settings/EditorLoadingSavingSettings.h
@@ -116,6 +116,10 @@ public:
 	UPROPERTY(meta = (DeprecatedProperty, DeprecationMessage = "Use RestoreOpenAssetTabsOnRestart instead"))
 	uint32 bRestoreOpenAssetTabsOnRestart_DEPRECATED:1;
 #endif
+
+	/** The number of recent assets the editor will keep track of */
+	UPROPERTY(EditAnywhere, config, Category = Startup, meta = (DisplayName = "Number of Recent Assets to keep track of", ClampMin = "1", ClampMax = "30"))
+	int32 NumRecentAssets = 20;
 	
 private:
 
diff --git a/Engine/Source/Editor/UnrealEd/Private/Subsystems/AssetEditorSubsystem.cpp b/Engine/Source/Editor/UnrealEd/Private/Subsystems/AssetEditorSubsystem.cpp
--- a/Engine/Source/Editor/UnrealEd/Private/Subsystems/AssetEditorSubsystem.cpp
+++ b/Engine/Source/Editor/UnrealEd/Private/Subsystems/AssetEditorSubsystem.cpp
@@ -14,7 +14,6 @@
 #include "Toolkits/AssetEditorToolkit.h"
 #include "Toolkits/SimpleAssetEditor.h"
 #include "Engine/MapBuildDataRegistry.h"
-#include "ContentBrowserModule.h"
 #include "MRUFavoritesList.h"
 #include "Settings/EditorLoadingSavingSettings.h"
 #include "UObject/UObjectIterator.h"
@@ -33,6 +32,9 @@
 #include "EditorModeManager.h"
 #include "Tools/LegacyEdMode.h"
 #include "ProfilingDebugging/StallDetector.h"
+#include "ToolMenus.h"
+#include "AssetRegistry/IAssetRegistry.h"
+#include "AssetRegistry/AssetRegistryModule.h"
 
 #include "Elements/SMInstance/SMInstanceElementData.h" // For SMInstanceElementDataUtil::SMInstanceElementsEnabled
 
@@ -59,7 +61,17 @@ void UAssetEditorSubsystem::Initialize(FSubsystemCollectionBase& Collection)
 	FCoreDelegates::OnEnginePreExit.AddUObject(this, &UAssetEditorSubsystem::UnregisterEditorModes);
 	FCoreDelegates::OnPostEngineInit.AddUObject(this, &UAssetEditorSubsystem::RegisterEditorModes);
 
+	if (FAssetRegistryModule* AssetRegistryModule = FModuleManager::GetModulePtr<FAssetRegistryModule>(TEXT("AssetRegistry")))
+	{
+		AssetRegistryModule->Get().OnAssetRemoved().AddUObject(this, &UAssetEditorSubsystem::OnAssetRemoved);
+		AssetRegistryModule->Get().OnAssetRenamed().AddUObject(this, &UAssetEditorSubsystem::OnAssetRenamed);
+	}
+
 	SMInstanceElementDataUtil::OnSMInstanceElementsEnabledChanged().AddUObject(this, &UAssetEditorSubsystem::OnSMInstanceElementsEnabled);
+
+	RegisterLevelEditorMenuExtensions();
+
+	InitializeRecentAssets();
 }
 
 void UAssetEditorSubsystem::Deinitialize()
@@ -70,6 +82,12 @@ void UAssetEditorSubsystem::Deinitialize()
 	FCoreDelegates::OnPostEngineInit.RemoveAll(this);
 	SMInstanceElementDataUtil::OnSMInstanceElementsEnabledChanged().RemoveAll(this);
 
+	if (FAssetRegistryModule* AssetRegistryModule = FModuleManager::GetModulePtr<FAssetRegistryModule>(TEXT("AssetRegistry")))
+	{
+		AssetRegistryModule->Get().OnAssetRemoved().RemoveAll(this);
+		AssetRegistryModule->Get().OnAssetRenamed().RemoveAll(this);
+	}
+
 	// Don't attempt to report usage stats if analytics isn't available
 	if (FEngineAnalytics::IsAvailable())
 	{
@@ -86,6 +104,154 @@ void UAssetEditorSubsystem::Deinitialize()
 			FEngineAnalytics::GetProvider().RecordEvent(EventName, EditorUsageAttribs);
 		}
 	}
+
+	SaveRecentAssets();
+	RecentAssetsList.Reset();
+}
+
+void UAssetEditorSubsystem::InitializeRecentAssets()
+{
+	// The current max allowed is 30 assets, we keep track of 40 to be safe since we are ignoring levels when showing the list
+	RecentAssetsList = MakeUnique<FMainMRUFavoritesList>(TEXT("AssetEditorSubsystemRecents"), 40);
+	RecentAssetsList->ReadFromINI();
+	
+	TArray<FString> RecentAssetEditors;
+
+	GConfig->GetArray(TEXT("AssetEditorSubsystem"), TEXT("RecentAssetEditors"), RecentAssetEditors, GEditorPerProjectIni);
+
+	bool bFoundRecentAssetEditorsCleanly = true;
+
+	// If the number of recent assets and recent asset editors don't match, something went wrong during saving and we have potentially corrupt data
+	if(RecentAssetsList->GetNumItems() != RecentAssetEditors.Num())
+	{
+		bFoundRecentAssetEditorsCleanly = false;
+		UE_LOG(LogAssetEditorSubsystem, Error, TEXT("Something went wrong while loading recent assets! Num Recent Assets  = %d, Num Recent Asset Editors = %d"), RecentAssetsList->GetNumItems(), RecentAssetEditors.Num());
+	}
+
+	// If we have corrupt data, simply ignore the asset editor list instead of potentially showing an asset in the wrong asset editor's recents
+	if(bFoundRecentAssetEditorsCleanly)
+	{
+		// Go in reverse since the first item should be the most recent
+		for(int i = 0; i < RecentAssetsList->GetNumItems(); ++i)
+		{
+			if(!RecentAssetEditors[i].IsEmpty())
+			{
+				RecentAssetToAssetEditorMap.Add(RecentAssetsList->GetMRUItem(i), RecentAssetEditors[i]);
+			}
+		}
+	}
+}
+
+void UAssetEditorSubsystem::SaveRecentAssets(const bool bOnShutdown)
+{
+	TArray<FString> RecentAssetEditorsToSave;
+
+	// If we are closing the editor, remove all assets that weren't actually saved to disk (transient)
+	if(bOnShutdown)
+	{
+		for(int i = 0; i < RecentAssetsList->GetNumItems(); ++i)
+		{
+			const FString& RecentAsset = RecentAssetsList->GetMRUItem(i);
+			if(!FPackageName::DoesPackageExist(RecentAsset))
+			{
+				RecentAssetsList->RemoveMRUItem(RecentAsset);
+				--i;
+			}
+		}
+	}
+	
+	for(int i = 0; i < RecentAssetsList->GetNumItems(); ++i)
+	{
+		FString CurRecentAsset = RecentAssetsList->GetMRUItem(i);
+
+		// If we have a valid asset editor for the current asset, save it
+		if(FString* CurrentAssetEditorName = RecentAssetToAssetEditorMap.Find(CurRecentAsset))
+		{
+			RecentAssetEditorsToSave.Add(*CurrentAssetEditorName);
+		}
+		// Otherwise add an empty entry (e.g levels) so the two arrays are always the same size
+		else
+		{
+			RecentAssetEditorsToSave.Add(FString());
+		}
+		
+	}
+
+	RecentAssetsList->WriteToINI();
+	GConfig->SetArray(TEXT("AssetEditorSubsystem"), TEXT("RecentAssetEditors"), RecentAssetEditorsToSave, GEditorPerProjectIni);
+}
+
+void UAssetEditorSubsystem::CullRecentAssetEditorsMap()
+{
+	/* Since the Recent Asset -> Asset Editor Map is not an MRU list, it can keep infinitely growing as the user opens assets
+	 * To keep it a reasonable size while also not culling it too often, we cull it when it gets twice as big as the MRU list
+	 */
+	if(RecentAssetToAssetEditorMap.Num() > 2 * RecentAssetsList->GetMaxItems())
+	{
+		for (TMap<FString, FString>::TIterator It(RecentAssetToAssetEditorMap); It; ++It)
+		{
+			// Remove any entries that are not in the mru list
+			if(RecentAssetsList->FindMRUItemIdx(It->Key) == INDEX_NONE)
+			{
+				It.RemoveCurrent();
+			}
+		}
+
+	}
+}
+
+void UAssetEditorSubsystem::OnAssetRemoved(const FAssetData& AssetData)
+{
+	FString PathName = FPaths::GetBaseFilename(AssetData.GetObjectPathString(), false);
+
+	// We need this early exit because FindMRUItemIdx has a check() for non valid long package names
+	if (!FPackageName::IsValidLongPackageName(PathName))
+	{
+		return;
+	}
+
+	// If the asset that was deleted was not found in the recent assets list, we have nothing to do
+	if(RecentAssetsList->FindMRUItemIdx(PathName) == INDEX_NONE)
+	{
+		return;
+	}
+
+	// Remove the asset from our list and map
+	RecentAssetsList->RemoveMRUItem(PathName);
+	RecentAssetToAssetEditorMap.Remove(PathName);
+
+	SaveRecentAssets();
+}
+
+void UAssetEditorSubsystem::OnAssetRenamed(const FAssetData& AssetData, const FString& AssetOldName)
+{
+	FString OldPathName = FPaths::GetBaseFilename(AssetOldName, false);
+	FString NewPathName = FPaths::GetBaseFilename(AssetData.GetObjectPathString(), false);
+
+	// We need this early exit because FindMRUItemIdx has a check() for non valid long package names
+	if (!FPackageName::IsValidLongPackageName(OldPathName) || !FPackageName::IsValidLongPackageName(NewPathName))
+	{
+		return;
+	}
+
+	// If the asset did not previously exist in the recents list, we have nothing to do
+	if(RecentAssetsList->FindMRUItemIdx(OldPathName) == INDEX_NONE)
+	{
+		return;
+	}
+
+	// Otherwise remove the old name of the asset, and re-add it with the new name
+	// NOTE: This has an unintentional side effect of bringing it to the top of the MRU list that can't be avoided
+	RecentAssetsList->RemoveMRUItem(OldPathName);
+	RecentAssetsList->AddMRUItem(NewPathName);
+
+	if(FString* AssetEditorName = RecentAssetToAssetEditorMap.Find(OldPathName))
+	{
+		RecentAssetToAssetEditorMap.Add(NewPathName, *AssetEditorName);
+		RecentAssetToAssetEditorMap.Remove(OldPathName);
+	}
+
+	SaveRecentAssets();
 }
 
 void UAssetEditorSubsystem::OnEditorClose()
@@ -216,8 +382,11 @@ void UAssetEditorSubsystem::NotifyAssetOpened(UObject* Asset, IAssetEditorInstan
 		OpenedEditorTimes.Add(InInstance, EditorTime);
 	}
 
+	FString AssetPath = Asset->GetOuter()->GetPathName();
+	
 	OpenedAssets.Add(Asset, InInstance);
 	OpenedEditors.Add(InInstance, Asset);
+	RecentAssetToAssetEditorMap.Add(AssetPath, InInstance->GetEditorName().ToString());
 
 	AssetOpenedInEditorEvent.Broadcast(Asset, InInstance);
 
@@ -432,11 +601,10 @@ bool UAssetEditorSubsystem::OpenEditorForAsset(UObject* Asset, const EToolkitMod
 		if (Asset->IsAsset() && !Asset->IsA(UMapBuildDataRegistry::StaticClass()))
 		{
 			FString AssetPath = Asset->GetOuter()->GetPathName();
-			FContentBrowserModule& CBModule = FModuleManager::LoadModuleChecked<FContentBrowserModule>(TEXT("ContentBrowser"));
-			FMainMRUFavoritesList* RecentlyOpenedAssets = CBModule.GetRecentlyOpenedAssets();
-			if (RecentlyOpenedAssets && FPackageName::IsValidLongPackageName(AssetPath))
+			if (FPackageName::IsValidLongPackageName(AssetPath))
 			{
-				RecentlyOpenedAssets->AddMRUItem(AssetPath);
+				RecentAssetsList->AddMRUItem(AssetPath);
+				CullRecentAssetEditorsMap();
 			}
 		}
 	}
@@ -602,9 +770,14 @@ void UAssetEditorSubsystem::OpenEditorForAsset(const FSoftObjectPath& AssetPath)
 
 void UAssetEditorSubsystem::OpenEditorForAsset(const FString& AssetPathName, const EAssetTypeActivationOpenedMethod OpenedMethod)
 {
-	// An asset needs loading
+	UPackage* Package = FindPackage(NULL, *AssetPathName);
 
-	UPackage* Package = LoadPackage(NULL, *AssetPathName, LOAD_NoRedirects);
+	if(!Package)
+	{
+		// An asset needs loading
+		Package =  LoadPackage(NULL, *AssetPathName, LOAD_NoRedirects);
+	}
+	
 	if (Package)
 	{
 		Package->FullyLoad();
@@ -937,8 +1110,174 @@ void UAssetEditorSubsystem::OnCancelRestorePreviouslyOpenAssets()
 	}
 }
 
+bool UAssetEditorSubsystem::ShouldShowRecentAsset(const FString& AssetName, int32 RecentAssetIndex, const FName& InAssetEditorName) const
+{
+	const FString* AssetEditorForCurrRecent = RecentAssetToAssetEditorMap.Find(AssetName);
+
+	// If this asset wasn't opened in any valid asset editor (e.g Levels)
+	if(!AssetEditorForCurrRecent)
+	{
+		return false;
+	}
+
+	// If we have a valid asset editor we are adding assets for
+	if(!InAssetEditorName.IsNone())
+	{
+		// If this asset was not opened in InAssetEditorName, ignore it
+		if(*AssetEditorForCurrRecent != InAssetEditorName.ToString())
+		{
+			return false;
+		}
+	}
+		
+	// If this asset does not pass the set filter, ignore it
+	if (!RecentAssetsList->MRUItemPassesCurrentFilter(RecentAssetIndex))
+	{
+		return false;
+	}
+
+	return true;
+}
+
+bool UAssetEditorSubsystem::ShouldShowRecentAssetsMenu(const FName& InAssetEditorName) const
+{
+	// If we have no recent assets at all
+	if(RecentAssetsList->GetNumItems() == 0)
+	{
+		return false;
+	}
+
+	const int32 MaxRecentAssets = GetDefault<UEditorLoadingSavingSettings>()->NumRecentAssets;
+
+	for ( int32 CurRecentIndex = 0; CurRecentIndex < RecentAssetsList->GetNumItems() && CurRecentIndex < MaxRecentAssets; ++CurRecentIndex )
+	{
+		const FString& CurRecent = RecentAssetsList->GetMRUItem(CurRecentIndex);
+
+		// If any of the assets in the recents wil be shown, we show the menu
+		if(ShouldShowRecentAsset(CurRecent, CurRecentIndex, InAssetEditorName))
+		{
+			return true;
+		}
+	}
+
+	return false;
+}
+
+void UAssetEditorSubsystem::CreateRecentAssetsMenu(UToolMenu* InMenu, const FName InAssetEditorName)
+{
+	const int32 MaxRecentAssets = GetDefault<UEditorLoadingSavingSettings>()->NumRecentAssets;
+
+	FToolMenuSection& Section = InMenu->FindOrAddSection("Recents");
+	
+	// Keep adding assets until we reach the end of the MRU list, or we reach the max allowed assets
+	for ( int32 CurRecentIndex = 0; CurRecentIndex < RecentAssetsList->GetNumItems() && CurRecentIndex < MaxRecentAssets; ++CurRecentIndex )
+	{
+		const FString& CurRecent = RecentAssetsList->GetMRUItem(CurRecentIndex);
+
+		if(!ShouldShowRecentAsset(CurRecent, CurRecentIndex, InAssetEditorName))
+		{
+			continue;
+		}
+
+		const FText ToolTip = FText::Format( LOCTEXT( "RecentAssetsToolTip", "Open {0}" ), FText::FromString( CurRecent ) );
+		const FText Label = FText::FromString( FPaths::GetBaseFilename(CurRecent) );
+
+		Section.AddMenuEntry(
+		NAME_None,
+		Label,
+		ToolTip,
+		FSlateIcon(),
+		FUIAction(
+			FExecuteAction::CreateLambda([this, CurRecent]()
+				{
+					OpenEditorForAsset(CurRecent);
+				})
+			)
+		);
+	}
+}
+
+void UAssetEditorSubsystem::RegisterLevelEditorMenuExtensions()
+{
+	UToolMenu* Menu = UToolMenus::Get()->ExtendMenu("LevelEditor.MainMenu.File");
+	
+	FToolMenuSection& Section = Menu->FindOrAddSection("FileAsset");
+
+	Section.AddDynamicEntry("FileRecentAssets", FNewToolMenuSectionDelegate::CreateLambda([this](FToolMenuSection& InSection)
+	{
+		// Since we want to show all asset types for the Level Editor, we use an empty asset editor name
+		const FName AssetEditorName = FName();
+
+		if (!ShouldShowRecentAssetsMenu(AssetEditorName))
+		{
+			return;
+		}
+		
+		InSection.AddSubMenu(
+			"RecentAssetsSubmenu",
+			LOCTEXT("RecentAssetsSubmenu_Label", "Recent Assets"),
+			LOCTEXT("RecentAssetsSubMenu_ToolTip", "Access your recently opened assets"),
+			FNewToolMenuDelegate::CreateUObject(this, &UAssetEditorSubsystem::CreateRecentAssetsMenu, AssetEditorName),
+			false,
+			FSlateIcon(FAppStyle::GetAppStyleSetName(), "Icons.RecentAssets"));
+	}));
+}
+
+void UAssetEditorSubsystem::CreateRecentAssetsMenuForEditor(const IAssetEditorInstance* InAssetEditorInstance, FToolMenuSection& InSection)
+{
+	if(!InAssetEditorInstance)
+	{
+		return;
+	}
+
+	InSection.AddDynamicEntry("FileRecentAssetEditorAssets", FNewToolMenuSectionDelegate::CreateLambda([this, InAssetEditorInstance](FToolMenuSection& InSection)
+	{
+		const FName EditingAssetTypeName = InAssetEditorInstance->GetEditingAssetTypeName();
+
+		// For generic asset editors (or any other special cases) that don't have one singular type of asset they are editing, show all recent assets
+		const FName AssetEditorName = EditingAssetTypeName.IsNone() ? FName() : InAssetEditorInstance->GetEditorName();
+
+		if(!ShouldShowRecentAssetsMenu(AssetEditorName))
+		{
+			return;
+		}
+
+		// Show all Recent Assets
+		if(AssetEditorName.IsNone())
+		{
+			InSection.AddSubMenu(
+			"RecentAssetsSubmenu",
+			LOCTEXT("RecentAssetsSubmenu_Label", "Recent Assets"),
+			LOCTEXT("RecentAssetsSubMenu_ToolTip", "Access your recently opened assets"),
+			FNewToolMenuDelegate::CreateUObject(this, &UAssetEditorSubsystem::CreateRecentAssetsMenu, AssetEditorName),
+			false,
+			FSlateIcon(FAppStyle::GetAppStyleSetName(), "Icons.RecentAssets")
+			);
+		}
+		// Only show recent assets opened in this Asset Editor
+		else
+		{
+			// Example submenu name: "Recent Material Assets" for the Material Editor
+			const FName RecentAssetsMenuName("Recent " + EditingAssetTypeName.ToString() + " Assets");
+			
+			InSection.AddSubMenu(
+				"RecentAssetEditorAssetsSubmenu",
+				FText::Format(LOCTEXT("RecentAssetEditorAssetsSubmenu_Label", "{0}"), FText::FromName(RecentAssetsMenuName)),
+				LOCTEXT("RecentAssetEditorAssetsSubMenu_ToolTip", "Access your recently opened assets"),
+				FNewToolMenuDelegate::CreateUObject(this, &UAssetEditorSubsystem::CreateRecentAssetsMenu, AssetEditorName),
+				false,
+				FSlateIcon(FAppStyle::GetAppStyleSetName(), "Icons.RecentAssets")
+			);
+		}
+		
+	}));
+;
+}
+
 void UAssetEditorSubsystem::SaveOpenAssetEditors(const bool bOnShutdown)
 {
+	SaveRecentAssets(bOnShutdown);
+
 	// Saving is disabled if bAutoRestoreAndDisableSaving is set
 	if (!bSavingOnShutdown && !bAutoRestoreAndDisableSaving.IsSet())
 	{
diff --git a/Engine/Source/Editor/UnrealEd/Private/Toolkits/AssetEditorToolkit.cpp b/Engine/Source/Editor/UnrealEd/Private/Toolkits/AssetEditorToolkit.cpp
--- a/Engine/Source/Editor/UnrealEd/Private/Toolkits/AssetEditorToolkit.cpp
+++ b/Engine/Source/Editor/UnrealEd/Private/Toolkits/AssetEditorToolkit.cpp
@@ -381,6 +381,21 @@ FText FAssetEditorToolkit::GetTabSuffix() const
 	return bDirtyState ? LOCTEXT("TabSuffixAsterix", "*") : FText::GetEmpty();
 }
 
+FName FAssetEditorToolkit::GetEditingAssetTypeName() const
+{
+	if(EditingObjects.IsEmpty())
+	{
+		return NAME_None;
+	}
+
+	if(UClass* EditingClass = EditingObjects[0]->GetClass())
+	{
+		return EditingClass->GetFName();
+	}
+
+	return NAME_None;
+}
+
 FText FAssetEditorToolkit::GetToolkitToolTipText() const
 {
 	const UObject* EditingObject = GetEditingObject();
@@ -940,6 +955,11 @@ bool FAssetEditorToolkit::ShouldPromptForNewFilesOnReload(const UObject& Editing
 	return true;
 }
 
+void FAssetEditorToolkit::FillDefaultFileMenuOpenCommands(FToolMenuSection& InSection)
+{
+	GEditor->GetEditorSubsystem<UAssetEditorSubsystem>()->CreateRecentAssetsMenuForEditor(this, InSection);
+}
+
 void FAssetEditorToolkit::FillDefaultFileMenuCommands(FToolMenuSection& InSection)
 {
 	const FToolMenuInsert InsertPosition(NAME_None, EToolMenuInsertType::First);
diff --git a/Engine/Source/Editor/UnrealEd/Private/Toolkits/SStandaloneAssetEditorToolkitHost.cpp b/Engine/Source/Editor/UnrealEd/Private/Toolkits/SStandaloneAssetEditorToolkitHost.cpp
--- a/Engine/Source/Editor/UnrealEd/Private/Toolkits/SStandaloneAssetEditorToolkitHost.cpp
+++ b/Engine/Source/Editor/UnrealEd/Private/Toolkits/SStandaloneAssetEditorToolkitHost.cpp
@@ -97,10 +97,26 @@ void SStandaloneAssetEditorToolkitHost::CreateDefaultStandaloneMenuBar(UToolMenu
 		{
 			const FName MenuName = *(InMenuBar->GetMenuName().ToString() + TEXT(".") + TEXT("File"));
 			UToolMenu* Menu = UToolMenus::Get()->ExtendMenu(MenuName);
+
+			FToolMenuSection& FileAssetSection = Menu->FindOrAddSection("FileAsset");
+
+			FileAssetSection.Label = LOCTEXT("FileAssetSectionHeading", "Open");
+			
+			FileAssetSection.AddDynamicEntry(NAME_None, FNewToolMenuSectionDelegate::CreateLambda([](FToolMenuSection& InSection)
+			{
+				if (UAssetEditorToolkitMenuContext* Context = InSection.FindContext<UAssetEditorToolkitMenuContext>())
+				{
+					Context->Toolkit.Pin()->FillDefaultFileMenuOpenCommands(InSection);
+				}
+			}));
+			
 			FToolMenuSection& Section = Menu->FindOrAddSection("FileLoadAndSave");
 			Section.AddDynamicEntry(NAME_None, FNewToolMenuSectionDelegate::CreateLambda([](FToolMenuSection& InSection)
 			{
-				InSection.FindContext<UAssetEditorToolkitMenuContext>()->Toolkit.Pin()->FillDefaultFileMenuCommands(InSection);
+				if (UAssetEditorToolkitMenuContext* Context = InSection.FindContext<UAssetEditorToolkitMenuContext>())
+				{
+					Context->Toolkit.Pin()->FillDefaultFileMenuCommands(InSection);
+				}
 			}));
 		}
 
@@ -109,7 +125,10 @@ void SStandaloneAssetEditorToolkitHost::CreateDefaultStandaloneMenuBar(UToolMenu
 			FToolMenuSection& Section = InMenu->AddSection("AssetEditorActions", LOCTEXT("ActionsHeading", "Actions"));
 			Section.AddDynamicEntry(NAME_None, FNewToolMenuSectionDelegate::CreateLambda([](FToolMenuSection& InSection)
 			{
-				InSection.FindContext<UAssetEditorToolkitMenuContext>()->Toolkit.Pin()->FillDefaultAssetMenuCommands(InSection);
+				if (UAssetEditorToolkitMenuContext* Context = InSection.FindContext<UAssetEditorToolkitMenuContext>())
+				{
+					Context->Toolkit.Pin()->FillDefaultAssetMenuCommands(InSection);
+				}
 			}));
 		}
 
diff --git a/Engine/Source/Editor/UnrealEd/Private/Toolkits/SimpleAssetEditor.cpp b/Engine/Source/Editor/UnrealEd/Private/Toolkits/SimpleAssetEditor.cpp
--- a/Engine/Source/Editor/UnrealEd/Private/Toolkits/SimpleAssetEditor.cpp
+++ b/Engine/Source/Editor/UnrealEd/Private/Toolkits/SimpleAssetEditor.cpp
@@ -336,4 +336,10 @@ void FSimpleAssetEditor::PostRegenerateMenusAndToolbars()
 	}
 }
 
+FName FSimpleAssetEditor::GetEditingAssetTypeName() const
+{
+	// We want the global recent assets menu for simple asset editors so we report our editing asset type as none
+	return NAME_None;
+}
+
 #undef LOCTEXT_NAMESPACE
diff --git a/Engine/Source/Editor/UnrealEd/Public/Subsystems/AssetEditorSubsystem.h b/Engine/Source/Editor/UnrealEd/Public/Subsystems/AssetEditorSubsystem.h
--- a/Engine/Source/Editor/UnrealEd/Public/Subsystems/AssetEditorSubsystem.h
+++ b/Engine/Source/Editor/UnrealEd/Public/Subsystems/AssetEditorSubsystem.h
@@ -10,6 +10,7 @@
 #include "Tools/Modes.h"
 #include "Misc/NamePermissionList.h"
 #include "AssetTypeActivationOpenedMethod.h"
+#include "MRUFavoritesList.h"
 #include "AssetEditorSubsystem.generated.h"
 
 class UAssetEditor;
@@ -19,6 +20,8 @@ class UClass;
 struct FAssetEditorRequestOpenAsset;
 class FEditorModeTools;
 class IMessageContext;
+class UToolMenu;
+struct FToolMenuSection;
 
 /**
  * This class keeps track of a currently open asset editor; allowing it to be
@@ -36,6 +39,7 @@ public:
 	virtual void InvokeTab(const struct FTabId& TabId) = 0;
 	UE_DEPRECATED(5.0, "Toolbar tab no longer exists and tab ID will return None; do not add it to layouts")
 	virtual FName GetToolbarTabId() const { return NAME_None; }
+	virtual FName GetEditingAssetTypeName() const { return NAME_None; }
 	virtual TSharedPtr<class FTabManager> GetAssociatedTabManager() = 0;
 	virtual double GetLastActivationTime() = 0;
 	virtual void RemoveEditingAsset(UObject* Asset) = 0;
@@ -166,7 +170,6 @@ public:
 	/** Request notification to restore the assets that were previously open when the editor was last closed */
 	void RequestRestorePreviouslyOpenAssets();
 
-	
 	void RegisterUAssetEditor(UAssetEditor* NewAssetEditor);
 	void UnregisterUAssetEditor(UAssetEditor* RemovedAssetEditor);
 	
@@ -213,6 +216,14 @@ public:
 	/** Get the permission list that controls which editor modes are exposed */
 	FNamePermissionList& GetAllowedEditorModes();
 
+	/** Fill in the given section with the recent assets for the given asset editor instance */
+	void CreateRecentAssetsMenuForEditor(const IAssetEditorInstance* InAssetEditorInstance, FToolMenuSection& InSection);
+
+	FMainMRUFavoritesList* GetRecentlyOpenedAssets() const
+	{
+		return RecentAssetsList.Get();
+	}
+
 private:
 
 	/** Handles FAssetEditorRequestOpenAsset messages. */
@@ -230,6 +241,19 @@ private:
 	/** Handler for when the "Don't Restore" button is clicked on the RestorePreviouslyOpenAssets notification */
 	void OnCancelRestorePreviouslyOpenAssets();
 
+	/** Register extensions to the Level Editor "File" Menu */
+	void RegisterLevelEditorMenuExtensions();
+
+	/**
+	 * Fill in the given menu with entries from the recent assets list, clicking on which will open up the asset
+	 * @param InMenu The menu to fill in
+	 * @param InAssetEditorName If provided, only assets in the list that were opened in this asset editor will be added to the menu
+	 */
+	void CreateRecentAssetsMenu(UToolMenu* InMenu, const FName InAssetEditorName = FName());
+
+	void OnAssetRemoved(const FAssetData& AssetData);
+
+	void OnAssetRenamed(const FAssetData& AssetData, const FString& AssetOldName);
 public:
 
 	/**
@@ -274,6 +298,13 @@ private:
 
 	bool IsEditorModeAllowed(const FName ModeId) const;
 
+	void InitializeRecentAssets();
+	void SaveRecentAssets(const bool bOnShutdown = false);
+	void CullRecentAssetEditorsMap();
+
+	bool ShouldShowRecentAsset(const FString& AssetName, int32 RecentAssetIndex,  const FName& InAssetEditorName) const;
+	bool ShouldShowRecentAssetsMenu(const FName& InAssetEditorName) const;
+
 private:
 
 	/** struct used by OpenedEditorTimes map to store editor names and times */
@@ -393,4 +424,10 @@ private:
 	 * Note that this does not disable or unregister disallowed modes, it simply removes them from the query results.
 	 */
 	FNamePermissionList AllowedEditorModes;
+
+	/** MRU list of recently opened assets */
+	TUniquePtr<FMainMRUFavoritesList> RecentAssetsList;
+	
+	/** Map keeping track of the asset editor each recent asset was opened in */
+	TMap<FString, FString> RecentAssetToAssetEditorMap;
 };
diff --git a/Engine/Source/Editor/UnrealEd/Public/Toolkits/AssetEditorToolkit.h b/Engine/Source/Editor/UnrealEd/Public/Toolkits/AssetEditorToolkit.h
--- a/Engine/Source/Editor/UnrealEd/Public/Toolkits/AssetEditorToolkit.h
+++ b/Engine/Source/Editor/UnrealEd/Public/Toolkits/AssetEditorToolkit.h
@@ -143,10 +143,18 @@ public:
 	virtual bool CloseWindow() override;
 	virtual bool IsPrimaryEditor() const override { return true; };
 	virtual void InvokeTab(const FTabId& TabId) override;
+	virtual FName GetEditingAssetTypeName() const override;
 	virtual TSharedPtr<FTabManager> GetAssociatedTabManager() override;
 	virtual double GetLastActivationTime() override;
 	virtual void RemoveEditingAsset(UObject* Asset) override;
 
+	/**
+	 * Fills in the supplied menu with commands for working with this asset file
+	 *
+	 * @param	MenuBuilder		The menu to add commands to
+	 */
+	void FillDefaultFileMenuOpenCommands(FToolMenuSection& InSection);
+
 	/**
 	 * Fills in the supplied menu with commands for working with this asset file
 	 *
diff --git a/Engine/Source/Editor/UnrealEd/Public/Toolkits/SimpleAssetEditor.h b/Engine/Source/Editor/UnrealEd/Public/Toolkits/SimpleAssetEditor.h
--- a/Engine/Source/Editor/UnrealEd/Public/Toolkits/SimpleAssetEditor.h
+++ b/Engine/Source/Editor/UnrealEd/Public/Toolkits/SimpleAssetEditor.h
@@ -41,6 +41,7 @@ public:
 	virtual FLinearColor GetWorldCentricTabColorScale() const override;
 	virtual bool IsPrimaryEditor() const override { return true; }
 	virtual bool IsSimpleAssetEditor() const override { return true; }
+	virtual FName GetEditingAssetTypeName() const override;
 
 	/** FAssetEditorToolkit interface */
 	virtual void PostRegenerateMenusAndToolbars() override;
diff --git a/Engine/Source/Runtime/SlateCore/Private/Styling/StarshipCoreStyle.cpp b/Engine/Source/Runtime/SlateCore/Private/Styling/StarshipCoreStyle.cpp
--- a/Engine/Source/Runtime/SlateCore/Private/Styling/StarshipCoreStyle.cpp
+++ b/Engine/Source/Runtime/SlateCore/Private/Styling/StarshipCoreStyle.cpp
@@ -410,6 +410,7 @@ TSharedRef<ISlateStyle> FStarshipCoreStyle::Create()
 		Style->Set("Icons.FlipHorizontal", new IMAGE_BRUSH_SVG("Starship/Common/FlipHorizontal", Icon16x16));
 		Style->Set("Icons.FlipVertical", new IMAGE_BRUSH_SVG("Starship/Common/FlipVertical", Icon16x16));
 		Style->Set("Icons.Layout", new IMAGE_BRUSH_SVG("Starship/Common/Layout", Icon16x16));
+		Style->Set("Icons.Recent", new IMAGE_BRUSH_SVG("Starship/Common/Recent", Icon16x16));
 
 		Style->Set("Icons.BadgeModified", new IMAGE_BRUSH_SVG("Starship/Common/badge-modified", Icon16x16));
 		
