Add a SLATE_ATTRIBUTE_DEPRECATED macro and deprecate unused attributes in SViewportToolbarComboMenu

diff --git a/Engine/Plugins/Experimental/Mutable/Source/CustomizableObjectEditor/Private/MuCOE/SCustomizableObjectEditorViewportToolBar.cpp b/Engine/Plugins/Experimental/Mutable/Source/CustomizableObjectEditor/Private/MuCOE/SCustomizableObjectEditorViewportToolBar.cpp
--- a/Engine/Plugins/Experimental/Mutable/Source/CustomizableObjectEditor/Private/MuCOE/SCustomizableObjectEditorViewportToolBar.cpp
+++ b/Engine/Plugins/Experimental/Mutable/Source/CustomizableObjectEditor/Private/MuCOE/SCustomizableObjectEditorViewportToolBar.cpp
@@ -922,7 +922,6 @@ TSharedRef<SWidget> SCustomizableObjectEditorViewportToolBar::GenerateRTSButtons
 		// Setup a GridSnapSetting with the UICommand
 		ToolbarBuilder.AddWidget(SNew(SViewportToolBarComboMenu)
 			.Cursor(EMouseCursor::Default)
-			.Style(ToolBarStyle)
 			.IsChecked(this, &SCustomizableObjectEditorViewportToolBar::IsRotationGridSnapChecked)
 			.OnCheckStateChanged(this, &SCustomizableObjectEditorViewportToolBar::HandleToggleRotationGridSnap)
 			.Label(this, &SCustomizableObjectEditorViewportToolBar::GetRotationGridLabel)
diff --git a/Engine/Plugins/Experimental/UVEditor/Source/UVEditor/Private/SUVEditor2DViewportToolBar.cpp b/Engine/Plugins/Experimental/UVEditor/Source/UVEditor/Private/SUVEditor2DViewportToolBar.cpp
--- a/Engine/Plugins/Experimental/UVEditor/Source/UVEditor/Private/SUVEditor2DViewportToolBar.cpp
+++ b/Engine/Plugins/Experimental/UVEditor/Source/UVEditor/Private/SUVEditor2DViewportToolBar.cpp
@@ -173,7 +173,6 @@ TSharedRef<SWidget> SUVEditor2DViewportToolBar::MakeTransformToolBar(const TShar
 		// Setup a GridSnapSetting with the UICommand
 		ToolbarBuilder.AddWidget(
 			SNew(SViewportToolBarComboMenu)
-			.Style(ToolBarStyle)
 			.IsChecked(this, &SUVEditor2DViewportToolBar::IsLocationGridSnapChecked)
 			.OnCheckStateChanged(this, &SUVEditor2DViewportToolBar::HandleToggleLocationGridSnap)
 			.Label(this, &SUVEditor2DViewportToolBar::GetLocationGridLabel)
@@ -214,7 +213,6 @@ TSharedRef<SWidget> SUVEditor2DViewportToolBar::MakeTransformToolBar(const TShar
 		// Setup a GridSnapSetting with the UICommand
 		ToolbarBuilder.AddWidget(
 			SNew(SViewportToolBarComboMenu)
-			.Style(ToolBarStyle)
 			.IsChecked(this, &SUVEditor2DViewportToolBar::IsRotationGridSnapChecked)
 			.OnCheckStateChanged(this, &SUVEditor2DViewportToolBar::HandleToggleRotationGridSnap)
 			.Label(this, &SUVEditor2DViewportToolBar::GetRotationGridLabel)
@@ -255,7 +253,6 @@ TSharedRef<SWidget> SUVEditor2DViewportToolBar::MakeTransformToolBar(const TShar
 		ToolbarBuilder.AddWidget(
 			SNew(SViewportToolBarComboMenu)
 			.Cursor(EMouseCursor::Default)
-			.Style(ToolBarStyle)
 			.IsChecked(this, &SUVEditor2DViewportToolBar::IsScaleGridSnapChecked)
 			.OnCheckStateChanged(this, &SUVEditor2DViewportToolBar::HandleToggleScaleGridSnap)
 			.Label(this, &SUVEditor2DViewportToolBar::GetScaleGridLabel)
diff --git a/Engine/Source/Editor/UMGEditor/Private/Designer/SDesignerToolBar.cpp b/Engine/Source/Editor/UMGEditor/Private/Designer/SDesignerToolBar.cpp
--- a/Engine/Source/Editor/UMGEditor/Private/Designer/SDesignerToolBar.cpp
+++ b/Engine/Source/Editor/UMGEditor/Private/Designer/SDesignerToolBar.cpp
@@ -48,7 +48,6 @@ TSharedRef< SWidget > SDesignerToolBar::MakeToolBar(const TSharedPtr< FExtender
 		FUICommandInfo* ToggleLocalizationPreviewCommand = FDesignerCommands::Get().ToggleLocalizationPreview.Get();
 
 		ToolbarBuilder.AddWidget(SNew(SViewportToolBarComboMenu)
-			.BlockLocation(EMultiBlockLocation::Start)
 			.Cursor(EMouseCursor::Default)
 			.IsChecked(this, &SDesignerToolBar::IsLocalizationPreviewChecked)
 			.OnCheckStateChanged(this, &SDesignerToolBar::HandleToggleLocalizationPreview)
@@ -90,8 +89,6 @@ TSharedRef< SWidget > SDesignerToolBar::MakeToolBar(const TSharedPtr< FExtender
 
 		// Setup a GridSnapSetting with the UICommand
 		ToolbarBuilder.AddWidget(SNew(SViewportToolBarComboMenu)
-			.Style(ToolBarStyle)
-			.BlockLocation(EMultiBlockLocation::Start)
 			.Cursor(EMouseCursor::Default)
 			.IsChecked(this, &SDesignerToolBar::IsLocationGridSnapChecked)
 			.OnCheckStateChanged(this, &SDesignerToolBar::HandleToggleLocationGridSnap)
diff --git a/Engine/Source/Editor/UnrealEd/Private/STransformViewportToolbar.cpp b/Engine/Source/Editor/UnrealEd/Private/STransformViewportToolbar.cpp
--- a/Engine/Source/Editor/UnrealEd/Private/STransformViewportToolbar.cpp
+++ b/Engine/Source/Editor/UnrealEd/Private/STransformViewportToolbar.cpp
@@ -255,7 +255,6 @@ TSharedRef< SWidget > STransformViewportToolBar::MakeTransformToolBar( const TSh
 		// Setup a GridSnapSetting with the UICommand
 		ToolbarBuilder.AddWidget(	
 			SNew(SViewportToolBarComboMenu)
-				.Style(ToolBarStyle)
 				.IsChecked(this, &STransformViewportToolBar::IsLocationGridSnapChecked)
 				.OnCheckStateChanged(this, &STransformViewportToolBar::HandleToggleLocationGridSnap)
 				.Label(this, &STransformViewportToolBar::GetLocationGridLabel)
@@ -296,7 +295,6 @@ TSharedRef< SWidget > STransformViewportToolBar::MakeTransformToolBar( const TSh
 		// Setup a GridSnapSetting with the UICommand
 		ToolbarBuilder.AddWidget(	
 			SNew(SViewportToolBarComboMenu)
-				.Style(ToolBarStyle)
 				.IsChecked(this, &STransformViewportToolBar::IsRotationGridSnapChecked)
 				.OnCheckStateChanged(this, &STransformViewportToolBar::HandleToggleRotationGridSnap)
 				.Label(this, &STransformViewportToolBar::GetRotationGridLabel)
@@ -334,7 +332,6 @@ TSharedRef< SWidget > STransformViewportToolBar::MakeTransformToolBar( const TSh
 
 		TSharedRef<SWidget> SnapLayerPickerWidget =
 			SNew(SViewportToolBarComboMenu)
-			.Style(ToolBarStyle)
 			.Visibility(this, &STransformViewportToolBar::IsLayer2DSnapVisible)
 			.IsChecked(this, &STransformViewportToolBar::IsLayer2DSnapChecked)
 			.OnCheckStateChanged(this, &STransformViewportToolBar::HandleToggleLayer2DSnap)
@@ -384,7 +381,6 @@ TSharedRef< SWidget > STransformViewportToolBar::MakeTransformToolBar( const TSh
 		ToolbarBuilder.AddWidget(	
 			SNew(SViewportToolBarComboMenu)
 				.Cursor( EMouseCursor::Default )
-				.Style(ToolBarStyle)
 				.IsChecked(this,&STransformViewportToolBar::IsScaleGridSnapChecked)
 				.OnCheckStateChanged(this, &STransformViewportToolBar::HandleToggleScaleGridSnap)
 				.Label(this ,&STransformViewportToolBar::GetScaleGridLabel)
diff --git a/Engine/Source/Editor/UnrealEd/Public/SViewportToolBarComboMenu.h b/Engine/Source/Editor/UnrealEd/Public/SViewportToolBarComboMenu.h
--- a/Engine/Source/Editor/UnrealEd/Public/SViewportToolBarComboMenu.h
+++ b/Engine/Source/Editor/UnrealEd/Public/SViewportToolBarComboMenu.h
@@ -53,7 +53,7 @@ public:
 		SLATE_ATTRIBUTE( FText, Label )
 
 		/** Overall style */
-		SLATE_ATTRIBUTE( FName, Style )
+		SLATE_ATTRIBUTE_DEPRECATED( FName, Style, 5.3, "The Style attribute is deprecated, styles are specified through stylesheets now" )
 
 		/** ToolTip shown on the menu button */
 		SLATE_ATTRIBUTE( FText, MenuButtonToolTip )
@@ -62,7 +62,7 @@ public:
 		SLATE_ATTRIBUTE( FText, ToggleButtonToolTip )
 
 		/** The button location */
-		SLATE_ARGUMENT( EMultiBlockLocation::Type, BlockLocation )
+		SLATE_ARGUMENT_DEPRECATED( EMultiBlockLocation::Type, BlockLocation, 5.3, "The BlockLocation argument is deprecated and not in use anymore" )
 
 		/** The minimum desired width of the menu button contents */
 		SLATE_ARGUMENT( float, MinDesiredButtonWidth )
diff --git a/Engine/Source/Runtime/SlateCore/Public/Widgets/DeclarativeSyntaxSupport.h b/Engine/Source/Runtime/SlateCore/Public/Widgets/DeclarativeSyntaxSupport.h
--- a/Engine/Source/Runtime/SlateCore/Public/Widgets/DeclarativeSyntaxSupport.h
+++ b/Engine/Source/Runtime/SlateCore/Public/Widgets/DeclarativeSyntaxSupport.h
@@ -194,6 +194,14 @@ template<typename WidgetType> struct TSlateBaseNamedArgs;
 		SLATE_PRIVATE_ATTRIBUTE_VARIABLE( AttrType, AttrName ); \
 		SLATE_PRIVATE_ATTRIBUTE_FUNCTION( AttrType, AttrName )
 
+/**
+ * Use when an argument used to be declared SLATE_ATTRIBUTE and it should no longer be used
+ */
+#define SLATE_ATTRIBUTE_DEPRECATED( AttrType, AttrName, DeprecationVersion, DeprecationMessage) \
+		SLATE_PRIVATE_ATTRIBUTE_VARIABLE( AttrType, AttrName ); \
+		UE_DEPRECATED(DeprecationVersion, DeprecationMessage) \
+		SLATE_PRIVATE_ATTRIBUTE_FUNCTION( AttrType, AttrName )
+
 /**
  * Use this macro to declare a slate argument.
  * Arguments differ from attributes in that they can only be values
