Fix missing style warning when opening a widget blueprint

diff --git a/Engine/Source/Editor/UnrealEd/Private/SViewportToolBarComboMenu.cpp b/Engine/Source/Editor/UnrealEd/Private/SViewportToolBarComboMenu.cpp
--- a/Engine/Source/Editor/UnrealEd/Private/SViewportToolBarComboMenu.cpp
+++ b/Engine/Source/Editor/UnrealEd/Private/SViewportToolBarComboMenu.cpp
@@ -25,10 +25,6 @@ struct FPointerEvent;
 
 void SViewportToolBarComboMenu::Construct( const FArguments& InArgs )
 {
-	const FToolBarStyle& ToolBarStyle = FAppStyle::Get().GetWidgetStyle<FToolBarStyle>(InArgs._Style.Get());
-
-	EMultiBlockLocation::Type BlockLocation = InArgs._BlockLocation;
-
 	const FButtonStyle& ButtonStyle = FAppStyle::Get().GetWidgetStyle<FButtonStyle>("EditorViewportToolBar.ComboMenu.ButtonStyle");
 	const FCheckBoxStyle& CheckBoxStyle = FAppStyle::Get().GetWidgetStyle<FCheckBoxStyle>("EditorViewportToolBar.ToggleButton.Start");
 	const FTextBlockStyle& LabelStyle = FAppStyle::Get().GetWidgetStyle<FTextBlockStyle>("EditorViewportToolBar.ComboMenu.LabelStyle");
