remove the non-bulk data chunks since they don't have hashes and will show up as no changes even when they do.

diff --git a/Engine/Source/Editor/UnrealEd/Private/Commandlets/DiffAssetBulkDataCommandlet.cpp b/Engine/Source/Editor/UnrealEd/Private/Commandlets/DiffAssetBulkDataCommandlet.cpp
--- a/Engine/Source/Editor/UnrealEd/Private/Commandlets/DiffAssetBulkDataCommandlet.cpp
+++ b/Engine/Source/Editor/UnrealEd/Private/Commandlets/DiffAssetBulkDataCommandlet.cpp
@@ -390,12 +390,19 @@ int32 UDiffAssetBulkDataCommandlet::Main(const FString& FullCommandLine)
 	}
 	
 	int32 TotalNewChunks = 0, TotalChangedChunks = 0, TotalDeletedChunks = 0;
-	UE_LOG(LogDiffAssetBulk, Display, TEXT("Modifications By IoStore Chunk:"));
+	UE_LOG(LogDiffAssetBulk, Display, TEXT("Modifications By IoStore Chunk (only bulk data tracked atm):"));
 	UE_LOG(LogDiffAssetBulk, Display, TEXT(""));
 	UE_LOG(LogDiffAssetBulk, Display, TEXT("    ChunkType                   New    Deleted    Changed"));
 	for (uint32 ChunkTypeIndex = 0; ChunkTypeIndex < (uint8)EIoChunkType::MAX; ChunkTypeIndex++)
 	{
 		EIoChunkType ChunkType = (EIoChunkType)ChunkTypeIndex;
+		if (ChunkType != EIoChunkType::BulkData &&
+			ChunkType != EIoChunkType::OptionalBulkData &&
+			ChunkType != EIoChunkType::MemoryMappedBulkData)
+		{
+			continue;
+		}
+
 		const TArray<FName>& NewChunksForType = NewChunksByType[ChunkTypeIndex];
 		const TArray<FName>& DeletedChunksForType = DeletedChunksByType[ChunkTypeIndex];
 		const TArray<FName>& ChangedChunksForType = ChangedChunksByType[ChunkTypeIndex];
