Makes the editor viewport screen percentage more intuitive

diff --git a/Engine/Source/Editor/LevelEditor/Private/SLevelViewportToolBar.cpp b/Engine/Source/Editor/LevelEditor/Private/SLevelViewportToolBar.cpp
--- a/Engine/Source/Editor/LevelEditor/Private/SLevelViewportToolBar.cpp
+++ b/Engine/Source/Editor/LevelEditor/Private/SLevelViewportToolBar.cpp
@@ -51,6 +51,7 @@
 #include "SLevelViewport.h"
 #include "SortHelper.h"
 #include "Interfaces/IMainFrameModule.h"
+#include "SCommonEditorViewportToolbarBase.h"
 
 #define LOCTEXT_NAMESPACE "LevelViewportToolBar"
 
@@ -733,7 +734,12 @@ void SLevelViewportToolBar::FillOptionsMenu(UToolMenu* Menu)
 				Section.AddEntry(FToolMenuEntry::InitWidget("FarViewPlane", GenerateFarViewPlaneMenu(), LOCTEXT("FarViewPlane", "Far View Plane")));
 			}
 
-			Section.AddEntry(FToolMenuEntry::InitWidget("ScreenPercentage", GenerateScreenPercentageMenu(), LOCTEXT("ScreenPercentage", "Screen Percentage")));
+			FEditorViewportClient& ViewportClient = Viewport.Pin()->GetLevelViewportClient();
+			Section.AddSubMenu(
+				"ScreenPercentageSubMenu",
+				LOCTEXT("ScreenPercentageSubMenu", "Screen Percentage"),
+				LOCTEXT("ScreenPercentageSubMenu_ToolTip", "Customize the viewport's screen percentage"),
+				FNewMenuDelegate::CreateStatic(&SCommonEditorViewportToolbarBase::ConstructScreenPercentageMenu, &ViewportClient));
 		}
 
 		{
@@ -1353,53 +1359,6 @@ void SLevelViewportToolBar::OnFOVValueChanged( float NewValue )
 	ViewportClient.Invalidate();
 }
 
-TSharedRef<SWidget> SLevelViewportToolBar::GenerateScreenPercentageMenu() const
-{
-	const int32 PreviewScreenPercentageMin = ISceneViewFamilyScreenPercentage::kMinTSRResolutionFraction * 100.0f;
-	const int32 PreviewScreenPercentageMax = ISceneViewFamilyScreenPercentage::kMaxTSRResolutionFraction * 100.0f;
-
-	return
-		SNew(SBox)
-		.HAlign(HAlign_Right)
-		.IsEnabled(this, &SLevelViewportToolBar::OnScreenPercentageIsEnabled)
-		[
-			SNew(SBox)
-			.Padding(FMargin(4.0f, 0.0f, 0.0f, 0.0f))
-			.WidthOverride(100.0f)
-			[
-				SNew ( SBorder )
-				.BorderImage(FAppStyle::Get().GetBrush("Menu.WidgetBorder"))
-				.Padding(FMargin(1.0f))
-				[
-					SNew(SSpinBox<int32>)
-					.Style(&FAppStyle::Get(), "Menu.SpinBox")
-					.Font(FAppStyle::GetFontStyle(TEXT("MenuItem.Font")))
-					.MinSliderValue(PreviewScreenPercentageMin)
-					.MaxSliderValue(PreviewScreenPercentageMax)
-					.Value(this, &SLevelViewportToolBar::OnGetScreenPercentageValue)
-					.OnValueChanged(const_cast<SLevelViewportToolBar*>(this), &SLevelViewportToolBar::OnScreenPercentageValueChanged)
-				]
-			]
-		];
-}
-
-int32 SLevelViewportToolBar::OnGetScreenPercentageValue() const
-{
-	return Viewport.Pin()->GetLevelViewportClient().GetPreviewScreenPercentage();
-}
-
-bool SLevelViewportToolBar::OnScreenPercentageIsEnabled() const
-{
-	return Viewport.Pin()->GetLevelViewportClient().SupportsPreviewResolutionFraction();
-}
-
-void SLevelViewportToolBar::OnScreenPercentageValueChanged(int32 NewValue)
-{
-	FLevelEditorViewportClient& ViewportClient = Viewport.Pin()->GetLevelViewportClient();
-	ViewportClient.SetPreviewScreenPercentage(NewValue);
-	ViewportClient.Invalidate();
-}
-
 TSharedRef<SWidget> SLevelViewportToolBar::GenerateFarViewPlaneMenu() const
 {
 	return
diff --git a/Engine/Source/Editor/LevelEditor/Private/SLevelViewportToolBar.h b/Engine/Source/Editor/LevelEditor/Private/SLevelViewportToolBar.h
--- a/Engine/Source/Editor/LevelEditor/Private/SLevelViewportToolBar.h
+++ b/Engine/Source/Editor/LevelEditor/Private/SLevelViewportToolBar.h
@@ -214,20 +214,6 @@ private:
 	/** Called when the FOV slider is adjusted in the perspective viewport */
 	void OnFOVValueChanged( float NewValue );
 
-	/**
-	* @return The widget containing the screen percentage.
-	*/
-	TSharedRef<SWidget> GenerateScreenPercentageMenu() const;
-
-	/** Called by the ScreenPercentage slider */
-	int32 OnGetScreenPercentageValue() const;
-
-	/** Called by the ScreenPercentage slider */
-	bool OnScreenPercentageIsEnabled() const;
-
-	/** Called when the ScreenPercentage slider is adjusted in the viewport */
-	void OnScreenPercentageValueChanged(int32 NewValue);
-
 	/**
 	 * @return The widget containing the far view plane slider.
 	 */
diff --git a/Engine/Source/Editor/UnrealEd/Private/EditorViewportClient.cpp b/Engine/Source/Editor/UnrealEd/Private/EditorViewportClient.cpp
--- a/Engine/Source/Editor/UnrealEd/Private/EditorViewportClient.cpp
+++ b/Engine/Source/Editor/UnrealEd/Private/EditorViewportClient.cpp
@@ -2811,6 +2811,16 @@ float FEditorViewportClient::GetDefaultPrimaryResolutionFractionTarget() const
 	return StaticHeuristic.ResolveResolutionFraction();
 }
 
+bool FEditorViewportClient::IsPreviewingScreenPercentage() const
+{
+	return bIsPreviewingResolutionFraction;
+}
+
+void FEditorViewportClient::SetPreviewingScreenPercentage(bool bIsPreviewing)
+{
+	bIsPreviewingResolutionFraction = bIsPreviewing;
+}
+
 int32 FEditorViewportClient::GetPreviewScreenPercentage() const
 {
 	float ResolutionFraction = 1.0f;
@@ -2832,23 +2842,7 @@ int32 FEditorViewportClient::GetPreviewScreenPercentage() const
 
 void FEditorViewportClient::SetPreviewScreenPercentage(int32 PreviewScreenPercentage)
 {
-	float AutoResolutionFraction = GetDefaultPrimaryResolutionFractionTarget();
-	int32 AutoScreenPercentage = FMath::RoundToInt(FMath::Clamp(
-		AutoResolutionFraction,
-		ISceneViewFamilyScreenPercentage::kMinTSRResolutionFraction,
-		ISceneViewFamilyScreenPercentage::kMaxTSRResolutionFraction) * 100.0f);
-
-	float NewResolutionFraction = PreviewScreenPercentage / 100.0f;
-	if (NewResolutionFraction >= ISceneViewFamilyScreenPercentage::kMinTSRResolutionFraction &&
-		NewResolutionFraction <= ISceneViewFamilyScreenPercentage::kMaxTSRResolutionFraction &&
-		PreviewScreenPercentage != AutoScreenPercentage)
-	{
-		PreviewResolutionFraction = NewResolutionFraction;
-	}
-	else
-	{
-		PreviewResolutionFraction.Reset();
-	}
+	PreviewResolutionFraction = PreviewScreenPercentage / 100.0f;
 }
 
 bool FEditorViewportClient::SupportsLowDPIPreview() const
@@ -4137,7 +4131,7 @@ void FEditorViewportClient::Draw(FViewport* InViewport, FCanvas* Canvas)
 			if ((!bStereoRendering || bInVREditViewMode) && 
 				SupportsPreviewResolutionFraction() && ViewFamily.SupportsScreenPercentage())
 			{
-				if (PreviewResolutionFraction.IsSet())
+				if (PreviewResolutionFraction.IsSet() && bIsPreviewingResolutionFraction)
 				{
 					GlobalResolutionFraction = PreviewResolutionFraction.GetValue();
 				}
diff --git a/Engine/Source/Editor/UnrealEd/Private/EditorViewportCommands.cpp b/Engine/Source/Editor/UnrealEd/Private/EditorViewportCommands.cpp
--- a/Engine/Source/Editor/UnrealEd/Private/EditorViewportCommands.cpp
+++ b/Engine/Source/Editor/UnrealEd/Private/EditorViewportCommands.cpp
@@ -173,6 +173,10 @@ void FEditorViewportCommands::RegisterCommands()
 	UI_COMMAND( ToggleAutoExposure, "Auto", "If enabled, enables automatic exposure", EUserInterfaceActionType::ToggleButton, FInputChord() );
 	UI_COMMAND( ToggleInGameExposure, "Game Settings", "If enabled, uses game settings", EUserInterfaceActionType::ToggleButton, FInputChord() );
 
+	UI_COMMAND(ToggleOverrideViewportScreenPercentage, "Custom Override", "Overrides the screen percentage of this viewport with custom value", EUserInterfaceActionType::ToggleButton, FInputChord());
+	UI_COMMAND(OpenEditorPerformanceProjectSettings, "Project Settings...", "Opens the project settings for viewport performance", EUserInterfaceActionType::Button, FInputChord());
+	UI_COMMAND(OpenEditorPerformanceEditorPreferences, "Editor Preferences...", "Opens the editor preferences for viewport performance", EUserInterfaceActionType::Button, FInputChord());
+
 	UI_COMMAND(ToggleInViewportContextMenu, "In-Viewport Context Menu", "Shows a contextual menu of key properties and actions for the selected items in the viewport.", EUserInterfaceActionType::ToggleButton, FInputChord(EKeys::Tab));
 }
 
diff --git a/Engine/Source/Editor/UnrealEd/Private/SCommonEditorViewportToolbarBase.cpp b/Engine/Source/Editor/UnrealEd/Private/SCommonEditorViewportToolbarBase.cpp
--- a/Engine/Source/Editor/UnrealEd/Private/SCommonEditorViewportToolbarBase.cpp
+++ b/Engine/Source/Editor/UnrealEd/Private/SCommonEditorViewportToolbarBase.cpp
@@ -6,6 +6,7 @@
 #include "Widgets/Layout/SBox.h"
 #include "Widgets/Input/SSpinBox.h"
 #include "Widgets/Input/STextComboBox.h"
+#include "Widgets/SToolTip.h"
 #include "Styling/AppStyle.h"
 
 #include "STransformViewportToolbar.h"
@@ -16,11 +17,15 @@
 #include "SEditorViewportToolBarButton.h"
 #include "SEditorViewportViewMenu.h"
 #include "Editor/EditorPerformanceSettings.h"
+#include "Settings/EditorProjectSettings.h"
 #include "Scalability.h"
 #include "SceneView.h"
 #include "SScalabilitySettings.h"
 #include "AssetEditorViewportLayout.h"
 #include "SAssetEditorViewport.h"
+#include "ToolMenu.h"
+#include "ToolMenuSection.h"
+
 
 #define LOCTEXT_NAMESPACE "SCommonEditorViewportToolbarBase"
 
@@ -187,6 +192,227 @@ void SCommonEditorViewportToolbarBase::Construct(const FArguments& InArgs, TShar
 	SViewportToolBar::Construct(SViewportToolBar::FArguments());
 }
 
+static FFormatNamedArguments GetScreenPercentageFormatArguments(const FEditorViewportClient& ViewportClient)
+{
+	static auto CVarEditorViewportDefaultScreenPercentageRealTimeMode = IConsoleManager::Get().FindConsoleVariable(TEXT("r.Editor.Viewport.ScreenPercentageMode.RealTime"));
+	static auto CVarEditorViewportDefaultScreenPercentagePathTracerMode = IConsoleManager::Get().FindConsoleVariable(TEXT("r.Editor.Viewport.ScreenPercentageMode.PathTracer"));
+	static auto CVarEditorViewportDefaultScreenPercentageMode = IConsoleManager::Get().FindConsoleVariable(TEXT("r.Editor.Viewport.ScreenPercentageMode.NonRealTime"));
+
+	const UEditorPerformanceProjectSettings* EditorProjectSettings = GetDefault<UEditorPerformanceProjectSettings>();
+	const UEditorPerformanceSettings* EditorUserSettings = GetDefault<UEditorPerformanceSettings>();
+	const FEngineShowFlags& EngineShowFlags = ViewportClient.EngineShowFlags;
+
+	const bool bViewModeSupportsScreenPercentage = ViewportClient.SupportsPreviewResolutionFraction();
+	const bool bIsRealTime = ViewportClient.IsRealtime();
+	const bool bIsPathTraced = EngineShowFlags.PathTracing;
+	const bool bIsPreviewScreenPercentage = ViewportClient.IsPreviewingScreenPercentage();
+
+	int32 DefaultScreenPercentage = FMath::RoundToInt(FMath::Clamp(
+		ViewportClient.GetDefaultPrimaryResolutionFractionTarget(),
+		ISceneViewFamilyScreenPercentage::kMinTSRResolutionFraction,
+		ISceneViewFamilyScreenPercentage::kMaxTSRResolutionFraction) * 100.0f);
+	int32 PreviewScreenPercentage = ViewportClient.GetPreviewScreenPercentage();
+	int32 FinalScreenPercentage = bIsPreviewScreenPercentage ? PreviewScreenPercentage : DefaultScreenPercentage;
+
+	FFormatNamedArguments FormatArguments;
+
+	EScreenPercentageMode ProjectSetting = EScreenPercentageMode::Manual;
+	EEditorUserScreenPercentageModeOverride UserPreference = EEditorUserScreenPercentageModeOverride::ProjectDefault;
+	IConsoleVariable* CVarDefaultScreenPercentage = nullptr;
+	if (bIsPathTraced)
+	{
+		FormatArguments.Add(TEXT("ViewportMode"), LOCTEXT("ScreenPercentage_ViewportMode_PathTraced", "Path Traced"));
+		ProjectSetting = EditorProjectSettings->PathTracerScreenPercentageMode;
+		UserPreference = EditorUserSettings->PathTracerScreenPercentageMode;
+		CVarDefaultScreenPercentage = CVarEditorViewportDefaultScreenPercentagePathTracerMode;
+	}
+	else if (bIsRealTime)
+	{
+		FormatArguments.Add(TEXT("ViewportMode"), LOCTEXT("ScreenPercentage_ViewportMode_Realtime", "Realtime"));
+		ProjectSetting = EditorProjectSettings->RealtimeScreenPercentageMode;
+		UserPreference = EditorUserSettings->RealtimeScreenPercentageMode;
+		CVarDefaultScreenPercentage = CVarEditorViewportDefaultScreenPercentageRealTimeMode;
+	}
+	else
+	{
+		FormatArguments.Add(TEXT("ViewportMode"), LOCTEXT("ScreenPercentage_ViewportMode_NonRealtime", "Non-Realtime"));
+		ProjectSetting = EditorProjectSettings->NonRealtimeScreenPercentageMode;
+		UserPreference = EditorUserSettings->NonRealtimeScreenPercentageMode;
+		CVarDefaultScreenPercentage = CVarEditorViewportDefaultScreenPercentageMode;
+	}
+
+	EScreenPercentageMode FinalScreenPercentageMode = EScreenPercentageMode::Manual;
+	if (!bViewModeSupportsScreenPercentage)
+	{
+		FormatArguments.Add(TEXT("SettingSource"), LOCTEXT("ScreenPercentage_SettingSource_UnsupportedByViewMode", "Unsupported by View mode"));
+		FinalScreenPercentageMode = EScreenPercentageMode::Manual;
+		FinalScreenPercentage = 100;
+	}
+	else if (bIsPreviewScreenPercentage)
+	{
+		FormatArguments.Add(TEXT("SettingSource"), LOCTEXT("ScreenPercentage_SettingSource_ViewportOverride", "Viewport Override"));
+		FinalScreenPercentageMode = EScreenPercentageMode::Manual;
+	}
+	else if ((CVarDefaultScreenPercentage->GetFlags() & ECVF_SetByMask) > ECVF_SetByProjectSetting)
+	{
+		FormatArguments.Add(TEXT("SettingSource"), LOCTEXT("ScreenPercentage_SettingSource_Cvar", "Console Variable"));
+		FinalScreenPercentageMode = EScreenPercentageMode(CVarDefaultScreenPercentage->GetInt());
+	}
+	else if (UserPreference == EEditorUserScreenPercentageModeOverride::ProjectDefault)
+	{
+		FormatArguments.Add(TEXT("SettingSource"), LOCTEXT("ScreenPercentage_SettingSource_ProjectSettigns", "Project Settings"));
+		FinalScreenPercentageMode = ProjectSetting;
+	}
+	else
+	{
+		FormatArguments.Add(TEXT("SettingSource"), LOCTEXT("ScreenPercentage_SettingSource_EditorPreferences", "Editor Preferences"));
+		if (UserPreference == EEditorUserScreenPercentageModeOverride::BasedOnDPIScale)
+		{
+			FinalScreenPercentageMode = EScreenPercentageMode::BasedOnDPIScale;
+		}
+		else if (UserPreference == EEditorUserScreenPercentageModeOverride::BasedOnDisplayResolution)
+		{
+			FinalScreenPercentageMode = EScreenPercentageMode::BasedOnDisplayResolution;
+		}
+		else
+		{
+			FinalScreenPercentageMode = EScreenPercentageMode::Manual;
+		}
+	}
+
+	if (FinalScreenPercentageMode == EScreenPercentageMode::BasedOnDPIScale)
+	{
+		FormatArguments.Add(TEXT("Setting"), LOCTEXT("ScreenPercentage_Setting_BasedOnDPIScale", "Based on OS's DPI scale"));
+	}
+	else if (FinalScreenPercentageMode == EScreenPercentageMode::BasedOnDisplayResolution)
+	{
+		FormatArguments.Add(TEXT("Setting"), LOCTEXT("ScreenPercentage_Setting_BasedOnDisplayResolution", "Based on display resolution"));
+	}
+	else
+	{
+		FormatArguments.Add(TEXT("Setting"), LOCTEXT("ScreenPercentage_Setting_Manual", "Manual"));
+	}
+
+	FormatArguments.Add(TEXT("CurrentScreenPercentage"), FText::FromString(FString::Printf(TEXT("%d"), FinalScreenPercentage)));
+
+	return FormatArguments;
+}
+
+void SCommonEditorViewportToolbarBase::ConstructScreenPercentageMenu(FMenuBuilder& MenuBuilder, FEditorViewportClient* InViewportClient)
+{
+	FEditorViewportClient& ViewportClient = *InViewportClient;
+
+	FMargin CommonPadding(26.0f, 3.0f);
+
+	const int32 PreviewScreenPercentageMin = ISceneViewFamilyScreenPercentage::kMinTSRResolutionFraction * 100.0f;
+	const int32 PreviewScreenPercentageMax = ISceneViewFamilyScreenPercentage::kMaxTSRResolutionFraction * 100.0f;
+
+	const FEditorViewportCommands& BaseViewportCommands = FEditorViewportCommands::Get();
+
+	MenuBuilder.BeginSection("Summary", LOCTEXT("Summary", "Summary"));
+	{
+		MenuBuilder.AddWidget(
+			SNew(SBox)
+			.Padding(CommonPadding)
+			[
+				SNew(STextBlock)
+				.ColorAndOpacity(FSlateColor::UseSubduedForeground())
+				.Text_Lambda([&ViewportClient]() {
+					FFormatNamedArguments FormatArguments = GetScreenPercentageFormatArguments(ViewportClient);
+					return FText::Format(LOCTEXT("ScreenPercentageCurrent_Display", "Current Screen Percentage: {CurrentScreenPercentage}"), FormatArguments);
+				})
+				.ToolTip(SNew(SToolTip).Text(LOCTEXT("ScreenPercentageCurrent_ToolTip", "Current Screen Percentage the viewport is rendered with.")))
+			],
+			FText::GetEmpty()
+		);
+		MenuBuilder.AddWidget(
+			SNew(SBox)
+			.Padding(CommonPadding)
+			[
+				SNew(STextBlock)
+				.ColorAndOpacity(FSlateColor::UseSubduedForeground())
+				.Text_Lambda([&ViewportClient]() {
+					FFormatNamedArguments FormatArguments = GetScreenPercentageFormatArguments(ViewportClient);
+					return FText::Format(LOCTEXT("ScreenPercentageActiveViewport", "Active Viewport: {ViewportMode}"), FormatArguments);
+				})
+			],
+			FText::GetEmpty()
+		);
+		MenuBuilder.AddWidget(
+			SNew(SBox)
+			.Padding(CommonPadding)
+			[
+				SNew(STextBlock)
+				.ColorAndOpacity(FSlateColor::UseSubduedForeground())
+				.Text_Lambda([&ViewportClient]() {
+					FFormatNamedArguments FormatArguments = GetScreenPercentageFormatArguments(ViewportClient);
+					return FText::Format(LOCTEXT("ScreenPercentageSetFrom", "Set From: {SettingSource}"), FormatArguments);
+				})
+			],
+			FText::GetEmpty()
+		);
+		MenuBuilder.AddWidget(
+			SNew(SBox)
+			.Padding(CommonPadding)
+			[
+				SNew(STextBlock)
+				.ColorAndOpacity(FSlateColor::UseSubduedForeground())
+				.Text_Lambda([&ViewportClient]() {
+					FFormatNamedArguments FormatArguments = GetScreenPercentageFormatArguments(ViewportClient);
+					return FText::Format(LOCTEXT("ScreenPercentageSetting", "Setting: {Setting}"), FormatArguments);
+				})
+			],
+			FText::GetEmpty()
+		);
+	}
+	MenuBuilder.EndSection();
+
+	MenuBuilder.BeginSection("ScreenPercentage", LOCTEXT("ScreenPercentage_ViewportOverride", "Viewport Override"));
+	{
+		MenuBuilder.AddMenuEntry(BaseViewportCommands.ToggleOverrideViewportScreenPercentage);
+		MenuBuilder.AddWidget(
+			SNew(SBox)
+			.HAlign(HAlign_Right)
+			.IsEnabled_Lambda([&ViewportClient]() {
+				return ViewportClient.IsPreviewingScreenPercentage() && ViewportClient.SupportsPreviewResolutionFraction();
+			})
+			[
+				SNew(SBox)
+				.Padding(FMargin(4.0f, 0.0f, 0.0f, 0.0f))
+				.WidthOverride(100.0f)
+				[
+					SNew(SBorder)
+					//.BorderImage(FAppStyle::Get().GetBrush("Menu.WidgetBorder"))
+					.Padding(FMargin(1.0f))
+					[
+						SNew(SSpinBox<int32>)
+						.Style(&FAppStyle::Get(), "Menu.SpinBox")
+						.Font(FAppStyle::GetFontStyle(TEXT("MenuItem.Font")))
+						.MinSliderValue(PreviewScreenPercentageMin)
+						.MaxSliderValue(PreviewScreenPercentageMax)
+						.Value_Lambda([&ViewportClient]() {
+							return ViewportClient.GetPreviewScreenPercentage();
+						})
+						.OnValueChanged_Lambda([&ViewportClient](int32 NewValue) {
+							ViewportClient.SetPreviewScreenPercentage(NewValue);
+							ViewportClient.Invalidate();
+						})
+					]
+				]
+			],
+			LOCTEXT("ScreenPercentage", "Screen Percentage")
+		);
+	}
+	MenuBuilder.EndSection();
+
+	MenuBuilder.BeginSection("ScreenPercentageSettings", LOCTEXT("ScreenPercentage_ViewportSettings", "Viewport Settings"));
+	{
+		MenuBuilder.AddMenuEntry(BaseViewportCommands.OpenEditorPerformanceProjectSettings);
+		MenuBuilder.AddMenuEntry(BaseViewportCommands.OpenEditorPerformanceEditorPreferences);
+	}
+	MenuBuilder.EndSection();
+}
+
 void SCommonEditorViewportToolbarBase::UpdateAssetViewerProfileList()
 {
 	if (PreviewProfileController)
@@ -311,7 +537,10 @@ TSharedRef<SWidget> SCommonEditorViewportToolbarBase::GenerateOptionsMenu() cons
 				OptionsMenuBuilder.AddWidget( GenerateFarViewPlaneMenu(), LOCTEXT("FarViewPlane", "Far View Plane") );
 			}
 
-			OptionsMenuBuilder.AddWidget(GenerateScreenPercentageMenu(), LOCTEXT("ScreenPercentage", "Screen Percentage"));
+			OptionsMenuBuilder.AddSubMenu(
+				LOCTEXT("ScreenPercentageSubMenu", "Screen Percentage"),
+				LOCTEXT("ScreenPercentageSubMenu_ToolTip", "Customize the viewport's screen percentage"),
+				FNewMenuDelegate::CreateStatic(&SCommonEditorViewportToolbarBase::ConstructScreenPercentageMenu, &GetViewportClient()));
 		}
 		OptionsMenuBuilder.EndSection();
 
@@ -496,54 +725,6 @@ void SCommonEditorViewportToolbarBase::OnFOVValueChanged(float NewValue) const
 	ViewportClient.Invalidate();
 }
 
-TSharedRef<SWidget> SCommonEditorViewportToolbarBase::GenerateScreenPercentageMenu() const
-{
-	const int32 PreviewScreenPercentageMin = static_cast<int32>(ISceneViewFamilyScreenPercentage::kMinTSRResolutionFraction * 100.0f);
-	const int32 PreviewScreenPercentageMax = static_cast<int32>(ISceneViewFamilyScreenPercentage::kMaxTSRResolutionFraction * 100.0f);
-
-	return
-		SNew(SBox)
-		.HAlign(HAlign_Right)
-		.IsEnabled(this, &SCommonEditorViewportToolbarBase::OnScreenPercentageIsEnabled)
-		[
-			SNew(SBox)
-			.Padding(FMargin(4.0f, 0.0f, 0.0f, 0.0f))
-			.WidthOverride(100.0f)
-			[
-				SNew ( SBorder )
-				.BorderImage(FAppStyle::Get().GetBrush("Menu.WidgetBorder"))
-				.Padding(FMargin(1.0f))
-				[
-					SNew(SSpinBox<int32>)
-					.Style(&FAppStyle::Get(), "Menu.SpinBox")
-					.Font(FAppStyle::GetFontStyle(TEXT("MenuItem.Font")))
-					.MinSliderValue(PreviewScreenPercentageMin)
-					.MaxSliderValue(PreviewScreenPercentageMax)
-					.Value(this, &SCommonEditorViewportToolbarBase::OnGetScreenPercentageValue)
-					.OnValueChanged(const_cast<SCommonEditorViewportToolbarBase*>(this), &SCommonEditorViewportToolbarBase::OnScreenPercentageValueChanged)
-				]
-			]
-		];
-}
-
-int32 SCommonEditorViewportToolbarBase::OnGetScreenPercentageValue() const
-{
-	return GetViewportClient().GetPreviewScreenPercentage();
-}
-
-bool SCommonEditorViewportToolbarBase::OnScreenPercentageIsEnabled() const
-{
-	return GetViewportClient().SupportsPreviewResolutionFraction();
-}
-
-void SCommonEditorViewportToolbarBase::OnScreenPercentageValueChanged(int32 NewValue)
-{
-	FEditorViewportClient& ViewportClient = GetViewportClient();
-	ViewportClient.SetPreviewScreenPercentage(NewValue);
-	ViewportClient.Invalidate();
-}
-
-
 TSharedRef<SWidget> SCommonEditorViewportToolbarBase::GenerateFarViewPlaneMenu() const
 {
 	return
diff --git a/Engine/Source/Editor/UnrealEd/Private/SEditorViewport.cpp b/Engine/Source/Editor/UnrealEd/Private/SEditorViewport.cpp
--- a/Engine/Source/Editor/UnrealEd/Private/SEditorViewport.cpp
+++ b/Engine/Source/Editor/UnrealEd/Private/SEditorViewport.cpp
@@ -24,6 +24,8 @@
 #include "GPUSkinCacheVisualizationMenuCommands.h"
 #include "GPUSkinCache.h"
 #include "Widgets/Colors/SComplexGradient.h"
+#include "Modules/ModuleManager.h"
+#include "ISettingsModule.h"
 
 #define LOCTEXT_NAMESPACE "EditorViewport"
 
@@ -374,6 +376,20 @@ void SEditorViewport::BindCommands()
 		FCanExecuteAction::CreateSP(this, &SEditorViewport::CanToggleInViewportContextMenu)
 	);
 
+	CommandListRef.MapAction(
+		Commands.ToggleOverrideViewportScreenPercentage,
+		FExecuteAction::CreateSP(this, &SEditorViewport::TogglePreviewingScreenPercentage),
+		FCanExecuteAction(),
+		FIsActionChecked::CreateSP(this, &SEditorViewport::IsPreviewingScreenPercentage));
+
+	CommandListRef.MapAction(
+		Commands.OpenEditorPerformanceProjectSettings,
+		FExecuteAction::CreateSP(this, &SEditorViewport::OnOpenViewportPerformanceProjectSettings));
+
+	CommandListRef.MapAction(
+		Commands.OpenEditorPerformanceEditorPreferences,
+		FExecuteAction::CreateSP(this, &SEditorViewport::OnOpenViewportPerformanceEditorPreferences));
+
 	// Simple macro for binding many view mode UI commands
 
 #define MAP_VIEWMODEPARAM_ACTION( ViewModeCommand, ViewModeParam ) \
@@ -774,6 +790,26 @@ bool SEditorViewport::OnIsSurfaceSnapEnabled()
 	return GetDefault<ULevelEditorViewportSettings>()->SnapToSurface.bEnabled;
 }
 
+bool SEditorViewport::IsPreviewingScreenPercentage() const
+{
+	return Client->IsPreviewingScreenPercentage();
+}
+
+void SEditorViewport::TogglePreviewingScreenPercentage()
+{
+	Client->SetPreviewingScreenPercentage(!IsPreviewingScreenPercentage());
+}
+
+void SEditorViewport::OnOpenViewportPerformanceProjectSettings()
+{
+	FModuleManager::LoadModuleChecked<ISettingsModule>("Settings").ShowViewer("Project", "Editor", "EditorPerformanceProjectSettings");
+}
+
+void SEditorViewport::OnOpenViewportPerformanceEditorPreferences()
+{
+	FModuleManager::LoadModuleChecked<ISettingsModule>("Settings").ShowViewer("Editor", "General", "EditorPerformanceSettings");
+}
+
 EActiveTimerReturnType SEditorViewport::EnsureTick( double InCurrentTime, float InDeltaTime )
 {
 	// Keep the timer going if we're realtime or were invalidated this frame
diff --git a/Engine/Source/Editor/UnrealEd/Public/EditorViewportClient.h b/Engine/Source/Editor/UnrealEd/Public/EditorViewportClient.h
--- a/Engine/Source/Editor/UnrealEd/Public/EditorViewportClient.h
+++ b/Engine/Source/Editor/UnrealEd/Public/EditorViewportClient.h
@@ -1423,6 +1423,12 @@ public:
 	/** @return default resolution fraction for UI based on display resolution and user settings. */
 	float GetDefaultPrimaryResolutionFractionTarget() const;
 
+	/** @return whether preview screen percentage for UI. */
+	bool IsPreviewingScreenPercentage() const;
+
+	/** Sets whether preview screen percentage for UI. */
+	void SetPreviewingScreenPercentage(bool bIsPreviewing);
+
 	/** @return preview screen percentage for UI. */
 	int32 GetPreviewScreenPercentage() const;
 
@@ -1984,6 +1990,7 @@ private:
 
 private:
 	/** Controles resolution fraction for previewing in editor viewport at different screen percentage. */
+	bool bIsPreviewingResolutionFraction = false;
 	TOptional<float> PreviewResolutionFraction;
 
 	// DPI mode for scene rendering.
diff --git a/Engine/Source/Editor/UnrealEd/Public/EditorViewportCommands.h b/Engine/Source/Editor/UnrealEd/Public/EditorViewportCommands.h
--- a/Engine/Source/Editor/UnrealEd/Public/EditorViewportCommands.h
+++ b/Engine/Source/Editor/UnrealEd/Public/EditorViewportCommands.h
@@ -221,6 +221,13 @@ public:
 	TSharedPtr< FUICommandInfo > ToggleInGameExposure;
 	/** Toggle the in-viewport context menu */
 	TSharedPtr< FUICommandInfo > ToggleInViewportContextMenu;
+
+	/** Toggle to override the viewport screen percentage */
+	TSharedPtr< FUICommandInfo > ToggleOverrideViewportScreenPercentage;
+	/** Opens the advanced viewport settings */
+	TSharedPtr< FUICommandInfo > OpenEditorPerformanceProjectSettings;
+	TSharedPtr< FUICommandInfo > OpenEditorPerformanceEditorPreferences;
+
 	/**
 	 * Grid commands
 	 */
diff --git a/Engine/Source/Editor/UnrealEd/Public/SCommonEditorViewportToolbarBase.h b/Engine/Source/Editor/UnrealEd/Public/SCommonEditorViewportToolbarBase.h
--- a/Engine/Source/Editor/UnrealEd/Public/SCommonEditorViewportToolbarBase.h
+++ b/Engine/Source/Editor/UnrealEd/Public/SCommonEditorViewportToolbarBase.h
@@ -79,6 +79,9 @@ public:
 
 	void Construct(const FArguments& InArgs, TSharedPtr<class ICommonEditorViewportToolbarInfoProvider> InInfoProvider);
 
+	/** Build the screen percentage menu. */
+	static void ConstructScreenPercentageMenu(FMenuBuilder& MenuBuilder, class FEditorViewportClient* ViewportClient);
+
 private:
 	/**
 	 * Returns the label for the "Camera" tool bar menu, which changes depending on the viewport type
@@ -145,17 +148,6 @@ private:
 	/** Called by the FOV slider in the perspective viewport to get the FOV value */
 	float OnGetFOVValue() const;
 
-	/**
-	* @return The widget containing the screen percentage.
-	*/
-	TSharedRef<SWidget> GenerateScreenPercentageMenu() const;
-
-	/** Called by the ScreenPercentage slider */
-	int32 OnGetScreenPercentageValue() const;
-
-	/** Called by the ScreenPercentage slider */
-	bool OnScreenPercentageIsEnabled() const;
-
 	/**
 	 * @return The widget containing the far view plane slider.
 	 */
diff --git a/Engine/Source/Editor/UnrealEd/Public/SEditorViewport.h b/Engine/Source/Editor/UnrealEd/Public/SEditorViewport.h
--- a/Engine/Source/Editor/UnrealEd/Public/SEditorViewport.h
+++ b/Engine/Source/Editor/UnrealEd/Public/SEditorViewport.h
@@ -105,6 +105,11 @@ public:
 	virtual void UpdateInViewportMenuLocation(const FVector2D InLocation);
 	virtual bool CanToggleInViewportContextMenu() { return false; }
 
+	bool IsPreviewingScreenPercentage() const;
+	void TogglePreviewingScreenPercentage();
+	void OnOpenViewportPerformanceProjectSettings();
+	void OnOpenViewportPerformanceEditorPreferences();
+
 ///////////////////////////////////////////////////////////////////////////////
 // begin feature level control functions block
 ///////////////////////////////////////////////////////////////////////////////
@@ -248,6 +253,7 @@ protected:
 	 */
 	static bool OnIsSurfaceSnapEnabled();
 
+
 protected:
 	TSharedPtr<SOverlay> ViewportOverlay;
 
