WorldPartition - Builder: Ignore empty arrays sent to UWorldPartitionBuilderCommandlet::OnFilesModified()
* Empty arrays could cause an empty list of files to be passed to the SCC providers, causing errors to be logged

diff --git a/Engine/Source/Editor/UnrealEd/Private/Commandlets/WorldPartitionBuilderCommandlet.cpp b/Engine/Source/Editor/UnrealEd/Private/Commandlets/WorldPartitionBuilderCommandlet.cpp
--- a/Engine/Source/Editor/UnrealEd/Private/Commandlets/WorldPartitionBuilderCommandlet.cpp
+++ b/Engine/Source/Editor/UnrealEd/Private/Commandlets/WorldPartitionBuilderCommandlet.cpp
@@ -241,7 +241,11 @@ bool UWorldPartitionBuilderCommandlet::RunBuilder(TSubclassOf<UWorldPartitionBui
 
 bool UWorldPartitionBuilderCommandlet::OnFilesModified(const TArray<FString>& InModifiedFiles, const FString& InChangeDescription)
 {
-	AutoSubmitFiles.Emplace(InChangeDescription, InModifiedFiles);
+	if (!InModifiedFiles.IsEmpty())
+	{
+		AutoSubmitFiles.Emplace(InChangeDescription, InModifiedFiles);
+	}
+
 	return true;
 }
 
