AssetEditorSubsystem: Change the "missing asset editors for recent assets in the config" log  to warning instead of error

diff --git a/Engine/Source/Editor/UnrealEd/Private/Subsystems/AssetEditorSubsystem.cpp b/Engine/Source/Editor/UnrealEd/Private/Subsystems/AssetEditorSubsystem.cpp
--- a/Engine/Source/Editor/UnrealEd/Private/Subsystems/AssetEditorSubsystem.cpp
+++ b/Engine/Source/Editor/UnrealEd/Private/Subsystems/AssetEditorSubsystem.cpp
@@ -125,7 +125,7 @@ void UAssetEditorSubsystem::InitializeRecentAssets()
 	if(RecentAssetsList->GetNumItems() != RecentAssetEditors.Num())
 	{
 		bFoundRecentAssetEditorsCleanly = false;
-		UE_LOG(LogAssetEditorSubsystem, Error, TEXT("Something went wrong while loading recent assets! Num Recent Assets  = %d, Num Recent Asset Editors = %d"), RecentAssetsList->GetNumItems(), RecentAssetEditors.Num());
+		UE_LOG(LogAssetEditorSubsystem, Warning, TEXT("Something went wrong while loading recent assets! Num Recent Assets  = %d, Num Recent Asset Editors = %d"), RecentAssetsList->GetNumItems(), RecentAssetEditors.Num());
 	}
 
 	// If we have corrupt data, simply ignore the asset editor list instead of potentially showing an asset in the wrong asset editor's recents
