Fixes a bug where the viewport screen percentage override could be set way to loo

diff --git a/Engine/Source/Editor/UnrealEd/Private/EditorViewportClient.cpp b/Engine/Source/Editor/UnrealEd/Private/EditorViewportClient.cpp
--- a/Engine/Source/Editor/UnrealEd/Private/EditorViewportClient.cpp
+++ b/Engine/Source/Editor/UnrealEd/Private/EditorViewportClient.cpp
@@ -2842,7 +2842,10 @@ int32 FEditorViewportClient::GetPreviewScreenPercentage() const
 
 void FEditorViewportClient::SetPreviewScreenPercentage(int32 PreviewScreenPercentage)
 {
-	PreviewResolutionFraction = PreviewScreenPercentage / 100.0f;
+	PreviewResolutionFraction = FMath::Clamp(
+		PreviewScreenPercentage / 100.0f,
+		ISceneViewFamilyScreenPercentage::kMinTSRResolutionFraction,
+		ISceneViewFamilyScreenPercentage::kMaxTSRResolutionFraction);
 }
 
 bool FEditorViewportClient::SupportsLowDPIPreview() const
