Polishes the new screen percentage menu based of feedback

* Adds rendering and display resolution
* Adds a decimal point on Current Screen Percentage
* Adds project settings and edifor preferences' icons

diff --git a/Engine/Source/Editor/UnrealEd/Private/SCommonEditorViewportToolbarBase.cpp b/Engine/Source/Editor/UnrealEd/Private/SCommonEditorViewportToolbarBase.cpp
--- a/Engine/Source/Editor/UnrealEd/Private/SCommonEditorViewportToolbarBase.cpp
+++ b/Engine/Source/Editor/UnrealEd/Private/SCommonEditorViewportToolbarBase.cpp
@@ -207,12 +207,12 @@ static FFormatNamedArguments GetScreenPercentageFormatArguments(const FEditorVie
 	const bool bIsPathTraced = EngineShowFlags.PathTracing;
 	const bool bIsPreviewScreenPercentage = ViewportClient.IsPreviewingScreenPercentage();
 
-	int32 DefaultScreenPercentage = FMath::RoundToInt(FMath::Clamp(
+	float DefaultScreenPercentage = FMath::Clamp(
 		ViewportClient.GetDefaultPrimaryResolutionFractionTarget(),
 		ISceneViewFamilyScreenPercentage::kMinTSRResolutionFraction,
-		ISceneViewFamilyScreenPercentage::kMaxTSRResolutionFraction) * 100.0f);
-	int32 PreviewScreenPercentage = ViewportClient.GetPreviewScreenPercentage();
-	int32 FinalScreenPercentage = bIsPreviewScreenPercentage ? PreviewScreenPercentage : DefaultScreenPercentage;
+		ISceneViewFamilyScreenPercentage::kMaxTSRResolutionFraction) * 100.0f;
+	float PreviewScreenPercentage = float(ViewportClient.GetPreviewScreenPercentage());
+	float FinalScreenPercentage = bIsPreviewScreenPercentage ? PreviewScreenPercentage : DefaultScreenPercentage;
 
 	FFormatNamedArguments FormatArguments;
 
@@ -293,7 +293,21 @@ static FFormatNamedArguments GetScreenPercentageFormatArguments(const FEditorVie
 		FormatArguments.Add(TEXT("Setting"), LOCTEXT("ScreenPercentage_Setting_Manual", "Manual"));
 	}
 
-	FormatArguments.Add(TEXT("CurrentScreenPercentage"), FText::FromString(FString::Printf(TEXT("%d"), FinalScreenPercentage)));
+	FormatArguments.Add(TEXT("CurrentScreenPercentage"), FText::FromString(FString::Printf(TEXT("%3.1f"), FMath::RoundToFloat(FinalScreenPercentage * 10.0f) / 10.0f)));
+
+	{
+		float FinalResolutionFraction = (FinalScreenPercentage / 100.0f);
+		FIntPoint DisplayResolution = ViewportClient.Viewport->GetSizeXY();
+		FIntPoint RenderingResolution;
+		RenderingResolution.X = FMath::CeilToInt(DisplayResolution.X * FinalResolutionFraction);
+		RenderingResolution.Y = FMath::CeilToInt(DisplayResolution.Y * FinalResolutionFraction);
+
+		FormatArguments.Add(TEXT("RenderingWidth"), FText::FromString(FString::Printf(TEXT("%d"), RenderingResolution.X)));
+		FormatArguments.Add(TEXT("RenderingHeight"), FText::FromString(FString::Printf(TEXT("%d"), RenderingResolution.Y)));
+		
+		FormatArguments.Add(TEXT("DisplayWidth"), FText::FromString(FString::Printf(TEXT("%d"), DisplayResolution.X)));
+		FormatArguments.Add(TEXT("DisplayHeight"), FText::FromString(FString::Printf(TEXT("%d"), DisplayResolution.Y)));
+	}
 
 	return FormatArguments;
 }
@@ -321,7 +335,21 @@ void SCommonEditorViewportToolbarBase::ConstructScreenPercentageMenu(FMenuBuilde
 					FFormatNamedArguments FormatArguments = GetScreenPercentageFormatArguments(ViewportClient);
 					return FText::Format(LOCTEXT("ScreenPercentageCurrent_Display", "Current Screen Percentage: {CurrentScreenPercentage}"), FormatArguments);
 				})
-				.ToolTip(SNew(SToolTip).Text(LOCTEXT("ScreenPercentageCurrent_ToolTip", "Current Screen Percentage the viewport is rendered with.")))
+				.ToolTip(SNew(SToolTip).Text(LOCTEXT("ScreenPercentageCurrent_ToolTip", "Current Screen Percentage the viewport is rendered with. The primary screen percentage can either be a spatial or temporal upscaler based of your anti-aliasing settings.")))
+			],
+			FText::GetEmpty()
+		);
+
+		MenuBuilder.AddWidget(
+			SNew(SBox)
+			.Padding(CommonPadding)
+			[
+				SNew(STextBlock)
+				.ColorAndOpacity(FSlateColor::UseSubduedForeground())
+				.Text_Lambda([&ViewportClient]() {
+					FFormatNamedArguments FormatArguments = GetScreenPercentageFormatArguments(ViewportClient);
+					return FText::Format(LOCTEXT("ScreenPercentageResolutions", "Resolution: {RenderingWidth}x{RenderingHeight} -> {DisplayWidth}x{DisplayHeight}"), FormatArguments);
+				})
 			],
 			FText::GetEmpty()
 		);
@@ -407,8 +435,16 @@ void SCommonEditorViewportToolbarBase::ConstructScreenPercentageMenu(FMenuBuilde
 
 	MenuBuilder.BeginSection("ScreenPercentageSettings", LOCTEXT("ScreenPercentage_ViewportSettings", "Viewport Settings"));
 	{
-		MenuBuilder.AddMenuEntry(BaseViewportCommands.OpenEditorPerformanceProjectSettings);
-		MenuBuilder.AddMenuEntry(BaseViewportCommands.OpenEditorPerformanceEditorPreferences);
+		MenuBuilder.AddMenuEntry(BaseViewportCommands.OpenEditorPerformanceProjectSettings,
+			/* InExtensionHook = */ NAME_None,
+			/* InLabelOverride = */ TAttribute<FText>(),
+			/* InToolTipOverride = */ TAttribute<FText>(),
+			FSlateIcon(FAppStyle::GetAppStyleSetName(), "ProjectSettings.TabIcon"));
+		MenuBuilder.AddMenuEntry(BaseViewportCommands.OpenEditorPerformanceEditorPreferences,
+			/* InExtensionHook = */ NAME_None,
+			/* InLabelOverride = */ TAttribute<FText>(),
+			/* InToolTipOverride = */ TAttribute<FText>(),
+			FSlateIcon(FAppStyle::GetAppStyleSetName(), "EditorPreferences.TabIcon"));
 	}
 	MenuBuilder.EndSection();
 }
