Fixes Mode switching can be inconsistent and editor modes submenu is sometimes not present.

diff --git a/Engine/Source/Editor/UnrealEd/Private/EditorModeManager.cpp b/Engine/Source/Editor/UnrealEd/Private/EditorModeManager.cpp
--- a/Engine/Source/Editor/UnrealEd/Private/EditorModeManager.cpp
+++ b/Engine/Source/Editor/UnrealEd/Private/EditorModeManager.cpp
@@ -51,6 +51,7 @@
 
 #include "Elements/Interfaces/TypedElementWorldInterface.h"
 #include "TextureResource.h"
+#include "Toolkits/ToolkitManager.h"
 
 /*------------------------------------------------------------------------------
 	FEditorModeTools.
@@ -768,6 +769,12 @@ void FEditorModeTools::DeactivateMode( FEditorModeID InID )
 			ActiveScriptableModes.RemoveAt(Index);
 
 			constexpr bool bIsEnteringMode = false;
+
+			if (const TSharedPtr<FModeToolkit> Toolkit = Mode->GetToolkit().Pin())
+			{
+				FToolkitManager::Get().CloseToolkit(Toolkit.ToSharedRef());
+			}
+			
 			BroadcastEditorModeIDChanged(InID, bIsEnteringMode);
 			break;
 		}
diff --git a/Engine/Source/Editor/UnrealEd/Private/Toolkits/BaseToolkit.cpp b/Engine/Source/Editor/UnrealEd/Private/Toolkits/BaseToolkit.cpp
--- a/Engine/Source/Editor/UnrealEd/Private/Toolkits/BaseToolkit.cpp
+++ b/Engine/Source/Editor/UnrealEd/Private/Toolkits/BaseToolkit.cpp
@@ -717,16 +717,9 @@ void FModeToolkit::OnModeIDChanged(const FEditorModeID& InID, bool bIsEntering)
 {
 	if (const FEditorModeInfo* ModeInfo = GetEditorModeInfo())
 	{
-		if (ModeInfo->ID != NAME_None && ModeInfo->ID == InID)
+		if (ModeInfo->ID != NAME_None && ModeInfo->ID == InID && bIsEntering)
 		{
-			if (bIsEntering)
-			{
-				FToolkitManager::Get().RegisterNewToolkit(SharedThis(this));
-			}
-			else
-			{
-				FToolkitManager::Get().CloseToolkit(SharedThis(this));
-			}
+			FToolkitManager::Get().RegisterNewToolkit(SharedThis(this));
 		}
 	}
 }
