Fix cached editor thumbnails for BPGC assets
Expose name of CachedEditorThumbnails.bin file

diff --git a/Engine/Source/Editor/UnrealEd/Private/ThumbnailExternalCache.cpp b/Engine/Source/Editor/UnrealEd/Private/ThumbnailExternalCache.cpp
--- a/Engine/Source/Editor/UnrealEd/Private/ThumbnailExternalCache.cpp
+++ b/Engine/Source/Editor/UnrealEd/Private/ThumbnailExternalCache.cpp
@@ -24,7 +24,6 @@ namespace ThumbnailExternalCache
 {
 	const int64 LatestVersion = 0;
 	const uint64 ExpectedHeaderId = 0x424d5548545f4555; // "UE_THUMB"
-	const FString ThumbnailFilenamePart(TEXT("CachedEditorThumbnails.bin"));
 	const FString ThumbnailImageFormatName(TEXT(""));
 
 	void ResizeThumbnailImage(FObjectThumbnail& Thumbnail, const int32 NewWidth, const int32 NewHeight)
@@ -111,7 +110,6 @@ public:
 	FSaveThumbnailCache();
 	~FSaveThumbnailCache();
 
-	void Save(const FString& InFilename, const TArrayView<FAssetData> InAssetDatas, const FThumbnailExternalCacheSettings& InSettings);
 	void Save(FArchive& Ar, const TArrayView<FAssetData> InAssetDatas, const FThumbnailExternalCacheSettings& InSettings);
 };
 
@@ -123,15 +121,6 @@ FSaveThumbnailCache::~FSaveThumbnailCache()
 {
 }
 
-void FSaveThumbnailCache::Save(const FString& InFilename, const TArrayView<FAssetData> InAssetDatas, const FThumbnailExternalCacheSettings& InSettings)
-{
-	if (TUniquePtr<FArchive> FileWriter = TUniquePtr<FArchive>(IFileManager::Get().CreateFileWriter(*InFilename)))
-	{
-		Save(*FileWriter, InAssetDatas, InSettings);
-		return;
-	}
-}
-
 void FSaveThumbnailCache::Save(FArchive& Ar, const TArrayView<FAssetData> InAssetDatas, const FThumbnailExternalCacheSettings& InSettings)
 {
 	// Reduce peak memory to support larger asset counts by loading then saving in batches
@@ -319,6 +308,12 @@ FThumbnailExternalCache& FThumbnailExternalCache::Get()
 	return ThumbnailExternalCache;
 }
 
+const FString& FThumbnailExternalCache::GetCachedEditorThumbnailsFilename()
+{
+	static const FString Filename = TEXT("CachedEditorThumbnails.bin");
+	return Filename;
+}
+
 void FThumbnailExternalCache::Init()
 {
 	if (!bHasInit)
@@ -326,7 +321,7 @@ void FThumbnailExternalCache::Init()
 		bHasInit = true;
 
 		// Load file for project
-		LoadCacheFileIndex(FPaths::ProjectDir() / ThumbnailExternalCache::ThumbnailFilenamePart);
+		LoadCacheFileIndex(FPaths::ProjectDir() / FThumbnailExternalCache::GetCachedEditorThumbnailsFilename());
 
 		// Load any thumbnail files for content plugins
 		TArray<TSharedRef<IPlugin>> ContentPlugins = IPluginManager::Get().GetEnabledPluginsWithContent();
@@ -366,7 +361,7 @@ bool FThumbnailExternalCache::LoadThumbnailsFromExternalCache(const TSet<FName>&
 		return false;
 	}
 
-	static const FString BlueprintGeneratedClassPrefix = TEXT("BlueprintGeneratedClass ");
+	static const FString BlueprintGeneratedClassPrefix = TEXT("/Script/Engine.BlueprintGeneratedClass ");
 
 	int32 NumLoaded = 0;
 	for (const FName ObjectFullName : InObjectFullNames)
@@ -381,7 +376,7 @@ bool FThumbnailExternalCache::LoadThumbnailsFromExternalCache(const TSet<FName>&
 		{
 			// Look for the thumbnail of the Blueprint version of this object instead
 			FNameBuilder ModifiedNameBuilder;
-			ModifiedNameBuilder.Append(TEXT("Blueprint "));
+			ModifiedNameBuilder.Append(TEXT("/Script/Engine.Blueprint "));
 			FStringView ViewToAppend = NameView;
 			ViewToAppend.RightChopInline(BlueprintGeneratedClassPrefix.Len());
 			ViewToAppend.LeftChopInline(2);
@@ -488,7 +483,7 @@ void FThumbnailExternalCache::OnContentPathDismounted(const FString& InAssetPath
 	{
 		if (FoundPlugin->CanContainContent())
 		{
-			FString Filename = FoundPlugin->GetBaseDir() / ThumbnailExternalCache::ThumbnailFilenamePart;
+			const FString Filename = FoundPlugin->GetBaseDir() / FThumbnailExternalCache::GetCachedEditorThumbnailsFilename();
 			CacheFiles.Remove(Filename);
 		}
 	}
@@ -498,7 +493,7 @@ void FThumbnailExternalCache::LoadCacheFileIndexForPlugin(const TSharedPtr<IPlug
 {
 	if (InPlugin && InPlugin->CanContainContent())
 	{
-		FString Filename = InPlugin->GetBaseDir() / ThumbnailExternalCache::ThumbnailFilenamePart;
+		const FString Filename = InPlugin->GetBaseDir() / FThumbnailExternalCache::GetCachedEditorThumbnailsFilename();
 		if (IFileManager::Get().FileExists(*Filename))
 		{
 			LoadCacheFileIndex(Filename);
diff --git a/Engine/Source/Editor/UnrealEd/Public/ThumbnailExternalCache.h b/Engine/Source/Editor/UnrealEd/Public/ThumbnailExternalCache.h
--- a/Engine/Source/Editor/UnrealEd/Public/ThumbnailExternalCache.h
+++ b/Engine/Source/Editor/UnrealEd/Public/ThumbnailExternalCache.h
@@ -28,6 +28,9 @@ public:
 	/** Get thumbnail external cache */
 	UNREALED_API static FThumbnailExternalCache& Get();
 
+	/** Gets the name of the editor thumbnail cache file */
+	UNREALED_API static const FString& GetCachedEditorThumbnailsFilename();
+
 	/** Load thumbnails for the given object names from external cache */
 	UNREALED_API bool LoadThumbnailsFromExternalCache(const TSet<FName>& InObjectFullNames, FThumbnailMap& InOutThumbnails);
 
@@ -99,6 +102,8 @@ private:
 
 	bool LoadCacheFileIndex(FArchive& Ar, const TSharedPtr<FThumbnailCacheFile>& CacheFile);
 
+private:
+
 	TMap<FString, TSharedPtr<FThumbnailCacheFile>> CacheFiles;
 
 	bool bHasInit = false;
