Fixed mismatched copy constructor/assignment operator in FFbxAnimCurveHandle.

diff --git a/Engine/Source/Editor/UnrealEd/Public/FbxImporter.h b/Engine/Source/Editor/UnrealEd/Public/FbxImporter.h
--- a/Engine/Source/Editor/UnrealEd/Public/FbxImporter.h
+++ b/Engine/Source/Editor/UnrealEd/Public/FbxImporter.h
@@ -270,18 +270,6 @@ public:
 		CurveType = NotTransform;
 	}
 
-	FFbxAnimCurveHandle(const FFbxAnimCurveHandle &CurveHandle)
-	{
-		UniqueId = CurveHandle.UniqueId;
-		Name = CurveHandle.Name;
-		ChannelIndex = CurveHandle.ChannelIndex;
-		CompositeIndex = CurveHandle.CompositeIndex;
-		KeyNumber = CurveHandle.KeyNumber;
-		AnimationTimeSecond = CurveHandle.AnimationTimeSecond;
-		AnimCurve = CurveHandle.AnimCurve;
-		CurveType = CurveHandle.CurveType;
-	}
-
 	//Identity Data
 	uint64 UniqueId;
 	FString Name;
