Cooker: Fix packages loaded by external actor packages to be cooked. Fixed in the !OnlyEditorOnly case only (which is the current engine default). Support for OnlyEditorOnly was incorrectly skipping these referenced packages even when OnlyEditorOnly is disabled.
Fixing this case when OnlyEditorOnly is enabled will come in a later change.
#rnx

diff --git a/Engine/Source/Editor/UnrealEd/Private/CookOnTheFlyServer.cpp b/Engine/Source/Editor/UnrealEd/Private/CookOnTheFlyServer.cpp
--- a/Engine/Source/Editor/UnrealEd/Private/CookOnTheFlyServer.cpp
+++ b/Engine/Source/Editor/UnrealEd/Private/CookOnTheFlyServer.cpp
@@ -2849,9 +2849,19 @@ void UCookOnTheFlyServer::QueueDiscoveredPackage(UE::Cook::FPackageData& Package
 	using namespace UE::Cook;
 
 	TArray<const ITargetPlatform*, TInlineAllocator<ExpectedMaxNumPlatforms>> BufferPlatforms;
-	if (PackageData.HasReachablePlatforms(
-		ReachablePlatforms.GetPlatforms(*this, &Instigator,
-			TConstArrayView<const ITargetPlatform*>(), &BufferPlatforms, true /* bAllowPartialInstigatorResults */)))
+	TConstArrayView<const ITargetPlatform*> DiscoveredPlatforms;
+	if (!bCanSkipEditorReferencedPackagesWhenCooking)
+	{
+		BufferPlatforms = PlatformManager->GetSessionPlatforms();
+		BufferPlatforms.Add(CookerLoadingPlatformKey);
+		DiscoveredPlatforms = BufferPlatforms;
+	}
+	else
+	{
+		DiscoveredPlatforms = ReachablePlatforms.GetPlatforms(*this, &Instigator,
+			TConstArrayView<const ITargetPlatform*>(), &BufferPlatforms, true /* bAllowPartialInstigatorResults */);
+	}
+	if (PackageData.HasReachablePlatforms(DiscoveredPlatforms))
 	{
 		// Not a new discovery; ignore
 		return;
diff --git a/Engine/Source/Editor/UnrealEd/Private/Cooker/CookRequestCluster.cpp b/Engine/Source/Editor/UnrealEd/Private/Cooker/CookRequestCluster.cpp
--- a/Engine/Source/Editor/UnrealEd/Private/Cooker/CookRequestCluster.cpp
+++ b/Engine/Source/Editor/UnrealEd/Private/Cooker/CookRequestCluster.cpp
@@ -108,14 +108,28 @@ FRequestCluster::FRequestCluster(UCookOnTheFlyServer& InCOTFS, TRingBuffer<FDisc
 	: FRequestCluster(InCOTFS)
 {
 	TArray<const ITargetPlatform*, TInlineAllocator<ExpectedMaxNumPlatforms>> BufferPlatforms;
+	if (!COTFS.bCanSkipEditorReferencedPackagesWhenCooking)
+	{
+		BufferPlatforms = COTFS.PlatformManager->GetSessionPlatforms();
+		BufferPlatforms.Add(CookerLoadingPlatformKey);
+	}
+
 	while (!DiscoveryQueue.IsEmpty())
 	{
 		FDiscoveryQueueElement* Discovery = &DiscoveryQueue.First();
 		FPackageData& PackageData = *Discovery->PackageData;
 
-		TConstArrayView<const ITargetPlatform*> NewReachablePlatforms =
-			Discovery->ReachablePlatforms.GetPlatforms(COTFS, &Discovery->Instigator,
+		TConstArrayView<const ITargetPlatform*> NewReachablePlatforms;
+		if (COTFS.bCanSkipEditorReferencedPackagesWhenCooking)
+		{
+			NewReachablePlatforms = Discovery->ReachablePlatforms.GetPlatforms(COTFS, &Discovery->Instigator,
 				TConstArrayView<const ITargetPlatform*>(), &BufferPlatforms);
+		}
+		else
+		{
+			NewReachablePlatforms = BufferPlatforms;
+		}
+
 		if (PackageData.HasReachablePlatforms(NewReachablePlatforms))
 		{
 			// If there are no new reachable platforms, add it to the cluster for cooking if it needs
@@ -148,12 +162,15 @@ FRequestCluster::FRequestCluster(UCookOnTheFlyServer& InCOTFS, TRingBuffer<FDisc
 				}
 
 				COTFS.OnDiscoveredPackageDebug(PackageData.GetPackageName(), Discovery->Instigator);
-				FPackageData* InstigatorPackageData = Discovery->Instigator.Referencer.IsNone() ? nullptr
-					: COTFS.PackageDatas->TryAddPackageDataByPackageName(Discovery->Instigator.Referencer);
-				if (InstigatorPackageData)
+				if (COTFS.bCanSkipEditorReferencedPackagesWhenCooking)
 				{
-					COTFS.DiscoveredDependencies.FindOrAdd(InstigatorPackageData->GetPackageName())
-						.Add(PackageData.GetPackageName());
+					FPackageData* InstigatorPackageData = Discovery->Instigator.Referencer.IsNone() ? nullptr
+						: COTFS.PackageDatas->TryAddPackageDataByPackageName(Discovery->Instigator.Referencer);
+					if (InstigatorPackageData)
+					{
+						COTFS.DiscoveredDependencies.FindOrAdd(InstigatorPackageData->GetPackageName())
+							.Add(PackageData.GetPackageName());
+					}
 				}
 			}
 		}
