Remove deprecated EditorAnalyticsEnabled setting from EditorSettings
* Sending Editor Analytics not controlled from this setting - it is not read or persisted.  Only results in an extra control in settings menu that has no effect.
* FYI - analytics sending controlled currently through UAnalyticsPrivacySettings, in UI: EditorPreferences->Privacy->Usage Data toggle.

diff --git a/Engine/Source/Editor/UnrealEd/Private/Settings/EditorLoadingSavingSettingsCustomization.h b/Engine/Source/Editor/UnrealEd/Private/Settings/EditorLoadingSavingSettingsCustomization.h
--- a/Engine/Source/Editor/UnrealEd/Private/Settings/EditorLoadingSavingSettingsCustomization.h
+++ b/Engine/Source/Editor/UnrealEd/Private/Settings/EditorLoadingSavingSettingsCustomization.h
@@ -60,8 +60,6 @@ protected:
 			TArray<UObject*> ObjectList;
 			ObjectList.Add(GetMutableDefault<UEditorSettings>());
 			StartupCategory.AddExternalObjectProperty(ObjectList, "bLoadTheMostRecentlyLoadedProjectAtStartup");
-
-			StartupCategory.AddExternalObjectProperty(ObjectList, "bEditorAnalyticsEnabled", EPropertyLocation::Advanced);
 		}
 	}
 };
diff --git a/Engine/Source/Editor/UnrealEd/Private/Settings/EditorSettings.cpp b/Engine/Source/Editor/UnrealEd/Private/Settings/EditorSettings.cpp
--- a/Engine/Source/Editor/UnrealEd/Private/Settings/EditorSettings.cpp
+++ b/Engine/Source/Editor/UnrealEd/Private/Settings/EditorSettings.cpp
@@ -21,7 +21,6 @@ UEditorSettings::UEditorSettings(const FObjectInitializer& ObjectInitializer)
 	: Super(ObjectInitializer)
 {
 	bCopyStarterContentPreference = false;
-	bEditorAnalyticsEnabled_DEPRECATED = true;
 	AutoScalabilityWorkScaleAmount = 1;
 
 	// Read the current state of the environment variables and cache it.
diff --git a/Engine/Source/Editor/UnrealEd/Public/Settings/EditorSettings.h b/Engine/Source/Editor/UnrealEd/Public/Settings/EditorSettings.h
--- a/Engine/Source/Editor/UnrealEd/Public/Settings/EditorSettings.h
+++ b/Engine/Source/Editor/UnrealEd/Public/Settings/EditorSettings.h
@@ -121,10 +121,6 @@ class UNREALED_API UEditorSettings : public UObject
 	UPROPERTY()
 	bool bLoadTheMostRecentlyLoadedProjectAtStartup; // Note that this property is NOT config since it is not necessary to save the value to ini. It is determined at startup in UEditorEngine::InitEditor().
 
-	/** Can the editor report usage analytics (types of assets being spawned, etc...) back to Epic in order for us to improve the editor user experience?  Note: The editor must be restarted for changes to take effect. */
-	UPROPERTY()
-	bool bEditorAnalyticsEnabled_DEPRECATED;
-
 	// =====================================================================
 	// The following options are NOT exposed in the preferences Editor
 	// (usually because there is a different way to set them interactively!)
