Cooker: Suppress spurious diagnostic message that editoronly packages that are not reachable were "rejected" for being "alreadycooked". They were actually never instigated and so do not need a message explaining their rejection.
#rnx

diff --git a/Engine/Source/Editor/UnrealEd/Private/CookOnTheFlyServer.cpp b/Engine/Source/Editor/UnrealEd/Private/CookOnTheFlyServer.cpp
--- a/Engine/Source/Editor/UnrealEd/Private/CookOnTheFlyServer.cpp
+++ b/Engine/Source/Editor/UnrealEd/Private/CookOnTheFlyServer.cpp
@@ -2514,7 +2514,17 @@ void UCookOnTheFlyServer::PumpRequests(UE::Cook::FTickStackData& StackData, int3
 		}
 		if (PackageData->GetPlatformsNeedingCookingNum() == 0)
 		{
-			DemoteToIdle(*PackageData, ESendFlags::QueueAdd, ESuppressCookReason::AlreadyCooked);
+			TArray<const ITargetPlatform*> ReachablePlatforms;
+			PackageData->GetReachablePlatforms(ReachablePlatforms);
+			if (ReachablePlatforms.IsEmpty())
+			{
+				UE_LOG(LogCook, Warning, TEXT("Package %s was reported as a RequestToLoad by a requestcluster, but it has no reachable platforms."),
+					*WriteToString<256>(PackageData->GetPackageName()));
+			}
+			ReachablePlatforms.Remove(CookerLoadingPlatformKey);
+			ESuppressCookReason SuppressCookReason = ReachablePlatforms.IsEmpty() ?
+				ESuppressCookReason::OnlyEditorOnly : ESuppressCookReason::AlreadyCooked;
+			DemoteToIdle(*PackageData, ESendFlags::QueueAdd, SuppressCookReason);
 			continue;
 		}
 		PackageData->SendToState(EPackageState::LoadPrepare, ESendFlags::QueueAdd);
@@ -2579,26 +2589,33 @@ void UCookOnTheFlyServer::DemoteToIdle(UE::Cook::FPackageData& PackageData, UE::
 	if (PackageData.IsInProgress())
 	{
 		WorkerRequests->ReportDemoteToIdle(PackageData, Reason);
+
 		// If per-package display is on, write a log statement explaining that the package was reachable but skipped.
-		if (!bCookListMode &
-			!!((GCookProgressDisplay & ((int32)ECookProgressDisplayMode::Instigators | (int32)ECookProgressDisplayMode::PackageNames))))
+		bool bPrintDiagnostic = !bCookListMode &
+			!!((GCookProgressDisplay & ((int32)ECookProgressDisplayMode::Instigators | (int32)ECookProgressDisplayMode::PackageNames)));
+
+		// Suppress the message in cases that cause large spam like NotInCurrentPlugin for DLC cooks.
+		bPrintDiagnostic &= (Reason != ESuppressCookReason::NotInCurrentPlugin);
+		if (bPrintDiagnostic)
 		{
-			// Suppress the message in cases that cause large spam like NotInCurrentPlugin for DLC cooks.
-			if (Reason != ESuppressCookReason::NotInCurrentPlugin)
+			TStringBuilder<256> PackageNameStr(InPlace, PackageData.GetPackageName());
+
+			// ExternalActors: Do not send a message for every NeverCook external Actor package; too much spam
+			if (Reason == ESuppressCookReason::NeverCook)
 			{
-				TStringBuilder<256> PackageNameStr(InPlace, PackageData.GetPackageName());
+				bPrintDiagnostic &= UE::String::FindFirst(PackageNameStr.ToView(),
+					ULevel::GetExternalActorsFolderName(), ESearchCase::IgnoreCase) == INDEX_NONE;
+			}
 
-				// ExternalActors: Do not send a message for every NeverCook external Actor package; too much spam
-				if (!(Reason == ESuppressCookReason::NeverCook &&
-					UE::String::FindFirst(PackageNameStr.ToView(),
-						ULevel::GetExternalActorsFolderName(), ESearchCase::IgnoreCase) != INDEX_NONE))
-				{
-					UE_CLOG((GCookProgressDisplay & (int32)ECookProgressDisplayMode::Instigators), LogCook, Display,
-						TEXT("Cooking %s, Instigator: { %s } -> Rejected %s"), *PackageNameStr,
-						*(PackageData.GetInstigator().ToString()), LexToString(Reason));
-					UE_CLOG(GCookProgressDisplay & (int32)ECookProgressDisplayMode::PackageNames, LogCook, Display,
-						TEXT("Cooking %s -> Rejected %s"), *PackageNameStr, LexToString(Reason));
-				}
+			// Reachability: Suppress the diagnostic that were found via cookload reference traversal but are not reachable on the target platforms
+			bPrintDiagnostic &= (PackageData.HasInstigator() || Reason != ESuppressCookReason::OnlyEditorOnly);
+			if (bPrintDiagnostic)
+			{
+				UE_CLOG((GCookProgressDisplay & (int32)ECookProgressDisplayMode::Instigators), LogCook, Display,
+					TEXT("Cooking %s, Instigator: { %s } -> Rejected %s"), *PackageNameStr,
+					*(PackageData.GetInstigator().ToString()), LexToString(Reason));
+				UE_CLOG(GCookProgressDisplay & (int32)ECookProgressDisplayMode::PackageNames, LogCook, Display,
+					TEXT("Cooking %s -> Rejected %s"), *PackageNameStr, LexToString(Reason));
 			}
 		}
 	}
