MPCook: Fix an MPCook-specific callsite that was not allowing packages loaded by external actor packages to be cooked. Fixed in the !OnlyEditorOnly case only.
#rnx

diff --git a/Engine/Source/Editor/UnrealEd/Private/Cooker/CookWorkerServer.cpp b/Engine/Source/Editor/UnrealEd/Private/Cooker/CookWorkerServer.cpp
--- a/Engine/Source/Editor/UnrealEd/Private/Cooker/CookWorkerServer.cpp
+++ b/Engine/Source/Editor/UnrealEd/Private/Cooker/CookWorkerServer.cpp
@@ -860,9 +860,18 @@ void FCookWorkerServer::QueueDiscoveredPackage(FDiscoveredPackageReplication&& D
 		DiscoveredPackage.NormalizedFileName);
 
 	TArray<const ITargetPlatform*, TInlineAllocator<ExpectedMaxNumPlatforms>> BufferPlatforms;
-	if (PackageData.HasReachablePlatforms(
-		Platforms.GetPlatforms(COTFS, &Instigator, OrderedSessionAndSpecialPlatforms, &BufferPlatforms,
-			true /* bAllowPartialInstigatorResults */)))
+	TConstArrayView<const ITargetPlatform*> DiscoveredPlatforms;
+	if (!COTFS.bCanSkipEditorReferencedPackagesWhenCooking)
+	{
+		DiscoveredPlatforms = OrderedSessionAndSpecialPlatforms;
+	}
+	else
+	{
+		DiscoveredPlatforms = Platforms.GetPlatforms(COTFS, &Instigator, OrderedSessionAndSpecialPlatforms, &BufferPlatforms,
+			true /* bAllowPartialInstigatorResults */);
+	}
+
+	if (PackageData.HasReachablePlatforms(DiscoveredPlatforms))
 	{
 		// The CookWorker thought this was a new package, but the Director already knows about it; ignore the report
 		return;
