BlueprintEditorUtils: fixed a missing "!" when checking for disallowed node class

diff --git a/Engine/Source/Editor/UnrealEd/Private/Kismet2/BlueprintEditorUtils.cpp b/Engine/Source/Editor/UnrealEd/Private/Kismet2/BlueprintEditorUtils.cpp
--- a/Engine/Source/Editor/UnrealEd/Private/Kismet2/BlueprintEditorUtils.cpp
+++ b/Engine/Source/Editor/UnrealEd/Private/Kismet2/BlueprintEditorUtils.cpp
@@ -10666,7 +10666,7 @@ bool FBlueprintEditorUtils::HasRestrictedNodes(const UBlueprint* BP, TArray<UEdG
 	// Check for any disallowed node types.
 	auto IsRestrictedNodeType = [BlueprintEditorSettings](const UEdGraphNode* InNode) -> bool
 	{
-		return BlueprintEditorSettings->IsClassAllowed(InNode->GetClass());
+		return !BlueprintEditorSettings->IsClassAllowed(InNode->GetClass());
 	};
 
 	// Check for any disallowed event graph implementations. The event node is considered restricted in this case.
