AssetEditorSubsystem: Fix a bug where the editor would sometimes hit a check when trying to delete a transient asset from the recents list on shutdown

diff --git a/Engine/Source/Editor/UnrealEd/Private/Subsystems/AssetEditorSubsystem.cpp b/Engine/Source/Editor/UnrealEd/Private/Subsystems/AssetEditorSubsystem.cpp
--- a/Engine/Source/Editor/UnrealEd/Private/Subsystems/AssetEditorSubsystem.cpp
+++ b/Engine/Source/Editor/UnrealEd/Private/Subsystems/AssetEditorSubsystem.cpp
@@ -132,11 +132,11 @@ void UAssetEditorSubsystem::InitializeRecentAssets()
 	if(bFoundRecentAssetEditorsCleanly)
 	{
 		// Go in reverse since the first item should be the most recent
-		for(int i = 0; i < RecentAssetsList->GetNumItems(); ++i)
+		for(int CurRecentIndex = 0; CurRecentIndex < RecentAssetsList->GetNumItems(); ++CurRecentIndex)
 		{
-			if(!RecentAssetEditors[i].IsEmpty())
+			if(!RecentAssetEditors[CurRecentIndex].IsEmpty())
 			{
-				RecentAssetToAssetEditorMap.Add(RecentAssetsList->GetMRUItem(i), RecentAssetEditors[i]);
+				RecentAssetToAssetEditorMap.Add(RecentAssetsList->GetMRUItem(CurRecentIndex), RecentAssetEditors[CurRecentIndex]);
 			}
 		}
 	}
@@ -149,20 +149,20 @@ void UAssetEditorSubsystem::SaveRecentAssets(const bool bOnShutdown)
 	// If we are closing the editor, remove all assets that weren't actually saved to disk (transient)
 	if(bOnShutdown)
 	{
-		for(int i = 0; i < RecentAssetsList->GetNumItems(); ++i)
+		for(int CurRecentIndex = 0; CurRecentIndex < RecentAssetsList->GetNumItems(); ++CurRecentIndex)
 		{
-			const FString& RecentAsset = RecentAssetsList->GetMRUItem(i);
+			const FString& RecentAsset = RecentAssetsList->GetMRUItem(CurRecentIndex);
 			if(!FPackageName::DoesPackageExist(RecentAsset))
 			{
-				RecentAssetsList->RemoveMRUItem(RecentAsset);
-				--i;
+				RecentAssetsList->RemoveMRUItem(CurRecentIndex);
+				--CurRecentIndex;
 			}
 		}
 	}
 	
-	for(int i = 0; i < RecentAssetsList->GetNumItems(); ++i)
+	for(int CurRecentIndex = 0; CurRecentIndex < RecentAssetsList->GetNumItems(); ++CurRecentIndex)
 	{
-		FString CurRecentAsset = RecentAssetsList->GetMRUItem(i);
+		FString CurRecentAsset = RecentAssetsList->GetMRUItem(CurRecentIndex);
 
 		// If we have a valid asset editor for the current asset, save it
 		if(FString* CurrentAssetEditorName = RecentAssetToAssetEditorMap.Find(CurRecentAsset))
