Make SyncBrowserToObjects use const arrays and add SyncBrowserToObject

diff --git a/Engine/Source/Editor/UnrealEd/Classes/Editor/EditorEngine.h b/Engine/Source/Editor/UnrealEd/Classes/Editor/EditorEngine.h
--- a/Engine/Source/Editor/UnrealEd/Classes/Editor/EditorEngine.h
+++ b/Engine/Source/Editor/UnrealEd/Classes/Editor/EditorEngine.h
@@ -2134,10 +2134,18 @@ public:
 	/**
 	 * Synchronizes the content or generic browser's selected objects to the collection specified.
 	 *
-	 * @param	ObjectSet	the list of objects to sync to
+	 * @param	InObjectsToSync	the list of objects to sync to
 	 */
-	void SyncBrowserToObjects( TArray<UObject*>& InObjectsToSync, bool bFocusContentBrowser = true );
-	void SyncBrowserToObjects( TArray<struct FAssetData>& InAssetsToSync, bool bFocusContentBrowser = true );
+	void SyncBrowserToObjects( const TArray<UObject*>& InObjectsToSync, bool bFocusContentBrowser = true );
+	void SyncBrowserToObjects( const TArray<struct FAssetData>& InAssetsToSync, bool bFocusContentBrowser = true );
+
+	/**
+	 * Synchronizes the content or generic browser's selected object to the one specified.
+	 *
+	 * @param	InObjectToSync	the object to sync to
+	 */
+	void SyncBrowserToObject( const UObject* InObjectToSync, bool bFocusContentBrowser = true );
+	void SyncBrowserToObject( const struct FAssetData& InAssetToSync, bool bFocusContentBrowser = true );
 
 	/**
 	 * Syncs the selected actors objects to the content browser
diff --git a/Engine/Source/Editor/UnrealEd/Private/EditorEngine.cpp b/Engine/Source/Editor/UnrealEd/Private/EditorEngine.cpp
--- a/Engine/Source/Editor/UnrealEd/Private/EditorEngine.cpp
+++ b/Engine/Source/Editor/UnrealEd/Private/EditorEngine.cpp
@@ -3023,19 +3023,29 @@ void UEditorEngine::LoadMapListFromIni(const FString& InSectionName, TArray<FStr
 	}
 }
 
-void UEditorEngine::SyncBrowserToObjects( TArray<UObject*>& InObjectsToSync, bool bFocusContentBrowser )
+void UEditorEngine::SyncBrowserToObjects( const TArray<UObject*>& InObjectsToSync, bool bFocusContentBrowser )
 {
 	FContentBrowserModule& ContentBrowserModule = FModuleManager::Get().LoadModuleChecked<FContentBrowserModule>("ContentBrowser");
 	ContentBrowserModule.Get().SyncBrowserToAssets( InObjectsToSync, false, bFocusContentBrowser );
 
 }
 
-void UEditorEngine::SyncBrowserToObjects( TArray<struct FAssetData>& InAssetsToSync, bool bFocusContentBrowser )
+void UEditorEngine::SyncBrowserToObjects( const TArray<struct FAssetData>& InAssetsToSync, bool bFocusContentBrowser )
 {
 	FContentBrowserModule& ContentBrowserModule = FModuleManager::Get().LoadModuleChecked<FContentBrowserModule>("ContentBrowser");
 	ContentBrowserModule.Get().SyncBrowserToAssets( InAssetsToSync, false, bFocusContentBrowser );
 }
 
+void UEditorEngine::SyncBrowserToObject( const UObject* InObjectToSync, bool bFocusContentBrowser )
+{
+	SyncBrowserToObjects({ InObjectToSync }, bFocusContentBrowser);
+}
+
+void UEditorEngine::SyncBrowserToObject( const FAssetData& InAssetToSync, bool bFocusContentBrowser )
+{
+	SyncBrowserToObjects({ InAssetToSync }, bFocusContentBrowser);
+}
+
 
 bool UEditorEngine::CanSyncToContentBrowser()
 {
