BSP: Make the Convert To Static Mesh function not transact static mesh creation because we can't undo that asset creation properly.

diff --git a/Engine/Source/Editor/UnrealEd/Private/StaticMeshEdit.cpp b/Engine/Source/Editor/UnrealEd/Private/StaticMeshEdit.cpp
--- a/Engine/Source/Editor/UnrealEd/Private/StaticMeshEdit.cpp
+++ b/Engine/Source/Editor/UnrealEd/Private/StaticMeshEdit.cpp
@@ -7,6 +7,7 @@
 #include "CoreMinimal.h"
 #include "Misc/CoreMisc.h"
 #include "Misc/FeedbackContext.h"
+#include "CoreGlobals.h" // GUndo
 #include "Engine/EngineTypes.h"
 #include "Model.h"
 #include "EditorFramework/AssetImportData.h"
@@ -257,6 +258,13 @@ void GetBrushMesh(ABrush* Brush, UModel* Model, FMeshDescription& MeshDescriptio
 
 UStaticMesh* CreateStaticMeshFromBrush(UObject* Outer, FName Name, ABrush* Brush, UModel* Model)
 {
+	// MINOR HACK: Make sure we don't transact static mesh object creation because we are currently
+	// unable to undo it properly, where undoing it places the asset in a broken state that can cause
+	// crashes. Instead, temporarily disconnect the current transaction object, if any.
+	ITransaction* UndoState = GUndo;
+	GUndo = nullptr; // Pretend we're not in a transaction
+	ON_SCOPE_EXIT{ GUndo = UndoState; }; // Revert
+
 	FScopedSlowTask SlowTask(0.0f, NSLOCTEXT("UnrealEd", "CreatingStaticMeshE", "Creating static mesh..."));
 	SlowTask.MakeDialog();
 
