OnlyEditorOnly:
1) Add Imports and SoftObjectPaths from saved packages into the Cook
2) Skip queuing of Unsolicited packages and loaded SoftObjectPaths when coking with OnlyEditorOnly enabled
3) Allow packages to be marked discovered by instigators that are not reachable themselves. We still need to add the hidden dependency from the instigator in case it becomes reachable later.
4)  Add different handling for unsolicited packages that are marked with a FCookLoadType.

CookerTools:
1) Add the -cookfirst argument. It fulfills the same purpose as cooksinglepackagenorefs, but instead of skipping all other packages it sorts the input packages first and then cooks all others. This allows running the same environment as a normal cook while still debugging a given set of packages. Useful when trying to debug package instigators.

#rnx

diff --git a/Engine/Source/Editor/UnrealEd/Private/CookOnTheFlyServer.cpp b/Engine/Source/Editor/UnrealEd/Private/CookOnTheFlyServer.cpp
--- a/Engine/Source/Editor/UnrealEd/Private/CookOnTheFlyServer.cpp
+++ b/Engine/Source/Editor/UnrealEd/Private/CookOnTheFlyServer.cpp
@@ -358,8 +358,9 @@ public:
 
 		FName StandardFileName(*FPaths::CreateStandardFilename(CookPackageRequest.Filename));
 		UE_LOG(LogCook, Verbose, TEXT("Enqueing cook request, Filename='%s', Platform='%s'"), *CookPackageRequest.Filename, *CookPackageRequest.PlatformName.ToString());
-		Cooker.WorkerRequests->AddCookOnTheFlyRequest(
-			FFilePlatformRequest(StandardFileName, EInstigator::CookOnTheFly, TargetPlatform, MoveTemp(CookPackageRequest.CompletionCallback)));
+		FFilePlatformRequest Request(StandardFileName, EInstigator::CookOnTheFly, TargetPlatform, MoveTemp(CookPackageRequest.CompletionCallback));
+		Request.SetUrgent(true);
+		Cooker.WorkerRequests->AddCookOnTheFlyRequest(MoveTemp(Request));
 
 		return true;
 	};
@@ -1260,9 +1261,9 @@ bool UCookOnTheFlyServer::RequestPackage(const FName& StandardFileName, const TA
 		}
 	}
 
-	WorkerRequests->AddPublicInterfaceRequest(
-		FFilePlatformRequest(StandardFileName, EInstigator::RequestPackageFunction, TargetPlatforms),
-		bForceFrontOfQueue);
+	FFilePlatformRequest Request(StandardFileName, EInstigator::RequestPackageFunction, TargetPlatforms);
+	Request.SetUrgent(bForceFrontOfQueue);
+	WorkerRequests->AddPublicInterfaceRequest(MoveTemp(Request), bForceFrontOfQueue);
 	return true;
 }
 
@@ -2420,7 +2421,7 @@ void UCookOnTheFlyServer::PumpExternalRequests(const UE::Cook::FCookerTimer& Coo
 #endif
 			bool bRequestsAreUrgent = IsCookOnTheFlyMode() && IsUsingLegacyCookOnTheFlyScheduling();
 			TRingBuffer<UE::Cook::FRequestCluster>& RequestClusters = PackageDatas->GetRequestQueue().GetRequestClusters();
-			RequestClusters.Add(UE::Cook::FRequestCluster(*this, MoveTemp(BuildRequests), bRequestsAreUrgent));
+			RequestClusters.Add(UE::Cook::FRequestCluster(*this, MoveTemp(BuildRequests)));
 		}
 	}
 }
@@ -2876,7 +2877,7 @@ void UCookOnTheFlyServer::QueueDiscoveredPackage(UE::Cook::FPackageData& Package
 	else
 	{
 		DiscoveredPlatforms = ReachablePlatforms.GetPlatforms(*this, &Instigator,
-			TConstArrayView<const ITargetPlatform*>(), &BufferPlatforms, true /* bAllowPartialInstigatorResults */);
+			TConstArrayView<const ITargetPlatform*>(), &BufferPlatforms);
 	}
 	if (PackageData.HasReachablePlatforms(DiscoveredPlatforms))
 	{
@@ -2886,10 +2887,7 @@ void UCookOnTheFlyServer::QueueDiscoveredPackage(UE::Cook::FPackageData& Package
 
 	if (bHiddenDependenciesDebug)
 	{
-		if (!PackageData.IsGenerated())
-		{
-			OnDiscoveredPackageDebug(PackageData.GetPackageName(), Instigator);
-		}
+		OnDiscoveredPackageDebug(PackageData.GetPackageName(), Instigator);
 	}
 	WorkerRequests->QueueDiscoveredPackage(*this, PackageData, MoveTemp(Instigator), MoveTemp(ReachablePlatforms));
 }
@@ -3123,7 +3121,9 @@ UE::Cook::EPollStatus UCookOnTheFlyServer::QueueGeneratedPackages(UE::Cook::FGen
 		for (const FCookGenerationInfo& ChildInfo: Generator.GetPackagesToGenerate())
 		{
 			FPackageData* ChildPackageData = ChildInfo.PackageData;
-			QueueDiscoveredPackage(*ChildPackageData, FInstigator(EInstigator::GeneratedPackage, OwnerName),
+			// Set the Instigator now rather than delaying it until the discovery queue is processed.
+			ChildPackageData->SetInstigator(Generator, FInstigator(EInstigator::GeneratedPackage, OwnerName));
+			QueueDiscoveredPackage(*ChildPackageData, FInstigator(ChildPackageData->GetInstigator()),
 				EDiscoveredPlatformSet::CopyFromInstigator);
 		}
 		Info.SetSaveStateComplete(FCookGenerationInfo::ESaveState::QueueGeneratedPackages);
@@ -4128,23 +4128,29 @@ void UCookOnTheFlyServer::ProcessUnsolicitedPackages(TArray<FName>* OutDiscovere
 			}
 		}
 
-		if (!PackageData->FindOrAddPlatformData(CookerLoadingPlatformKey).IsReachable())
+		if (PackageData->FindOrAddPlatformData(CookerLoadingPlatformKey).IsReachable() || Instigator.Category == EInstigator::EditorOnlyLoad)
+		{
+			// This load was expected so we do not need to add a hidden dependency for it.
+			// If we are using legacy WhatGetsCookedRules, queue the package for cooking because it was loaded.
+			if (!bCanSkipEditorReferencedPackagesWhenCooking)
+			{
+				QueueDiscoveredPackage(*PackageData, FInstigator(Instigator), EDiscoveredPlatformSet::CopyFromInstigator);
+			}
+		}
+		else if (Instigator.Category == EInstigator::SaveTimeHardDependency || Instigator.Category == EInstigator::SaveTimeSoftDependency)
+		{
+			// A load undeclared in AssetRegistry dependencies, but one that has been marked up as a known issue and needed at runtime.
+			// Queue it for discovery and do not log a warning about it.
+			QueueDiscoveredPackage(*PackageData, FInstigator(Instigator), EDiscoveredPlatformSet::CopyFromInstigator);
+		}
+		else
 		{
 			// Possibly an unexpected load, but it might be an expected load after we finish processing some other unexpected loads that were
 			// discovered earlier during the same LoadPackage call. Queue it for discovery and if it is still unreachable when its turn comes
 			// we will add the hidden dependency for it.
 			// as reachable in all of its instigator's reachable platforms
-			QueueDiscoveredPackage(*PackageData, MoveTemp(Instigator), EDiscoveredPlatformSet::CopyFromInstigator);
-		}
-		else
-		{
-			// This load was expected so we do not need to add a hidden dependency for it.
-			// ONLYEDITORONLY_TODO: Read the list of Imports in FinishPlatform out of the SavePackage results instead, and then turn off bMarkExpectedLoadsReachableAsWell
-			constexpr bool bMarkExpectedLoadsReachableAsWell = true;
-			if (bMarkExpectedLoadsReachableAsWell)
-			{
-				QueueDiscoveredPackage(*PackageData, MoveTemp(Instigator), EDiscoveredPlatformSet::CopyFromInstigator);
-			}
+			QueueDiscoveredPackage(*PackageData, FInstigator(Instigator), EDiscoveredPlatformSet::CopyFromInstigator);
+			// TODO: Add diagnostics for this unexpected load
 		}
 	}
 }
@@ -4167,7 +4173,7 @@ private: // Used only by UCookOnTheFlyServer, which has private access
 
 	// General Package Data
 	UCookOnTheFlyServer& COTFS;
-	UE::Cook::FPackageData& PackageData;
+	FPackageData& PackageData;
 	TArrayView<const ITargetPlatform*> PlatformsForPackage;
 	FTickStackData& StackData;
 	UPackage* Package;
@@ -6049,6 +6055,30 @@ void FSaveCookedPackageContext::FinishPlatform()
 			MoveTemp(OverrideAssetPackageData), MoveTemp(OverridePackageDependencies));
 	}
 
+	if (COTFS.bCanSkipEditorReferencedPackagesWhenCooking)
+	{
+		// If the save succeeded, add the imports and softobjectpaths from the save to the cook for the platform
+		if (bSuccessful)
+		{
+			FName PackageFName = Package->GetFName();
+			TConstArrayView<const ITargetPlatform*> ReachablePlatforms(&TargetPlatform, 1);
+			for (const TArray<FName>& DependencyNames : { SavePackageResult.ImportPackages, SavePackageResult.SoftPackageReferences })
+			{
+				EInstigator InstigatorType = &DependencyNames == &SavePackageResult.ImportPackages ?
+					EInstigator::SaveTimeHardDependency : EInstigator::SaveTimeSoftDependency;
+				for (FName DependencyName : DependencyNames)
+				{
+					UE::Cook::FPackageData* DependencyData = COTFS.PackageDatas->TryAddPackageDataByPackageName(DependencyName);
+					if (DependencyData)
+					{
+						COTFS.QueueDiscoveredPackage(*DependencyData,
+							UE::Cook::FInstigator(InstigatorType, PackageFName), FDiscoveredPlatformSet(ReachablePlatforms));
+					}
+				}
+			}
+		}
+	}
+
 	// If not retrying, mark the package as cooked, either successfully or with failure
 	bool bIsRetryErrorCode = IsRetryErrorCode(SavePackageResult.Result);
 	if (!bIsRetryErrorCode)
@@ -6132,9 +6162,16 @@ void FSaveCookedPackageContext::FinishPackage()
 			UE::Cook::FPackageData* SoftObjectPackageData = COTFS.PackageDatas->TryAddPackageDataByPackageName(SoftObjectPackage);
 			if (SoftObjectPackageData)
 			{
-				COTFS.QueueDiscoveredPackage(*SoftObjectPackageData,
-					UE::Cook::FInstigator(UE::Cook::EInstigator::SoftDependency, PackageFName),
-					FDiscoveredPlatformSet(ReachablePlatforms));
+				if (!COTFS.bCanSkipEditorReferencedPackagesWhenCooking)
+				{
+					COTFS.QueueDiscoveredPackage(*SoftObjectPackageData,
+						UE::Cook::FInstigator(UE::Cook::EInstigator::SoftDependency, PackageFName),
+						FDiscoveredPlatformSet(ReachablePlatforms));
+				}
+				else
+				{
+					// TODO: Add OnlyEditorOnly diagnostic
+				}
 			}
 		}
 	}
@@ -10464,6 +10501,20 @@ void UCookOnTheFlyServer::GenerateInitialRequests(FBeginCookContext& BeginContex
 		UE_SCOPED_HIERARCHICAL_COOKTIMER(GenerateLongPackageName);
 		GenerateLongPackageNames(FilesInPath, FilesInPathInstigators);
 	}
+	TSet<FName> CookFirstPackages;
+	bool bCookFirst = FParse::Param(FCommandLine::Get(), TEXT("CookFirst"));
+	if (bCookFirst)
+	{
+		for (const FString& CookMap : CookMaps)
+		{
+			FString LongPackageName;
+			if (FPackageName::TryConvertFilenameToLongPackageName(CookMap, LongPackageName))
+			{
+				CookFirstPackages.Add(FName(*LongPackageName));
+			}
+		}
+	}
+
 	// add all the files to the cook list for the requested platforms
 	for (FName PackageName : FilesInPath)
 	{
@@ -10477,7 +10528,9 @@ void UCookOnTheFlyServer::GenerateInitialRequests(FBeginCookContext& BeginContex
 		UE::Cook::FInstigator& Instigator = FilesInPathInstigators.FindChecked(PackageName);
 		if (!PackageFileFName.IsNone())
 		{
-			WorkerRequests->AddStartCookByTheBookRequest(UE::Cook::FFilePlatformRequest(PackageFileFName, MoveTemp(Instigator), TargetPlatforms));
+			UE::Cook::FFilePlatformRequest Request(PackageFileFName, MoveTemp(Instigator), TargetPlatforms);
+			Request.SetUrgent(bCookFirst && CookFirstPackages.Contains(PackageName));
+			WorkerRequests->AddStartCookByTheBookRequest(MoveTemp(Request));
 		}
 		else if (!FLinkerLoad::IsKnownMissingPackage(PackageName))
 		{
@@ -11434,7 +11487,8 @@ void UCookOnTheFlyServer::OnDiscoveredPackageDebug(FName PackageName, const UE::
 	}
 	switch (Instigator.Category)
 	{
-	case EInstigator::StartupPackage:
+	case EInstigator::StartupPackage: [[fallthrough]];
+	case EInstigator::GeneratedPackage:
 		// Not a Hidden dependency
 		return;
 	default:
diff --git a/Engine/Source/Editor/UnrealEd/Private/Cooker/CookDirector.cpp b/Engine/Source/Editor/UnrealEd/Private/Cooker/CookDirector.cpp
--- a/Engine/Source/Editor/UnrealEd/Private/Cooker/CookDirector.cpp
+++ b/Engine/Source/Editor/UnrealEd/Private/Cooker/CookDirector.cpp
@@ -384,6 +384,11 @@ void FCookDirector::AssignRequests(TArray<FWorkerId>&& InWorkers, TArray<TRefCou
 		{
 			WorkerId = WorkerIdConstraint;
 		}
+		// Override the loadbalancer's assignment to force it local if the Package is urgent
+		else if (Requests[RequestIndex]->GetIsUrgent())
+		{
+			WorkerId = FWorkerId::Local();
+		}
 
 		if (!WorkerId.IsLocal())
 		{
diff --git a/Engine/Source/Editor/UnrealEd/Private/Cooker/CookPackageData.cpp b/Engine/Source/Editor/UnrealEd/Private/Cooker/CookPackageData.cpp
--- a/Engine/Source/Editor/UnrealEd/Private/Cooker/CookPackageData.cpp
+++ b/Engine/Source/Editor/UnrealEd/Private/Cooker/CookPackageData.cpp
@@ -176,23 +176,18 @@ bool FPackageData::AreAllReachablePlatformsVisitedByCluster() const
 }
 
 void FPackageData::GetReachablePlatformsForInstigator(UCookOnTheFlyServer& COTFS, FName InInstigator,
-	TArray<const ITargetPlatform*>& Platforms, bool bAllowPartialInstigatorResults)
+	TArray<const ITargetPlatform*>& Platforms)
 {
 	return GetReachablePlatformsForInstigator(COTFS,
-		COTFS.PackageDatas->TryAddPackageDataByPackageName(InInstigator), Platforms, bAllowPartialInstigatorResults);
+		COTFS.PackageDatas->TryAddPackageDataByPackageName(InInstigator), Platforms);
 }
 
 void FPackageData::GetReachablePlatformsForInstigator(UCookOnTheFlyServer& COTFS, UE::Cook::FPackageData* InInstigator,
-	TArray<const ITargetPlatform*>& Platforms, bool bAllowPartialInstigatorResults)
+	TArray<const ITargetPlatform*>& Platforms)
 {
 	if (InInstigator)
 	{
 		InInstigator->GetReachablePlatforms(Platforms);
-		if (Platforms.IsEmpty() && !bAllowPartialInstigatorResults)
-		{
-			UE_LOG(LogCook, Error, TEXT("An instigator package is not marked as reachable for any platform, should be impossible: %s"),
-				*InInstigator->GetPackageName().ToString());
-		}
 	}
 	else
 	{
@@ -274,6 +269,11 @@ void FPackageData::SetInstigator(FCookWorkerClient& Cluster, FInstigator&& InIns
 	SetInstigatorInternal(MoveTemp(InInstigator));
 }
 
+void FPackageData::SetInstigator(FGeneratorPackage& Cluster, FInstigator&& InInstigator)
+{
+	SetInstigatorInternal(MoveTemp(InInstigator));
+}
+
 void FPackageData::SetInstigatorInternal(FInstigator&& InInstigator)
 {
 	if (this->Instigator.Category == EInstigator::NotYetRequested)
diff --git a/Engine/Source/Editor/UnrealEd/Private/Cooker/CookPackageData.h b/Engine/Source/Editor/UnrealEd/Private/Cooker/CookPackageData.h
--- a/Engine/Source/Editor/UnrealEd/Private/Cooker/CookPackageData.h
+++ b/Engine/Source/Editor/UnrealEd/Private/Cooker/CookPackageData.h
@@ -526,6 +526,7 @@ public:
 	/** Setting the instigator is mostly private - it should only be done during clustering. */
 	void SetInstigator(FRequestCluster& Cluster, FInstigator&& InInstigator);
 	void SetInstigator(FCookWorkerClient& Cluster, FInstigator&& InInstigator);
+	void SetInstigator(FGeneratorPackage& Cluster, FInstigator&& InInstigator);
 
 	/** Get whether COTFS is keeping this package referenced referenced during GC. */
 	bool IsKeepReferencedDuringGC() const { return static_cast<bool>(bKeepReferencedDuringGC); }
@@ -552,9 +553,9 @@ public:
 	 * If the package does not exist, return the COTFS's list of Session platforms
 	 */
 	static void GetReachablePlatformsForInstigator(UCookOnTheFlyServer& COTFS, FPackageData* InInstigator,
-		TArray<const ITargetPlatform*>& Platforms, bool bAllowPartialInstigatorResults=false);
+		TArray<const ITargetPlatform*>& Platforms);
 	static void GetReachablePlatformsForInstigator(UCookOnTheFlyServer& COTFS, FName InInstigator,
-		TArray<const ITargetPlatform*>& Platforms, bool bAllowPartialInstigatorResults=false);
+		TArray<const ITargetPlatform*>& Platforms);
 private:
 	friend struct UE::Cook::FPackageDatas;
 
diff --git a/Engine/Source/Editor/UnrealEd/Private/Cooker/CookRequestCluster.cpp b/Engine/Source/Editor/UnrealEd/Private/Cooker/CookRequestCluster.cpp
--- a/Engine/Source/Editor/UnrealEd/Private/Cooker/CookRequestCluster.cpp
+++ b/Engine/Source/Editor/UnrealEd/Private/Cooker/CookRequestCluster.cpp
@@ -82,10 +82,9 @@ FRequestCluster::FRequestCluster(UCookOnTheFlyServer& InCOTFS)
 	}
 }
 
-FRequestCluster::FRequestCluster(UCookOnTheFlyServer& InCOTFS, TArray<FFilePlatformRequest>&& InRequests, bool bInExternalRequestsAreUrgent)
+FRequestCluster::FRequestCluster(UCookOnTheFlyServer& InCOTFS, TArray<FFilePlatformRequest>&& InRequests)
 	: FRequestCluster(InCOTFS)
 {
-	bExternalRequestsAreUrgent = bInExternalRequestsAreUrgent;
 	ReserveInitialRequests(InRequests.Num());
 	FilePlatformRequests = MoveTemp(InRequests);
 	InRequests.Empty();
@@ -144,24 +143,22 @@ FRequestCluster::FRequestCluster(UCookOnTheFlyServer& InCOTFS, TRingBuffer<FDisc
 			continue;
 		}
 
-		if (NewReachablePlatforms.Contains(CookerLoadingPlatformKey) &&
-			!PackageData.FindOrAddPlatformData(CookerLoadingPlatformKey).IsReachable())
+		// Startup packages and Generated packages do not need to add hidden dependencies or log warnings
+		if (Discovery->Instigator.Category != EInstigator::StartupPackage &&
+			Discovery->Instigator.Category != EInstigator::GeneratedPackage)
 		{
-			// We are adding the expectation that this package will be loaded for something during the cook. Doing so
-			// is only expected from a few types of instigators, or from external package requests, or from cluster
-			// exploration. If not expected, add a diagnostic message, and add it as a hidden dependency.
-			if (Discovery->Instigator.Category != EInstigator::StartupPackage &&
-				Discovery->Instigator.Category != EInstigator::GeneratedPackage)
+			// If there are other discovered packages we have already added to this cluster, then defer this one
+			// until we have explored those; add this one to the next cluster. Exploring those earlier discoveries
+			// might add this one through cluster exploration and not require a hidden dependency.
+			if (!OwnedPackageDatas.IsEmpty())
 			{
-				// If there are other discovered packages we have already added to this cluster, then defer this one
-				// until we have explored those; add this one to the next cluster. Exploring those earlier discoveries
-				// might uncover this one and make it expected.
-				if (!OwnedPackageDatas.IsEmpty())
-				{
-					break;
-				}
+				break;
+			}
 
-				COTFS.OnDiscoveredPackageDebug(PackageData.GetPackageName(), Discovery->Instigator);
+			if (Discovery->ReachablePlatforms.GetSource() == EDiscoveredPlatformSet::CopyFromInstigator)
+			{
+				// Add it as a hidden dependency so that future platforms discovered as reachable in
+				// the instigator will also be marked as reachable in the dependency.
 				if (COTFS.bCanSkipEditorReferencedPackagesWhenCooking)
 				{
 					FPackageData* InstigatorPackageData = Discovery->Instigator.Referencer.IsNone() ? nullptr
@@ -173,6 +170,14 @@ FRequestCluster::FRequestCluster(UCookOnTheFlyServer& InCOTFS, TRingBuffer<FDisc
 					}
 				}
 			}
+
+			// Adding packages to the cook should happen only for a few types of instigators, from external package
+			// requests, or during cluster exploration. If not expected, add a diagnostic message.
+			if (Discovery->Instigator.Category != EInstigator::SaveTimeHardDependency &&
+				Discovery->Instigator.Category != EInstigator::SaveTimeSoftDependency)
+			{
+				COTFS.OnDiscoveredPackageDebug(PackageData.GetPackageName(), Discovery->Instigator);
+			}
 		}
 		// Add the new reachable platforms
 		PackageData.AddReachablePlatforms(*this, NewReachablePlatforms, MoveTemp(Discovery->Instigator));
@@ -249,7 +254,7 @@ void FRequestCluster::FetchPackageNames(const FCookerTimer& CookerTimer, bool& b
 		{
 			PackageData->AddReachablePlatforms(*this, Request.GetPlatforms(), MoveTemp(Request.GetInstigator()));
 			PullIntoCluster(*PackageData);
-			PackageData->AddUrgency(bExternalRequestsAreUrgent, false /* bAllowUpdateState */);
+			PackageData->AddUrgency(Request.IsUrgent(), false /* bAllowUpdateState */);
 		}
 		else
 		{
@@ -257,13 +262,13 @@ void FRequestCluster::FetchPackageNames(const FCookerTimer& CookerTimer, bool& b
 			{
 				// If it's already in progress with no new platforms, we don't need to add it to the cluster, but add
 				// add on our urgency setting
-				PackageData->AddUrgency(bExternalRequestsAreUrgent, true /* bAllowUpdateState */);
+				PackageData->AddUrgency(Request.IsUrgent(), true /* bAllowUpdateState */);
 			}
 			else if (PackageData->GetPlatformsNeedingCookingNum() > 0)
 			{
 				// If it's missing cookable platforms and not in progress we need to add it to the cluster for cooking
 				PullIntoCluster(*PackageData);
-				PackageData->AddUrgency(bExternalRequestsAreUrgent, true /* bAllowUpdateState */);
+				PackageData->AddUrgency(Request.IsUrgent(), true /* bAllowUpdateState */);
 			}
 		}
 		// Add on our completion callback, or call it immediately if already done
diff --git a/Engine/Source/Editor/UnrealEd/Private/Cooker/CookRequestCluster.h b/Engine/Source/Editor/UnrealEd/Private/Cooker/CookRequestCluster.h
--- a/Engine/Source/Editor/UnrealEd/Private/Cooker/CookRequestCluster.h
+++ b/Engine/Source/Editor/UnrealEd/Private/Cooker/CookRequestCluster.h
@@ -47,8 +47,7 @@ namespace UE::Cook
 class FRequestCluster
 {
 public:
-	FRequestCluster(UCookOnTheFlyServer& COTFS, TArray<FFilePlatformRequest>&& InRequests,
-		bool bInExternalRequestsAreUrgent);
+	FRequestCluster(UCookOnTheFlyServer& COTFS, TArray<FFilePlatformRequest>&& InRequests);
 	FRequestCluster(UCookOnTheFlyServer& COTFS, FPackageDataSet&& InRequests);
 	FRequestCluster(UCookOnTheFlyServer& COTFS, TRingBuffer<FDiscoveryQueueElement>& DiscoveryQueue);
 	FRequestCluster(FRequestCluster&&) = default;
@@ -331,7 +330,6 @@ private:
 	IAssetRegistry& AssetRegistry;
 	FPackageTracker& PackageTracker;
 	FBuildDefinitions& BuildDefinitions;
-	bool bExternalRequestsAreUrgent = false;
 	bool bAllowHardDependencies = true;
 	bool bAllowSoftDependencies = true;
 	bool bHybridIterativeEnabled = true;
diff --git a/Engine/Source/Editor/UnrealEd/Private/Cooker/CookRequests.cpp b/Engine/Source/Editor/UnrealEd/Private/Cooker/CookRequests.cpp
--- a/Engine/Source/Editor/UnrealEd/Private/Cooker/CookRequests.cpp
+++ b/Engine/Source/Editor/UnrealEd/Private/Cooker/CookRequests.cpp
@@ -56,6 +56,7 @@ FFilePlatformRequest::FFilePlatformRequest(const FFilePlatformRequest& InFilePla
 	: Filename(InFilePlatformRequest.Filename)
 	, Platforms(InFilePlatformRequest.Platforms)
 	, Instigator(InFilePlatformRequest.Instigator)
+	, bUrgent(InFilePlatformRequest.bUrgent)
 {
 	check(!InFilePlatformRequest.CompletionCallback); // CompletionCallbacks can not be copied, so the caller's intent is not clear in this constructor if the input has one
 }
@@ -65,6 +66,7 @@ FFilePlatformRequest::FFilePlatformRequest(FFilePlatformRequest&& InFilePlatform
 	, Platforms(MoveTemp(InFilePlatformRequest.Platforms))
 	, CompletionCallback(MoveTemp(InFilePlatformRequest.CompletionCallback))
 	, Instigator(MoveTemp(InFilePlatformRequest.Instigator))
+	, bUrgent(InFilePlatformRequest.bUrgent)
 {
 }
 
@@ -75,6 +77,7 @@ FFilePlatformRequest& FFilePlatformRequest::operator=(FFilePlatformRequest&& InF
 	check(!CompletionCallback); // We don't support holding multiple completion callbacks
 	CompletionCallback = MoveTemp(InFileRequest.CompletionCallback);
 	Instigator = MoveTemp(InFileRequest.Instigator);
+	bUrgent = InFileRequest.bUrgent;
 	return *this;
 }
 
diff --git a/Engine/Source/Editor/UnrealEd/Private/Cooker/CookRequests.h b/Engine/Source/Editor/UnrealEd/Private/Cooker/CookRequests.h
--- a/Engine/Source/Editor/UnrealEd/Private/Cooker/CookRequests.h
+++ b/Engine/Source/Editor/UnrealEd/Private/Cooker/CookRequests.h
@@ -37,6 +37,7 @@ protected:
 	TArray<const ITargetPlatform*> Platforms;
 	FCompletionCallback CompletionCallback;
 	FInstigator Instigator;
+	bool bUrgent = false;
 
 public:
 	FFilePlatformRequest() = default;
@@ -56,6 +57,9 @@ public:
 	const FName& GetFilename() const;
 	FInstigator& GetInstigator();
 
+	void SetUrgent(bool bInUrgent) { bUrgent = bInUrgent; }
+	bool IsUrgent() const { return bUrgent; }
+
 	const TArray<const ITargetPlatform*>& GetPlatforms() const;
 	TArray<const ITargetPlatform*>& GetPlatforms();
 	void RemovePlatform(const ITargetPlatform* Platform);
diff --git a/Engine/Source/Editor/UnrealEd/Private/Cooker/CookTypes.cpp b/Engine/Source/Editor/UnrealEd/Private/Cooker/CookTypes.cpp
--- a/Engine/Source/Editor/UnrealEd/Private/Cooker/CookTypes.cpp
+++ b/Engine/Source/Editor/UnrealEd/Private/Cooker/CookTypes.cpp
@@ -6,6 +6,7 @@
 #include "CookPackageData.h"
 #include "Containers/StringView.h"
 #include "DerivedDataRequest.h"
+#include "Editor.h"
 #include "HAL/PlatformTLS.h"
 #include "HAL/PlatformTime.h"
 #include "Interfaces/ITargetPlatform.h"
@@ -15,6 +16,7 @@
 #include "Misc/CommandLine.h"
 #include "Misc/Parse.h"
 #include "PackageTracker.h"
+#include "Serialization/PackageWriterToSharedBuffer.h"
 
 LLM_DEFINE_TAG(Cooker_CachedPlatformData);
 
@@ -384,8 +386,7 @@ void FDiscoveredPlatformSet::ConvertToBitfield(TConstArrayView<const ITargetPlat
 
 TConstArrayView<const ITargetPlatform*> FDiscoveredPlatformSet::GetPlatforms(UCookOnTheFlyServer& COTFS,
 	FInstigator* Instigator, TConstArrayView<const ITargetPlatform*> OrderedPlatforms,
-	TArray<const ITargetPlatform*, TInlineAllocator<ExpectedMaxNumPlatforms>>* OutBuffer,
-	bool bAllowPartialInstigatorResults)
+	TArray<const ITargetPlatform*, TInlineAllocator<ExpectedMaxNumPlatforms>>* OutBuffer)
 {
 	switch (Source)
 	{
@@ -409,8 +410,7 @@ TConstArrayView<const ITargetPlatform*> FDiscoveredPlatformSet::GetPlatforms(UCo
 	case EDiscoveredPlatformSet::CopyFromInstigator:
 		Platforms.Reset();
 		check(Instigator);
-		FPackageData::GetReachablePlatformsForInstigator(COTFS, Instigator->Referencer, Platforms,
-			bAllowPartialInstigatorResults);
+		FPackageData::GetReachablePlatformsForInstigator(COTFS, Instigator->Referencer, Platforms);
 		return Platforms;
 	default:
 		checkNoEntry();
diff --git a/Engine/Source/Editor/UnrealEd/Private/Cooker/CookTypes.h b/Engine/Source/Editor/UnrealEd/Private/Cooker/CookTypes.h
--- a/Engine/Source/Editor/UnrealEd/Private/Cooker/CookTypes.h
+++ b/Engine/Source/Editor/UnrealEd/Private/Cooker/CookTypes.h
@@ -486,8 +486,7 @@ struct FDiscoveredPlatformSet
 	void OnPlatformAddedToSession(const ITargetPlatform* Platform);
 	TConstArrayView<const ITargetPlatform*> GetPlatforms(UCookOnTheFlyServer& COTFS,
 		FInstigator* Instigator, TConstArrayView<const ITargetPlatform*> OrderedPlatforms,
-		TArray<const ITargetPlatform*, TInlineAllocator<ExpectedMaxNumPlatforms>>* OutBuffer,
-		bool bAllowPartialInstigatorResults=false);
+		TArray<const ITargetPlatform*, TInlineAllocator<ExpectedMaxNumPlatforms>>* OutBuffer);
 	/** If the current type is EmbeddedBitField, change it to EmbeddedList. */
 	void ConvertFromBitfield(TConstArrayView<const ITargetPlatform*> OrderedPlatforms);
 	/** If the current type is EmbeddedList, change it to EmbeddedBitfield. Asserts if the type is already EmbeddedBitfield. */
diff --git a/Engine/Source/Editor/UnrealEd/Private/Cooker/CookWorkerServer.cpp b/Engine/Source/Editor/UnrealEd/Private/Cooker/CookWorkerServer.cpp
--- a/Engine/Source/Editor/UnrealEd/Private/Cooker/CookWorkerServer.cpp
+++ b/Engine/Source/Editor/UnrealEd/Private/Cooker/CookWorkerServer.cpp
@@ -867,8 +867,7 @@ void FCookWorkerServer::QueueDiscoveredPackage(FDiscoveredPackageReplication&& D
 	}
 	else
 	{
-		DiscoveredPlatforms = Platforms.GetPlatforms(COTFS, &Instigator, OrderedSessionAndSpecialPlatforms, &BufferPlatforms,
-			true /* bAllowPartialInstigatorResults */);
+		DiscoveredPlatforms = Platforms.GetPlatforms(COTFS, &Instigator, OrderedSessionAndSpecialPlatforms, &BufferPlatforms);
 	}
 
 	if (PackageData.HasReachablePlatforms(DiscoveredPlatforms))
diff --git a/Engine/Source/Editor/UnrealEd/Private/Cooker/PackageTracker.cpp b/Engine/Source/Editor/UnrealEd/Private/Cooker/PackageTracker.cpp
--- a/Engine/Source/Editor/UnrealEd/Private/Cooker/PackageTracker.cpp
+++ b/Engine/Source/Editor/UnrealEd/Private/Cooker/PackageTracker.cpp
@@ -140,8 +140,21 @@ void FPackageTracker::NotifyUObjectCreated(const class UObjectBase* Object, int3
 #else
 			FName ReferencerName(NAME_None);
 #endif
-			FInstigator Instigator(EInstigator::Unsolicited, ReferencerName);
-			if (COTFS.bHiddenDependenciesDebug)
+			EInstigator InstigatorType;
+			switch (FCookLoadScope::GetCurrentValue())
+			{
+			case ECookLoadType::EditorOnly:
+				InstigatorType = EInstigator::EditorOnlyLoad;
+				break;
+			case ECookLoadType::UsedInGame:
+				InstigatorType = EInstigator::SaveTimeSoftDependency;
+				break;
+			default:
+				InstigatorType = EInstigator::Unsolicited;
+				break;
+			}
+			FInstigator Instigator(InstigatorType, ReferencerName);
+			if (InstigatorType == EInstigator::Unsolicited && COTFS.bHiddenDependenciesDebug)
 			{
 				COTFS.OnDiscoveredPackageDebug(Package->GetFName(), Instigator);
 			}
