Load ScriptableEditorWidgets alongside Blutilities and UMGEditor, as it is now owns the DetailsView that used to be in UMGEditor, this fixes projects that were using the DetailsView (which was moved in 25196253) project implicitly via the engine dependencies on UMGEditor

diff --git a/Engine/Source/Editor/UnrealEd/Private/EditorEngine.cpp b/Engine/Source/Editor/UnrealEd/Private/EditorEngine.cpp
--- a/Engine/Source/Editor/UnrealEd/Private/EditorEngine.cpp
+++ b/Engine/Source/Editor/UnrealEd/Private/EditorEngine.cpp
@@ -1150,6 +1150,7 @@ void UEditorEngine::Init(IEngineLoop* InEngineLoop)
 			TEXT("EditorSettingsViewer"),
 			TEXT("ProjectSettingsViewer"),
 			TEXT("Blutility"),
+			TEXT("ScriptableEditorWidgets"),
 			TEXT("XmlParser"),
 			TEXT("UndoHistory"),
 			TEXT("DeviceProfileEditor"),
diff --git a/Engine/Source/Editor/UnrealEd/UnrealEd.Build.cs b/Engine/Source/Editor/UnrealEd/UnrealEd.Build.cs
--- a/Engine/Source/Editor/UnrealEd/UnrealEd.Build.cs
+++ b/Engine/Source/Editor/UnrealEd/UnrealEd.Build.cs
@@ -246,6 +246,7 @@ public class UnrealEd : ModuleRules
 				"SettingsEditor",
 				"StringTableEditor",
 				"Blutility",
+				"ScriptableEditorWidgets",
 				"WorkspaceMenuStructure",
 				"PlacementMode",
 				"MeshUtilities",
