SourceControl: Fixed crash in FBookmarkScoped destructor.

diff --git a/Engine/Source/Editor/UnrealEd/Private/Bookmarks/BookmarkScoped.cpp b/Engine/Source/Editor/UnrealEd/Private/Bookmarks/BookmarkScoped.cpp
--- a/Engine/Source/Editor/UnrealEd/Private/Bookmarks/BookmarkScoped.cpp
+++ b/Engine/Source/Editor/UnrealEd/Private/Bookmarks/BookmarkScoped.cpp
@@ -23,9 +23,11 @@ FBookmarkScoped::~FBookmarkScoped()
 	// Restore all viewport bookmarks
 	for (FEditorViewportClient* ViewportClient : GEditor->GetAllViewportClients())
 	{
-		UBookMark* BookMark = BookMarks.FindChecked(ViewportClient);
-		TSharedPtr<FBookmarkBaseJumpToSettings> Settings = MakeShared<FBookmarkJumpToSettings>();
-		Actions->JumpToBookmark(BookMark, Settings, *ViewportClient);
+		if (UBookMark** BookMark = BookMarks.Find(ViewportClient))
+		{
+			TSharedPtr<FBookmarkBaseJumpToSettings> Settings = MakeShared<FBookmarkJumpToSettings>();
+			Actions->JumpToBookmark(*BookMark, Settings, *ViewportClient);
+		}
 	}
 
 	for (TPair<FEditorViewportClient*, UBookMark*>& Pair : BookMarks)
