Fix TextAssetCommandlet compile errors when text archive support is disabled

diff --git a/Engine/Source/Editor/UnrealEd/Private/Commandlets/TextAssetCommandlet.cpp b/Engine/Source/Editor/UnrealEd/Private/Commandlets/TextAssetCommandlet.cpp
--- a/Engine/Source/Editor/UnrealEd/Private/Commandlets/TextAssetCommandlet.cpp
+++ b/Engine/Source/Editor/UnrealEd/Private/Commandlets/TextAssetCommandlet.cpp
@@ -217,6 +217,7 @@ void GenerateClassSchema(UClass* Class, FStructuredArchiveRecord Record, TArray<
 	}
 }
 
+#if WITH_TEXT_ARCHIVE_SUPPORT
 void GenerateSchema()
 {
 	//static const FName NAME_SpecificClass(TEXT("TextAssetTestObject"));
@@ -253,6 +254,7 @@ void GenerateSchema()
 
 	StructuredArchive.Close();
 }
+#endif
 
 bool UTextAssetCommandlet::DoTextAssetProcessing(const FString& InCommandLine)
 {
@@ -315,7 +317,11 @@ bool UTextAssetCommandlet::DoTextAssetProcessing(const FProcessingArgs& InArgs)
 		return true;
 
 	case ETextAssetCommandletMode::GenerateSchema:
+#if WITH_TEXT_ARCHIVE_SUPPORT
 		GenerateSchema();
+#else 
+		UE_LOG(LogTextAsset, Error, TEXT("Unable to generate schema when compiled with WITH_TEXT_ARCHIVE_SUPPORT=0"));
+#endif
 		break;
 
 	default:
