PR #10368: GroupActors: Add a FixupGroupActors function (Contributed by [REDACTED_PERSON])

diff --git a/Engine/Source/Editor/LevelEditor/Private/LevelEditor.cpp b/Engine/Source/Editor/LevelEditor/Private/LevelEditor.cpp
--- a/Engine/Source/Editor/LevelEditor/Private/LevelEditor.cpp
+++ b/Engine/Source/Editor/LevelEditor/Private/LevelEditor.cpp
@@ -1843,6 +1843,11 @@ void FLevelEditorModule::BindGlobalLevelEditorCommands()
 		Commands.OpenMergeActor,
 		FExecuteAction::CreateStatic(&FLevelEditorActionCallbacks::OpenMergeActor_Clicked)
 		);
+
+	ActionList.MapAction(
+		Commands.FixupGroupActor,
+		FExecuteAction::CreateStatic(&FLevelEditorActionCallbacks::FixupGroupActor_Clicked)
+	);
 }
 
 TSharedPtr<FLevelEditorOutlinerSettings> FLevelEditorModule::GetLevelEditorOutlinerSettings() const
diff --git a/Engine/Source/Editor/LevelEditor/Private/LevelEditorActions.cpp b/Engine/Source/Editor/LevelEditor/Private/LevelEditorActions.cpp
--- a/Engine/Source/Editor/LevelEditor/Private/LevelEditorActions.cpp
+++ b/Engine/Source/Editor/LevelEditor/Private/LevelEditorActions.cpp
@@ -3736,6 +3736,7 @@ void FLevelEditorCommands::RegisterCommands()
 	UI_COMMAND( RemoveActorsFromGroup, "Remove from Group", "Removes the selected actors from the selected groups", EUserInterfaceActionType::Button, FInputChord() );
 	UI_COMMAND( LockGroup, "Lock", "Locks the selected groups", EUserInterfaceActionType::Button, FInputChord() );
 	UI_COMMAND( UnlockGroup, "Unlock", "Unlocks the selected groups", EUserInterfaceActionType::Button, FInputChord() );
+	UI_COMMAND( FixupGroupActor, "Fixup Group Actor", "Removes null actors and deletes the GroupActor if it is empty.", EUserInterfaceActionType::Button, FInputChord());
 
 #if PLATFORM_MAC
 	UI_COMMAND( ShowAll, "Show All Actors", "Shows all actors", EUserInterfaceActionType::Button, FInputChord( EModifierKey::Command, EKeys::H ) );
@@ -3903,4 +3904,12 @@ void FLevelEditorCommands::RegisterCommands()
 
 UE_ENABLE_OPTIMIZATION_SHIP
 
+void FLevelEditorActionCallbacks::FixupGroupActor_Clicked()
+{
+	if (UActorGroupingUtils::IsGroupingActive())
+	{
+		AGroupActor::FixupGroupActor();
+	}
+}
+
 #undef LOCTEXT_NAMESPACE
diff --git a/Engine/Source/Editor/LevelEditor/Private/LevelEditorContextMenu.cpp b/Engine/Source/Editor/LevelEditor/Private/LevelEditorContextMenu.cpp
--- a/Engine/Source/Editor/LevelEditor/Private/LevelEditorContextMenu.cpp
+++ b/Engine/Source/Editor/LevelEditor/Private/LevelEditorContextMenu.cpp
@@ -1354,6 +1354,11 @@ void FLevelEditorContextMenuImpl::FillGroupMenu( UToolMenu* Menu )
 		{ 
 			Section.AddMenuEntry( FLevelEditorCommands::Get().AddActorsToGroup );
 		}
+
+		if (AGroupActor::SelectedGroupNeedsFixup())
+		{
+			Section.AddMenuEntry(FLevelEditorCommands::Get().FixupGroupActor);
+		}
 	}
 }
 
diff --git a/Engine/Source/Editor/LevelEditor/Public/LevelEditorActions.h b/Engine/Source/Editor/LevelEditor/Public/LevelEditorActions.h
--- a/Engine/Source/Editor/LevelEditor/Public/LevelEditorActions.h
+++ b/Engine/Source/Editor/LevelEditor/Public/LevelEditorActions.h
@@ -629,6 +629,8 @@ public:
         // Open merge actor command
 	TSharedPtr< FUICommandInfo > OpenMergeActor;
 
+	TSharedPtr< FUICommandInfo > FixupGroupActor;
+
 };
 
 /**
@@ -1422,5 +1424,8 @@ private:
 	 * @param SimActor		Simulating Actor in PIE or SIE
 	 */
 	static bool SaveAnimationFromSkeletalMeshComponent(AActor * EditorActor, AActor * SimActor, TArray<class USkeletalMeshComponent*> & OutEditorComponents);
+
+public:
+	static void FixupGroupActor_Clicked();
 };
 
diff --git a/Engine/Source/Editor/UnrealEd/Classes/Editor/GroupActor.h b/Engine/Source/Editor/UnrealEd/Classes/Editor/GroupActor.h
--- a/Engine/Source/Editor/UnrealEd/Classes/Editor/GroupActor.h
+++ b/Engine/Source/Editor/UnrealEd/Classes/Editor/GroupActor.h
@@ -221,6 +221,10 @@ private:
 	 * Helper function for Remove() and ClearAndRemove() - cleans up actors when grouping is no longer needed
 	 */
 	void PostRemove();
+
+public:
+	static UNREALED_API bool SelectedGroupNeedsFixup();
+	static UNREALED_API void FixupGroupActor();
 };
 
 namespace GroupActorHelpers
diff --git a/Engine/Source/Editor/UnrealEd/Private/GroupActor.cpp b/Engine/Source/Editor/UnrealEd/Private/GroupActor.cpp
--- a/Engine/Source/Editor/UnrealEd/Private/GroupActor.cpp
+++ b/Engine/Source/Editor/UnrealEd/Private/GroupActor.cpp
@@ -986,5 +986,68 @@ int32 AGroupActor::GetActorNum() const
 	return GroupActors.Num();
 }
 
+namespace UE::Editor::GroupActorUtil
+{
+	TArray<AGroupActor*> GetSelectedGroupActors(TArray<TObjectPtr<AActor>> ActiveGroupActors)
+	{
+		TArray<AGroupActor*> Result;
+		for (TObjectPtr<AActor> Actor : ActiveGroupActors)
+		{
+			if (AGroupActor* SelectedGroup = Cast<AGroupActor>(Actor))
+			{
+				if (SelectedGroup->IsSelected())
+				{
+					Result.Add(SelectedGroup);
+				}
+			}
+		}
+		return Result;
+	}
+}
+
+bool AGroupActor::SelectedGroupNeedsFixup()
+{
+	if (UWorld* EditorWorld = GEditor->GetEditorWorldContext().World())
+	{
+		TArray<AGroupActor*> SelectedGroupActors = UE::Editor::GroupActorUtil::GetSelectedGroupActors(EditorWorld->ActiveGroupActors);
+		if (!SelectedGroupActors.IsEmpty())
+		{
+			//check if there's at least 1 GroupActor contains a nullptr within it's list of Actors
+			return SelectedGroupActors.ContainsByPredicate([](const AGroupActor* SelectedGroupActor)
+				{
+					if (SelectedGroupActor)
+					{
+						return SelectedGroupActor->GroupActors.ContainsByPredicate([](const TObjectPtr<class AActor> Actor) { return Actor == nullptr; });
+					}
+					return false;
+				});
+		}
+	}
+	return false;
+} // namespace UE::Editor::GroupActorUtil
 
+void AGroupActor::FixupGroupActor()
+{
+	if (UWorld* EditorWorld = GEditor->GetEditorWorldContext().World())
+	{
+		TArray<AGroupActor*> SelectedGroupActors = UE::Editor::GroupActorUtil::GetSelectedGroupActors(EditorWorld->ActiveGroupActors);
+		if (!SelectedGroupActors.IsEmpty())
+		{
+			for (AGroupActor* SelectedGroupActor : SelectedGroupActors)
+			{
+				if (SelectedGroupActor && SelectedGroupActor->GroupActors.ContainsByPredicate([](const TObjectPtr<class AActor> Actor) { return Actor == nullptr; }))
+				{
+					//remove all nullptr entries in the GroupActors array.
+					SelectedGroupActor->Modify();
+					SelectedGroupActor->GroupActors.RemoveAll([](const TObjectPtr<class AActor> Actor) { return Actor == nullptr; });
+					SelectedGroupActor->GroupActors.Shrink();
 
+					if (SelectedGroupActor->GroupActors.IsEmpty())
+					{
+						SelectedGroupActor->PostRemove();
+					}
+				}
+			}
+		}
+	}
+}
