FEditorModeTools::DeactivateAllModes: Fix Toolkit leak

#rnx

diff --git a/Engine/Source/Editor/UnrealEd/Private/EditorModeManager.cpp b/Engine/Source/Editor/UnrealEd/Private/EditorModeManager.cpp
--- a/Engine/Source/Editor/UnrealEd/Private/EditorModeManager.cpp
+++ b/Engine/Source/Editor/UnrealEd/Private/EditorModeManager.cpp
@@ -757,6 +757,22 @@ void FEditorModeTools::RemoveAllDelegateHandlers()
 	OnCoordSystemChanged().Clear();
 }
 
+void FEditorModeTools::DeactivateModeAtIndex(int32 Index)
+{
+	UEdMode* Mode = ActiveScriptableModes[Index];
+	const FEditorModeID ModeID = Mode->GetID();
+	PendingDeactivateModes.Emplace(ModeID, Mode);
+	ActiveScriptableModes.RemoveAt(Index);
+		
+	if (const TSharedPtr<FModeToolkit> Toolkit = Mode->GetToolkit().Pin())
+	{
+		FToolkitManager::Get().CloseToolkit(Toolkit.ToSharedRef());
+	}
+
+	constexpr bool bIsEnteringMode = false;
+	BroadcastEditorModeIDChanged(ModeID, bIsEnteringMode);
+}
+
 void FEditorModeTools::DeactivateMode( FEditorModeID InID )
 {
 	// Find the mode from the ID and exit it.
@@ -765,17 +781,7 @@ void FEditorModeTools::DeactivateMode( FEditorModeID InID )
 		UEdMode* Mode = ActiveScriptableModes[Index];
 		if (Mode->GetID() == InID)
 		{
-			PendingDeactivateModes.Emplace(InID, Mode);
-			ActiveScriptableModes.RemoveAt(Index);
-
-			constexpr bool bIsEnteringMode = false;
-
-			if (const TSharedPtr<FModeToolkit> Toolkit = Mode->GetToolkit().Pin())
-			{
-				FToolkitManager::Get().CloseToolkit(Toolkit.ToSharedRef());
-			}
-			
-			BroadcastEditorModeIDChanged(InID, bIsEnteringMode);
+			DeactivateModeAtIndex(Index);
 			break;
 		}
 	};
@@ -785,12 +791,7 @@ void FEditorModeTools::DeactivateAllModes()
 {
 	for (int32 Index = ActiveScriptableModes.Num() - 1; Index >= 0; --Index)
 	{
-		FEditorModeID ModeID = ActiveScriptableModes[Index]->GetID();
-		PendingDeactivateModes.Emplace(ModeID, ActiveScriptableModes[Index]);
-		ActiveScriptableModes.RemoveAt(Index);
-		
-		constexpr bool bIsEnteringMode = false;
-		BroadcastEditorModeIDChanged(ModeID, bIsEnteringMode);
+		DeactivateModeAtIndex(Index);
 	};
 }
 
diff --git a/Engine/Source/Editor/UnrealEd/Public/EditorModeManager.h b/Engine/Source/Editor/UnrealEd/Public/EditorModeManager.h
--- a/Engine/Source/Editor/UnrealEd/Public/EditorModeManager.h
+++ b/Engine/Source/Editor/UnrealEd/Public/EditorModeManager.h
@@ -109,6 +109,7 @@ protected:
 	/** Removes the mode ID from the tools manager when a mode is unregistered */
 	void OnModeUnregistered(FEditorModeID ModeID);
 		
+	void DeactivateModeAtIndex(int32 Index);
 public:
 
 	/**
