Use the Component's outer to get the container values instead of the CDO. This fixes an issue where subobjects that are stored in container properties (TArray, TSet, and TMap) would not show up in the details panel if they were created dynamically at runtime.

diff --git a/Engine/Source/Editor/UnrealEd/Private/Kismet2/ComponentEditorUtils.cpp b/Engine/Source/Editor/UnrealEd/Private/Kismet2/ComponentEditorUtils.cpp
--- a/Engine/Source/Editor/UnrealEd/Private/Kismet2/ComponentEditorUtils.cpp
+++ b/Engine/Source/Editor/UnrealEd/Private/Kismet2/ComponentEditorUtils.cpp
@@ -176,7 +176,6 @@ FProperty* FComponentEditorUtils::GetPropertyForEditableNativeComponent(const UA
 	// Note: We aren't concerned with whether the component is marked editable - the component itself is responsible for determining which of its properties are editable	
 	UObject* ComponentOuter = (NativeComponent ? NativeComponent->GetOuter() : nullptr);
 	UClass* OwnerClass = (ComponentOuter ? ComponentOuter->GetClass() : nullptr);
-	UObject* OwnerCDO = (OwnerClass ? OwnerClass->GetDefaultObject() : nullptr);
 
 	if (OwnerClass != nullptr)
 	{
@@ -199,12 +198,12 @@ FProperty* FComponentEditorUtils::GetPropertyForEditableNativeComponent(const UA
 	
 		// We have to check for array properties as well because they are not FObjectProperties and we want to be able to
 		// edit the inside of it
-		if (OwnerCDO != nullptr)
+		if (ComponentOuter != nullptr)
 		{
 			for (TFieldIterator<FArrayProperty> PropIt(OwnerClass, EFieldIteratorFlags::IncludeSuper); PropIt; ++PropIt)
 			{
 				FArrayProperty* TestProperty = *PropIt;
-				void* ArrayPropInstAddress = TestProperty->ContainerPtrToValuePtr<void>(OwnerCDO);
+				void* ArrayPropInstAddress = TestProperty->ContainerPtrToValuePtr<void>(ComponentOuter);
 
 				// Ensure that this property is valid
 				FObjectProperty* ArrayEntryProp = CastField<FObjectProperty>(TestProperty->Inner);
@@ -232,7 +231,7 @@ FProperty* FComponentEditorUtils::GetPropertyForEditableNativeComponent(const UA
 			for (TFieldIterator<FMapProperty> PropIt(OwnerClass, EFieldIteratorFlags::IncludeSuper); PropIt; ++PropIt)
 			{
 				FMapProperty* TestProperty = *PropIt;
-				void* MapPropInstAddress = TestProperty->ContainerPtrToValuePtr<void>(OwnerCDO);
+				void* MapPropInstAddress = TestProperty->ContainerPtrToValuePtr<void>(ComponentOuter);
 				
 				// Ensure that this property is valid and that it is marked as visible in the editor
 				FObjectProperty* MapValProp = CastField<FObjectProperty>(TestProperty->ValueProp);
@@ -261,7 +260,7 @@ FProperty* FComponentEditorUtils::GetPropertyForEditableNativeComponent(const UA
 			for (TFieldIterator<FSetProperty> PropIt(OwnerClass, EFieldIteratorFlags::IncludeSuper); PropIt; ++PropIt)
 			{
 				FSetProperty* TestProperty = *PropIt;
-				void* SetPropInstAddress = TestProperty->ContainerPtrToValuePtr<void>(OwnerCDO);
+				void* SetPropInstAddress = TestProperty->ContainerPtrToValuePtr<void>(ComponentOuter);
 
 				// Ensure that this property is valid and that it is marked visible
 				FObjectProperty* SetValProp = CastField<FObjectProperty>(TestProperty->ElementProp);
