Prioritize thumbnail found from the on-disk package if it's cooked because the asset cannot change and rendering it without editor-only data might not give the same result as when it was rendered uncooked.

diff --git a/Engine/Source/Editor/UnrealEd/Private/AssetThumbnail.cpp b/Engine/Source/Editor/UnrealEd/Private/AssetThumbnail.cpp
--- a/Engine/Source/Editor/UnrealEd/Private/AssetThumbnail.cpp
+++ b/Engine/Source/Editor/UnrealEd/Private/AssetThumbnail.cpp
@@ -41,6 +41,24 @@
 #include "IVREditorModule.h"
 #include "Framework/Application/SlateApplication.h"
 
+namespace AssetThumbnailPool
+{
+	const FObjectThumbnail* LoadThumbnailFromPackage(const FAssetData& AssetData, FThumbnailMap& OutThumbnailMap)
+	{
+		FString PackageFilename;
+		if (FPackageName::DoesPackageExist(AssetData.PackageName.ToString(), &PackageFilename))
+		{
+			const FName ObjectFullName = FName(*AssetData.GetFullName());
+			TSet<FName> ObjectFullNames;
+			ObjectFullNames.Add(ObjectFullName);
+
+			ThumbnailTools::LoadThumbnailsFromPackage(PackageFilename, ObjectFullNames, OutThumbnailMap);
+			return OutThumbnailMap.Find(ObjectFullName);
+		}
+		return nullptr;
+	};
+}
+
 FName FAssetThumbnailPool::CustomThumbnailTagName = "CustomThumbnail";
 
 class SAssetThumbnail : public SCompoundWidget
@@ -1177,84 +1195,99 @@ void FAssetThumbnailPool::Tick( float DeltaTime )
 bool FAssetThumbnailPool::LoadThumbnail(TSharedRef<FThumbnailInfo> ThumbnailInfo, bool& bIsAssetStillCompiling, const FAssetData& CustomAssetToRender)
 {
 	const FAssetData& AssetData = CustomAssetToRender.IsValid() ? CustomAssetToRender : ThumbnailInfo->AssetData;
-	UObject* Asset = AssetData.FastGetAsset();
-	if (Asset && !IsValidChecked(Asset))
+
+	FThumbnailMap ThumbnailMap;
+	const FObjectThumbnail* FoundThumbnail = nullptr;
+
+	// Prioritize thumbnail found from the on-disk package if it's cooked because the asset cannot change and 
+	// rendering it without editor-only data might not give the same result as when it was rendered uncooked.
+	if (AssetData.PackageFlags & PKG_Cooked)
 	{
-		Asset = nullptr;
+		FoundThumbnail = AssetThumbnailPool::LoadThumbnailFromPackage(AssetData, ThumbnailMap);
 	}
 
-	const bool bAreShadersCompiling = (GShaderCompilingManager && GShaderCompilingManager->IsCompiling());
-	if (Asset && !bAreShadersCompiling)
+	if (!FoundThumbnail)
 	{
-		//Avoid rendering the thumbnail of an asset that is currently edited asynchronously
-		const IInterface_AsyncCompilation* Interface_AsyncCompilation = Cast<IInterface_AsyncCompilation>(Asset);
-		bIsAssetStillCompiling = Interface_AsyncCompilation && Interface_AsyncCompilation->IsCompiling();
-		if (!bIsAssetStillCompiling)
+		// Render a fresh thumbnail from the loaded asset if possible
+		UObject* Asset = AssetData.FastGetAsset();
+		if (Asset && !IsValidChecked(Asset))
 		{
-			FThumbnailRenderingInfo* RenderInfo = GUnrealEd->GetThumbnailManager()->GetRenderingInfo(Asset);
-			if (RenderInfo != nullptr && RenderInfo->Renderer != nullptr)
-			{
-				FThumbnailInfo_RenderThread ThumbInfo = ThumbnailInfo.Get();
+			Asset = nullptr;
+		}
 
-				auto EnqueueThumbnailRender = [Asset, ThumbInfo, ThumbnailInfo]()
-				{
-					ENQUEUE_RENDER_COMMAND(SyncSlateTextureCommand)(
-						[ThumbInfo](FRHICommandListImmediate& RHICmdList)
-						{
-							if (ThumbInfo.ThumbnailTexture->GetTypedResource() != ThumbInfo.ThumbnailRenderTarget->GetTextureRHI())
-							{
-								ThumbInfo.ThumbnailTexture->ClearTextureData();
-								ThumbInfo.ThumbnailTexture->ReleaseDynamicRHI();
-								ThumbInfo.ThumbnailTexture->SetRHIRef(ThumbInfo.ThumbnailRenderTarget->GetTextureRHI(), ThumbInfo.Width, ThumbInfo.Height);
-							}
-						});
-
-					// // We have to wait for the render command to finish since thumbnail rendering cannot be done on the GPU currently
-					FRenderCommandFence Fence;
-					Fence.BeginFence();
-					Fence.Wait();
-
-					//@todo: this should be done on the GPU only but it is not supported by thumbnail tools yet
-					ThumbnailTools::RenderThumbnail(
-						Asset,
-						ThumbnailInfo->Width,
-						ThumbnailInfo->Height,
-						ThumbnailTools::EThumbnailTextureFlushMode::NeverFlush,
-						ThumbnailInfo->ThumbnailRenderTarget
-					);
-				};
-
-				EThumbnailRenderFrequency ThumbnailRenderFrequency = RenderInfo->Renderer->GetThumbnailRenderFrequency(Asset);
-
-				switch (ThumbnailRenderFrequency)
+		const bool bAreShadersCompiling = (GShaderCompilingManager && GShaderCompilingManager->IsCompiling());
+		if (Asset && !bAreShadersCompiling)
+		{
+			//Avoid rendering the thumbnail of an asset that is currently edited asynchronously
+			const IInterface_AsyncCompilation* Interface_AsyncCompilation = Cast<IInterface_AsyncCompilation>(Asset);
+			bIsAssetStillCompiling = Interface_AsyncCompilation && Interface_AsyncCompilation->IsCompiling();
+			if (!bIsAssetStillCompiling)
+			{
+				FThumbnailRenderingInfo* RenderInfo = GUnrealEd->GetThumbnailManager()->GetRenderingInfo(Asset);
+				if (RenderInfo != nullptr && RenderInfo->Renderer != nullptr)
 				{
-					case EThumbnailRenderFrequency::Realtime:
+					FThumbnailInfo_RenderThread ThumbInfo = ThumbnailInfo.Get();
+
+					auto EnqueueThumbnailRender = [Asset, ThumbInfo, ThumbnailInfo]()
 					{
-						EnqueueThumbnailRender();
-						return true;
-					}
-					case EThumbnailRenderFrequency::OnPropertyChange:
+						ENQUEUE_RENDER_COMMAND(SyncSlateTextureCommand)(
+							[ThumbInfo](FRHICommandListImmediate& RHICmdList)
+							{
+								if (ThumbInfo.ThumbnailTexture->GetTypedResource() != ThumbInfo.ThumbnailRenderTarget->GetTextureRHI())
+								{
+									ThumbInfo.ThumbnailTexture->ClearTextureData();
+									ThumbInfo.ThumbnailTexture->ReleaseDynamicRHI();
+									ThumbInfo.ThumbnailTexture->SetRHIRef(ThumbInfo.ThumbnailRenderTarget->GetTextureRHI(), ThumbInfo.Width, ThumbInfo.Height);
+								}
+							});
+
+						// // We have to wait for the render command to finish since thumbnail rendering cannot be done on the GPU currently
+						FRenderCommandFence Fence;
+						Fence.BeginFence();
+						Fence.Wait();
+
+						//@todo: this should be done on the GPU only but it is not supported by thumbnail tools yet
+						ThumbnailTools::RenderThumbnail(
+							Asset,
+							ThumbnailInfo->Width,
+							ThumbnailInfo->Height,
+							ThumbnailTools::EThumbnailTextureFlushMode::NeverFlush,
+							ThumbnailInfo->ThumbnailRenderTarget
+						);
+					};
+
+					EThumbnailRenderFrequency ThumbnailRenderFrequency = RenderInfo->Renderer->GetThumbnailRenderFrequency(Asset);
+
+					switch (ThumbnailRenderFrequency)
 					{
-						if (ThumbnailInfo->LastUpdateTime <= 0.0f)
+						case EThumbnailRenderFrequency::Realtime:
 						{
 							EnqueueThumbnailRender();
 							return true;
 						}
-						break;
-					}
-					case EThumbnailRenderFrequency::OnAssetSave:
-					{
-						// OnAssetSave is default behavior below, so nohing to do
-						break;
-					}
-					case EThumbnailRenderFrequency::Once:
-					{
-						// Eagerly return if we aren't interested in cached thumbnails
-						return true;
-					}
-					default:
-					{
-						break;
+						case EThumbnailRenderFrequency::OnPropertyChange:
+						{
+							if (ThumbnailInfo->LastUpdateTime <= 0.0f)
+							{
+								EnqueueThumbnailRender();
+								return true;
+							}
+							break;
+						}
+						case EThumbnailRenderFrequency::OnAssetSave:
+						{
+							// OnAssetSave is default behavior below, so nohing to do
+							break;
+						}
+						case EThumbnailRenderFrequency::Once:
+						{
+							// Eagerly return if we aren't interested in cached thumbnails
+							return true;
+						}
+						default:
+						{
+							break;
+						}
 					}
 				}
 			}
@@ -1262,22 +1295,16 @@ bool FAssetThumbnailPool::LoadThumbnail(TSharedRef<FThumbnailInfo> ThumbnailInfo
 	}
 
 	// If we could not render a fresh thumbnail, see if we already have a cached one to load
-	const FObjectThumbnail* FoundThumbnail = ThumbnailTools::FindCachedThumbnail(AssetData.GetFullName());
-
-	// If we don't have a cached thumbnail, try to find it on disk
 	FThumbnailMap ThumbnailMap;
 	if (!FoundThumbnail)
 	{
-		FString PackageFilename;
-		if (FPackageName::DoesPackageExist(AssetData.PackageName.ToString(), &PackageFilename))
-		{
-			const FName ObjectFullName = FName(*AssetData.GetFullName());
-			TSet<FName> ObjectFullNames;
-			ObjectFullNames.Add(ObjectFullName);
+		FoundThumbnail = ThumbnailTools::FindCachedThumbnail(AssetData.GetFullName());
+	}
 
-			ThumbnailTools::LoadThumbnailsFromPackage(PackageFilename, ObjectFullNames, ThumbnailMap);
-			FoundThumbnail = ThumbnailMap.Find(ObjectFullName);
-		}
+	// If we don't have a thumbnail cached in memory, try to find it on disk
+	if (!FoundThumbnail && !(AssetData.PackageFlags & PKG_Cooked))
+	{
+		FoundThumbnail = AssetThumbnailPool::LoadThumbnailFromPackage(AssetData, ThumbnailMap);
 	}
 
 	if (FoundThumbnail)
