Add ensure to verify unstated assumption. If there are chains of form old->less_old->new in the object replacment map then we would need to flatten them to ensure that no one gets a references to old replaced with a reference to a less_old object, rather than the final 'new' object

diff --git a/Engine/Source/Editor/UnrealEd/Private/Kismet2/KismetReinstanceUtilities.cpp b/Engine/Source/Editor/UnrealEd/Private/Kismet2/KismetReinstanceUtilities.cpp
--- a/Engine/Source/Editor/UnrealEd/Private/Kismet2/KismetReinstanceUtilities.cpp
+++ b/Engine/Source/Editor/UnrealEd/Private/Kismet2/KismetReinstanceUtilities.cpp
@@ -1,6 +1,7 @@
 // Copyright Epic Games, Inc. All Rights Reserved.
 
 #include "Kismet2/KismetReinstanceUtilities.h"
+#include "Algo/ForEach.h"
 #include "BlueprintCompilationManager.h"
 #include "ComponentInstanceDataCache.h"
 #include "Engine/Blueprint.h"
@@ -59,6 +60,61 @@ static FAutoConsoleVariableRef CVarUseLegacyAnimInstanceReinstancingBehavior(
 
 struct FReplaceReferenceHelper
 {
+	static void ValidateReplacementMappings(const TMap<UObject*, UObject*>& OldToNewInstanceMap)
+	{
+		// Test long unstated assumption - alternatively we could 'flatten' the chains
+		// but would then have to guard against cycles:
+		bool bFoundChains = false;
+		for (const TPair<UObject*, UObject*>& OldToNew : OldToNewInstanceMap)
+		{
+			bFoundChains = OldToNewInstanceMap.Find(OldToNew.Value) != nullptr;
+			if (bFoundChains)
+			{
+				break;
+			}
+		}
+
+		if (!bFoundChains)
+		{
+			return;
+		}
+
+		TSet<UObject*> ObjectsInvolved;
+		Algo::ForEach(OldToNewInstanceMap, [&ObjectsInvolved, &OldToNewInstanceMap](const TPair<UObject*, UObject*>& OldToNew)
+			{
+				UObject* const* NewMappedToNew = OldToNewInstanceMap.Find(OldToNew.Value);
+				if (NewMappedToNew != nullptr)
+				{
+					if(OldToNew.Key)
+					{
+						ObjectsInvolved.Add(OldToNew.Key);
+					}
+					if (OldToNew.Value)
+					{
+						ObjectsInvolved.Add(OldToNew.Value);
+					}
+					if(*NewMappedToNew)
+					{
+						ObjectsInvolved.Add(*NewMappedToNew);
+					}
+				}
+			});
+
+		TSet<UClass*> ClassesInvolved;
+		Algo::ForEach(ObjectsInvolved, [&ClassesInvolved](UObject* Object)
+			{
+				ClassesInvolved.Add(Object->GetClass());
+			});
+
+		TStringBuilder<256> NamesOfClasses;
+		Algo::ForEach(ClassesInvolved, [&NamesOfClasses](UClass* Class) { NamesOfClasses.Append(Class->GetName() + TEXT("\n")); });
+		TStringBuilder<256> NamesOfObjects;
+		Algo::ForEach(ObjectsInvolved, [&NamesOfObjects](UObject* Obj) { NamesOfObjects.Append(Obj->GetName() + TEXT("\n")); });
+
+		ensureMsgf(false, TEXT("Found chains of replacement objects while updating class layouts, please report a bug involving Classes:\n%sAnd Objects:\n%s"), 
+			*NamesOfClasses, *NamesOfObjects);
+	}
+
 	static void IncludeCDO(UClass* OldClass, UClass* NewClass, TMap<UObject*, UObject*> &OldToNewInstanceMap, TArray<UObject*> &SourceObjects, UObject* OriginalCDO)
 	{
 		UObject* OldCDO = OldClass->GetDefaultObject();
@@ -2551,6 +2607,7 @@ void FBlueprintCompileReinstancer::ReplaceInstancesOfClass_Inner(const TMap<UCla
 		}
 	}
 
+	FReplaceReferenceHelper::ValidateReplacementMappings(OldToNewInstanceMap);
 	FReplaceReferenceHelper::FindAndReplaceReferences(SourceObjects, ObjectsThatShouldUseOldStuff, ObjectsReplaced, OldToNewInstanceMap, ReinstancedObjectsWeakReferenceMap);
 	
 	for (UObject* Obj : ObjectsReplaced)
