Fix assert in UEditorLevelUtils::CopyOrMoveActorsToLevel when cancelling the task.

diff --git a/Engine/Source/Editor/UnrealEd/Private/EditorLevelUtils.cpp b/Engine/Source/Editor/UnrealEd/Private/EditorLevelUtils.cpp
--- a/Engine/Source/Editor/UnrealEd/Private/EditorLevelUtils.cpp
+++ b/Engine/Source/Editor/UnrealEd/Private/EditorLevelUtils.cpp
@@ -190,69 +190,77 @@ int32 UEditorLevelUtils::CopyOrMoveActorsToLevel(const TArray<AActor*>& ActorsTo
 				
 				OwningWorld->SetCurrentLevel(DestLevel);
 
+				bool bSelectionChanged = false;
+				FDelegateHandle SelectionChangedEvent = USelection::SelectionChangedEvent.AddLambda([&bSelectionChanged](UObject* Object) { bSelectionChanged = true; });
+
 				const bool bDuplicate = false;
 				const bool bOffsetLocations = false;
 				const bool bWarnIfHidden = false;
 				GEditor->edactPasteSelected(OwningWorld, bDuplicate, bOffsetLocations, bWarnIfHidden, &DestinationData);
 
+				USelection::SelectionChangedEvent.Remove(SelectionChangedEvent);
+
 				// Restore the original current level
 				OwningWorld->SetCurrentLevel(OldCurrentLevel);
 
-				// Build a remapping of old to new names so we can do a fixup
-				for (FSelectionIterator It(GEditor->GetSelectedActorIterator()); It; ++It)
+				if (bSelectionChanged)
 				{
-					AActor* Actor = static_cast<AActor*>(*It);
-
-					if (ActorPathMapping.IsValidIndex(Actor->CopyPasteId))
+					// Build a remapping of old to new names so we can do a fixup
+					for (FSelectionIterator It(GEditor->GetSelectedActorIterator()); It; ++It)
 					{
-						TTuple<FSoftObjectPath, FSoftObjectPath>& Tuple = ActorPathMapping[Actor->CopyPasteId];
-						check(Tuple.Value.IsNull());
+						AActor* Actor = static_cast<AActor*>(*It);
 
-						Tuple.Value = FSoftObjectPath(Actor);
-						if (OutActors)
+						if (ActorPathMapping.IsValidIndex(Actor->CopyPasteId))
 						{
-							OutActors->Add(Actor);
+							TTuple<FSoftObjectPath, FSoftObjectPath>& Tuple = ActorPathMapping[Actor->CopyPasteId];
+							check(Tuple.Value.IsNull());
+
+							Tuple.Value = FSoftObjectPath(Actor);
+							if (OutActors)
+							{
+								OutActors->Add(Actor);
+							}
 						}
-					}
-					else
-					{
-						UE_LOG(LogLevelTools, Error, TEXT("Cannot find remapping for moved actor ID %s, any soft references pointing to it will be broken!"), *Actor->GetPathName());
-					}
-					// Reset CopyPasteId on new actors
-					Actor->CopyPasteId = INDEX_NONE;
-				}
-
-				// Only do Asset Rename on Move (Copy should not affect existing references)
-				if (bMoveActors)
-				{
-					FAssetToolsModule& AssetToolsModule = FModuleManager::GetModuleChecked<FAssetToolsModule>(TEXT("AssetTools"));
-					TArray<FAssetRenameData> RenameData;
-
-					for (TTuple<FSoftObjectPath, FSoftObjectPath>& Pair : ActorPathMapping)
-					{
-						if (Pair.Value.IsValid())
+						else
 						{
-							RenameData.Add(FAssetRenameData(Pair.Key, Pair.Value, true));
+							UE_LOG(LogLevelTools, Error, TEXT("Cannot find remapping for moved actor ID %s, any soft references pointing to it will be broken!"), *Actor->GetPathName());
 						}
+						// Reset CopyPasteId on new actors
+						Actor->CopyPasteId = INDEX_NONE;
 					}
 
-					if (RenameData.Num() > 0)
+					// Only do Asset Rename on Move (Copy should not affect existing references)
+					if (bMoveActors)
 					{
-						if (bWarnAboutRenaming)
+						FAssetToolsModule& AssetToolsModule = FModuleManager::GetModuleChecked<FAssetToolsModule>(TEXT("AssetTools"));
+						TArray<FAssetRenameData> RenameData;
+
+						for (TTuple<FSoftObjectPath, FSoftObjectPath>& Pair : ActorPathMapping)
 						{
-							AssetToolsModule.Get().RenameAssetsWithDialog(RenameData);
+							if (Pair.Value.IsValid())
+							{
+								RenameData.Add(FAssetRenameData(Pair.Key, Pair.Value, true));
+							}
 						}
-						else
+
+						if (RenameData.Num() > 0)
 						{
-							AssetToolsModule.Get().RenameAssets(RenameData);
+							if (bWarnAboutRenaming)
+							{
+								AssetToolsModule.Get().RenameAssetsWithDialog(RenameData);
+							}
+							else
+							{
+								AssetToolsModule.Get().RenameAssets(RenameData);
+							}
 						}
 					}
-				}
 
-				// Restore new level visibility to previous state
-				if (!bLevelVisible)
-				{
-					UEditorLevelUtils::SetLevelVisibility(DestLevel, false, false);
+					// Restore new level visibility to previous state
+					if (!bLevelVisible)
+					{
+						UEditorLevelUtils::SetLevelVisibility(DestLevel, false, false);
+					}
 				}
 			}
 
