Fix re-instantiation of sub objects that are not default sub-object and are containing themselves.

diff --git a/Engine/Source/Editor/UnrealEd/Private/Kismet2/KismetReinstanceUtilities.cpp b/Engine/Source/Editor/UnrealEd/Private/Kismet2/KismetReinstanceUtilities.cpp
--- a/Engine/Source/Editor/UnrealEd/Private/Kismet2/KismetReinstanceUtilities.cpp
+++ b/Engine/Source/Editor/UnrealEd/Private/Kismet2/KismetReinstanceUtilities.cpp
@@ -2409,10 +2409,50 @@ void FBlueprintCompileReinstancer::ReplaceInstancesOfClass_Inner(const TMap<UCla
 					}
 
 					AActor* OldActor = Cast<AActor>(OldObject);
+					bool bIsValid = IsValid(OldObject);
+
+					// If the object is not valid, it might been replaced by the IntancedReferenceSubobjectHelper code.
+					// Anyhow we still need to replace the object class as it only duplicated the object, 
+					// So check if it has replacement and process it now if the IntancedReferenceSubobjectHelper reinstanciated the class being replaced.
+					UObject* PrevOldObject = nullptr;
+					if (!bIsValid)
+					{
+						// Was it replaced because it owner got re-instantiated?
+						if (UObject** ReplacedOldObject = OldToNewInstanceMap.Find(OldObject))
+						{
+							// If the replaced object does not have the right class
+							if ((*ReplacedOldObject)->GetClass() != NewClass)
+							{
+								// Remove the old entry as it points to the wrong instance to replace
+								OldToNewInstanceMap.Remove(OldObject);
+								ObjectRemappingHelper.ReplacedObjects.Remove(OldObject);
+
+								// Make sure the replaced object has not already been re-instantiated
+								if (UObject** NewUObject = OldToNewInstanceMap.Find(*ReplacedOldObject))
+								{
+									// Add missing entry in the map
+									checkf(IsValid(*NewUObject), TEXT("Expecting the new object to already be valid"));
+									OldToNewInstanceMap.Add(OldObject, *NewUObject);
+									ObjectRemappingHelper.ReplacedObjects.Add(OldObject, *NewUObject);
+								}
+								else
+								{
+									// Remember old object and re-instantiate the replaced object
+									PrevOldObject = OldObject;
+									OldObject = *ReplacedOldObject;
+									bIsValid = IsValid(OldObject);
+								}
+							}
+							else
+							{
+								checkf(IsValid(*ReplacedOldObject), TEXT("Expecting the replaced object to already be valid"));
+							}
+						}
+					}
 
 					// Skip archetype instances, EXCEPT for component templates and child actor templates
 					const bool bIsChildActorTemplate = OldActor && OldActor->GetOuter()->IsA<UChildActorComponent>();
-					if ((!IsValid(OldObject) && !bIsScriptComponent) || // @todo: why do we need to replace PendingKill script components?
+					if ((!bIsValid && !bIsScriptComponent) || // @todo: why do we need to replace PendingKill script components?
 						(!bIsComponent && !bIsChildActorTemplate && OldObject->IsTemplate() && !bIsScriptComponent) ||
 						(InstancesThatShouldUseOldClass && InstancesThatShouldUseOldClass->Contains(OldObject)))
 					{
@@ -2427,7 +2467,17 @@ void FBlueprintCompileReinstancer::ReplaceInstancesOfClass_Inner(const TMap<UCla
 						UObject* NewUObject = nullptr;
 						ReplaceObjectHelper(OldObject, OldClass, NewUObject, NewClass, OldToNewInstanceMap, OldToNewNameMap, OldObjIndex, ObjectsToReplace, PotentialEditorsForRefreshing, OwnersToRerunConstructionScript, &FDirectAttachChildrenAccessor::Get, bIsComponent, bArchetypesAreUpToDate);
 						UpdateObjectBeingDebugged(OldObject, NewUObject);
-						ObjectsReplaced.Add(OldObject);
+
+						if (PrevOldObject)
+						{
+							OldToNewInstanceMap.Add(PrevOldObject, NewUObject);
+							ObjectRemappingHelper.ReplacedObjects.Add(PrevOldObject, NewUObject);
+							ObjectsReplaced.Add(PrevOldObject);
+						}
+						else
+						{
+							ObjectsReplaced.Add(OldObject);
+						}
 
 						if (bLogConversions)
 						{
@@ -2440,13 +2490,13 @@ void FBlueprintCompileReinstancer::ReplaceInstancesOfClass_Inner(const TMap<UCla
 
 
 		FDelegateHandle OnLevelActorDeletedHandle = GEngine ? GEngine->OnLevelActorDeleted().AddLambda([&OldToNewInstanceMap](AActor* DestroyedActor)
+		{
+			if (UObject** ReplacementObject = OldToNewInstanceMap.Find(DestroyedActor))
 			{
-				if (UObject** ReplacementObject = OldToNewInstanceMap.Find(DestroyedActor))
-				{
-					AActor* ReplacementActor = CastChecked<AActor>(*ReplacementObject);
-					ReplacementActor->GetWorld()->EditorDestroyActor(ReplacementActor, /*bShouldModifyLevel =*/true);
-				}
-			}) : FDelegateHandle();
+				AActor* ReplacementActor = CastChecked<AActor>(*ReplacementObject);
+				ReplacementActor->GetWorld()->EditorDestroyActor(ReplacementActor, /*bShouldModifyLevel =*/true);
+			}
+		}) : FDelegateHandle();
 
 		// WARNING: for (TPair<UClass*, UClass*> OldToNewClass : InOldToNewClassMap) duplicated above 
 		// this loop only handles actors - which need to be reconstructed *after* their owned components 
diff --git a/Engine/Source/Runtime/Engine/Private/InstancedReferenceSubobjectHelper.cpp b/Engine/Source/Runtime/Engine/Private/InstancedReferenceSubobjectHelper.cpp
--- a/Engine/Source/Runtime/Engine/Private/InstancedReferenceSubobjectHelper.cpp
+++ b/Engine/Source/Runtime/Engine/Private/InstancedReferenceSubobjectHelper.cpp
@@ -241,7 +241,7 @@ void FFindInstancedReferenceSubobjectHelper::Duplicate(
 	UObject* NewObject, 
 	TMap<UObject*, UObject*>& ReferenceReplacementMap, 
 	TArray<UObject*>& DuplicatedObjects, 
-	const TMap<UObject*, UObject*>* OptionalMappings)
+	TMap<UObject*, UObject*>* OptionalMappings)
 {
 	if (OldObject->GetClass()->HasAnyClassFlags(CLASS_HasInstancedReference) &&
 		NewObject->GetClass()->HasAnyClassFlags(CLASS_HasInstancedReference))
@@ -279,6 +279,10 @@ void FFindInstancedReferenceSubobjectHelper::Duplicate(
 
 							UObject* NewEditInlineSubobject = StaticDuplicateObject(Obj, NewObject, Obj->GetFName());
 							ReferenceReplacementMap.Add(Obj, NewEditInlineSubobject);
+							if(OptionalMappings)
+							{
+								OptionalMappings->Add(Obj, NewEditInlineSubobject);
+							}
 
 							// NOTE: we cannot patch OldObject's linker table here, since we don't 
 							//       know the relation between the two objects (one could be of a 
diff --git a/Engine/Source/Runtime/Engine/Public/InstancedReferenceSubobjectHelper.h b/Engine/Source/Runtime/Engine/Public/InstancedReferenceSubobjectHelper.h
--- a/Engine/Source/Runtime/Engine/Public/InstancedReferenceSubobjectHelper.h
+++ b/Engine/Source/Runtime/Engine/Public/InstancedReferenceSubobjectHelper.h
@@ -119,7 +119,7 @@ public:
 		UObject* NewObject, 
 		TMap<UObject*, UObject*>& ReferenceReplacementMap, 
 		TArray<UObject*>& DuplicatedObjects, 
-		const TMap<UObject*, UObject*>* OptionalMappings);
+		TMap<UObject*, UObject*>* OptionalMappings);
 
 	template<typename T>
 	static void ForEachInstancedSubObject(FInstancedPropertyPath& PropertyPath, T ContainerAddress, TFunctionRef<void(const FInstancedSubObjRef& Ref, T PropertyValueAddress)> ObjRefFunc);
