Fixed Memcpy looping over on itself indefinite

- Removed for loop indefinite in CopyBrickToAtlasVolumeTexture

diff --git a/Engine/Source/Editor/UnrealEd/Private/Lightmass/ImportVolumetricLightmap.cpp b/Engine/Source/Editor/UnrealEd/Private/Lightmass/ImportVolumetricLightmap.cpp
--- a/Engine/Source/Editor/UnrealEd/Private/Lightmass/ImportVolumetricLightmap.cpp
+++ b/Engine/Source/Editor/UnrealEd/Private/Lightmass/ImportVolumetricLightmap.cpp
@@ -68,10 +68,7 @@ void CopyBetweenAtlasVolumeTextures(
 		{
 			const int32 DestIndex = DestZIndex + (DestBrickMin.Y + YIndex) * DestPitch;
 			const int32 SourceIndex = SourceZIndex + (SourceBrickMin.Y + YIndex) * SourcePitch;
-			for (int32 XIndex = 0; XIndex < BrickSize.X * FormatSize; XIndex++)
-			{
-				FMemory::Memcpy((uint8*)&DestData[DestIndex], (const uint8*)&SourceData[SourceIndex], BrickSize.X * FormatSize);
-			}
+			FMemory::Memcpy((uint8*)&DestData[DestIndex], (const uint8*)&SourceData[SourceIndex], BrickSize.X * FormatSize);
 		}
 	}
 }
