Remove unused variable (heh).
#rnx

diff --git a/Engine/Source/Editor/UnrealEd/Private/Commandlets/CookCommandlet.cpp b/Engine/Source/Editor/UnrealEd/Private/Commandlets/CookCommandlet.cpp
--- a/Engine/Source/Editor/UnrealEd/Private/Commandlets/CookCommandlet.cpp
+++ b/Engine/Source/Editor/UnrealEd/Private/Commandlets/CookCommandlet.cpp
@@ -536,10 +536,7 @@ void UCookCommandlet::RunCookByTheBookCook(UCookOnTheFlyServer* CookOnTheFlyServ
 		DECLARE_SCOPE_CYCLE_COUNTER(TEXT("CookByTheBook.MainLoop"), STAT_CookByTheBook_MainLoop, STATGROUP_LoadTime);
 		while (CookOnTheFlyServer->IsInSession())
 		{
-			uint32 TickResults = 0;
-			uint32 UnusedVariable = 0;
-
-			TickResults = CookOnTheFlyServer->TickCookByTheBook(MAX_flt,
+			uint32 TickResults = TickResults = CookOnTheFlyServer->TickCookByTheBook(MAX_flt,
 				ShowProgress ? ECookTickFlags::None : ECookTickFlags::HideProgressDisplay);
 			ConditionalCollectGarbage(TickResults, *CookOnTheFlyServer);
 		}
