Asset Editors: Deprecate OnClose and RequestClose in favor of similar functions that take in a close reason
Fix a crash in some asset editors when you force deleted an asset with pending changes and hit "cancel" on the save dialog

diff --git a/Engine/Plugins/Editor/WaveformEditor/Source/WaveformEditor/Private/WaveformEditorInstantiator.cpp b/Engine/Plugins/Editor/WaveformEditor/Source/WaveformEditor/Private/WaveformEditorInstantiator.cpp
--- a/Engine/Plugins/Editor/WaveformEditor/Source/WaveformEditor/Private/WaveformEditorInstantiator.cpp
+++ b/Engine/Plugins/Editor/WaveformEditor/Source/WaveformEditor/Private/WaveformEditorInstantiator.cpp
@@ -132,7 +132,7 @@ void FWaveformEditorInstantiator::CreateWaveformEditor(TArray<USoundWave*> Sound
 			if (!WaveformEditor->Init(EToolkitMode::Standalone, nullptr, SoundWavePtr))
 			{
 				UE_LOG(LogWaveformEditor, Warning, TEXT("Could not open waveform editor for soundwave %s, initialization failed"), *(SoundWavePtr->GetName()))
-				WaveformEditor->CloseWindow();
+				WaveformEditor->CloseWindow(EAssetEditorCloseReason::AssetUnloadingOrInvalid);
 			}
 		}
 	}
diff --git a/Engine/Plugins/Enterprise/DataprepEditor/Source/DataprepEditor/Private/DataprepEditor.cpp b/Engine/Plugins/Enterprise/DataprepEditor/Source/DataprepEditor/Private/DataprepEditor.cpp
--- a/Engine/Plugins/Enterprise/DataprepEditor/Source/DataprepEditor/Private/DataprepEditor.cpp
+++ b/Engine/Plugins/Enterprise/DataprepEditor/Source/DataprepEditor/Private/DataprepEditor.cpp
@@ -1528,10 +1528,10 @@ void FDataprepEditor::UpdatePreviewPanels(bool bInclude3DViewport)
 	}
 }
 
-bool FDataprepEditor::OnRequestClose()
+bool FDataprepEditor::OnRequestClose(EAssetEditorCloseReason InCloseReason)
 {
 	const int ActorCount = PreviewWorld->GetActorCount();
-	if( bWorldBuilt && !bIgnoreCloseRequest && ActorCount > DefaultActorsInPreviewWorld.Num() )
+	if( InCloseReason != EAssetEditorCloseReason::AssetForceDeleted &&  bWorldBuilt && !bIgnoreCloseRequest && ActorCount > DefaultActorsInPreviewWorld.Num() )
 	{
 		// World was imported and is not empty -- show warning message
 		const FText Title( LOCTEXT( "DataprepEditor_ProceedWithClose", "Proceed with close") );
diff --git a/Engine/Plugins/Enterprise/DataprepEditor/Source/DataprepEditor/Private/DataprepEditor.h b/Engine/Plugins/Enterprise/DataprepEditor/Source/DataprepEditor/Private/DataprepEditor.h
--- a/Engine/Plugins/Enterprise/DataprepEditor/Source/DataprepEditor/Private/DataprepEditor.h
+++ b/Engine/Plugins/Enterprise/DataprepEditor/Source/DataprepEditor/Private/DataprepEditor.h
@@ -227,7 +227,7 @@ private:
 	bool CanBuildWorld();
 	bool CanCommitWorld();
 
-	virtual bool OnRequestClose() override;
+	virtual bool OnRequestClose(EAssetEditorCloseReason InCloseReason) override;
 
 	/** Create a snapshot of the world and tracked assets */
 	void TakeSnapshot();
diff --git a/Engine/Plugins/Enterprise/VariantManager/Source/VariantManager/Private/LevelVariantSetsEditorToolkit.cpp b/Engine/Plugins/Enterprise/VariantManager/Source/VariantManager/Private/LevelVariantSetsEditorToolkit.cpp
--- a/Engine/Plugins/Enterprise/VariantManager/Source/VariantManager/Private/LevelVariantSetsEditorToolkit.cpp
+++ b/Engine/Plugins/Enterprise/VariantManager/Source/VariantManager/Private/LevelVariantSetsEditorToolkit.cpp
@@ -39,7 +39,7 @@ void FLevelVariantSetsEditorToolkit::Initialize(const EToolkitMode::Type Mode, c
 
 			if (AssetEditorInstance.IsValid())
 			{
-				InLevelVariantSetsAssetEditor.Pin()->CloseWindow();
+				InLevelVariantSetsAssetEditor.Pin()->CloseWindow(EAssetEditorCloseReason::AssetEditorHostClosed);
 			}
 		}
 	};
@@ -124,7 +124,7 @@ TSharedRef<SDockTab> FLevelVariantSetsEditorToolkit::HandleTabManagerSpawnTab(co
 		.TabRole(ETabRole::PanelTab);
 }
 
-bool FLevelVariantSetsEditorToolkit::OnRequestClose()
+bool FLevelVariantSetsEditorToolkit::OnRequestClose(EAssetEditorCloseReason InCloseReason)
 {
 	return true;
 }
diff --git a/Engine/Plugins/Enterprise/VariantManager/Source/VariantManager/Private/LevelVariantSetsEditorToolkit.h b/Engine/Plugins/Enterprise/VariantManager/Source/VariantManager/Private/LevelVariantSetsEditorToolkit.h
--- a/Engine/Plugins/Enterprise/VariantManager/Source/VariantManager/Private/LevelVariantSetsEditorToolkit.h
+++ b/Engine/Plugins/Enterprise/VariantManager/Source/VariantManager/Private/LevelVariantSetsEditorToolkit.h
@@ -33,7 +33,7 @@ public:
 	{
 		return TEXT("FLevelVariantSetsEditorToolkit");
 	}
-	virtual bool OnRequestClose() override;
+	virtual bool OnRequestClose(EAssetEditorCloseReason InCloseReason) override;
 	virtual bool CanFindInContentBrowser() const override;
 	virtual void FocusWindow(UObject* ObjectToFocusOn = nullptr) override;
 
diff --git a/Engine/Plugins/Experimental/ChaosClothAssetEditor/Source/ChaosClothAssetEditor/Private/ChaosClothAsset/ClothEditorToolkit.cpp b/Engine/Plugins/Experimental/ChaosClothAssetEditor/Source/ChaosClothAssetEditor/Private/ChaosClothAsset/ClothEditorToolkit.cpp
--- a/Engine/Plugins/Experimental/ChaosClothAssetEditor/Source/ChaosClothAssetEditor/Private/ChaosClothAsset/ClothEditorToolkit.cpp
+++ b/Engine/Plugins/Experimental/ChaosClothAssetEditor/Source/ChaosClothAssetEditor/Private/ChaosClothAsset/ClothEditorToolkit.cpp
@@ -366,7 +366,7 @@ void FChaosClothAssetEditorToolkit::RemoveViewportOverlayWidget(TSharedRef<SWidg
 	ViewportWidget->RemoveOverlayWidget(InViewportOverlayWidget);
 }
 
-bool FChaosClothAssetEditorToolkit::OnRequestClose()
+bool FChaosClothAssetEditorToolkit::OnRequestClose(EAssetEditorCloseReason InCloseReason)
 {
 	// Note: This needs a bit of adjusting, because currently OnRequestClose seems to be 
 	// called multiple times when the editor itself is being closed. We can take the route 
@@ -386,7 +386,7 @@ bool FChaosClothAssetEditorToolkit::OnRequestClose()
 	// duplication for each opening event.
 	GetEditorModeManager().ActivateDefaultMode();
 
-	return FAssetEditorToolkit::OnRequestClose();
+	return FAssetEditorToolkit::OnRequestClose(InCloseReason);
 }
 
 void FChaosClothAssetEditorToolkit::PostInitAssetEditor()
diff --git a/Engine/Plugins/Experimental/ChaosClothAssetEditor/Source/ChaosClothAssetEditor/Private/ChaosClothAsset/ClothEditorToolkit.h b/Engine/Plugins/Experimental/ChaosClothAssetEditor/Source/ChaosClothAssetEditor/Private/ChaosClothAsset/ClothEditorToolkit.h
--- a/Engine/Plugins/Experimental/ChaosClothAssetEditor/Source/ChaosClothAssetEditor/Private/ChaosClothAsset/ClothEditorToolkit.h
+++ b/Engine/Plugins/Experimental/ChaosClothAssetEditor/Source/ChaosClothAssetEditor/Private/ChaosClothAsset/ClothEditorToolkit.h
@@ -75,7 +75,7 @@ private:
 	// FAssetEditorToolkit
 	virtual void AddViewportOverlayWidget(TSharedRef<SWidget> InViewportOverlayWidget) override;
 	virtual void RemoveViewportOverlayWidget(TSharedRef<SWidget> InViewportOverlayWidget) override;
-	virtual bool OnRequestClose() override;
+	virtual bool OnRequestClose(EAssetEditorCloseReason InCloseReason) override;
 	virtual void PostInitAssetEditor() override;
 	virtual void InitToolMenuContext(FToolMenuContext& MenuContext) override;
 
diff --git a/Engine/Plugins/Experimental/CharacterFXEditor/BaseCharacterFXEditor/Source/BaseCharacterFXEditor/Private/BaseCharacterFXEditorToolkit.cpp b/Engine/Plugins/Experimental/CharacterFXEditor/BaseCharacterFXEditor/Source/BaseCharacterFXEditor/Private/BaseCharacterFXEditorToolkit.cpp
--- a/Engine/Plugins/Experimental/CharacterFXEditor/BaseCharacterFXEditor/Source/BaseCharacterFXEditor/Private/BaseCharacterFXEditorToolkit.cpp
+++ b/Engine/Plugins/Experimental/CharacterFXEditor/BaseCharacterFXEditor/Source/BaseCharacterFXEditor/Private/BaseCharacterFXEditorToolkit.cpp
@@ -84,7 +84,7 @@ void FBaseCharacterFXEditorToolkit::OnToolkitHostingFinished(const TSharedRef<IT
 	ModeUILayer->OnToolkitHostingFinished(Toolkit);
 }
 
-bool FBaseCharacterFXEditorToolkit::OnRequestClose()
+bool FBaseCharacterFXEditorToolkit::OnRequestClose(EAssetEditorCloseReason InCloseReason)
 {
 	// Note: This needs a bit of adjusting, because currently OnRequestClose seems to be 
 	// called multiple times when the editor itself is being closed. We can take the route 
@@ -105,7 +105,7 @@ bool FBaseCharacterFXEditorToolkit::OnRequestClose()
 	// duplication for each opening event.
 	GetEditorModeManager().DeactivateAllModes();
 
-	return FAssetEditorToolkit::OnRequestClose();
+	return FAssetEditorToolkit::OnRequestClose(InCloseReason);
 }
 
 AssetEditorViewportFactoryFunction FBaseCharacterFXEditorToolkit::GetViewportDelegate()
diff --git a/Engine/Plugins/Experimental/CharacterFXEditor/BaseCharacterFXEditor/Source/BaseCharacterFXEditor/Public/BaseCharacterFXEditorToolkit.h b/Engine/Plugins/Experimental/CharacterFXEditor/BaseCharacterFXEditor/Source/BaseCharacterFXEditor/Public/BaseCharacterFXEditorToolkit.h
--- a/Engine/Plugins/Experimental/CharacterFXEditor/BaseCharacterFXEditor/Source/BaseCharacterFXEditor/Public/BaseCharacterFXEditorToolkit.h
+++ b/Engine/Plugins/Experimental/CharacterFXEditor/BaseCharacterFXEditor/Source/BaseCharacterFXEditor/Public/BaseCharacterFXEditorToolkit.h
@@ -26,7 +26,7 @@ public:
 
 	virtual void OnToolkitHostingStarted(const TSharedRef<IToolkit>& Toolkit) override;
 	virtual void OnToolkitHostingFinished(const TSharedRef<IToolkit>& Toolkit) override;
-	virtual bool OnRequestClose() override;
+	virtual bool OnRequestClose(EAssetEditorCloseReason InCloseReason) override;
 	virtual void RegisterTabSpawners(const TSharedRef<FTabManager>& InTabManager) override;
 
 	virtual AssetEditorViewportFactoryFunction GetViewportDelegate() override;
diff --git a/Engine/Plugins/Experimental/Mutable/Source/CustomizableObjectEditor/Private/MuCOE/CustomizableObjectEditorViewportClient.cpp b/Engine/Plugins/Experimental/Mutable/Source/CustomizableObjectEditor/Private/MuCOE/CustomizableObjectEditorViewportClient.cpp
--- a/Engine/Plugins/Experimental/Mutable/Source/CustomizableObjectEditor/Private/MuCOE/CustomizableObjectEditorViewportClient.cpp
+++ b/Engine/Plugins/Experimental/Mutable/Source/CustomizableObjectEditor/Private/MuCOE/CustomizableObjectEditorViewportClient.cpp
@@ -2805,7 +2805,7 @@ bool FCustomizableObjectEditorViewportClient::ManageBakingAction(const FString&
 			for (IAssetEditorInstance* ObjectEditorInstance : ObjectEditors)
 			{
 				// Close the editors that contains this asset
-				if (!ObjectEditorInstance->CloseWindow())
+				if (!ObjectEditorInstance->CloseWindow(EAssetEditorCloseReason::AssetEditorHostClosed))
 				{
 					FText Caption = LOCTEXT("OpenExisitngFile", "Open File");
 					FText Message = FText::Format(LOCTEXT("CantCloseAsset", "This Obejct \"{0}\" is open in an editor and can't be closed automatically. Please close the editor and try to bake it again"), FText::FromString(ObjName));
diff --git a/Engine/Plugins/Experimental/UVEditor/Source/UVEditor/Private/UVEditorToolkit.cpp b/Engine/Plugins/Experimental/UVEditor/Source/UVEditor/Private/UVEditorToolkit.cpp
--- a/Engine/Plugins/Experimental/UVEditor/Source/UVEditor/Private/UVEditorToolkit.cpp
+++ b/Engine/Plugins/Experimental/UVEditor/Source/UVEditor/Private/UVEditorToolkit.cpp
@@ -222,7 +222,7 @@ void FUVEditorToolkit::RegisterTabSpawners(const TSharedRef<FTabManager>& InTabM
 		.SetIcon(FSlateIcon(FAppStyle::GetAppStyleSetName(), "LevelEditor.Tabs.Viewports"));
 }
 
-bool FUVEditorToolkit::OnRequestClose()
+bool FUVEditorToolkit::OnRequestClose(EAssetEditorCloseReason InCloseReason)
 {
 	// Note: This needs a bit of adjusting, because currently OnRequestClose seems to be 
 	// called multiple times when the editor itself is being closed. We can take the route 
@@ -236,6 +236,12 @@ bool FUVEditorToolkit::OnRequestClose()
 		return true; 
 	}
 
+	// If the asset is being force deleted, we don't want to apply changes
+	if(InCloseReason == EAssetEditorCloseReason::AssetForceDeleted)
+	{
+		return true;
+	}
+
 	bool bHasUnappliedChanges = UVMode->HaveUnappliedChanges();
 	bool bCanApplyChanges = UVMode->CanApplyChanges();
 
@@ -282,7 +288,7 @@ bool FUVEditorToolkit::OnRequestClose()
 		}
 	}
 
-	return FAssetEditorToolkit::OnRequestClose();
+	return FAssetEditorToolkit::OnRequestClose(InCloseReason);
 }
 
 void FUVEditorToolkit::OnClose()
diff --git a/Engine/Plugins/Experimental/UVEditor/Source/UVEditor/Public/UVEditorToolkit.h b/Engine/Plugins/Experimental/UVEditor/Source/UVEditor/Public/UVEditorToolkit.h
--- a/Engine/Plugins/Experimental/UVEditor/Source/UVEditor/Public/UVEditorToolkit.h
+++ b/Engine/Plugins/Experimental/UVEditor/Source/UVEditor/Public/UVEditorToolkit.h
@@ -49,7 +49,7 @@ public:
 	virtual FText GetBaseToolkitName() const override;
 	virtual FText GetToolkitToolTipText() const override;
 	virtual void RegisterTabSpawners(const TSharedRef<FTabManager>& InTabManager) override;
-	virtual bool OnRequestClose() override;
+	virtual bool OnRequestClose(EAssetEditorCloseReason InCloseReason) override;
 	virtual void OnClose() override;
 	virtual void SaveAsset_Execute() override;
 	virtual bool CanSaveAsset() const override;
diff --git a/Engine/Plugins/FX/Niagara/Source/NiagaraEditor/Private/Toolkits/NiagaraParameterDefinitionsToolkit.cpp b/Engine/Plugins/FX/Niagara/Source/NiagaraEditor/Private/Toolkits/NiagaraParameterDefinitionsToolkit.cpp
--- a/Engine/Plugins/FX/Niagara/Source/NiagaraEditor/Private/Toolkits/NiagaraParameterDefinitionsToolkit.cpp
+++ b/Engine/Plugins/FX/Niagara/Source/NiagaraEditor/Private/Toolkits/NiagaraParameterDefinitionsToolkit.cpp
@@ -182,9 +182,9 @@ void FNiagaraParameterDefinitionsToolkit::SaveAssetAs_Execute()
 	FAssetEditorToolkit::SaveAsset_Execute();
 }
 
-bool FNiagaraParameterDefinitionsToolkit::OnRequestClose()
+bool FNiagaraParameterDefinitionsToolkit::OnRequestClose(EAssetEditorCloseReason InCloseReason)
 {
-	if (bChangesDiscarded == false && OnApplyEnabled())
+	if (InCloseReason != EAssetEditorCloseReason::AssetForceDeleted && bChangesDiscarded == false && OnApplyEnabled())
 	{
 		// find out the user wants to do with this dirty NiagaraScript
 		EAppReturnType::Type YesNoCancelReply = FMessageDialog::Open(
diff --git a/Engine/Plugins/FX/Niagara/Source/NiagaraEditor/Private/Toolkits/NiagaraParameterDefinitionsToolkit.h b/Engine/Plugins/FX/Niagara/Source/NiagaraEditor/Private/Toolkits/NiagaraParameterDefinitionsToolkit.h
--- a/Engine/Plugins/FX/Niagara/Source/NiagaraEditor/Private/Toolkits/NiagaraParameterDefinitionsToolkit.h
+++ b/Engine/Plugins/FX/Niagara/Source/NiagaraEditor/Private/Toolkits/NiagaraParameterDefinitionsToolkit.h
@@ -40,7 +40,7 @@ protected:
 	virtual void GetSaveableObjects(TArray<UObject*>& OutObjects) const override;
 	virtual void SaveAsset_Execute() override;
 	virtual void SaveAssetAs_Execute() override;
-	virtual bool OnRequestClose() override;
+	virtual bool OnRequestClose(EAssetEditorCloseReason InCloseReason) override;
 
 private:
 	TSharedRef<SDockTab> SpawnTab_ParameterDefinitionsDetails(const FSpawnTabArgs& Args);
diff --git a/Engine/Plugins/FX/Niagara/Source/NiagaraEditor/Private/Toolkits/NiagaraScriptToolkit.cpp b/Engine/Plugins/FX/Niagara/Source/NiagaraEditor/Private/Toolkits/NiagaraScriptToolkit.cpp
--- a/Engine/Plugins/FX/Niagara/Source/NiagaraEditor/Private/Toolkits/NiagaraScriptToolkit.cpp
+++ b/Engine/Plugins/FX/Niagara/Source/NiagaraEditor/Private/Toolkits/NiagaraScriptToolkit.cpp
@@ -912,7 +912,7 @@ void FNiagaraScriptToolkit::UpdateOriginalNiagaraScript()
 }
 
 
-bool FNiagaraScriptToolkit::OnRequestClose()
+bool FNiagaraScriptToolkit::OnRequestClose(EAssetEditorCloseReason InCloseReason)
 {
 	if (bChangesDiscarded == false && IsEditScriptDifferentFromOriginalScript())
 	{
diff --git a/Engine/Plugins/FX/Niagara/Source/NiagaraEditor/Private/Toolkits/NiagaraScriptToolkit.h b/Engine/Plugins/FX/Niagara/Source/NiagaraEditor/Private/Toolkits/NiagaraScriptToolkit.h
--- a/Engine/Plugins/FX/Niagara/Source/NiagaraEditor/Private/Toolkits/NiagaraScriptToolkit.h
+++ b/Engine/Plugins/FX/Niagara/Source/NiagaraEditor/Private/Toolkits/NiagaraScriptToolkit.h
@@ -84,7 +84,7 @@ protected:
 	virtual void GetSaveableObjects(TArray<UObject*>& OutObjects) const override;
 	virtual void SaveAsset_Execute() override;
 	virtual void SaveAssetAs_Execute() override;
-	virtual bool OnRequestClose() override;
+	virtual bool OnRequestClose(EAssetEditorCloseReason InCloseReason) override;
 
 private:
 	void OnEditedScriptPropertyFinishedChanging(const FPropertyChangedEvent& InEvent);
diff --git a/Engine/Plugins/FX/Niagara/Source/NiagaraEditor/Private/Toolkits/NiagaraSystemToolkit.cpp b/Engine/Plugins/FX/Niagara/Source/NiagaraEditor/Private/Toolkits/NiagaraSystemToolkit.cpp
--- a/Engine/Plugins/FX/Niagara/Source/NiagaraEditor/Private/Toolkits/NiagaraSystemToolkit.cpp
+++ b/Engine/Plugins/FX/Niagara/Source/NiagaraEditor/Private/Toolkits/NiagaraSystemToolkit.cpp
@@ -1236,7 +1236,7 @@ void FNiagaraSystemToolkit::SaveAssetAs_Execute()
 	SystemViewModel->NotifyPostSave();
 }
 
-bool FNiagaraSystemToolkit::OnRequestClose()
+bool FNiagaraSystemToolkit::OnRequestClose(EAssetEditorCloseReason InCloseReason)
 {
 	if (GbLogNiagaraSystemChanges > 0)
 	{
@@ -1271,7 +1271,7 @@ bool FNiagaraSystemToolkit::OnRequestClose()
 		}
 	}
 
-	if (bScratchPadChangesDiscarded == false && bHasUnappliedScratchPadChanges)
+	if (InCloseReason != EAssetEditorCloseReason::AssetForceDeleted && bScratchPadChangesDiscarded == false && bHasUnappliedScratchPadChanges)
 	{
 		// find out the user wants to do with their dirty scratch pad scripts.
 		EAppReturnType::Type YesNoCancelReply = FMessageDialog::Open(EAppMsgType::YesNoCancel,
@@ -1291,7 +1291,7 @@ bool FNiagaraSystemToolkit::OnRequestClose()
 		}
 	}
 
-	if (SystemToolkitMode == ESystemToolkitMode::Emitter)
+	if (SystemToolkitMode == ESystemToolkitMode::Emitter && InCloseReason != EAssetEditorCloseReason::AssetForceDeleted)
 	{
 		TSharedPtr<FNiagaraEmitterViewModel> EmitterViewModel = SystemViewModel->GetEmitterHandleViewModels()[0]->GetEmitterViewModel();
 		if (bChangesDiscarded == false && (EmitterViewModel->GetEmitter().Emitter->GetChangeId() != LastSyncedEmitterChangeId || bEmitterThumbnailUpdated))
@@ -1324,7 +1324,7 @@ bool FNiagaraSystemToolkit::OnRequestClose()
 	}
 	
 	GEngine->ForceGarbageCollection(true);
-	return FAssetEditorToolkit::OnRequestClose();
+	return FAssetEditorToolkit::OnRequestClose(InCloseReason);
 }
 
 void FNiagaraSystemToolkit::EmitterAssetSelected(const FAssetData& AssetData)
diff --git a/Engine/Plugins/FX/Niagara/Source/NiagaraEditor/Private/Toolkits/NiagaraSystemToolkit.h b/Engine/Plugins/FX/Niagara/Source/NiagaraEditor/Private/Toolkits/NiagaraSystemToolkit.h
--- a/Engine/Plugins/FX/Niagara/Source/NiagaraEditor/Private/Toolkits/NiagaraSystemToolkit.h
+++ b/Engine/Plugins/FX/Niagara/Source/NiagaraEditor/Private/Toolkits/NiagaraSystemToolkit.h
@@ -165,7 +165,7 @@ protected:
 	virtual void GetSaveableObjects(TArray<UObject*>& OutObjects) const override;
 	virtual void SaveAsset_Execute() override;
 	virtual void SaveAssetAs_Execute() override;
-	virtual bool OnRequestClose() override;
+	virtual bool OnRequestClose(EAssetEditorCloseReason InCloseReason) override;
 	
 private:
 	void InitializeInternal(const EToolkitMode::Type Mode, const TSharedPtr<class IToolkitHost>& InitToolkitHost, const FGuid& MessageLogGuid);
diff --git a/Engine/Plugins/Media/MediaFrameworkUtilities/Source/MediaFrameworkUtilitiesEditor/Private/AssetEditor/MediaProfileEditorToolkit.cpp b/Engine/Plugins/Media/MediaFrameworkUtilities/Source/MediaFrameworkUtilitiesEditor/Private/AssetEditor/MediaProfileEditorToolkit.cpp
--- a/Engine/Plugins/Media/MediaFrameworkUtilities/Source/MediaFrameworkUtilitiesEditor/Private/AssetEditor/MediaProfileEditorToolkit.cpp
+++ b/Engine/Plugins/Media/MediaFrameworkUtilities/Source/MediaFrameworkUtilitiesEditor/Private/AssetEditor/MediaProfileEditorToolkit.cpp
@@ -133,7 +133,7 @@ void FMediaProfileEditorToolkit::HandleCoreObjectPropertyChanged(UObject* Object
 
 void FMediaProfileEditorToolkit::HandleMediaProxiesChanged()
 {
-	CloseWindow();
+	CloseWindow(EAssetEditorCloseReason::AssetUnloadingOrInvalid);
 }
 
 
@@ -149,28 +149,35 @@ void FMediaProfileEditorToolkit::SaveAssetAs_Execute()
 	Super::SaveAssetAs_Execute();
 }
 
-bool FMediaProfileEditorToolkit::OnRequestClose()
+bool FMediaProfileEditorToolkit::OnRequestClose(EAssetEditorCloseReason InCloseReason)
 {
 	bool bClose = true;
-	if (UMediaProfile* MediaProfile = Cast<UMediaProfile>(GetEditingObject()))
+
+	// We don't want to apply changes if the asset is being force deleted
+
+	if(InCloseReason != EAssetEditorCloseReason::AssetForceDeleted)
 	{
-		if (MediaProfile->bNeedToBeReapplied)
+		if (UMediaProfile* MediaProfile = Cast<UMediaProfile>(GetEditingObject()))
 		{
-			// find out the user wants to do with this dirty profile
-			EAppReturnType::Type YesNoCancelReply = FMessageDialog::Open(EAppMsgType::YesNoCancel, LOCTEXT("Prompt_EditorClose", "Would you like to apply the modifications?"));
-
-			switch (YesNoCancelReply)
+			if (MediaProfile->bNeedToBeReapplied)
 			{
-			case EAppReturnType::Yes:
-				ApplyMediaProfile();
-				break;
-			case EAppReturnType::Cancel:
-				bClose = false;
-				break;
+				// find out the user wants to do with this dirty profile
+				EAppReturnType::Type YesNoCancelReply = FMessageDialog::Open(EAppMsgType::YesNoCancel, LOCTEXT("Prompt_EditorClose", "Would you like to apply the modifications?"));
+
+				switch (YesNoCancelReply)
+				{
+				case EAppReturnType::Yes:
+					ApplyMediaProfile();
+					break;
+				case EAppReturnType::Cancel:
+					bClose = false;
+					break;
+				}
 			}
 		}
 	}
-	return Super::OnRequestClose();
+	
+	return Super::OnRequestClose(InCloseReason);
 }
 
 #undef LOCTEXT_NAMESPACE
diff --git a/Engine/Plugins/Media/MediaFrameworkUtilities/Source/MediaFrameworkUtilitiesEditor/Private/AssetEditor/MediaProfileEditorToolkit.h b/Engine/Plugins/Media/MediaFrameworkUtilities/Source/MediaFrameworkUtilitiesEditor/Private/AssetEditor/MediaProfileEditorToolkit.h
--- a/Engine/Plugins/Media/MediaFrameworkUtilities/Source/MediaFrameworkUtilitiesEditor/Private/AssetEditor/MediaProfileEditorToolkit.h
+++ b/Engine/Plugins/Media/MediaFrameworkUtilities/Source/MediaFrameworkUtilitiesEditor/Private/AssetEditor/MediaProfileEditorToolkit.h
@@ -39,7 +39,7 @@ protected:
 
 	virtual void SaveAsset_Execute() override;
 	virtual void SaveAssetAs_Execute() override;
-	virtual bool OnRequestClose() override;
+	virtual bool OnRequestClose(EAssetEditorCloseReason InCloseReason) override;
 
 private:
 
diff --git a/Engine/Plugins/MovieScene/LevelSequenceEditor/Source/LevelSequenceEditor/Private/LevelSequenceEditorToolkit.cpp b/Engine/Plugins/MovieScene/LevelSequenceEditor/Source/LevelSequenceEditor/Private/LevelSequenceEditorToolkit.cpp
--- a/Engine/Plugins/MovieScene/LevelSequenceEditor/Source/LevelSequenceEditor/Private/LevelSequenceEditorToolkit.cpp
+++ b/Engine/Plugins/MovieScene/LevelSequenceEditor/Source/LevelSequenceEditor/Private/LevelSequenceEditorToolkit.cpp
@@ -705,7 +705,7 @@ void FLevelSequenceEditorToolkit::HandleMapChanged(class UWorld* NewWorld, EMapC
 	if( ( MapChangeType == EMapChangeType::LoadMap || MapChangeType == EMapChangeType::NewMap || MapChangeType == EMapChangeType::TearDownWorld) )
 	{
 		Sequencer->GetSpawnRegister().CleanUp(*Sequencer);
-		CloseWindow();
+		CloseWindow(EAssetEditorCloseReason::AssetUnloadingOrInvalid);
 	}
 }
 
diff --git a/Engine/Plugins/MovieScene/TemplateSequence/Source/TemplateSequenceEditor/Private/TemplateSequenceEditorToolkit.cpp b/Engine/Plugins/MovieScene/TemplateSequence/Source/TemplateSequenceEditor/Private/TemplateSequenceEditorToolkit.cpp
--- a/Engine/Plugins/MovieScene/TemplateSequence/Source/TemplateSequenceEditor/Private/TemplateSequenceEditorToolkit.cpp
+++ b/Engine/Plugins/MovieScene/TemplateSequence/Source/TemplateSequenceEditor/Private/TemplateSequenceEditorToolkit.cpp
@@ -275,11 +275,11 @@ void FTemplateSequenceEditorToolkit::HandleMapChanged(UWorld* NewWorld, EMapChan
 	if ((MapChangeType == EMapChangeType::LoadMap || MapChangeType == EMapChangeType::NewMap || MapChangeType == EMapChangeType::TearDownWorld))
 	{
 		Sequencer->GetSpawnRegister().CleanUp(*Sequencer);
-		CloseWindow();
+		CloseWindow(EAssetEditorCloseReason::AssetUnloadingOrInvalid);
 	}
 }
 
-bool FTemplateSequenceEditorToolkit::OnRequestClose()
+bool FTemplateSequenceEditorToolkit::OnRequestClose(EAssetEditorCloseReason InCloseReason)
 {
 	return true;
 }
diff --git a/Engine/Plugins/MovieScene/TemplateSequence/Source/TemplateSequenceEditor/Private/TemplateSequenceEditorToolkit.h b/Engine/Plugins/MovieScene/TemplateSequence/Source/TemplateSequenceEditor/Private/TemplateSequenceEditorToolkit.h
--- a/Engine/Plugins/MovieScene/TemplateSequence/Source/TemplateSequenceEditor/Private/TemplateSequenceEditorToolkit.h
+++ b/Engine/Plugins/MovieScene/TemplateSequence/Source/TemplateSequenceEditor/Private/TemplateSequenceEditorToolkit.h
@@ -60,7 +60,7 @@ public:
 		return TEXT("FTemplateSequenceEditorToolkit");
 	}
 
-	virtual bool OnRequestClose() override;
+	virtual bool OnRequestClose(EAssetEditorCloseReason InCloseReason) override;
 	virtual bool CanFindInContentBrowser() const override;
 
 public:
diff --git a/Engine/Plugins/Runtime/nDisplay/Source/DisplayClusterConfigurator/Private/DisplayClusterConfiguratorBlueprintEditor.cpp b/Engine/Plugins/Runtime/nDisplay/Source/DisplayClusterConfigurator/Private/DisplayClusterConfiguratorBlueprintEditor.cpp
--- a/Engine/Plugins/Runtime/nDisplay/Source/DisplayClusterConfigurator/Private/DisplayClusterConfiguratorBlueprintEditor.cpp
+++ b/Engine/Plugins/Runtime/nDisplay/Source/DisplayClusterConfigurator/Private/DisplayClusterConfiguratorBlueprintEditor.cpp
@@ -1075,9 +1075,9 @@ TStatId FDisplayClusterConfiguratorBlueprintEditor::GetStatId() const
 	RETURN_QUICK_DECLARE_CYCLE_STAT(FDisplayClusterConfiguratorBlueprintEditor, STATGROUP_Tickables);
 }
 
-bool FDisplayClusterConfiguratorBlueprintEditor::OnRequestClose()
+bool FDisplayClusterConfiguratorBlueprintEditor::OnRequestClose(EAssetEditorCloseReason InCloseReason)
 {
-	return FBlueprintEditor::OnRequestClose();
+	return FBlueprintEditor::OnRequestClose(InCloseReason);
 }
 
 void FDisplayClusterConfiguratorBlueprintEditor::OnClose()
diff --git a/Engine/Plugins/Runtime/nDisplay/Source/DisplayClusterConfigurator/Private/DisplayClusterConfiguratorBlueprintEditor.h b/Engine/Plugins/Runtime/nDisplay/Source/DisplayClusterConfigurator/Private/DisplayClusterConfiguratorBlueprintEditor.h
--- a/Engine/Plugins/Runtime/nDisplay/Source/DisplayClusterConfigurator/Private/DisplayClusterConfiguratorBlueprintEditor.h
+++ b/Engine/Plugins/Runtime/nDisplay/Source/DisplayClusterConfigurator/Private/DisplayClusterConfiguratorBlueprintEditor.h
@@ -142,7 +142,7 @@ protected:
 	virtual void Tick(float DeltaTime) override;
 	virtual bool IsTickable() const override { return true; }
 	virtual TStatId GetStatId() const override;
-	virtual bool OnRequestClose() override;
+	virtual bool OnRequestClose(EAssetEditorCloseReason InCloseReason) override;
 	virtual void OnClose() override;
 	virtual void Compile() override;
 	//~ End FBlueprintEditor Interface
diff --git a/Engine/Plugins/Tests/EditorTests/Source/EditorTests/Private/UnrealEd/BlueprintAutomationTests.cpp b/Engine/Plugins/Tests/EditorTests/Source/EditorTests/Private/UnrealEd/BlueprintAutomationTests.cpp
--- a/Engine/Plugins/Tests/EditorTests/Source/EditorTests/Private/UnrealEd/BlueprintAutomationTests.cpp
+++ b/Engine/Plugins/Tests/EditorTests/Source/EditorTests/Private/UnrealEd/BlueprintAutomationTests.cpp
@@ -327,7 +327,7 @@ public:
 		if (EditorInst != nullptr)
 		{
 			UE_LOG(LogBlueprintAutomationTests, Log, TEXT("Closing '%s' so we don't invalidate the open version when unloading it."), *BlueprintObj->GetName());
-			EditorInst->CloseWindow();
+			EditorInst->CloseWindow(EAssetEditorCloseReason::AssetUnloadingOrInvalid);
 		}
 	}
 
diff --git a/Engine/Plugins/VirtualProduction/CameraCalibration/Source/CameraCalibrationEditor/Private/AssetEditor/CameraCalibrationToolkit.cpp b/Engine/Plugins/VirtualProduction/CameraCalibration/Source/CameraCalibrationEditor/Private/AssetEditor/CameraCalibrationToolkit.cpp
--- a/Engine/Plugins/VirtualProduction/CameraCalibration/Source/CameraCalibrationEditor/Private/AssetEditor/CameraCalibrationToolkit.cpp
+++ b/Engine/Plugins/VirtualProduction/CameraCalibration/Source/CameraCalibrationEditor/Private/AssetEditor/CameraCalibrationToolkit.cpp
@@ -173,7 +173,7 @@ void FCameraCalibrationToolkit::UnregisterTabSpawners(const TSharedRef<class FTa
 	Super::UnregisterTabSpawners(InTabManager);
 }
 
-bool FCameraCalibrationToolkit::OnRequestClose()
+bool FCameraCalibrationToolkit::OnRequestClose(EAssetEditorCloseReason InCloseReason)
 {
 	return true;
 }
diff --git a/Engine/Plugins/VirtualProduction/CameraCalibration/Source/CameraCalibrationEditor/Private/AssetEditor/CameraCalibrationToolkit.h b/Engine/Plugins/VirtualProduction/CameraCalibration/Source/CameraCalibrationEditor/Private/AssetEditor/CameraCalibrationToolkit.h
--- a/Engine/Plugins/VirtualProduction/CameraCalibration/Source/CameraCalibrationEditor/Private/AssetEditor/CameraCalibrationToolkit.h
+++ b/Engine/Plugins/VirtualProduction/CameraCalibration/Source/CameraCalibrationEditor/Private/AssetEditor/CameraCalibrationToolkit.h
@@ -54,7 +54,7 @@ protected:
 	virtual FString GetWorldCentricTabPrefix() const override;
 	virtual void RegisterTabSpawners(const TSharedRef<FTabManager>& TabManager) override;
 	virtual void UnregisterTabSpawners(const TSharedRef<FTabManager>& TabManager) override;
-	virtual bool OnRequestClose() override;
+	virtual bool OnRequestClose(EAssetEditorCloseReason InCloseReason) override;
 	//~ End IToolkit interface
 
 private:
diff --git a/Engine/Plugins/VirtualProduction/RemoteControl/Source/RemoteControlUI/Private/AssetEditor/RemoteControlPresetEditorToolkit.cpp b/Engine/Plugins/VirtualProduction/RemoteControl/Source/RemoteControlUI/Private/AssetEditor/RemoteControlPresetEditorToolkit.cpp
--- a/Engine/Plugins/VirtualProduction/RemoteControl/Source/RemoteControlUI/Private/AssetEditor/RemoteControlPresetEditorToolkit.cpp
+++ b/Engine/Plugins/VirtualProduction/RemoteControl/Source/RemoteControlUI/Private/AssetEditor/RemoteControlPresetEditorToolkit.cpp
@@ -105,7 +105,7 @@ void FRemoteControlPresetEditorToolkit::UnregisterTabSpawners(const TSharedRef<F
 	InTabManager->UnregisterTabSpawner(FRemoteControlUIModule::RemoteControlPanelTabName);
 }
 
-bool FRemoteControlPresetEditorToolkit::OnRequestClose()
+bool FRemoteControlPresetEditorToolkit::OnRequestClose(EAssetEditorCloseReason InCloseReason)
 {
 	return true;
 }
@@ -170,7 +170,7 @@ void FRemoteControlPresetEditorToolkit::InvokePanelTab()
 		{
 			if (const TSharedPtr<IAssetEditorInstance> AssetEditorInstance = InRemoteControlPreset.Pin())
 			{
-				InRemoteControlPreset.Pin()->CloseWindow();
+				InRemoteControlPreset.Pin()->CloseWindow(EAssetEditorCloseReason::AssetEditorHostClosed);
 			}
 		}
 	};
diff --git a/Engine/Plugins/VirtualProduction/RemoteControl/Source/RemoteControlUI/Private/AssetEditor/RemoteControlPresetEditorToolkit.h b/Engine/Plugins/VirtualProduction/RemoteControl/Source/RemoteControlUI/Private/AssetEditor/RemoteControlPresetEditorToolkit.h
--- a/Engine/Plugins/VirtualProduction/RemoteControl/Source/RemoteControlUI/Private/AssetEditor/RemoteControlPresetEditorToolkit.h
+++ b/Engine/Plugins/VirtualProduction/RemoteControl/Source/RemoteControlUI/Private/AssetEditor/RemoteControlPresetEditorToolkit.h
@@ -32,7 +32,7 @@ public:
 	virtual FString GetWorldCentricTabPrefix() const override;
 	virtual void RegisterTabSpawners(const TSharedRef<FTabManager>& TabManager) override;
 	virtual void UnregisterTabSpawners(const TSharedRef<FTabManager>& TabManager) override;
-	virtual bool OnRequestClose() override;
+	virtual bool OnRequestClose(EAssetEditorCloseReason InCloseReason) override;
 	//~ End IToolkit interface
 
 
diff --git a/Engine/Source/Editor/AudioEditor/Private/Editors/SoundEffectPresetEditor.cpp b/Engine/Source/Editor/AudioEditor/Private/Editors/SoundEffectPresetEditor.cpp
--- a/Engine/Source/Editor/AudioEditor/Private/Editors/SoundEffectPresetEditor.cpp
+++ b/Engine/Source/Editor/AudioEditor/Private/Editors/SoundEffectPresetEditor.cpp
@@ -163,9 +163,9 @@ void FSoundEffectPresetEditor::Init(const EToolkitMode::Type Mode, const TShared
 		bUseSmallIcons);
 }
 
-bool FSoundEffectPresetEditor::CloseWindow()
+bool FSoundEffectPresetEditor::CloseWindow(EAssetEditorCloseReason InCloseReason)
 {
-	if (FAssetEditorToolkit::CloseWindow())
+	if (FAssetEditorToolkit::CloseWindow(InCloseReason))
 	{
 		UserWidgets.Reset();
 		return true;
diff --git a/Engine/Source/Editor/AudioEditor/Private/Editors/SoundEffectPresetEditor.h b/Engine/Source/Editor/AudioEditor/Private/Editors/SoundEffectPresetEditor.h
--- a/Engine/Source/Editor/AudioEditor/Private/Editors/SoundEffectPresetEditor.h
+++ b/Engine/Source/Editor/AudioEditor/Private/Editors/SoundEffectPresetEditor.h
@@ -36,7 +36,7 @@ public:
 	void Init(const EToolkitMode::Type Mode, const TSharedPtr<IToolkitHost>& InitToolkitHost, USoundEffectPreset* PresetToEdit, const TArray<UUserWidget*>& InWidgetBlueprints);
 
 	/** FAssetEditorToolkit interface */
-	virtual bool CloseWindow() override;
+	virtual bool CloseWindow(EAssetEditorCloseReason InCloseReason) override;
 	virtual FName GetEditorName() const override;
 	virtual FName GetToolkitFName() const override;
 	virtual FText GetBaseToolkitName() const override;
diff --git a/Engine/Source/Editor/AudioEditor/Private/SoundSubmixGraph.cpp b/Engine/Source/Editor/AudioEditor/Private/SoundSubmixGraph.cpp
--- a/Engine/Source/Editor/AudioEditor/Private/SoundSubmixGraph.cpp
+++ b/Engine/Source/Editor/AudioEditor/Private/SoundSubmixGraph.cpp
@@ -49,7 +49,7 @@ void USoundSubmixGraph::SetRootSoundSubmix(USoundSubmixBase* InSoundSubmix)
 							TArray<IAssetEditorInstance*> SubmixEditors = EditorSubsystem->FindEditorsForAsset(Submix);
 							for (IAssetEditorInstance* Editor : SubmixEditors)
 							{
-								Editor->CloseWindow();
+								Editor->CloseWindow(EAssetEditorCloseReason::EditorRefreshRequested);
 							}
 						}
 
diff --git a/Engine/Source/Editor/AudioEditor/Private/SoundSubmixGraphSchema.cpp b/Engine/Source/Editor/AudioEditor/Private/SoundSubmixGraphSchema.cpp
--- a/Engine/Source/Editor/AudioEditor/Private/SoundSubmixGraphSchema.cpp
+++ b/Engine/Source/Editor/AudioEditor/Private/SoundSubmixGraphSchema.cpp
@@ -499,7 +499,7 @@ void USoundSubmixGraphSchema::DroppedAssetsOnGraph(const TArray<FAssetData>& Ass
 			// (to avoid modification of multiple graph editors representing the same branch of submixes)
 			if (SubmixEditor->GetGraph() != Graph)
 			{
-				Editor->CloseWindow();
+				Editor->CloseWindow(EAssetEditorCloseReason::AssetUnloadingOrInvalid);
 			}
 		}
 
diff --git a/Engine/Source/Editor/BlueprintGraph/Private/BlueprintEditorSettings.cpp b/Engine/Source/Editor/BlueprintGraph/Private/BlueprintEditorSettings.cpp
--- a/Engine/Source/Editor/BlueprintGraph/Private/BlueprintEditorSettings.cpp
+++ b/Engine/Source/Editor/BlueprintGraph/Private/BlueprintEditorSettings.cpp
@@ -178,7 +178,7 @@ void UBlueprintEditorSettings::PostEditChangeProperty(FPropertyChangedEvent& Pro
 		FBlueprintEditorModule& BlueprintEditorModule = FModuleManager::LoadModuleChecked<FBlueprintEditorModule>("Kismet");
 		for (const TSharedRef<IBlueprintEditor>& BlueprintEditor : BlueprintEditorModule.GetBlueprintEditors())
 		{
-			BlueprintEditor->CloseWindow();
+			BlueprintEditor->CloseWindow(EAssetEditorCloseReason::EditorRefreshRequested);
 		}
 	}
 
diff --git a/Engine/Source/Editor/Kismet/Private/BlueprintEditor.cpp b/Engine/Source/Editor/Kismet/Private/BlueprintEditor.cpp
--- a/Engine/Source/Editor/Kismet/Private/BlueprintEditor.cpp
+++ b/Engine/Source/Editor/Kismet/Private/BlueprintEditor.cpp
@@ -797,9 +797,9 @@ bool FBlueprintEditor::IsInAScriptingMode() const
 	return IsModeCurrent(FBlueprintEditorApplicationModes::StandardBlueprintEditorMode) || IsModeCurrent(FBlueprintEditorApplicationModes::BlueprintMacroMode);
 }
 
-bool FBlueprintEditor::OnRequestClose()
+bool FBlueprintEditor::OnRequestClose(EAssetEditorCloseReason InCloseReason)
 {
-	return FWorkflowCentricApplication::OnRequestClose();
+	return FWorkflowCentricApplication::OnRequestClose(InCloseReason);
 }
 
 void FBlueprintEditor::OnClose()
@@ -9138,7 +9138,7 @@ void FBlueprintEditor::Tick(float DeltaTime)
 	{
 		IAssetEditorInstance* EditorInst = GEditor->GetEditorSubsystem<UAssetEditorSubsystem>()->FindEditorForAsset(GetBlueprintObj(), /*bFocusIfOpen =*/false);
 		check(EditorInst != nullptr);
-		EditorInst->CloseWindow();
+		EditorInst->CloseWindow(EAssetEditorCloseReason::AssetUnloadingOrInvalid);
 	}
 	else
 	{
diff --git a/Engine/Source/Editor/Kismet/Public/BlueprintEditor.h b/Engine/Source/Editor/Kismet/Public/BlueprintEditor.h
--- a/Engine/Source/Editor/Kismet/Public/BlueprintEditor.h
+++ b/Engine/Source/Editor/Kismet/Public/BlueprintEditor.h
@@ -247,7 +247,7 @@ public:
 
 public:
 	//~ Begin FAssetEditorToolkit Interface
-	virtual bool OnRequestClose() override;
+	virtual bool OnRequestClose(EAssetEditorCloseReason InCloseReason) override;
 	virtual void OnClose() override;
 	// End of FAssetEditorToolkit 
 
diff --git a/Engine/Source/Editor/LevelEditor/Private/SLevelEditor.cpp b/Engine/Source/Editor/LevelEditor/Private/SLevelEditor.cpp
--- a/Engine/Source/Editor/LevelEditor/Private/SLevelEditor.cpp
+++ b/Engine/Source/Editor/LevelEditor/Private/SLevelEditor.cpp
@@ -635,7 +635,7 @@ TSharedPtr<SDockTab> SLevelEditor::AttachSequencer( TSharedPtr<SWidget> Sequence
 
 			if (AssetEditorInstance.IsValid())
 			{
-				InSequencerAssetEditor.Pin()->CloseWindow();
+				InSequencerAssetEditor.Pin()->CloseWindow(EAssetEditorCloseReason::AssetEditorHostClosed);
 			}
 		}
 	};
@@ -653,7 +653,7 @@ TSharedPtr<SDockTab> SLevelEditor::AttachSequencer( TSharedPtr<SWidget> Sequence
 				// Closing the window will invoke this method again but we are handling reopening with a new movie scene ourselves
 				TGuardValue<bool> ReentrantGuard(bIsReentrant, true);
 				// Shutdown cleanly
-				SequencerAssetEditor.Pin()->CloseWindow();
+				SequencerAssetEditor.Pin()->CloseWindow(EAssetEditorCloseReason::AssetEditorHostClosed);
 			}
 
 			if (!FGlobalTabmanager::Get()->OnOverrideDockableAreaRestore_Handler.IsBound())
diff --git a/Engine/Source/Editor/MaterialEditor/Private/MaterialEditor.cpp b/Engine/Source/Editor/MaterialEditor/Private/MaterialEditor.cpp
--- a/Engine/Source/Editor/MaterialEditor/Private/MaterialEditor.cpp
+++ b/Engine/Source/Editor/MaterialEditor/Private/MaterialEditor.cpp
@@ -2053,10 +2053,16 @@ void FMaterialEditor::SaveAssetAs_Execute()
 	}
 }
 
-bool FMaterialEditor::OnRequestClose()
+bool FMaterialEditor::OnRequestClose(EAssetEditorCloseReason InCloseReason)
 {
 	DestroyColorPicker();
 
+	// If the asset has been deleted, we don't want to show the save changes prompt
+	if(InCloseReason == EAssetEditorCloseReason::AssetForceDeleted)
+	{
+		bMaterialDirty = false;
+	}
+	
 	if (bMaterialDirty)
 	{
 		// find out the user wants to do with this dirty material
diff --git a/Engine/Source/Editor/MaterialEditor/Private/MaterialEditor.h b/Engine/Source/Editor/MaterialEditor/Private/MaterialEditor.h
--- a/Engine/Source/Editor/MaterialEditor/Private/MaterialEditor.h
+++ b/Engine/Source/Editor/MaterialEditor/Private/MaterialEditor.h
@@ -487,7 +487,7 @@ protected:
 	virtual void GetSaveableObjects(TArray<UObject*>& OutObjects) const override;
 	virtual void SaveAsset_Execute() override;
 	virtual void SaveAssetAs_Execute() override;
-	virtual bool OnRequestClose() override;
+	virtual bool OnRequestClose(EAssetEditorCloseReason InCloseReason) override;
 
 protected:
 	/** Called when the selection changes in the GraphEditor */
diff --git a/Engine/Source/Editor/MaterialEditor/Private/MaterialInstanceEditor.cpp b/Engine/Source/Editor/MaterialEditor/Private/MaterialInstanceEditor.cpp
--- a/Engine/Source/Editor/MaterialEditor/Private/MaterialInstanceEditor.cpp
+++ b/Engine/Source/Editor/MaterialEditor/Private/MaterialInstanceEditor.cpp
@@ -763,9 +763,9 @@ bool FMaterialInstanceEditor::OnApplyVisible() const
 	return MaterialEditorInstance && MaterialEditorInstance->bIsFunctionPreviewMaterial == true;
 }
 
-bool FMaterialInstanceEditor::OnRequestClose()
+bool FMaterialInstanceEditor::OnRequestClose(EAssetEditorCloseReason InCloseReason)
 {
-	if (MaterialEditorInstance->bIsFunctionInstanceDirty)
+	if (MaterialEditorInstance->bIsFunctionInstanceDirty && InCloseReason != EAssetEditorCloseReason::AssetForceDeleted)
 	{
 		// Find out the user wants to do with this dirty function instance
 		EAppReturnType::Type YesNoCancelReply = FMessageDialog::Open(EAppMsgType::YesNoCancel,
diff --git a/Engine/Source/Editor/MaterialEditor/Private/MaterialInstanceEditor.h b/Engine/Source/Editor/MaterialEditor/Private/MaterialInstanceEditor.h
--- a/Engine/Source/Editor/MaterialEditor/Private/MaterialInstanceEditor.h
+++ b/Engine/Source/Editor/MaterialEditor/Private/MaterialInstanceEditor.h
@@ -126,7 +126,7 @@ protected:
 	//~ FAssetEditorToolkit interface
 	virtual void SaveAsset_Execute() override;
 	virtual void SaveAssetAs_Execute() override;
-	virtual bool OnRequestClose() override;
+	virtual bool OnRequestClose(EAssetEditorCloseReason InCloseReason) override;
 
 	/** Saves editor settings. */
 	void SaveSettings();
diff --git a/Engine/Source/Editor/MaterialEditor/Private/Tests/MaterialEditorTests.cpp b/Engine/Source/Editor/MaterialEditor/Private/Tests/MaterialEditorTests.cpp
--- a/Engine/Source/Editor/MaterialEditor/Private/Tests/MaterialEditorTests.cpp
+++ b/Engine/Source/Editor/MaterialEditor/Private/Tests/MaterialEditorTests.cpp
@@ -209,7 +209,7 @@ namespace MaterialEditorPromotionTestHelper
 						MaterialEditor->UpdateMaterialAfterGraphChange();
 						MaterialEditor->bMaterialDirty = false; // Remove dirty flag so window closes without prompt
 						UE_LOG(LogEditorMaterialEditorPromotionTests, Display, TEXT("Compiled the new material"));
-						MaterialEditor->CloseWindow();
+						MaterialEditor->CloseWindow(EAssetEditorCloseReason::AssetEditorHostClosed);
 
 					}
 				}
diff --git a/Engine/Source/Editor/SequenceRecorder/Private/ActorRecording.cpp b/Engine/Source/Editor/SequenceRecorder/Private/ActorRecording.cpp
--- a/Engine/Source/Editor/SequenceRecorder/Private/ActorRecording.cpp
+++ b/Engine/Source/Editor/SequenceRecorder/Private/ActorRecording.cpp
@@ -134,7 +134,7 @@ bool UActorRecording::StartRecording(ULevelSequence* CurrentSequence, float Curr
 			if (EditorInstance)
 			{
 				UE_LOG(LogAnimation, Log, TEXT("Closing '%s' so we don't invalidate the open version when unloading it."), *TargetAnimation->GetName());
-				EditorInstance->CloseWindow();
+				EditorInstance->CloseWindow(EAssetEditorCloseReason::AssetUnloadingOrInvalid);
 			}
 		}
 
diff --git a/Engine/Source/Editor/SkeletalMeshEditor/Private/SkeletalMeshEditor.cpp b/Engine/Source/Editor/SkeletalMeshEditor/Private/SkeletalMeshEditor.cpp
--- a/Engine/Source/Editor/SkeletalMeshEditor/Private/SkeletalMeshEditor.cpp
+++ b/Engine/Source/Editor/SkeletalMeshEditor/Private/SkeletalMeshEditor.cpp
@@ -145,7 +145,7 @@ bool IsReductionParentBaseLODUseSkeletalMeshBuildWorkflow(USkeletalMesh* Skeleta
 	return false;
 }
 
-bool FSkeletalMeshEditor::OnRequestClose()
+bool FSkeletalMeshEditor::OnRequestClose(EAssetEditorCloseReason InCloseReason)
 {
 	bool bAllowClose = true;
 
@@ -185,7 +185,7 @@ bool FSkeletalMeshEditor::OnRequestClose()
 			}
 		}
 
-		if (bHaveModifiedLOD)
+		if (bHaveModifiedLOD && InCloseReason != EAssetEditorCloseReason::AssetForceDeleted)
 		{
 			// find out the user wants to do with this dirty material
 			EAppReturnType::Type OkCancelReply = FMessageDialog::Open(
diff --git a/Engine/Source/Editor/SkeletalMeshEditor/Private/SkeletalMeshEditor.h b/Engine/Source/Editor/SkeletalMeshEditor/Private/SkeletalMeshEditor.h
--- a/Engine/Source/Editor/SkeletalMeshEditor/Private/SkeletalMeshEditor.h
+++ b/Engine/Source/Editor/SkeletalMeshEditor/Private/SkeletalMeshEditor.h
@@ -82,7 +82,7 @@ public:
 	virtual bool ProcessCommandBindings(const FKeyEvent& InKeyEvent) const override;
 	
 	//~ Begin FAssetEditorToolkit Interface.
-	virtual bool OnRequestClose() override;
+	virtual bool OnRequestClose(EAssetEditorCloseReason InCloseReason) override;
 	virtual FName GetMainTabName() const override;
 	//~ End FAssetEditorToolkit Interface.
 
diff --git a/Engine/Source/Editor/StaticMeshEditor/Private/StaticMeshEditor.cpp b/Engine/Source/Editor/StaticMeshEditor/Private/StaticMeshEditor.cpp
--- a/Engine/Source/Editor/StaticMeshEditor/Private/StaticMeshEditor.cpp
+++ b/Engine/Source/Editor/StaticMeshEditor/Private/StaticMeshEditor.cpp
@@ -2451,10 +2451,10 @@ void FStaticMeshEditor::OnConvexDecomposition()
 	TabManager->TryInvokeTab(CollisionTabId);
 }
 
-bool FStaticMeshEditor::OnRequestClose()
+bool FStaticMeshEditor::OnRequestClose(EAssetEditorCloseReason InCloseReason)
 {
 	bool bAllowClose = true;
-	if (StaticMeshDetails.IsValid() && StaticMeshDetails.Pin()->IsApplyNeeded())
+	if (InCloseReason != EAssetEditorCloseReason::AssetForceDeleted && StaticMeshDetails.IsValid() && StaticMeshDetails.Pin()->IsApplyNeeded())
 	{
 		// find out the user wants to do with this dirty material
 		EAppReturnType::Type YesNoCancelReply = FMessageDialog::Open(
diff --git a/Engine/Source/Editor/StaticMeshEditor/Private/StaticMeshEditor.h b/Engine/Source/Editor/StaticMeshEditor/Private/StaticMeshEditor.h
--- a/Engine/Source/Editor/StaticMeshEditor/Private/StaticMeshEditor.h
+++ b/Engine/Source/Editor/StaticMeshEditor/Private/StaticMeshEditor.h
@@ -473,7 +473,7 @@ private:
 	void OnConvexDecomposition();
 
 	//~ Begin FAssetEditorToolkit Interface.
-	virtual bool OnRequestClose() override;
+	virtual bool OnRequestClose(EAssetEditorCloseReason InCloseReason) override;
 	//~ End FAssetEditorToolkit Interface.
 
 	//~ Begin FEditorUndoClient Interface
diff --git a/Engine/Source/Editor/TranslationEditor/Private/ITranslationEditor.cpp b/Engine/Source/Editor/TranslationEditor/Private/ITranslationEditor.cpp
--- a/Engine/Source/Editor/TranslationEditor/Private/ITranslationEditor.cpp
+++ b/Engine/Source/Editor/TranslationEditor/Private/ITranslationEditor.cpp
@@ -23,7 +23,7 @@ void ITranslationEditor::OpenTranslationEditor(const FString& InManifestFile, co
 
 		if (!bLoadedSuccessfully)
 		{
-			NewTranslationEditor->CloseWindow();
+			NewTranslationEditor->CloseWindow(EAssetEditorCloseReason::AssetUnloadingOrInvalid);
 		}
 	}
 	else
@@ -57,7 +57,7 @@ void ITranslationEditor::OpenTranslationEditor(ULocalizationTarget* const Locali
 
 		if (!bLoadedSuccessfully)
 		{
-			NewTranslationEditor->CloseWindow();
+			NewTranslationEditor->CloseWindow(EAssetEditorCloseReason::AssetUnloadingOrInvalid);
 		}
 	}
 	else
diff --git a/Engine/Source/Editor/UMGEditor/Private/WidgetBlueprintEditor.cpp b/Engine/Source/Editor/UMGEditor/Private/WidgetBlueprintEditor.cpp
--- a/Engine/Source/Editor/UMGEditor/Private/WidgetBlueprintEditor.cpp
+++ b/Engine/Source/Editor/UMGEditor/Private/WidgetBlueprintEditor.cpp
@@ -1576,9 +1576,9 @@ void FWidgetBlueprintEditor::Compile()
 	}
 }
 
-bool FWidgetBlueprintEditor::OnRequestClose()
+bool FWidgetBlueprintEditor::OnRequestClose(EAssetEditorCloseReason InCloseReason)
 {
-	bool bAllowClose = Super::OnRequestClose();
+	bool bAllowClose = Super::OnRequestClose(InCloseReason);
 
 	// Give any active modes a chance to shutdown while the toolkit host is still alive
 	// Note: This along side with the default tool palette extension tab being closed 
diff --git a/Engine/Source/Editor/UMGEditor/Public/WidgetBlueprintEditor.h b/Engine/Source/Editor/UMGEditor/Public/WidgetBlueprintEditor.h
--- a/Engine/Source/Editor/UMGEditor/Public/WidgetBlueprintEditor.h
+++ b/Engine/Source/Editor/UMGEditor/Public/WidgetBlueprintEditor.h
@@ -73,7 +73,7 @@ public:
 	//~ End FBlueprintEditor interface
 	
 	//~ Begin FAssetEditorToolkit Interface
-	virtual bool OnRequestClose() override;
+	virtual bool OnRequestClose(EAssetEditorCloseReason InCloseReason) override;
 	// End of FAssetEditorToolkit 
 
 	//~ Begin FGCObjectInterface interface
diff --git a/Engine/Source/Editor/UnrealEd/Private/MiniCurveEditor.cpp b/Engine/Source/Editor/UnrealEd/Private/MiniCurveEditor.cpp
--- a/Engine/Source/Editor/UnrealEd/Private/MiniCurveEditor.cpp
+++ b/Engine/Source/Editor/UnrealEd/Private/MiniCurveEditor.cpp
@@ -70,6 +70,11 @@ void SMiniCurveEditor::FocusWindow(UObject* ObjectToFocusOn)
 }
 
 bool SMiniCurveEditor::CloseWindow()
+{
+	return CloseWindow(EAssetEditorCloseReason::AssetEditorHostClosed);
+}
+
+bool SMiniCurveEditor::CloseWindow(EAssetEditorCloseReason InCloseReason)
 {
 	if(WidgetWindow.IsValid())
 	{
diff --git a/Engine/Source/Editor/UnrealEd/Private/ObjectTools.cpp b/Engine/Source/Editor/UnrealEd/Private/ObjectTools.cpp
--- a/Engine/Source/Editor/UnrealEd/Private/ObjectTools.cpp
+++ b/Engine/Source/Editor/UnrealEd/Private/ObjectTools.cpp
@@ -3242,7 +3242,7 @@ namespace ObjectTools
 				TArray<IAssetEditorInstance*> ObjectEditors = GEditor->GetEditorSubsystem<UAssetEditorSubsystem>()->FindEditorsForAssetAndSubObjects(Object);
 				for (IAssetEditorInstance* ObjectEditorInstance : ObjectEditors)
 				{
-					if (!ObjectEditorInstance->CloseWindow())
+					if (!ObjectEditorInstance->CloseWindow(EAssetEditorCloseReason::AssetForceDeleted))
 					{
 						bClosedAllEditors = false;
 					}
diff --git a/Engine/Source/Editor/UnrealEd/Private/Settings/BlueprintEditorProjectSettings.cpp b/Engine/Source/Editor/UnrealEd/Private/Settings/BlueprintEditorProjectSettings.cpp
--- a/Engine/Source/Editor/UnrealEd/Private/Settings/BlueprintEditorProjectSettings.cpp
+++ b/Engine/Source/Editor/UnrealEd/Private/Settings/BlueprintEditorProjectSettings.cpp
@@ -42,7 +42,7 @@ void UBlueprintEditorProjectSettings::PostEditChangeProperty(struct FPropertyCha
 		FBlueprintEditorModule& BlueprintEditorModule = FModuleManager::LoadModuleChecked<FBlueprintEditorModule>("Kismet");
 		for (const TSharedRef<IBlueprintEditor>& BlueprintEditor : BlueprintEditorModule.GetBlueprintEditors())
 		{
-			BlueprintEditor->CloseWindow();
+			BlueprintEditor->CloseWindow(EAssetEditorCloseReason::AssetUnloadingOrInvalid);
 		}
 	}
 }
diff --git a/Engine/Source/Editor/UnrealEd/Private/StaticLightingSystem/StaticLightingSystem.cpp b/Engine/Source/Editor/UnrealEd/Private/StaticLightingSystem/StaticLightingSystem.cpp
--- a/Engine/Source/Editor/UnrealEd/Private/StaticLightingSystem/StaticLightingSystem.cpp
+++ b/Engine/Source/Editor/UnrealEd/Private/StaticLightingSystem/StaticLightingSystem.cpp
@@ -2530,7 +2530,7 @@ void UEditorEngine::BuildLighting(const FLightingBuildOptions& Options)
 			IAssetEditorInstance* Editor = AssetEditorSubsystem->FindEditorForAsset(EditedAsset, false);
 			if (Editor)
 			{
-				Editor->CloseWindow();
+				Editor->CloseWindow(EAssetEditorCloseReason::AssetUnloadingOrInvalid);
 			}
 		}
 	}
diff --git a/Engine/Source/Editor/UnrealEd/Private/Subsystems/AssetEditorSubsystem.cpp b/Engine/Source/Editor/UnrealEd/Private/Subsystems/AssetEditorSubsystem.cpp
--- a/Engine/Source/Editor/UnrealEd/Private/Subsystems/AssetEditorSubsystem.cpp
+++ b/Engine/Source/Editor/UnrealEd/Private/Subsystems/AssetEditorSubsystem.cpp
@@ -308,7 +308,7 @@ int32 UAssetEditorSubsystem::CloseAllEditorsForAsset(UObject* Asset)
 	{
 		if (EditorInstance)
 		{
-			EditorInstance->CloseWindow();
+			EditorInstance->CloseWindow(EAssetEditorCloseReason::CloseAllEditorsForAsset);
 		}
 	}
 
@@ -342,7 +342,7 @@ void UAssetEditorSubsystem::CloseOtherEditors(UObject* Asset, IAssetEditorInstan
 		IAssetEditorInstance* Editor = It.Value();
 		if (Asset == It.Key().RawPtr && Editor != OnlyEditor)
 		{
-			Editor->CloseWindow();
+			Editor->CloseWindow(EAssetEditorCloseReason::CloseOtherEditors);
 		}
 	}
 
@@ -443,7 +443,7 @@ bool UAssetEditorSubsystem::CloseAllAssetEditors()
 		IAssetEditorInstance* Editor = It.Key();
 		if (Editor != nullptr)
 		{
-			if (!Editor->CloseWindow())
+			if (!Editor->CloseWindow(EAssetEditorCloseReason::CloseAllAssetEditors))
 			{
 				bAllEditorsClosed = false;
 			}
diff --git a/Engine/Source/Editor/UnrealEd/Private/Toolkits/AssetEditorToolkit.cpp b/Engine/Source/Editor/UnrealEd/Private/Toolkits/AssetEditorToolkit.cpp
--- a/Engine/Source/Editor/UnrealEd/Private/Toolkits/AssetEditorToolkit.cpp
+++ b/Engine/Source/Editor/UnrealEd/Private/Toolkits/AssetEditorToolkit.cpp
@@ -223,7 +223,7 @@ void FAssetEditorToolkit::InitAssetEditor( const EToolkitMode::Type Mode, const
 					}
 					return true;
 				})
-			.OnRequestClose(this, &FAssetEditorToolkit::OnRequestClose)
+			.OnRequestClose(this, &FAssetEditorToolkit::OnRequestClose, EAssetEditorCloseReason::AssetEditorHostClosed)
 			.OnClose(this, &FAssetEditorToolkit::OnClose)
 		);
 
@@ -485,10 +485,22 @@ void FAssetEditorToolkit::FocusWindow(UObject* ObjectToFocusOn)
 	BringToolkitToFront();
 }
 
+bool FAssetEditorToolkit::OnRequestClose(EAssetEditorCloseReason)
+{
+	PRAGMA_DISABLE_DEPRECATION_WARNINGS
+	return OnRequestClose();
+	PRAGMA_ENABLE_DEPRECATION_WARNINGS
+}
 
 bool FAssetEditorToolkit::CloseWindow()
 {
-	if (OnRequestClose())
+	// We use AssetEditorHostClosed as the default close reason for legacy cases
+	return CloseWindow(EAssetEditorCloseReason::AssetEditorHostClosed);
+}
+
+bool FAssetEditorToolkit::CloseWindow(EAssetEditorCloseReason InCloseReason)
+{
+	if (OnRequestClose(InCloseReason))
 	{
 		OnClose();
 
@@ -783,7 +795,7 @@ void FAssetEditorToolkit::RemoveEditingAsset(UObject* Asset)
 	// Just close the editor tab if it's the last element
 	if (EditingObjects.Num() == 1 && EditingObjects.Contains(Asset))
 	{
-		CloseWindow();
+		CloseWindow(EAssetEditorCloseReason::AssetUnloadingOrInvalid);
 	}
 	else
 	{
diff --git a/Engine/Source/Editor/UnrealEd/Private/Tools/UAssetEditor.cpp b/Engine/Source/Editor/UnrealEd/Private/Tools/UAssetEditor.cpp
--- a/Engine/Source/Editor/UnrealEd/Private/Tools/UAssetEditor.cpp
+++ b/Engine/Source/Editor/UnrealEd/Private/Tools/UAssetEditor.cpp
@@ -33,11 +33,6 @@ void UAssetEditor::FocusWindow(UObject* ObjectToFocusOn /*= nullptr*/)
 {
 }
 
-bool UAssetEditor::CloseWindow()
-{
-	return true;
-}
-
 TSharedPtr<class FTabManager> UAssetEditor::GetAssociatedTabManager()
 {
 	return ToolkitInstance->GetAssociatedTabManager();
diff --git a/Engine/Source/Editor/UnrealEd/Private/WorkflowCentricApplication.cpp b/Engine/Source/Editor/UnrealEd/Private/WorkflowCentricApplication.cpp
--- a/Engine/Source/Editor/UnrealEd/Private/WorkflowCentricApplication.cpp
+++ b/Engine/Source/Editor/UnrealEd/Private/WorkflowCentricApplication.cpp
@@ -122,7 +122,7 @@ void FWorkflowCentricApplication::PushTabFactories(FWorkflowAllowedTabSet& Facto
 	}
 }
 
-bool FWorkflowCentricApplication::OnRequestClose()
+bool FWorkflowCentricApplication::OnRequestClose(EAssetEditorCloseReason InCloseReason)
 {
 	return FSlateApplication::Get().IsNormalExecution();
 }
diff --git a/Engine/Source/Editor/UnrealEd/Public/MiniCurveEditor.h b/Engine/Source/Editor/UnrealEd/Public/MiniCurveEditor.h
--- a/Engine/Source/Editor/UnrealEd/Public/MiniCurveEditor.h
+++ b/Engine/Source/Editor/UnrealEd/Public/MiniCurveEditor.h
@@ -32,6 +32,7 @@ public:
 	virtual FName GetEditorName() const override;
 	virtual void FocusWindow(UObject* ObjectToFocusOn) override;
 	virtual bool CloseWindow() override;
+	virtual bool CloseWindow(EAssetEditorCloseReason InCloseReason) override;
 	virtual bool IsPrimaryEditor() const override { return true; };
 	virtual void InvokeTab(const struct FTabId& TabId) override {}
 	virtual TSharedPtr<class FTabManager> GetAssociatedTabManager() override;
diff --git a/Engine/Source/Editor/UnrealEd/Public/Subsystems/AssetEditorSubsystem.h b/Engine/Source/Editor/UnrealEd/Public/Subsystems/AssetEditorSubsystem.h
--- a/Engine/Source/Editor/UnrealEd/Public/Subsystems/AssetEditorSubsystem.h
+++ b/Engine/Source/Editor/UnrealEd/Public/Subsystems/AssetEditorSubsystem.h
@@ -23,6 +23,28 @@ class IMessageContext;
 class UToolMenu;
 struct FToolMenuSection;
 
+/** The way that editors were requested to close */
+enum class EAssetEditorCloseReason : uint8
+{
+	/* NOTE: All close reasons can be passed into OnRequestClose() and CloseWindow(), but only some are broadcast by
+	 * UAssetEditorSubsystem::AssetEditorRequestCloseEvent currently while others are for the asset editors themselves
+	 */
+	
+	// Close reasons broadcast by UAssetEditorSubsystem
+	
+	CloseAllEditorsForAsset,   // All asset editors operating on a specific asset are being requested to close
+	CloseOtherEditors,         // An asset editor is requesting all asset editors using an asset except itself to close
+	RemoveAssetFromAllEditors, // An asset is being removed from all asset editors (which may or may not actually close)
+	CloseAllAssetEditors,      // Every single Asset Editor is being requested to close 
+
+	// Close reasons for individual asset editors
+	
+	AssetEditorHostClosed,     // The "default" reason for an asset editor to close, e.g when the close button is clicked
+	AssetUnloadingOrInvalid,   // The asset being edited is being unloaded or is no longer valid
+	EditorRefreshRequested,    // The asset editor wants to close and re-open to edit the same asset
+	AssetForceDeleted          // The asset being edited has been force deleted
+};
+
 /**
  * This class keeps track of a currently open asset editor; allowing it to be
  * brought into focus, closed, etc..., without concern for how the editor was
@@ -34,7 +56,17 @@ public:
 
 	virtual FName GetEditorName() const = 0;
 	virtual void FocusWindow(UObject* ObjectToFocusOn = nullptr) = 0;
-	virtual bool CloseWindow() = 0;
+
+	UE_DEPRECATED(5.3, "Use CloseWindow that takes in an EAssetEditorCloseReason instead")
+	virtual bool CloseWindow() { return true; }
+	
+	virtual bool CloseWindow(EAssetEditorCloseReason InCloseReason)
+	{
+		PRAGMA_DISABLE_DEPRECATION_WARNINGS
+		return CloseWindow();
+		PRAGMA_ENABLE_DEPRECATION_WARNINGS
+	}
+
 	virtual bool IsPrimaryEditor() const = 0;
 	virtual void InvokeTab(const struct FTabId& TabId) = 0;
 	UE_DEPRECATED(5.0, "Toolbar tab no longer exists and tab ID will return None; do not add it to layouts")
@@ -45,15 +77,6 @@ public:
 	virtual void RemoveEditingAsset(UObject* Asset) = 0;
 };
 
-/** The way that editors were requested to close */
-enum class EAssetEditorCloseReason : uint8
-{
-	CloseAllEditorsForAsset,
-	CloseOtherEditors,
-	RemoveAssetFromAllEditors,
-	CloseAllAssetEditors,
-};
-
 struct UNREALED_API FRegisteredModeInfo
 {
 	TWeakObjectPtr<UClass> ModeClass;
diff --git a/Engine/Source/Editor/UnrealEd/Public/Toolkits/AssetEditorToolkit.h b/Engine/Source/Editor/UnrealEd/Public/Toolkits/AssetEditorToolkit.h
--- a/Engine/Source/Editor/UnrealEd/Public/Toolkits/AssetEditorToolkit.h
+++ b/Engine/Source/Editor/UnrealEd/Public/Toolkits/AssetEditorToolkit.h
@@ -140,7 +140,9 @@ public:
 	/** IAssetEditorInstance interface */
 	virtual FName GetEditorName() const override;
 	virtual void FocusWindow(UObject* ObjectToFocusOn = nullptr) override;
+	UE_DEPRECATED(5.3, "Use CloseWindow that takes in an EAssetEditorCloseReason instead")
 	virtual bool CloseWindow() override;
+	virtual bool CloseWindow(EAssetEditorCloseReason InCloseReason) override;
 	virtual bool IsPrimaryEditor() const override { return true; };
 	virtual void InvokeTab(const FTabId& TabId) override;
 	virtual FName GetEditingAssetTypeName() const override;
@@ -352,11 +354,14 @@ protected:
 	virtual bool ShouldPromptForNewFilesOnReload(const UObject& object) const;
 
 	/** Called when this toolkit is requested to close. Returns false if closing should be prevented. */
+	UE_DEPRECATED(5.3, "Use OnRequestClose that takes in an EAssetEditorCloseReason instead")
 	virtual bool OnRequestClose() { return true; }
 
+	/** Called when this toolkit is requested to close. Returns false if closing should be prevented. */
+	virtual bool OnRequestClose(EAssetEditorCloseReason InCloseReason);
+
 	/** Called when this toolkit is being closed */
 	virtual void OnClose() {}
-
 	/**
 	  * Static: Called when "Switch to Standalone Editor" is clicked for the asset editor
 	  *
diff --git a/Engine/Source/Editor/UnrealEd/Public/Tools/UAssetEditor.h b/Engine/Source/Editor/UnrealEd/Public/Tools/UAssetEditor.h
--- a/Engine/Source/Editor/UnrealEd/Public/Tools/UAssetEditor.h
+++ b/Engine/Source/Editor/UnrealEd/Public/Tools/UAssetEditor.h
@@ -28,7 +28,6 @@ public:
 
 	virtual FName GetEditorName() const override;
 	virtual void FocusWindow(UObject* ObjectToFocusOn = nullptr) override;
-	virtual bool CloseWindow() override;
 	virtual bool IsPrimaryEditor() const override
 	{
 		return true;
diff --git a/Engine/Source/Editor/UnrealEd/Public/WorkflowOrientedApp/WorkflowCentricApplication.h b/Engine/Source/Editor/UnrealEd/Public/WorkflowOrientedApp/WorkflowCentricApplication.h
--- a/Engine/Source/Editor/UnrealEd/Public/WorkflowOrientedApp/WorkflowCentricApplication.h
+++ b/Engine/Source/Editor/UnrealEd/Public/WorkflowOrientedApp/WorkflowCentricApplication.h
@@ -25,7 +25,7 @@ public:
 	// End of IToolkit interface
 	
 	// FAssetEditorToolkit interface
-	virtual bool OnRequestClose() override;
+	virtual bool OnRequestClose(EAssetEditorCloseReason InCloseReason) override;
 	virtual void OnClose() override;
 	virtual FName GetToolMenuToolbarName(FName& OutParentName) const override;
 	// End of FAssetEditorToolkit interface
diff --git a/Engine/Source/Editor/VREditor/Private/VREditorActions.cpp b/Engine/Source/Editor/VREditor/Private/VREditorActions.cpp
--- a/Engine/Source/Editor/VREditor/Private/VREditorActions.cpp
+++ b/Engine/Source/Editor/VREditor/Private/VREditorActions.cpp
@@ -342,7 +342,7 @@ void FVREditorActionCallbacks::CloseSequencer(UMovieSceneSequence* OpenSequence)
 	IAssetEditorInstance* SequencerEditor = GEditor->GetEditorSubsystem<UAssetEditorSubsystem>()->FindEditorForAsset(OpenSequence, false);
 	if (SequencerEditor != nullptr)
 	{
-		SequencerEditor->CloseWindow();
+		SequencerEditor->CloseWindow(EAssetEditorCloseReason::AssetEditorHostClosed);
 	}
 }
 
diff --git a/Engine/Source/Runtime/Engine/Private/SoundSubmix.cpp b/Engine/Source/Runtime/Engine/Private/SoundSubmix.cpp
--- a/Engine/Source/Runtime/Engine/Private/SoundSubmix.cpp
+++ b/Engine/Source/Runtime/Engine/Private/SoundSubmix.cpp
@@ -1152,7 +1152,7 @@ ENGINE_API void SubmixUtils::RefreshEditorForSubmix(const USoundSubmixBase* InSu
 				TArray<IAssetEditorInstance*> SubmixEditors = EditorSubsystem->FindEditorsForAsset(WeakSubmix.Get());
 				for (IAssetEditorInstance* Editor : SubmixEditors)
 				{
-					Editor->CloseWindow();
+					Editor->CloseWindow(EAssetEditorCloseReason::EditorRefreshRequested);
 				}
 
 				EditorSubsystem->OpenEditorForAsset(WeakSubmix.Get());
