Properly namespace enums so names do note end up shadowing
#rnx

diff --git a/Engine/Source/Editor/UnrealEd/Classes/Commandlets/GatherTextCommandletBase.h b/Engine/Source/Editor/UnrealEd/Classes/Commandlets/GatherTextCommandletBase.h
--- a/Engine/Source/Editor/UnrealEd/Classes/Commandlets/GatherTextCommandletBase.h
+++ b/Engine/Source/Editor/UnrealEd/Classes/Commandlets/GatherTextCommandletBase.h
@@ -21,7 +21,7 @@ struct UNREALED_API FGatherTextDelegates
 class FFuzzyPathMatcher
 {
 public:
-	enum EPathMatch
+	enum class EPathMatch
 	{
 		Included,
 		Excluded,
@@ -34,7 +34,7 @@ public:
 	EPathMatch TestPath(const FString& InPathToTest) const;
 
 private:
-	enum EPathType : uint8
+	enum class EPathType : uint8
 	{
 		Include,
 		Exclude,
diff --git a/Engine/Source/Editor/UnrealEd/Private/Commandlets/GatherTextFromAssetsCommandlet.cpp b/Engine/Source/Editor/UnrealEd/Private/Commandlets/GatherTextFromAssetsCommandlet.cpp
--- a/Engine/Source/Editor/UnrealEd/Private/Commandlets/GatherTextFromAssetsCommandlet.cpp
+++ b/Engine/Source/Editor/UnrealEd/Private/Commandlets/GatherTextFromAssetsCommandlet.cpp
@@ -681,7 +681,7 @@ int32 UGatherTextFromAssetsCommandlet::Main(const FString& Params)
 		}
 
 		// Filter out assets whose package file paths do not pass the "fuzzy path" filters.
-		if (FuzzyPathMatcher.TestPath(PackageFilePathWithExtension) != FFuzzyPathMatcher::Included)
+		if (FuzzyPathMatcher.TestPath(PackageFilePathWithExtension) != FFuzzyPathMatcher::EPathMatch::Included)
 		{
 			return true;
 		}
diff --git a/Engine/Source/Editor/UnrealEd/Private/Commandlets/GatherTextFromMetadataCommandlet.cpp b/Engine/Source/Editor/UnrealEd/Private/Commandlets/GatherTextFromMetadataCommandlet.cpp
--- a/Engine/Source/Editor/UnrealEd/Private/Commandlets/GatherTextFromMetadataCommandlet.cpp
+++ b/Engine/Source/Editor/UnrealEd/Private/Commandlets/GatherTextFromMetadataCommandlet.cpp
@@ -310,7 +310,7 @@ void UGatherTextFromMetaDataCommandlet::GatherTextFromUObjects(const TArray<FStr
 		check(!SourceFilePath.IsEmpty());
 
 		const FFuzzyPathMatcher::EPathMatch PathMatch = FuzzyPathMatcher.TestPath(SourceFilePath);
-		if (PathMatch != FFuzzyPathMatcher::Included)
+		if (PathMatch != FFuzzyPathMatcher::EPathMatch::Included)
 		{
 			continue;
 		}
diff --git a/Engine/Source/Editor/UnrealEd/Private/Commandlets/GatherTextFromSourceCommandlet.cpp b/Engine/Source/Editor/UnrealEd/Private/Commandlets/GatherTextFromSourceCommandlet.cpp
--- a/Engine/Source/Editor/UnrealEd/Private/Commandlets/GatherTextFromSourceCommandlet.cpp
+++ b/Engine/Source/Editor/UnrealEd/Private/Commandlets/GatherTextFromSourceCommandlet.cpp
@@ -199,7 +199,7 @@ int32 UGatherTextFromSourceCommandlet::Main( const FString& Params )
 	FilesToProcess.RemoveAll([&FuzzyPathMatcher](const FString& FoundFile)
 	{
 		// Filter out assets whose package file paths do not pass the "fuzzy path" filters.
-		if (FuzzyPathMatcher.TestPath(FoundFile) != FFuzzyPathMatcher::Included)
+		if (FuzzyPathMatcher.TestPath(FoundFile) != FFuzzyPathMatcher::EPathMatch::Included)
 		{
 			return true;
 		}
