Disable ensure while investigating UE-185118

diff --git a/Engine/Source/Editor/UnrealEd/Private/Kismet2/KismetReinstanceUtilities.cpp b/Engine/Source/Editor/UnrealEd/Private/Kismet2/KismetReinstanceUtilities.cpp
--- a/Engine/Source/Editor/UnrealEd/Private/Kismet2/KismetReinstanceUtilities.cpp
+++ b/Engine/Source/Editor/UnrealEd/Private/Kismet2/KismetReinstanceUtilities.cpp
@@ -2657,7 +2657,7 @@ void FBlueprintCompileReinstancer::ReplaceInstancesOfClass_Inner(const TMap<UCla
 		}
 	}
 
-	FReplaceReferenceHelper::ValidateReplacementMappings(OldToNewInstanceMap);
+	//FReplaceReferenceHelper::ValidateReplacementMappings(OldToNewInstanceMap);
 	FReplaceReferenceHelper::FindAndReplaceReferences(SourceObjects, ObjectsThatShouldUseOldStuff, ObjectsReplaced, OldToNewInstanceMap, ReinstancedObjectsWeakReferenceMap);
 	
 	for (UObject* Obj : ObjectsReplaced)
