[Perf] Fix for hitch when highlighting pins in debugger
UDN: https://epicgames.lightning.force.com/lightning/r/Case/5004z00001nzTuKAAU/view

diff --git a/Engine/Source/Editor/Kismet/Private/WatchPointViewer.cpp b/Engine/Source/Editor/Kismet/Private/WatchPointViewer.cpp
--- a/Engine/Source/Editor/Kismet/Private/WatchPointViewer.cpp
+++ b/Engine/Source/Editor/Kismet/Private/WatchPointViewer.cpp
@@ -78,7 +78,7 @@ namespace
 			FText InBlueprintName,
 			FText InGraphName,
 			FText InNodeName,
-			const FPropertyInstanceInfo& Info
+			FPropertyInstanceInfo& Info
 		)
 			: BP(InBP)
 			, Node(InNode)
@@ -98,7 +98,7 @@ namespace
 
 			for (const TSharedPtr<FPropertyInstanceInfo>& ChildInfo : Info.GetChildren())
 			{
-				Children.Add(MakeShared<FWatchRow>(InBP, InNode, InPin, InObjectBeingDebugged, BlueprintName, GraphName, NodeName, MoveTemp(*ChildInfo)));
+				Children.Add(MakeShared<FWatchRow>(InBP, InNode, InPin, InObjectBeingDebugged, BlueprintName, GraphName, NodeName, *ChildInfo));
 			}
 		}
 
@@ -206,7 +206,7 @@ namespace
 							BlueprintName,
 							MoveTemp(GraphName),
 							MoveTemp(NodeName),
-							MoveTemp(*DebugInfo)
+							*DebugInfo
 						)
 					);
 				}
diff --git a/Engine/Source/Editor/UnrealEd/Private/Kismet2/KismetDebugUtilities.cpp b/Engine/Source/Editor/UnrealEd/Private/Kismet2/KismetDebugUtilities.cpp
--- a/Engine/Source/Editor/UnrealEd/Private/Kismet2/KismetDebugUtilities.cpp
+++ b/Engine/Source/Editor/UnrealEd/Private/Kismet2/KismetDebugUtilities.cpp
@@ -2033,15 +2033,16 @@ void FKismetDebugUtilities::GetDebugInfo_InContainer(int32 Index, TSharedPtr<FPr
 
 void FKismetDebugUtilities::GetDebugInfoInternal(TSharedPtr<FPropertyInstanceInfo> &DebugInfo, const FProperty* Property, const void* PropertyValue)
 {
-	TMap<FPropertyInstanceInfo::FPropertyInstance, TSharedPtr<FPropertyInstanceInfo>> VisitedNodes;
-	DebugInfo = FPropertyInstanceInfo::FindOrMake({Property, PropertyValue}, VisitedNodes);
+	DebugInfo = FPropertyInstanceInfo::Make({Property, PropertyValue}, nullptr);
 }
 
-FPropertyInstanceInfo::FPropertyInstanceInfo(FPropertyInstance PropertyInstance) :
+FPropertyInstanceInfo::FPropertyInstanceInfo(FPropertyInstance PropertyInstance, const TSharedPtr<FPropertyInstanceInfo>& InParent) :
 	Name(FText::FromString(PropertyInstance.Property->GetName())),
 	DisplayName(PropertyInstance.Property->GetDisplayNameText()),
 	Type(UEdGraphSchema_K2::TypeToText(PropertyInstance.Property)),
-	Property(PropertyInstance.Property)
+	Property(PropertyInstance.Property),
+	ValueAddress(PropertyInstance.Value),
+	Parent(InParent)
 {
 	const FProperty* ResolvedProperty = Property.Get();
 	check(ResolvedProperty);
@@ -2212,47 +2213,12 @@ FPropertyInstanceInfo::FPropertyInstanceInfo(FPropertyInstance PropertyInstance)
 	ensureMsgf(false, TEXT("Failed to identify property type. This function may need to be expanded to include it: %s"), *ResolvedProperty->GetClass()->GetName());
 }
 
-TSharedPtr<FPropertyInstanceInfo> FPropertyInstanceInfo::FindOrMake(FPropertyInstance PropertyInstance,
-	TMap<FPropertyInstance, TSharedPtr<FPropertyInstanceInfo>>& VisitedNodes)
+TSharedPtr<FPropertyInstanceInfo> FPropertyInstanceInfo::Make(FPropertyInstance PropertyInstance, const TSharedPtr<FPropertyInstanceInfo>& Parent)
 {
-	if(TSharedPtr<FPropertyInstanceInfo>* Found = VisitedNodes.Find(PropertyInstance))
-	{
-		return *Found;
-	}
-
-	// wait to populate children until after inserting into VisitedNodes.
-	// this way we will catch circular references
-	TSharedPtr<FPropertyInstanceInfo> DebugInfo = MakeShared<FPropertyInstanceInfo>(PropertyInstance);
-
-	bool bFoundObject = false;
-	if (DebugInfo->Object.IsValid())
-	{
-		for (const TPair<FPropertyInstance, TSharedPtr<FPropertyInstanceInfo>>& VisitedNode : VisitedNodes)
-		{
-			if (VisitedNode.Value->Object == DebugInfo->Object)
-			{
-				bFoundObject = true;
-				DebugInfo->ReferencedObject = VisitedNode.Value;
-				break;
-			}
-		}
-	}
-	
-	VisitedNodes.Add(
-		PropertyInstance,
-		DebugInfo
-	);
-
-	if (!bFoundObject)
-	{
-		DebugInfo->PopulateChildren(PropertyInstance, VisitedNodes);
-	}
-
-	return DebugInfo;
+	return MakeShared<FPropertyInstanceInfo>(PropertyInstance, Parent);
 }
 
-void FPropertyInstanceInfo::PopulateChildren(FPropertyInstance PropertyInstance,
-	TMap<FPropertyInstance, TSharedPtr<FPropertyInstanceInfo>>& VisitedNodes)
+void FPropertyInstanceInfo::PopulateChildren(FPropertyInstance PropertyInstance)
 {
 	check(PropertyInstance.Property);
 	if (PropertyInstance.Value == nullptr)
@@ -2279,11 +2245,12 @@ void FPropertyInstanceInfo::PopulateChildren(FPropertyInstance PropertyInstance,
 				ArrayProperty->Inner,
 				ArrayHelper.GetRawPtr(i)
 			};
-			const TSharedPtr<FPropertyInstanceInfo> ChildInfo = FindOrMake(ChildProperty, VisitedNodes);
+			const TSharedPtr<FPropertyInstanceInfo> ChildInfo = Make(ChildProperty, AsShared());
 			
 			// overwrite the display name with the array index for the current element
 			ChildInfo->DisplayName = FText::Format(LOCTEXT("ArrayIndexName", "[{0}]"), FText::AsNumber(i));
 			ChildInfo->bIsInContainer = true;
+			ChildInfo->ContainerIndex = i;
 			Children.Add(ChildInfo);
 		}
 	}
@@ -2295,7 +2262,7 @@ void FPropertyInstanceInfo::PopulateChildren(FPropertyInstance PropertyInstance,
 				*It,
 				It->ContainerPtrToValuePtr<void>(PropertyInstance.Value, 0)
 			};
-			const TSharedPtr<FPropertyInstanceInfo> ChildInfo = FindOrMake(ChildProperty, VisitedNodes);
+			const TSharedPtr<FPropertyInstanceInfo> ChildInfo = Make(ChildProperty, AsShared());
 			
 			Children.Add(ChildInfo);
 		}
@@ -2303,60 +2270,53 @@ void FPropertyInstanceInfo::PopulateChildren(FPropertyInstance PropertyInstance,
 	else if (const FMapProperty* MapProperty = CastField<FMapProperty>(Property.Get()))
 	{
 		FScriptMapHelper MapHelper(MapProperty, PropertyInstance.Value);
-		uint8* PropData = MapHelper.GetPairPtr(0);
 
 		int32 Index = 0;
-		for (int32 Count = MapHelper.Num(); Count; PropData += MapProperty->MapLayout.SetLayout.Size, ++Index)
+		for (FScriptMapHelper::FIterator Itr = MapHelper.CreateIterator(); Itr; ++Itr, ++Index)
 		{
-			if (MapHelper.IsValidIndex(Index))
-			{
-				FPropertyInstance ChildProperty = {
-					MapProperty->ValueProp,
-					PropData + MapProperty->MapLayout.ValueOffset
-				};
-				const TSharedPtr<FPropertyInstanceInfo> ChildInfo = FindOrMake(ChildProperty, VisitedNodes);
+			uint8* KeyData = MapHelper.GetKeyPtr(*Itr);
+			uint8* ValData = MapHelper.GetValuePtr(*Itr);
+			FPropertyInstance ChildProperty = {
+				MapProperty->ValueProp,
+				ValData
+			};
+			const TSharedPtr<FPropertyInstanceInfo> ChildInfo = Make(ChildProperty, AsShared());
 				
-				FString NameStr = TEXT("[");
-				MapProperty->KeyProp->ExportTextItem_Direct(
-					NameStr,
-					PropData,
-					nullptr,
-					nullptr,
-					PPF_PropertyWindow | PPF_BlueprintDebugView | PPF_Delimited,
-					nullptr
-				);
-				NameStr += TEXT("] ");
-				ChildInfo->DisplayName = FText::FromString(NameStr);
-				ChildInfo->bIsInContainer = true;
+			FString NameStr = TEXT("[");
+			MapProperty->KeyProp->ExportTextItem_Direct(
+				NameStr,
+				KeyData,
+				nullptr,
+				nullptr,
+				PPF_PropertyWindow | PPF_BlueprintDebugView | PPF_Delimited,
+				nullptr
+			);
+			NameStr += TEXT("] ");
+			ChildInfo->DisplayName = FText::FromString(NameStr);
+			ChildInfo->bIsInContainer = true;
+			ChildInfo->ContainerIndex = Index;
 				
-				Children.Add(ChildInfo);
-				--Count;
-			}
+			Children.Add(ChildInfo);
 		}
 	}
 	else if (const FSetProperty* SetProperty = CastField<FSetProperty>(Property.Get()))
 	{
 		FScriptSetHelper SetHelper(SetProperty, PropertyInstance.Value);
-		uint8* PropData = SetHelper.GetElementPtr(0);
-
 		int32 Index = 0;
-		for (int32 Count = SetHelper.Num(); Count; PropData += SetProperty->SetLayout.Size, ++Index)
+		for (FScriptSetHelper::FIterator Itr = SetHelper.CreateIterator(); Itr; ++Itr, ++Index)
 		{
-			if (SetHelper.IsValidIndex(Index))
-			{
-				FPropertyInstance ChildProperty = {
-					SetProperty->ElementProp,
-					PropData
-				};
-				const TSharedPtr<FPropertyInstanceInfo> ChildInfo = FindOrMake(ChildProperty, VisitedNodes);
+			uint8* PropData = SetHelper.GetElementPtr(*Itr);
+			FPropertyInstance ChildProperty = {
+				SetProperty->ElementProp,
+				PropData
+			};
+			const TSharedPtr<FPropertyInstanceInfo> ChildInfo = Make(ChildProperty, AsShared());
 				
-				ChildInfo->DisplayName = FText::Format(LOCTEXT("SetIndexName", "[{0}]"), FText::AsNumber(Index));
-				ChildInfo->bIsInContainer = true;
-
-				Children.Add(ChildInfo);
+			ChildInfo->DisplayName = FText::Format(LOCTEXT("SetIndexName", "[{0}]"), FText::AsNumber(Index));
+			ChildInfo->bIsInContainer = true;
+			ChildInfo->ContainerIndex = Index;
 
-				--Count;
-			}
+			Children.Add(ChildInfo);
 		}
 	}
 	else if (const FObjectPropertyBase* ObjectPropertyBase = CastField<FObjectPropertyBase>(Property.Get()))
@@ -2371,7 +2331,7 @@ void FPropertyInstanceInfo::PopulateChildren(FPropertyInstance PropertyInstance,
 						*It,
 						It->ContainerPtrToValuePtr<void*>(Obj)
 					};
-					const TSharedPtr<FPropertyInstanceInfo> ChildInfo = FindOrMake(ChildProperty, VisitedNodes);
+					const TSharedPtr<FPropertyInstanceInfo> ChildInfo = Make(ChildProperty, AsShared());
 					Children.Add(ChildInfo);
 				}
 			}
@@ -2452,17 +2412,100 @@ FString FPropertyInstanceInfo::GetWatchText() const
 	return Value.ToString();
 }
 
-const TArray<TSharedPtr<FPropertyInstanceInfo>>& FPropertyInstanceInfo::GetChildren() const
+static bool IsAContainer(const FProperty* Property)
+{
+	return Property->IsA<FArrayProperty>() || Property->IsA<FSetProperty>() || Property->IsA<FMapProperty>() ||
+		Property->IsA<FStructProperty>() || Property->IsA<FObjectPropertyBase>();
+}
+
+const TArray<TSharedPtr<FPropertyInstanceInfo>>& FPropertyInstanceInfo::GetChildren()
 {
-	// If this node represents an object that was already added elsewhere in the tree, return its children instead
-	if (ReferencedObject.IsValid())
+	if (Children.IsEmpty() && IsAContainer(Property.Get()))
 	{
-		return ReferencedObject->GetChildren();
+		// Children haven't been generated yet. Generate them.
+		const FPropertyInstance PropertyInstance = GetPropertyInstance();
+		if (PropertyInstance.Value != nullptr)
+		{
+			TMap<FPropertyInstance, TSharedPtr<FPropertyInstanceInfo>> VisitedNodes;
+			PopulateChildren(PropertyInstance);
+		}
 	}
 
 	return Children;
 }
 
+FPropertyInstanceInfo::FPropertyInstance FPropertyInstanceInfo::GetPropertyInstance() const
+{
+	// reiterate the property tree to rebuild data pointer in case it was invalidated
+	
+	TArray<const FPropertyInstanceInfo*> Path;
+	const FPropertyInstanceInfo* Itr = this;
+	while (const FPropertyInstanceInfo* ItrParent = Itr->Parent.Pin().Get())
+	{
+		Path.Add(Itr);
+		Itr = ItrParent;
+	}
+	Path.Add(Itr);
+	Algo::Reverse(Path);
+	const void* Data = Itr->ValueAddress;
+	for (int I = 1; I < Path.Num() && Data != nullptr; ++I)
+	{
+		const FPropertyInstanceInfo* ParentInfo = Path[I - 1];
+		const FPropertyInstanceInfo* Info = Path[I];
+		if (const FArrayProperty* ArrayProperty = CastField<FArrayProperty>(ParentInfo->Property.Get()))
+		{
+			FScriptArrayHelper ArrayHelper(ArrayProperty, Data);
+			const int32 Index = Info->ContainerIndex;
+			if (ArrayHelper.IsValidIndex(Index))
+			{
+				Data = ArrayHelper.GetRawPtr(Index);
+			}
+			else
+			{
+				Data = nullptr;
+				break;
+			}
+		}
+		else if (const FSetProperty* SetProperty = CastField<FSetProperty>(ParentInfo->Property.Get()))
+		{
+			FScriptSetHelper SetHelper(SetProperty, Data);
+			const int32 Index = Info->ContainerIndex;
+			Data = SetHelper.FindNthElementPtr(Index);
+		}
+		else if (const FMapProperty* MapProperty = CastField<FMapProperty>(ParentInfo->Property.Get()))
+		{
+			FScriptMapHelper MapHelper(MapProperty, Data);
+			const int32 Index = Info->ContainerIndex;
+			Data = MapHelper.FindNthValuePtr(Index);
+		}
+		else if (const FStructProperty* StructProperty = CastField<FStructProperty>(ParentInfo->Property.Get()))
+		{
+			Data = Info->Property->ContainerPtrToValuePtr<void>(Data);
+		}
+		else if (const FObjectPropertyBase* ObjectPropertyBase = CastField<FObjectPropertyBase>(ParentInfo->Property.Get()))
+		{
+			if (UObject* Obj = ObjectPropertyBase->GetObjectPropertyValue(Data))
+			{
+				if (!Obj->IsA(Info->Property->GetOwner<UClass>()))
+				{
+					// This node is invalid. (this can happen if it was removed from an array, map, or set)
+					Data = nullptr;
+					break;
+				}
+				Data = Info->Property->ContainerPtrToValuePtr<void>(Obj);
+			}
+		}
+		if (Data != Info->ValueAddress)
+		{
+			// This node is invalid. (this can happen if it was removed from an array, map, or set)
+			Data = nullptr;
+			break;
+		}
+	}
+	
+	return Data ? FPropertyInstance{Property.Get(), Data} : FPropertyInstance{nullptr, nullptr};
+}
+
 FText FKismetDebugUtilities::GetAndClearLastExceptionMessage()
 {
 	FKismetDebugUtilitiesData& Data = FKismetDebugUtilitiesData::Get();
diff --git a/Engine/Source/Editor/UnrealEd/Public/Kismet2/KismetDebugUtilities.h b/Engine/Source/Editor/UnrealEd/Public/Kismet2/KismetDebugUtilities.h
--- a/Engine/Source/Editor/UnrealEd/Public/Kismet2/KismetDebugUtilities.h
+++ b/Engine/Source/Editor/UnrealEd/Public/Kismet2/KismetDebugUtilities.h
@@ -50,7 +50,7 @@ struct FKismetTraceSample
 // FDebugInfo
 
 // call FKismetDebugUtilities::GetDebugInfo to construct
-struct FPropertyInstanceInfo
+struct FPropertyInstanceInfo : TSharedFromThis<FPropertyInstanceInfo>
 {
 	/**
 	 * used to determine whether an object's property has
@@ -66,28 +66,26 @@ struct FPropertyInstanceInfo
 	* Helper constructor. call FKismetDebugUtilities::GetDebugInfo or
 	* FKismetDebugUtilities::GetDebugInfoInternal instead
 	*/
-    FPropertyInstanceInfo(FPropertyInstance PropertyInstance);
+    FPropertyInstanceInfo(FPropertyInstance PropertyInstance, const TSharedPtr<FPropertyInstanceInfo>& Parent = nullptr);
 
 	/**
-	 * looks for an existing FPropertyInstance inside VisitedNodes
-	 * if one isn't found, it makes one, stores it in VisitedNodes, and returns it
+	 * Constructs a Shared FPropertyInstanceInfo
 	 */
-	static TSharedPtr<FPropertyInstanceInfo> FindOrMake(FPropertyInstance PropertyInstance,
-		TMap<FPropertyInstance, TSharedPtr<FPropertyInstanceInfo>> &VisitedNodes);
+	static TSharedPtr<FPropertyInstanceInfo> Make(FPropertyInstance PropertyInstance, const TSharedPtr<FPropertyInstanceInfo>& Parent);
 
 	/**
-	* populates a FDebugInfo::Children with sub-properties
+	* populates a FDebugInfo::Children with sub-properties (non-recursive)
 	*/
-	void PopulateChildren(FPropertyInstance PropertyInstance,
-		TMap<FPropertyInstance, TSharedPtr<FPropertyInstanceInfo>> &VisitedNodes);
+	void PopulateChildren(FPropertyInstance PropertyInstance);
 
 	/** Resolves the PathToProperty treating this PropertyInstance as the head of the path */
 	UNREALED_API TSharedPtr<FPropertyInstanceInfo> ResolvePathToProperty(const TArray<FName>& InPathToProperty);
 
 	/** Returns the watch text for info popup bubbles on the graph */
 	UNREALED_API FString GetWatchText() const;
-	
-	UNREALED_API const TArray<TSharedPtr<FPropertyInstanceInfo>>& GetChildren() const;
+
+	/** Returns children of this node (generating them if necessary) */
+	UNREALED_API const TArray<TSharedPtr<FPropertyInstanceInfo>>& GetChildren();
 
 	FText Name;
 	FText DisplayName;
@@ -96,16 +94,16 @@ struct FPropertyInstanceInfo
 	TWeakObjectPtr<UObject> Object = nullptr; // only filled if property is a UObject
 	TFieldPath<const FProperty> Property;
 	bool bIsInContainer = false;
+	int32 ContainerIndex = INDEX_NONE;
 
 private:
+	// Warning: only the head ValueAddress in the tree is guaranteed to be valid. Call GetPropertyInstance() instead of
+	// using ValueAddress directly.
+	const void* ValueAddress;
+	FPropertyInstance GetPropertyInstance() const;
+	
+	TWeakPtr<FPropertyInstanceInfo> Parent;
 	TArray<TSharedPtr<FPropertyInstanceInfo>> Children;
-
-	/** 
-	 * Only populated when this is an object and was already seen in the debug tree
-	 * used to avoid circular references in cases where the same object may be referenced by multiple properties  
-	 */
-	TSharedPtr<FPropertyInstanceInfo> ReferencedObject; 
-
 };
 
 inline uint32 GetTypeHash(const FPropertyInstanceInfo::FPropertyInstance& PropertyInstance)
