Implement a way to disable the Streaming Disabled warning for templates/projects that don't intend to enable streaming
#rnx

diff --git a/Engine/Source/Editor/UnrealEd/Private/SWorldPartitionViewportWidget.cpp b/Engine/Source/Editor/UnrealEd/Private/SWorldPartitionViewportWidget.cpp
--- a/Engine/Source/Editor/UnrealEd/Private/SWorldPartitionViewportWidget.cpp
+++ b/Engine/Source/Editor/UnrealEd/Private/SWorldPartitionViewportWidget.cpp
@@ -13,6 +13,7 @@
 #include "WorldPartition/IWorldPartitionEditorModule.h"
 #include "Styling/AppStyle.h"
 #include "Modules/ModuleManager.h"
+#include "GameFramework/WorldSettings.h"
 
 #define LOCTEXT_NAMESPACE "WorldPartitionViewportWidget"
 
@@ -87,7 +88,8 @@ EVisibility SWorldPartitionViewportWidget::GetVisibility(UWorld* InWorld)
 		}
 		else
 		{
-			if (WorldPartition->IsEnablingStreamingJustified())
+			const AWorldSettings* WorldSettings = InWorld->GetWorldSettings();
+			if ((WorldSettings != nullptr && !WorldSettings->bHideEnableStreamingWarning) && WorldPartition->IsEnablingStreamingJustified())
 			{
 				Message = LOCTEXT("StreamingDisabledText","Streaming Disabled");
 				Tooltip = LOCTEXT("StreamingDisabledToolTip", "The size of your world has grown enough to justify enabling streaming.");
diff --git a/Engine/Source/Runtime/Engine/Classes/GameFramework/WorldSettings.h b/Engine/Source/Runtime/Engine/Classes/GameFramework/WorldSettings.h
--- a/Engine/Source/Runtime/Engine/Classes/GameFramework/WorldSettings.h
+++ b/Engine/Source/Runtime/Engine/Classes/GameFramework/WorldSettings.h
@@ -571,6 +571,10 @@ public:
 	uint8 bGenerateSingleClusterForLevel : 1;
 
 #if WITH_EDITORONLY_DATA
+	/** if set to true, this hide the streaming disabled warning available in the viewport */
+	UPROPERTY(EditAnywhere, Category = WorldPartitionSetup, AdvancedDisplay, meta = (EditCondition = "WorldPartition != nullptr"))
+	uint8 bHideEnableStreamingWarning : 1;
+	
 	/** 
 	 * Whether Foliage actors of this world contain their grid size in their name. This should only be changed by UWorldPartitionFoliageBuilder or when creating new worlds so that older worlds are unaffected
 	 * and is used by the UActorPartitionSubsystem to find existing foliage actors by name.
diff --git a/Engine/Source/Runtime/Engine/Private/WorldPartition/WorldPartition.cpp b/Engine/Source/Runtime/Engine/Private/WorldPartition/WorldPartition.cpp
--- a/Engine/Source/Runtime/Engine/Private/WorldPartition/WorldPartition.cpp
+++ b/Engine/Source/Runtime/Engine/Private/WorldPartition/WorldPartition.cpp
@@ -1439,7 +1439,7 @@ void UWorldPartition::Tick(float DeltaSeconds)
 					{
 						AllActorsBounds += EditorBounds;
 
-						// Warn the user if the world becomes larger that 4km in any axis
+						// Warn the user if the world becomes larger that WorldExtent in any axis
 						if (AllActorsBounds.GetSize().GetMax() >= UWorldPartition::WorldExtentToEnableStreaming)
 						{
 							bEnablingStreamingJustified = true;
diff --git a/Engine/Source/Runtime/Engine/Private/WorldSettings.cpp b/Engine/Source/Runtime/Engine/Private/WorldSettings.cpp
--- a/Engine/Source/Runtime/Engine/Private/WorldSettings.cpp
+++ b/Engine/Source/Runtime/Engine/Private/WorldSettings.cpp
@@ -91,6 +91,7 @@ PRAGMA_ENABLE_DEPRECATION_WARNINGS
 	InstancedFoliageGridSize = DefaultPlacementGridSize = LandscapeSplineMeshesGridSize = 25600;
 	NavigationDataChunkGridSize = 102400;
 	NavigationDataBuilderLoadingCellSize = 102400 * 4;
+	bHideEnableStreamingWarning = false;
 	bIncludeGridSizeInNameForFoliageActors = false;
 	bIncludeGridSizeInNameForPartitionedActors = false;
 #endif
