Add optional DeleteMinimumAgeDays argument to FixupRedirects.
#PR 9976
#rnx

diff --git a/Engine/Source/Editor/UnrealEd/Private/Commandlets/ContentCommandlets.cpp b/Engine/Source/Editor/UnrealEd/Private/Commandlets/ContentCommandlets.cpp
--- a/Engine/Source/Editor/UnrealEd/Private/Commandlets/ContentCommandlets.cpp
+++ b/Engine/Source/Editor/UnrealEd/Private/Commandlets/ContentCommandlets.cpp
@@ -1381,6 +1381,13 @@ int32 UResavePackagesCommandlet::Main( const FString& Params )
 	IAssetRegistry& AssetRegistry = AssetRegistryModule.Get();
 	AssetRegistry.Tick(-1.0f);
 
+	int32 DeleteMinimumAgeDays = 0;
+	if (FParse::Value(*Params, TEXT("DeleteMinimumAgeDays="), DeleteMinimumAgeDays))
+	{
+		DeleteMinimumAgeDays = FMath::Max(DeleteMinimumAgeDays, 0);
+	}
+	FDateTime Now = FDateTime::Now();
+
 	// Delete unreferenced redirector packages
 	for (int32 PackageIndex = 0; PackageIndex < RedirectorsToFixup.Num(); PackageIndex++)
 	{
@@ -1393,19 +1400,37 @@ int32 UResavePackagesCommandlet::Main( const FString& Params )
 
 		AssetRegistry.GetReferencers(PackageName, Referencers);
 
-		if (Referencers.Num() == 0)
+		if (Referencers.Num() > 0)
 		{
 			if (Verbosity != ONLY_ERRORS)
 			{
-				UE_LOG(LogContentCommandlet, Display, TEXT("Deleting unreferenced redirector [%s]"), *Filename);
+				UE_LOG(LogContentCommandlet, Display, TEXT("Can't delete redirector [%s], unsaved packages reference it"), *Filename);
 			}
+			continue;
+		}
 
-			DeleteOnePackage(Filename);
+		if (DeleteMinimumAgeDays > 0)
+		{
+			FDateTime FileModificationTime = IFileManager::Get().GetTimeStamp(*Filename);
+			int32 DaysSinceModification = (Now - FileModificationTime).GetDays();
+			if (DaysSinceModification < DeleteMinimumAgeDays)
+			{
+				if (Verbosity != ONLY_ERRORS)
+				{
+					UE_LOG(LogContentCommandlet, Display,
+						TEXT("Not deleting unreferenced redirector [%s] (age threshold set to %d days, redirector was modified %d days ago)"),
+						*Filename, DeleteMinimumAgeDays, DaysSinceModification);
+				}
+				continue;
+			}
 		}
-		else if (Verbosity != ONLY_ERRORS)
+
+		if (Verbosity != ONLY_ERRORS)
 		{
-			UE_LOG(LogContentCommandlet, Display, TEXT("Can't delete redirector [%s], unsaved packages reference it"), *Filename);
+			UE_LOG(LogContentCommandlet, Display, TEXT("Deleting unreferenced redirector [%s]"), *Filename);
 		}
+
+		DeleteOnePackage(Filename);
 	}
 
 	// Submit the results to source control
