Fixes screenshot crashes on HMD without altering ReadPixels in order to avoid crashes for non-screenshot uses. All screenshot pathways now explicitly indicate screenshot size, and use GetRenderTargetTextureSizeXY() instead of GetSizeXY() in order to accommodate the stereo case in which RTTSize varies from viewport size, which can cause a crash if we request a screenshot larger than the actual texture.

diff --git a/Engine/Source/Editor/LevelEditor/Private/HighResScreenshotUI.cpp b/Engine/Source/Editor/LevelEditor/Private/HighResScreenshotUI.cpp
--- a/Engine/Source/Editor/LevelEditor/Private/HighResScreenshotUI.cpp
+++ b/Engine/Source/Editor/LevelEditor/Private/HighResScreenshotUI.cpp
@@ -399,13 +399,13 @@ FReply SHighResScreenshotDialog::OnCaptureClicked()
 		auto ConfigViewport = Config.TargetViewport.Pin();
 		if (ConfigViewport.IsValid())
 		{
-			GScreenshotResolutionX = ConfigViewport->GetSizeXY().X * Config.ResolutionMultiplier;
-			GScreenshotResolutionY = ConfigViewport->GetSizeXY().Y * Config.ResolutionMultiplier;
+			GScreenshotResolutionX = ConfigViewport->GetRenderTargetTextureSizeXY().X * Config.ResolutionMultiplier;
+			GScreenshotResolutionY = ConfigViewport->GetRenderTargetTextureSizeXY().Y * Config.ResolutionMultiplier;
 			FIntRect ScaledCaptureRegion = Config.UnscaledCaptureRegion;
 
 			if (ScaledCaptureRegion.Area() > 0)
 			{
-				ScaledCaptureRegion.Clip(FIntRect(FIntPoint::ZeroValue, ConfigViewport->GetSizeXY()));
+				ScaledCaptureRegion.Clip(FIntRect(FIntPoint::ZeroValue, ConfigViewport->GetRenderTargetTextureSizeXY()));
 				ScaledCaptureRegion *= Config.ResolutionMultiplier;
 			}
 
diff --git a/Engine/Source/Editor/UnrealEd/Private/EditorViewportClient.cpp b/Engine/Source/Editor/UnrealEd/Private/EditorViewportClient.cpp
--- a/Engine/Source/Editor/UnrealEd/Private/EditorViewportClient.cpp
+++ b/Engine/Source/Editor/UnrealEd/Private/EditorViewportClient.cpp
@@ -5739,19 +5739,20 @@ void FEditorViewportClient::TakeScreenshot(FViewport* InViewport, bool bInValida
 
 	// Read the contents of the viewport into an array.
 	bool bHdrEnabled = InViewport->GetSceneHDREnabled();
+	const FIntRect CaptureRect = FIntRect(0, 0, InViewport->GetRenderTargetTextureSizeXY().X, InViewport->GetRenderTargetTextureSizeXY().Y);
 	
 	if (!bHdrEnabled)
 	{
 		TArray<FColor> RawPixels;
-		RawPixels.SetNum(InViewport->GetSizeXY().X * InViewport->GetSizeXY().Y);
-		if( !InViewport->ReadPixels(RawPixels) )
+		RawPixels.SetNum(CaptureRect.Area());
+		if(!InViewport->ReadPixels(RawPixels, FReadSurfaceDataFlags(RCM_UNorm, CubeFace_MAX), CaptureRect))
 		{
 			// Failed to read the image from the viewport
 			SaveMessagePtr->SetText(NSLOCTEXT( "UnrealEd", "ScreenshotFailedViewport", "Screenshot failed, unable to read image from viewport" ));
 			return;
 		}
 
-		TUniquePtr<TImagePixelData<FColor>> PixelData = MakeUnique<TImagePixelData<FColor>>(InViewport->GetSizeXY(), TArray64<FColor>(MoveTemp(RawPixels)));
+		TUniquePtr<TImagePixelData<FColor>> PixelData = MakeUnique<TImagePixelData<FColor>>(InViewport->GetRenderTargetTextureSizeXY(), TArray64<FColor>(MoveTemp(RawPixels)));
 
 		check(PixelData->IsDataWellFormed());
 		ImageTask->PixelData = MoveTemp(PixelData);
@@ -5762,8 +5763,8 @@ void FEditorViewportClient::TakeScreenshot(FViewport* InViewport, bool bInValida
 	else
 	{
 		TArray<FLinearColor> RawPixels;
-		RawPixels.SetNum(InViewport->GetSizeXY().X * InViewport->GetSizeXY().Y);
-		if (!InViewport->ReadLinearColorPixels(RawPixels))
+		RawPixels.SetNum(CaptureRect.Area());
+		if (!InViewport->ReadLinearColorPixels(RawPixels, FReadSurfaceDataFlags(RCM_UNorm, CubeFace_MAX), CaptureRect))
 		{
 			// Failed to read the image from the viewport
 			SaveMessagePtr->SetText(NSLOCTEXT("UnrealEd", "ScreenshotFailedViewport", "Screenshot failed, unable to read image from viewport"));
@@ -5772,7 +5773,7 @@ void FEditorViewportClient::TakeScreenshot(FViewport* InViewport, bool bInValida
 
 		ConvertPixelDataToSCRGB(RawPixels, InViewport->GetDisplayOutputFormat());
 
-		TUniquePtr<TImagePixelData<FLinearColor>> PixelData = MakeUnique<TImagePixelData<FLinearColor>>(InViewport->GetSizeXY(), TArray64<FLinearColor>(MoveTemp(RawPixels)));
+		TUniquePtr<TImagePixelData<FLinearColor>> PixelData = MakeUnique<TImagePixelData<FLinearColor>>(InViewport->GetRenderTargetTextureSizeXY(), TArray64<FLinearColor>(MoveTemp(RawPixels)));
 
 		check(PixelData->IsDataWellFormed());
 		ImageTask->PixelData = MoveTemp(PixelData);
diff --git a/Engine/Source/Runtime/Engine/Private/GameViewportClient.cpp b/Engine/Source/Runtime/Engine/Private/GameViewportClient.cpp
--- a/Engine/Source/Runtime/Engine/Private/GameViewportClient.cpp
+++ b/Engine/Source/Runtime/Engine/Private/GameViewportClient.cpp
@@ -2002,7 +2002,7 @@ bool UGameViewportClient::ProcessScreenShots(FViewport* InViewport)
 
 		bool bScreenshotSuccessful = false;
 		bool bIsUI = false;
-		FIntVector Size(InViewport->GetSizeXY().X, InViewport->GetSizeXY().Y, 0);
+		FIntVector Size(InViewport->GetRenderTargetTextureSizeXY().X, InViewport->GetRenderTargetTextureSizeXY().Y, 0);
 
 		EDisplayOutputFormat ViewportOutputFormat = InViewport->GetDisplayOutputFormat();
 		bool bHdrEnabled = InViewport->GetSceneHDREnabled();
@@ -2025,12 +2025,12 @@ bool UGameViewportClient::ProcessScreenShots(FViewport* InViewport)
 		}
 		else if (bHdrEnabled)
 		{
-			bScreenshotSuccessful = GetViewportScreenShotHDR(InViewport, BitmapHDR);
+			bScreenshotSuccessful = GetViewportScreenShotHDR(InViewport, BitmapHDR, FIntRect(0, 0, Size.X, Size.Y));
 			ConvertPixelDataToSCRGB(BitmapHDR, ViewportOutputFormat);
 		}
 		else
 		{
-			bScreenshotSuccessful = GetViewportScreenShot(InViewport, Bitmap);
+			bScreenshotSuccessful = GetViewportScreenShot(InViewport, Bitmap, FIntRect(0, 0, Size.X, Size.Y));
 		}
 
 		if (bScreenshotSuccessful)
@@ -3759,8 +3759,8 @@ bool UGameViewportClient::HandleScreenshotCommand( const TCHAR* Cmd, FOutputDevi
 
 		FScreenshotRequest::RequestScreenshot(FileName, bShowUI, bAddFilenameSuffix, Viewport->GetSceneHDREnabled());
 
-		GScreenshotResolutionX = Viewport->GetSizeXY().X;
-		GScreenshotResolutionY = Viewport->GetSizeXY().Y;
+		GScreenshotResolutionX = Viewport->GetRenderTargetTextureSizeXY().X;
+		GScreenshotResolutionY = Viewport->GetRenderTargetTextureSizeXY().Y;
 	}
 	return true;
 }
diff --git a/Engine/Source/Runtime/Engine/Private/UnrealClient.cpp b/Engine/Source/Runtime/Engine/Private/UnrealClient.cpp
--- a/Engine/Source/Runtime/Engine/Private/UnrealClient.cpp
+++ b/Engine/Source/Runtime/Engine/Private/UnrealClient.cpp
@@ -2459,7 +2459,7 @@ ENGINE_API bool GetViewportScreenShot(FViewport* Viewport, TArray<FColor>& Bitma
 	// Read the contents of the viewport into an array.
 	if (Viewport->ReadPixels(Bitmap, FReadSurfaceDataFlags(), ViewRect))
 	{
-		check(Bitmap.Num() == ViewRect.Area() || (Bitmap.Num() == Viewport->GetSizeXY().X * Viewport->GetSizeXY().Y));
+		check(Bitmap.Num() == ViewRect.Area() || (Bitmap.Num() == Viewport->GetRenderTargetTextureSizeXY().X * Viewport->GetRenderTargetTextureSizeXY().Y));
 		return true;
 	}
 
@@ -2471,7 +2471,7 @@ ENGINE_API bool GetViewportScreenShotHDR(FViewport* Viewport, TArray<FLinearColo
 	// Read the contents of the viewport into an array.
 	if (Viewport->ReadLinearColorPixels(Bitmap, FReadSurfaceDataFlags(RCM_MinMax), ViewRect))
 	{
-		check(Bitmap.Num() == ViewRect.Area() || (Bitmap.Num() == Viewport->GetSizeXY().X * Viewport->GetSizeXY().Y));
+		check(Bitmap.Num() == ViewRect.Area() || (Bitmap.Num() == Viewport->GetRenderTargetTextureSizeXY().X * Viewport->GetRenderTargetTextureSizeXY().Y));
 		return true;
 	}
 
