OnlyEditorOnly: Make a separate setting for OnlyEditorOnly rather than using CanSkipEditorReferencedPackagesWhenCooking, so we can deprecate that separately later and still test OnlyEditorOnly.
OnlyEditorOnly: Add Instrumentation to explain differences between regular cook and OnlyEditorOnly cook.
#rnx

diff --git a/Engine/Config/BaseEditor.ini b/Engine/Config/BaseEditor.ini
--- a/Engine/Config/BaseEditor.ini
+++ b/Engine/Config/BaseEditor.ini
@@ -315,6 +315,9 @@ MinFreeMemory=1024
 MaxUsedMemory=16384
 
 [CookSettings]
+; If true, the cook will only cook what is referenced by saved packages, rather than cooking every loaded package.
+; Packages that are only referenced by editoronly properties will be skipped.
+OnlyEditorOnly=false
 BulkDataRegistryEnabled=true
 ; HybridIterative uses the TargetDomain and BuildDependencies to decide when packages need to be recooked. If enabled,
 ; every cook is iterative for packages with allowed classes (IterativeClassAllowList) unless -noiterative is passed.
diff --git a/Engine/Source/Editor/UnrealEd/Classes/CookOnTheSide/CookOnTheFlyServer.h b/Engine/Source/Editor/UnrealEd/Classes/CookOnTheSide/CookOnTheFlyServer.h
--- a/Engine/Source/Editor/UnrealEd/Classes/CookOnTheSide/CookOnTheFlyServer.h
+++ b/Engine/Source/Editor/UnrealEd/Classes/CookOnTheSide/CookOnTheFlyServer.h
@@ -158,6 +158,7 @@ namespace UE::Cook
 	class FCookDirector;
 	class FCookWorkerClient;
 	class FCookWorkerServer;
+	class FDiagnostics;
 	class FPackageWriterMPCollector;
 	class FRequestCluster;
 	class FRequestQueue;
@@ -881,6 +882,7 @@ private:
 	 * that predates COTF2?
 	 */
 	bool IsUsingLegacyCookOnTheFlyScheduling() const;
+	bool IsDebugRecordUnsolicited() const;
 
 	/** Update accumulators of editor data and consume editor changes after last cook when starting cooks in the editor. */
 	void BeginCookEditorSystems();
@@ -1418,6 +1420,7 @@ private:
 	 */
 	bool bHybridIterativeEnabled = true;
 	bool bHiddenDependenciesDebug = false;
+	bool bOnlyEditorOnlyDebug = false;
 	bool bHiddenDependenciesClassPathFilterListIsAllowList = true;
 	bool bFirstCookInThisProcessInitialized = false;
 	bool bFirstCookInThisProcess = true;
@@ -1429,7 +1432,7 @@ private:
 	/** Multithreaded synchronization of Pollables, accessible only inside PollablesLock. */
 	bool bPollablesInTick = false;
 	/** Config value that specifies whether the Skip-Only-Editor-Only feature is enabled. */
-	bool bCanSkipEditorReferencedPackagesWhenCooking = false;
+	bool bSkipOnlyEditorOnly = false;
 	/** True if we're running cooklist; tweak output */
 	bool bCookListMode = false;
 
@@ -1500,6 +1503,7 @@ private:
 	friend UE::Cook::FCookGenerationInfo;
 	friend UE::Cook::FCookWorkerClient;
 	friend UE::Cook::FCookWorkerServer;
+	friend UE::Cook::FDiagnostics;
 	friend UE::Cook::FGeneratorPackage;
 	friend UE::Cook::FInitializeConfigSettings;
 	friend UE::Cook::FPackageData;
diff --git a/Engine/Source/Editor/UnrealEd/Private/CookOnTheFlyServer.cpp b/Engine/Source/Editor/UnrealEd/Private/CookOnTheFlyServer.cpp
--- a/Engine/Source/Editor/UnrealEd/Private/CookOnTheFlyServer.cpp
+++ b/Engine/Source/Editor/UnrealEd/Private/CookOnTheFlyServer.cpp
@@ -20,6 +20,7 @@
 #include "Containers/DirectoryTree.h"
 #include "Containers/RingBuffer.h"
 #include "Cooker/AsyncIODelete.h"
+#include "Cooker/CookDiagnostics.h"
 #include "Cooker/CookDirector.h"
 #include "Cooker/CookMPCollector.h"
 #include "Cooker/CookOnTheFlyServerInterface.h"
@@ -2868,7 +2869,7 @@ void UCookOnTheFlyServer::QueueDiscoveredPackage(UE::Cook::FPackageData& Package
 
 	TArray<const ITargetPlatform*, TInlineAllocator<ExpectedMaxNumPlatforms>> BufferPlatforms;
 	TConstArrayView<const ITargetPlatform*> DiscoveredPlatforms;
-	if (!bCanSkipEditorReferencedPackagesWhenCooking)
+	if (!bSkipOnlyEditorOnly)
 	{
 		BufferPlatforms = PlatformManager->GetSessionPlatforms();
 		BufferPlatforms.Add(CookerLoadingPlatformKey);
@@ -4132,7 +4133,7 @@ void UCookOnTheFlyServer::ProcessUnsolicitedPackages(TArray<FName>* OutDiscovere
 		{
 			// This load was expected so we do not need to add a hidden dependency for it.
 			// If we are using legacy WhatGetsCookedRules, queue the package for cooking because it was loaded.
-			if (!bCanSkipEditorReferencedPackagesWhenCooking)
+			if (!bSkipOnlyEditorOnly)
 			{
 				QueueDiscoveredPackage(*PackageData, FInstigator(Instigator), EDiscoveredPlatformSet::CopyFromInstigator);
 			}
@@ -4150,7 +4151,15 @@ void UCookOnTheFlyServer::ProcessUnsolicitedPackages(TArray<FName>* OutDiscovere
 			// we will add the hidden dependency for it.
 			// as reachable in all of its instigator's reachable platforms
 			QueueDiscoveredPackage(*PackageData, FInstigator(Instigator), EDiscoveredPlatformSet::CopyFromInstigator);
-			// TODO: Add diagnostics for this unexpected load
+
+			if (IsDebugRecordUnsolicited() && !Instigator.Referencer.IsNone())
+			{
+				FPackageData* InstigatorPackage = PackageDatas->TryAddPackageDataByPackageName(Instigator.Referencer);
+				if (InstigatorPackage && InstigatorPackage != PackageData)
+				{
+					InstigatorPackage->CreateOrGetUnsolicited().FindOrAdd(PackageData, Instigator.Category);
+				}
+			}
 		}
 	}
 }
@@ -4175,6 +4184,7 @@ private: // Used only by UCookOnTheFlyServer, which has private access
 	UCookOnTheFlyServer& COTFS;
 	FPackageData& PackageData;
 	TArrayView<const ITargetPlatform*> PlatformsForPackage;
+	TSet<FPackageData*> SaveReferences;
 	FTickStackData& StackData;
 	UPackage* Package;
 	const FString PackageName;
@@ -5697,7 +5707,8 @@ void UCookOnTheFlyServer::SaveCookedPackage(UE::Cook::FSaveCookedPackageContext&
 		UE_ADD_CUSTOM_COOKTIMER_META(SaveCookedPackage, PackageName, *WriteToString<256>(Context.PackageData.GetFileName()));
 
 	UPackage* Package = Context.Package;
-	uint32 OriginalPackageFlags = Package->GetPackageFlags();
+	FOnScopeExit ScopedPackageFlags([Package, OriginalPackageFlags = Package->GetPackageFlags()]()
+		{ Package->SetPackageFlagsTo(OriginalPackageFlags); });
 
 	Context.SetupPackage();
 
@@ -5708,7 +5719,6 @@ void UCookOnTheFlyServer::SaveCookedPackage(UE::Cook::FSaveCookedPackageContext&
 	TGuardValue<bool> ScopedIsSavingPackage(bIsSavingPackage, true);
 	// For legacy reasons we set GIsCookerLoadingPackage == true during save. Some classes use it to conditionally execute cook operations in both save and load
 	TGuardValue<bool> ScopedIsCookerLoadingPackage(GIsCookerLoadingPackage, true);
-	ON_SCOPE_EXIT{ Package->SetPackageFlagsTo(OriginalPackageFlags); };
 
 	bool bFirstPlatform = true;
 	for (const ITargetPlatform* TargetPlatform : Context.PlatformsForPackage)
@@ -5764,9 +5774,16 @@ void UCookOnTheFlyServer::SaveCookedPackage(UE::Cook::FSaveCookedPackageContext&
 		bFirstPlatform = false;
 	}
 
+	// Need to restore flags before calling FinishPackage because it might need to save again
+	ScopedPackageFlags.ExitEarly();
 	Context.FinishPackage();
 }
 
+bool UCookOnTheFlyServer::IsDebugRecordUnsolicited() const
+{
+	return bOnlyEditorOnlyDebug;
+}
+
 namespace UE::Cook
 {
 FSaveCookedPackageContext::FSaveCookedPackageContext(UCookOnTheFlyServer& InCOTFS, UE::Cook::FPackageData& InPackageData,
@@ -6055,24 +6072,25 @@ void FSaveCookedPackageContext::FinishPlatform()
 			MoveTemp(OverrideAssetPackageData), MoveTemp(OverridePackageDependencies));
 	}
 
-	if (COTFS.bCanSkipEditorReferencedPackagesWhenCooking)
+	if (bSuccessful && COTFS.bSkipOnlyEditorOnly)
 	{
 		// If the save succeeded, add the imports and softobjectpaths from the save to the cook for the platform
-		if (bSuccessful)
+		FName PackageFName = Package->GetFName();
+		TConstArrayView<const ITargetPlatform*> ReachablePlatforms(&TargetPlatform, 1);
+		for (const TArray<FName>& DependencyNames : { SavePackageResult.ImportPackages, SavePackageResult.SoftPackageReferences })
 		{
-			FName PackageFName = Package->GetFName();
-			TConstArrayView<const ITargetPlatform*> ReachablePlatforms(&TargetPlatform, 1);
-			for (const TArray<FName>& DependencyNames : { SavePackageResult.ImportPackages, SavePackageResult.SoftPackageReferences })
+			EInstigator InstigatorType = &DependencyNames == &SavePackageResult.ImportPackages ?
+				EInstigator::SaveTimeHardDependency : EInstigator::SaveTimeSoftDependency;
+			for (FName DependencyName : DependencyNames)
 			{
-				EInstigator InstigatorType = &DependencyNames == &SavePackageResult.ImportPackages ?
-					EInstigator::SaveTimeHardDependency : EInstigator::SaveTimeSoftDependency;
-				for (FName DependencyName : DependencyNames)
+				UE::Cook::FPackageData* DependencyData = COTFS.PackageDatas->TryAddPackageDataByPackageName(DependencyName);
+				if (DependencyData)
 				{
-					UE::Cook::FPackageData* DependencyData = COTFS.PackageDatas->TryAddPackageDataByPackageName(DependencyName);
-					if (DependencyData)
+					COTFS.QueueDiscoveredPackage(*DependencyData,
+						UE::Cook::FInstigator(InstigatorType, PackageFName), FDiscoveredPlatformSet(ReachablePlatforms));
+					if (COTFS.IsDebugRecordUnsolicited())
 					{
-						COTFS.QueueDiscoveredPackage(*DependencyData,
-							UE::Cook::FInstigator(InstigatorType, PackageFName), FDiscoveredPlatformSet(ReachablePlatforms));
+						SaveReferences.Add(DependencyData);
 					}
 				}
 			}
@@ -6127,9 +6145,9 @@ void FSaveCookedPackageContext::FinishPackage()
 {
 	// Add soft references discovered from the package
 	FName PackageFName = Package->GetFName();
+	TConstArrayView<const ITargetPlatform*> ReachablePlatforms = PlatformsForPackage;
 	if (!COTFS.CookByTheBookOptions->bSkipSoftReferences)
 	{
-		TConstArrayView<const ITargetPlatform*> ReachablePlatforms = PlatformsForPackage;
 		// Also request any localized variants of this package
 		for (FName LocalizedPackageName : FRequestCluster::GetLocalizationReferences(PackageFName, COTFS))
 		{
@@ -6142,8 +6160,8 @@ void FSaveCookedPackageContext::FinishPackage()
 			}
 		}
 
-		// Add SoftObjectPaths to the cook. This has to be done after the package save to catch any SoftObjectPaths that were added during save.
-		// ONLYEDITORONLY_TODO: Read the list of SoftPackagePaths in FinishPlatform out of the SavePackage results instead
+		// When using legacy WhatGetsCookedRules, add all the SoftObjectPaths discovered during the package's load, plus any added on
+		// during save, to the cook for all platforms
 		TSet<FName> SoftObjectPackages;
 		GRedirectCollector.ProcessSoftObjectPathPackageList(PackageFName, false /* bGetEditorOnly */, SoftObjectPackages);
 		for (FName SoftObjectPackage : SoftObjectPackages)
@@ -6162,19 +6180,26 @@ void FSaveCookedPackageContext::FinishPackage()
 			UE::Cook::FPackageData* SoftObjectPackageData = COTFS.PackageDatas->TryAddPackageDataByPackageName(SoftObjectPackage);
 			if (SoftObjectPackageData)
 			{
-				if (!COTFS.bCanSkipEditorReferencedPackagesWhenCooking)
+				if (!COTFS.bSkipOnlyEditorOnly)
 				{
 					COTFS.QueueDiscoveredPackage(*SoftObjectPackageData,
 						UE::Cook::FInstigator(UE::Cook::EInstigator::SoftDependency, PackageFName),
 						FDiscoveredPlatformSet(ReachablePlatforms));
 				}
-				else
+
+				if (COTFS.IsDebugRecordUnsolicited())
 				{
-					// TODO: Add OnlyEditorOnly diagnostic
+					PackageData.CreateOrGetUnsolicited().Add(SoftObjectPackageData, EInstigator::SoftDependency);
 				}
 			}
 		}
 	}
+	if (COTFS.bOnlyEditorOnlyDebug)
+	{
+		COTFS.ProcessUnsolicitedPackages();
+		FDiagnostics::AnalyzeOnlyEditorOnlySave(COTFS, PackageData, PackageData.DetachUnsolicited(),
+			SaveReferences, ReachablePlatforms);
+	}
 
 	if (!bHasRetryErrorCode)
 	{
@@ -6511,9 +6536,19 @@ void UCookOnTheFlyServer::SetInitializeConfigSettings(UE::Cook::FInitializeConfi
 
 	// The rest of this function parses config settings that are reparsed on every CookDirector and CookWorker rather than
 	// being replicated from CookDirector to CookWorker
-	bCanSkipEditorReferencedPackagesWhenCooking = UE::SavePackageUtilities::CanSkipEditorReferencedPackagesWhenCooking();
 
 	// Debugging hidden dependencies
+	bOnlyEditorOnlyDebug = FParse::Param(FCommandLine::Get(), TEXT("OnlyEditorOnlyDebug"));
+	bSkipOnlyEditorOnly = false;
+	GConfig->GetBool(TEXT("CookSettings"), TEXT("OnlyEditorOnly"), bSkipOnlyEditorOnly, GEditorIni);
+	bSkipOnlyEditorOnly |= FParse::Param(FCommandLine::Get(), TEXT("OnlyEditorOnly"));
+	bSkipOnlyEditorOnly |= bOnlyEditorOnlyDebug;
+	if (bSkipOnlyEditorOnly)
+	{
+		UE_LOG(LogCook, Display,
+			TEXT("OnlyEditorOnly is enabled, unsolicited packages will not be cooked unless they are referenced from the cooked version of the instigator package."));
+	}
+
 	bHiddenDependenciesDebug = FParse::Param(FCommandLine::Get(), TEXT("HiddenDependenciesDebug"));
 	if (bHiddenDependenciesDebug)
 	{
diff --git a/Engine/Source/Editor/UnrealEd/Private/Cooker/CookDiagnostics.cpp b/Engine/Source/Editor/UnrealEd/Private/Cooker/CookDiagnostics.cpp
new file mode 100644
--- /dev/null
+++ b/Engine/Source/Editor/UnrealEd/Private/Cooker/CookDiagnostics.cpp
@@ -0,0 +1,417 @@
+// Copyright Epic Games, Inc. All Rights Reserved.
+
+#include "CookDiagnostics.h"
+
+#include "AssetRegistry/IAssetRegistry.h"
+#include "Containers/Array.h"
+#include "Containers/RingBuffer.h"
+#include "Cooker/CookPackageData.h"
+#include "Cooker/CookRequestCluster.h"
+#include "Cooker/CookTypes.h"
+#include "CookOnTheSide/CookLog.h"
+#include "CookOnTheSide/CookOnTheFlyServer.h"
+#include "Misc/Optional.h"
+#include "Misc/StringBuilder.h"
+#include "Serialization/ArchiveSerializedPropertyChain.h"
+#include "UObject/ICookInfo.h"
+#include "UObject/LinkerLoad.h"
+#include "UObject/Package.h"
+#include "UObject/SoftObjectPath.h"
+#include "UObject/UObjectGlobals.h"
+
+class ITargetPlatform;
+
+namespace UE::Cook
+{
+
+struct FPackageReferenceAndPropertyChain
+{
+	FName PackageName;
+	TArray<FProperty*, TInlineAllocator<1>> Properties;
+};
+/**
+ * From a passed in export, collects referenced exports and imports, and stores the property chain
+ * that references each import.
+ */
+class FWhyImportedCollector : public FArchiveUObject
+{
+public:
+	explicit FWhyImportedCollector(UPackage* InRootPackage)
+		: RootPackage(InRootPackage)
+	{
+		ArIsObjectReferenceCollector = true;
+		SetIsSaving(true);
+		SetIsPersistent(true);
+	}
+
+	void Reset()
+	{
+		Exports.Reset();
+		Imports.Reset();
+	}
+	virtual FArchive& operator<<(UObject*& Obj) override
+	{
+		if (Obj)
+		{
+			UPackage* Package = Obj->GetPackage();
+			if (Package == RootPackage)
+			{
+				Exports.Add(Obj);
+			}
+			else
+			{
+				Imports.Add(FPackageReferenceAndPropertyChain{ Package->GetFName(), ConvertSerializedPropertyChain() });
+			}
+		}
+		return *this;
+	}
+	virtual FArchive& operator<<(FSoftObjectPath& Value) override
+	{
+		FName Package = Value.GetLongPackageFName();
+		if (Package != RootPackage->GetFName() && !Package.IsNone())
+		{
+			Imports.Add(FPackageReferenceAndPropertyChain{ Package, ConvertSerializedPropertyChain() });
+		}
+		return *this;
+	}
+
+	TArray<FProperty*, TInlineAllocator<1>> ConvertSerializedPropertyChain()
+	{
+		TArray<FProperty*, TInlineAllocator<1>> Result;
+		const FArchiveSerializedPropertyChain* Chain = GetSerializedPropertyChain();
+		if (Chain && Chain->GetNumProperties())
+		{
+			Result.Reserve(Chain->GetNumProperties());
+			for (int32 Index = 0; Index < Chain->GetNumProperties(); ++Index)
+			{
+				Result.Add(Chain->GetPropertyFromRoot(Index));
+			}
+		}
+		return Result;
+	}
+
+public:
+	UPackage* RootPackage;
+	TArray<UObject*> Exports;
+	TArray<FPackageReferenceAndPropertyChain> Imports;
+};
+
+void FDiagnostics::AnalyzeOnlyEditorOnlySave(UCookOnTheFlyServer& COTFS, FPackageData& PackageData,
+	TMap<FPackageData*, EInstigator>&& UnsolicitedForPackage, TSet<FPackageData*>& SaveReferences,
+	TConstArrayView<const ITargetPlatform*> ReachablePlatforms)
+{
+	TOptional<TArray<FName>> ExpectedDependencies;
+	TOptional<TMap<FName, TMap<UObject*, TArray<FProperty*, TInlineAllocator<1>>>>> ExportsUsingPackageName;
+	bool bHasResaved = false;
+	FMultiPackageReaderResults Resave;
+	UPackage* Package = PackageData.GetPackage();
+	check(Package); // We are called from SavePackage so the PackageData still has a pointer to it
+
+	// Unsolicited packages are added to the cook when not using OnlyEditorOnly.
+	// When using OnlyEditorOnly, only packages referenced by the cooked SavePackage are added.
+	// For the period where we are transitioning the engine to always run OnlyEditorOnly, there will be
+	// some false negatives: OnlyEditorOnly will fail to add a package that is necessary.
+	// This function investigates each unsolicited package that was not added and writes information about
+	// why it was not added. Each of these needs to be investigated, and either fixed somehow so that it is
+	// referenced from the cooked SavePackage, or marked as an expected correct difference using FCookLoadScope.
+	for (const TPair<FPackageData*, EInstigator>& UnsolicitedPair : UnsolicitedForPackage)
+	{
+		FPackageData* Unsolicited = UnsolicitedPair.Key;
+		EInstigator Instigator = UnsolicitedPair.Value;
+		if (SaveReferences.Contains(Unsolicited))
+		{
+			// This package is included by OnlyEditorOnly as well. Remove it from our list to investigate.
+			continue;
+		}
+		FName UnsolicitedPackageName = Unsolicited->GetPackageName();
+
+		// If the package is not cookable on any of the platforms we are cooking, then even though it was loaded
+		// it will not be cooked by LegacyWhatShouldBeCooked. Remove it from our list to investigate.
+		UPackage* UnsolicitedUPackage = FindPackage(nullptr, *WriteToString<256>(UnsolicitedPackageName));
+		bool bSuppressedOnAllPlatforms = true;
+		for (const ITargetPlatform* ReachablePlatform : ReachablePlatforms)
+		{
+			if (UnsolicitedUPackage && FCoreUObjectDelegates::ShouldCookPackageForPlatform.IsBound())
+			{
+				if (!FCoreUObjectDelegates::ShouldCookPackageForPlatform.Execute(UnsolicitedUPackage, ReachablePlatform))
+				{
+					continue;
+				}
+			}
+			bool bCookable;
+			bool bExplorable;
+			ESuppressCookReason Reason;
+			FRequestCluster::IsRequestCookable(ReachablePlatform, *Unsolicited, COTFS, Reason, bCookable, bExplorable);
+			if (!bCookable)
+			{
+				continue;
+			}
+			bSuppressedOnAllPlatforms = false;
+			break;
+		}
+		if (bSuppressedOnAllPlatforms)
+		{
+			continue;
+		}
+
+		if (PackageData.GetGeneratorPackage())
+		{
+			// TODO: Collect SaveReferences for all generated packages and collect Unsolicited from all externalactor packages
+			// and run the missing solicited test when the GeneratorPackage finishes saving all generated
+			continue;
+		}
+
+		// If the referenced package was one of the declared dependencies of the source package, then it is not unsolicited.
+		// We only need to log unsolicited dependencies. Declared dependencies are either added already, or they were editoronly
+		// and were intentionally not added to the cook. They would still be cooked by LegacyWhatShouldBeCooked, and that is a difference,
+		// but we will trust the system that declared them as editoronly and therefore not log them as needs investigation.
+		if (!ExpectedDependencies)
+		{
+			ExpectedDependencies.Emplace();
+			COTFS.AssetRegistry->GetDependencies(PackageData.GetPackageName(), *ExpectedDependencies,
+				UE::AssetRegistry::EDependencyCategory::Package, UE::AssetRegistry::EDependencyQuery::NoRequirements);
+		}
+		if (ExpectedDependencies->Contains(Unsolicited->GetPackageName()))
+		{
+			continue;
+		}
+
+		// Maybe the unsolicited package was a used by a previous version the source package, but postload steps
+		// removed it. Or maybe it was previously undeclared as a dependency but new code in savepackage declares it.
+		// Try resaving the source package, and if the unsolicited package was an import but is no longer, or it is
+		// now an editor-only import, then it would no longer be unsolicited at head version and it is correct that
+		// OnlyEditorOnly removes it.
+		FLinkerLoad* LinkerLoad = Package->GetLinker();
+		FName ClassNameOfUPackage = UPackage::StaticClass()->GetFName();
+		FName ClassPackageOfUPackage = UPackage::StaticClass()->GetOuter()->GetFName();
+		bool bIsImportOfOriginalPackageLoad = false;
+		if (LinkerLoad)
+		{
+			for (FObjectImport& Import : LinkerLoad->ImportMap)
+			{
+				if (Import.ClassName == ClassNameOfUPackage && Import.ClassPackage == ClassPackageOfUPackage &&
+					Import.OuterIndex.IsNull() && Import.ObjectName == UnsolicitedPackageName)
+				{
+					bIsImportOfOriginalPackageLoad = true;
+					break;
+				}
+			}
+			for (FName SoftPackageReference : LinkerLoad->SoftPackageReferenceList)
+			{
+				if (SoftPackageReference == UnsolicitedPackageName)
+				{
+					bIsImportOfOriginalPackageLoad = true;
+					break;
+				}
+			}
+		}
+
+		if (!bHasResaved)
+		{
+			bHasResaved = true;
+			FSavePackageArgs SaveArgs;
+			SaveArgs.TopLevelFlags = RF_Standalone;
+			SaveArgs.SaveFlags = SAVE_NoError // Do not crash the SaveServer on an error
+				| SAVE_BulkDataByReference	// EditorDomain saves reference bulkdata from the WorkspaceDomain rather than duplicating it
+				| SAVE_Async				// SavePackage support for PackageWriter is only implemented with SAVE_Async
+				;
+			SaveArgs.bSlowTask = false;
+			Resave = GetSaveExportsAndImports(Package, nullptr, SaveArgs);
+		}
+
+		if (Resave.Realms[0].bValid)
+		{
+			FPackageReaderResults& Realm = Resave.Realms[0];
+			FSoftObjectPath PackagePath(UnsolicitedPackageName, NAME_None, FString());
+			FSoftObjectPath PackageClassPath(UPackage::StaticClass());
+			bool bIsImportOfResavedPackage = false;
+			bool bIsEditorOnlyImportOfResavedPackage = false;
+			for (TPair<FSoftObjectPath, FPackageReader::FObjectData>& Pair : Realm.Imports)
+			{
+				if (Pair.Value.ClassPath == PackageClassPath && Pair.Key == PackagePath)
+				{
+					bIsImportOfResavedPackage = true;
+					bIsEditorOnlyImportOfResavedPackage = Pair.Value.bUsedInGame;
+					break;
+				}
+			}
+			for (TPair<FName, bool>& Pair : Realm.SoftPackageReferences)
+			{
+				if (Pair.Key == UnsolicitedPackageName)
+				{
+					bIsImportOfResavedPackage = true;
+					bIsEditorOnlyImportOfResavedPackage = !Pair.Value;
+					break;
+				}
+			}
+
+			if (bIsImportOfOriginalPackageLoad && !bIsImportOfResavedPackage)
+			{
+				// An import that was removed by upgrade steps; ignore it
+				continue;
+			}
+
+			if (bIsEditorOnlyImportOfResavedPackage)
+			{
+				// An import that the new save code identifies as editoronly; ignore it
+				continue;
+			}
+		}
+
+		// No more filtering, we're going to log this as an unexpected difference.
+		// Try to find out how it was loaded so we can give information about what code needs to
+		// be modified to support OnlyEditorOnly.
+
+		
+		// Serialize all referenced exports in the source package and record all the packages each one
+		// references. If any of those packages is the unsolicited package, print out the export and property
+		// that refers to it.
+		if (!ExportsUsingPackageName)
+		{
+			ExportsUsingPackageName.Emplace();
+			TSet<UObject*> Exports;
+			TRingBuffer<UObject*> ExportsQueue;
+			ForEachObjectWithPackage(Package, [&Exports, &ExportsQueue](UObject* Object)
+				{
+					if (Object->HasAnyFlags(RF_Public))
+					{
+						bool bAlreadyExists;
+						Exports.Add(Object, &bAlreadyExists);
+						if (!bAlreadyExists)
+						{
+							ExportsQueue.Add(Object);
+						}
+					}
+					return true;
+				});
+			FWhyImportedCollector Collector(Package);
+			while (!ExportsQueue.IsEmpty())
+			{
+				UObject* Export = ExportsQueue.PopFrontValue();
+				Export->Serialize(Collector);
+				for (UObject* Dependency : Collector.Exports)
+				{
+					bool bAlreadyExists;
+					Exports.Add(Dependency, &bAlreadyExists);
+					if (!bAlreadyExists)
+					{
+						ExportsQueue.Add(Dependency);
+					}
+				}
+				for (FPackageReferenceAndPropertyChain& ImportChain : Collector.Imports)
+				{
+					TArray<FProperty*, TInlineAllocator<1>>& ExistingChain = ExportsUsingPackageName->FindOrAdd(ImportChain.PackageName).FindOrAdd(Export);
+					if (ExistingChain.IsEmpty())
+					{
+						ExistingChain = MoveTemp(ImportChain.Properties);
+					}
+				}
+				Collector.Reset();
+			}
+		}
+		TMap<UObject*, TArray<FProperty*, TInlineAllocator<1>>>* ExportsUsingUnsolicited = ExportsUsingPackageName->Find(UnsolicitedPackageName);
+		TStringBuilder<256> WhyReferenced;
+
+		int32 PackageNameLen = Package->GetName().Len();
+		if (ExportsUsingUnsolicited)
+		{
+			int32 Count = 0;
+			constexpr int32 MaxCount = 3;
+			WhyReferenced << TEXT("\n\tReferencers:");
+			for (TPair<UObject*, TArray<FProperty*, TInlineAllocator<1>>>&Pair : *ExportsUsingUnsolicited)
+			{
+				UObject* Export = Pair.Key;
+				WhyReferenced << TEXT("\n\t\t");
+				if (Count++ == MaxCount)
+				{
+					WhyReferenced << TEXT("...");
+					break;
+				}
+				WhyReferenced << Export->GetClass()->GetPathName() << TEXT(" ");
+				WhyReferenced << Export->GetPathName().RightChop(PackageNameLen + 1);
+				WhyReferenced << TEXT("#");
+				if (Pair.Value.IsEmpty())
+				{
+					WhyReferenced << TEXT("<UnknownProperty>");
+				}
+				else
+				{
+					int32 NumProperties = Pair.Value.Num();
+					bool bAddedSeparator = false;
+					for (int32 PropertyIndex = 0; PropertyIndex < NumProperties; ++PropertyIndex)
+					{
+						FProperty* Property = Pair.Value[PropertyIndex];
+						if (PropertyIndex < NumProperties - 1 && Property->GetFName() == Pair.Value[PropertyIndex + 1]->GetFName() &&
+							Property->IsA(FArrayProperty::StaticClass()))
+						{
+							// Omit adding a duplicate entry for inner property of an array which has the same name as its array property
+							continue;
+						}
+						WhyReferenced << Property->GetName() << TEXT("#");
+						bAddedSeparator = true;
+					}
+					if (bAddedSeparator)
+					{
+						WhyReferenced.RemoveSuffix(1);
+					}
+				}
+			}
+		}
+		else
+		{
+			// If we didn't find any exports referring to the unsolicited package, then just print out the list of class types
+			// in the package to show which code needs to be investigated.
+			enum class EClassPriority
+			{
+				PrimaryUAsset,
+				Asset,
+				Public,
+				Private,
+				Count,
+			};
+			TMap<UClass*, EClassPriority> ObjectClasses;
+			FName PackageLeafName = FName(*FPaths::GetBaseFilename(Package->GetName(), true /* bRemovePath */));
+			ForEachObjectWithPackage(Package, [&ObjectClasses, PackageLeafName](UObject* Object)
+				{
+					UClass* Class = Object->GetClass();
+					EClassPriority NewPriority = EClassPriority::Private;
+					if (Object->IsAsset())
+					{
+						NewPriority = Object->GetFName() == PackageLeafName ? EClassPriority::PrimaryUAsset : EClassPriority::Asset;
+					}
+					else
+					{
+						NewPriority = Object->HasAnyFlags(RF_Public) ? EClassPriority::Public : EClassPriority::Private;
+					}
+					EClassPriority& OldPriority = ObjectClasses.FindOrAdd(Class, NewPriority);
+					OldPriority = static_cast<EClassPriority>(FMath::Min(static_cast<int32>(OldPriority), static_cast<int32>(NewPriority)));
+					return true;
+				});
+			WhyReferenced << TEXT("\n\tNo exports found referencing the dependency. Classes in the referencer package:");
+			int32 ClassCount = 0;
+			constexpr int32 MaxClassCount = 10;
+			ObjectClasses.ValueSort([](EClassPriority A, EClassPriority B)
+				{
+					if (A != B)
+					{
+						return (int32)A < (int32)B;
+					}
+					return false;
+				});
+			for (TPair<UClass*, EClassPriority>& ClassPair : ObjectClasses)
+			{
+				if (ClassCount++ >= MaxClassCount)
+				{
+					break;
+				}
+				WhyReferenced << TEXT("\n\t\t") << ClassPair.Key->GetPathName();
+			}
+		}
+
+		UE_LOG(LogCook, Display, TEXT("OnlyEditorOnly: Skipped adding %s %s -> %s%s"),
+			LexToString(Instigator), *WriteToString<256>(PackageData.GetPackageName()),
+			*WriteToString<256>(Unsolicited->GetPackageName()), *WhyReferenced);
+	}
+}
+
+}
\ No newline at end of file
diff --git a/Engine/Source/Editor/UnrealEd/Private/Cooker/CookDiagnostics.h b/Engine/Source/Editor/UnrealEd/Private/Cooker/CookDiagnostics.h
new file mode 100644
--- /dev/null
+++ b/Engine/Source/Editor/UnrealEd/Private/Cooker/CookDiagnostics.h
@@ -0,0 +1,31 @@
+// Copyright Epic Games, Inc. All Rights Reserved.
+
+#pragma once
+
+#include "Containers/ArrayView.h"
+#include "Containers/Map.h"
+#include "Containers/Set.h"
+
+class ITargetPlatform;
+class UCookOnTheFlyServer;
+namespace UE::Cook { enum class EInstigator : uint8; }
+namespace UE::Cook { struct FPackageData; }
+
+namespace UE::Cook
+{
+
+class FDiagnostics
+{
+public:
+	/**
+	 * Compares packages that would be added to the cook when using legacy WhatGetsCookedRules to the packages that
+	 * are added to the cook using OnlyEditorOnly rules. Silently ignores expected differences, but logs a message
+	 * with diagnostics for unexpected differences.
+	 */
+	static void AnalyzeOnlyEditorOnlySave(UCookOnTheFlyServer& COTFS, FPackageData& PackageData,
+		TMap<FPackageData*, EInstigator>&& UnsolicitedForPackage, TSet<FPackageData*>& SaveReferences,
+		TConstArrayView<const ITargetPlatform*> ReachablePlatforms);
+
+};
+
+}
\ No newline at end of file
diff --git a/Engine/Source/Editor/UnrealEd/Private/Cooker/CookPackageData.cpp b/Engine/Source/Editor/UnrealEd/Private/Cooker/CookPackageData.cpp
--- a/Engine/Source/Editor/UnrealEd/Private/Cooker/CookPackageData.cpp
+++ b/Engine/Source/Editor/UnrealEd/Private/Cooker/CookPackageData.cpp
@@ -1449,6 +1449,31 @@ FConstructPackageData FPackageData::CreateConstructData()
 	return ConstructData;
 }
 
+TMap<FPackageData*, EInstigator>& FPackageData::CreateOrGetUnsolicited()
+{
+	if (!Unsolicited)
+	{
+		Unsolicited = MakeUnique<TMap<FPackageData*, EInstigator>>();
+	}
+	return *Unsolicited;
+}
+
+TMap<FPackageData*, EInstigator> FPackageData::DetachUnsolicited()
+{
+	TMap<FPackageData*, EInstigator> Result;
+	if (Unsolicited)
+	{
+		Result = MoveTemp(*Unsolicited);
+		Unsolicited.Reset();
+	}
+	return Result;
+}
+
+void FPackageData::ClearUnsolicited()
+{
+	Unsolicited.Reset();
+}
+
 }
 
 FCbWriter& operator<<(FCbWriter& Writer, const UE::Cook::FConstructPackageData& PackageData)
diff --git a/Engine/Source/Editor/UnrealEd/Private/Cooker/CookPackageData.h b/Engine/Source/Editor/UnrealEd/Private/Cooker/CookPackageData.h
--- a/Engine/Source/Editor/UnrealEd/Private/Cooker/CookPackageData.h
+++ b/Engine/Source/Editor/UnrealEd/Private/Cooker/CookPackageData.h
@@ -548,6 +548,11 @@ public:
 	/** Marshall this PackageData to a ConstructData that can be used later or on a remote machine to reconstruct it. */
 	FConstructPackageData CreateConstructData();
 
+	/** Only used When IsDebugRecordUnsolicited: storage for the list of unsolicited packages for this package */
+	TMap<FPackageData*, EInstigator>& CreateOrGetUnsolicited();
+	TMap<FPackageData*, EInstigator> DetachUnsolicited();
+	void ClearUnsolicited();
+
 	/**
 	 * Return the platforms for which the given Package has been marked reachable.
 	 * If the package does not exist, return the COTFS's list of Session platforms
@@ -630,6 +635,7 @@ private:
 	TSortedMap<const ITargetPlatform*, FPackagePlatformData, TInlineAllocator<1>> PlatformDatas;
 	TArray<FWeakObjectPtr> CachedObjectsInOuter;
 	FCompletionCallback CompletionCallback;
+	TUniquePtr<TMap<FPackageData*, EInstigator>> Unsolicited;
 	FName PackageName;
 	FName FileName;
 
diff --git a/Engine/Source/Editor/UnrealEd/Private/Cooker/CookRequestCluster.cpp b/Engine/Source/Editor/UnrealEd/Private/Cooker/CookRequestCluster.cpp
--- a/Engine/Source/Editor/UnrealEd/Private/Cooker/CookRequestCluster.cpp
+++ b/Engine/Source/Editor/UnrealEd/Private/Cooker/CookRequestCluster.cpp
@@ -107,7 +107,7 @@ FRequestCluster::FRequestCluster(UCookOnTheFlyServer& InCOTFS, TRingBuffer<FDisc
 	: FRequestCluster(InCOTFS)
 {
 	TArray<const ITargetPlatform*, TInlineAllocator<ExpectedMaxNumPlatforms>> BufferPlatforms;
-	if (!COTFS.bCanSkipEditorReferencedPackagesWhenCooking)
+	if (!COTFS.bSkipOnlyEditorOnly)
 	{
 		BufferPlatforms = COTFS.PlatformManager->GetSessionPlatforms();
 		BufferPlatforms.Add(CookerLoadingPlatformKey);
@@ -119,7 +119,7 @@ FRequestCluster::FRequestCluster(UCookOnTheFlyServer& InCOTFS, TRingBuffer<FDisc
 		FPackageData& PackageData = *Discovery->PackageData;
 
 		TConstArrayView<const ITargetPlatform*> NewReachablePlatforms;
-		if (COTFS.bCanSkipEditorReferencedPackagesWhenCooking)
+		if (COTFS.bSkipOnlyEditorOnly)
 		{
 			NewReachablePlatforms = Discovery->ReachablePlatforms.GetPlatforms(COTFS, &Discovery->Instigator,
 				TConstArrayView<const ITargetPlatform*>(), &BufferPlatforms);
@@ -159,7 +159,7 @@ FRequestCluster::FRequestCluster(UCookOnTheFlyServer& InCOTFS, TRingBuffer<FDisc
 			{
 				// Add it as a hidden dependency so that future platforms discovered as reachable in
 				// the instigator will also be marked as reachable in the dependency.
-				if (COTFS.bCanSkipEditorReferencedPackagesWhenCooking)
+				if (COTFS.bSkipOnlyEditorOnly)
 				{
 					FPackageData* InstigatorPackageData = Discovery->Instigator.Referencer.IsNone() ? nullptr
 						: COTFS.PackageDatas->TryAddPackageDataByPackageName(Discovery->Instigator.Referencer);
@@ -961,7 +961,7 @@ void FRequestCluster::FGraphSearch::ExploreVertexEdges(FVertexData& Vertex)
 	if (bFetchAnyTargetPlatform)
 	{
 		EDependencyQuery FlagsForHardDependencyQuery;
-		if (Cluster.COTFS.bCanSkipEditorReferencedPackagesWhenCooking)
+		if (Cluster.COTFS.bSkipOnlyEditorOnly)
 		{
 			FlagsForHardDependencyQuery = EDependencyQuery::Game | EDependencyQuery::Hard;
 		}
@@ -979,7 +979,7 @@ void FRequestCluster::FGraphSearch::ExploreVertexEdges(FVertexData& Vertex)
 		}
 		if (Cluster.bAllowSoftDependencies)
 		{
-			// bCanSkipEditorReferencedPackagesWhenCooking does not affect soft dependencies; skip editoronly soft dependencies
+			// bSkipOnlyEditorOnly is always true for soft dependencies; skip editoronly soft dependencies
 			Cluster.AssetRegistry.GetDependencies(PackageName, SoftGameDependencies, EDependencyCategory::Package,
 				EDependencyQuery::Game | EDependencyQuery::Soft);
 
@@ -1438,20 +1438,42 @@ TMap<FPackageData*, TArray<FPackageData*>>& FRequestCluster::FGraphSearch::GetGr
 	return GraphEdges;
 }
 
+void FRequestCluster::IsRequestCookable(const ITargetPlatform* Platform, FPackageData& PackageData,
+	UCookOnTheFlyServer& COTFS, ESuppressCookReason& OutReason, bool& bOutCookable, bool& bOutExplorable)
+{
+	FString LocalDLCPath;
+	bool bLocalErrorOnEngineContentUse = false;
+	bool bLocalAllowUncookedAssetReferences = false;
+	if (!COTFS.IsCookOnTheFlyMode())
+	{
+		UE::Cook::FCookByTheBookOptions& Options = *COTFS.CookByTheBookOptions;
+		bLocalErrorOnEngineContentUse = Options.bErrorOnEngineContentUse;
+		bLocalAllowUncookedAssetReferences = Options.bAllowUncookedAssetReferences;
+	}
+	if (bLocalErrorOnEngineContentUse)
+	{
+		LocalDLCPath = FPaths::Combine(*COTFS.GetBaseDirectoryForDLC(), TEXT("Content"));
+		FPaths::MakeStandardFilename(LocalDLCPath);
+	}
+
+	IsRequestCookable(Platform, PackageData.GetPackageName(), PackageData, *COTFS.PackageDatas,
+		*COTFS.PackageTracker, LocalDLCPath, bLocalErrorOnEngineContentUse, bLocalAllowUncookedAssetReferences,
+		COTFS.bSkipOnlyEditorOnly, OutReason, bOutCookable, bOutExplorable);
+}
+
 void FRequestCluster::IsRequestCookable(const ITargetPlatform* Platform, FName PackageName, FPackageData& PackageData,
 	ESuppressCookReason& OutReason, bool& bOutCookable, bool& bOutExplorable)
 {
 	return IsRequestCookable(Platform, PackageName, PackageData, PackageDatas, PackageTracker,
 		DLCPath, bErrorOnEngineContentUse, bAllowUncookedAssetReferences,
-		COTFS.bCanSkipEditorReferencedPackagesWhenCooking,
+		COTFS.bSkipOnlyEditorOnly,
 		OutReason, bOutCookable, bOutExplorable);
 }
 
 void FRequestCluster::IsRequestCookable(const ITargetPlatform* Platform, FName PackageName, FPackageData& PackageData,
 	FPackageDatas& InPackageDatas, FPackageTracker& InPackageTracker,
 	FStringView InDLCPath, bool bInErrorOnEngineContentUse, bool bInAllowUncookedAssetReferences,
-	bool bCanSkipEditorReferencedPackagesWhenCooking,
-	ESuppressCookReason& OutReason, bool& bOutCookable, bool& bOutExplorable)
+	bool bSkipOnlyEditorOnly, ESuppressCookReason& OutReason, bool& bOutCookable, bool& bOutExplorable)
 {
 	check(Platform != CookerLoadingPlatformKey); // IsRequestCookable should not be called for The CookerLoadingPlatform; it has different rules
 
@@ -1477,7 +1499,7 @@ void FRequestCluster::IsRequestCookable(const ITargetPlatform* Platform, FName P
 			OutReason = ESuppressCookReason::NeverCook;
 			bOutCookable = false;
 			bOutExplorable = true;
-			if (bCanSkipEditorReferencedPackagesWhenCooking)
+			if (bSkipOnlyEditorOnly)
 			{
 				bOutExplorable = true;
 			}
diff --git a/Engine/Source/Editor/UnrealEd/Private/Cooker/CookRequestCluster.h b/Engine/Source/Editor/UnrealEd/Private/Cooker/CookRequestCluster.h
--- a/Engine/Source/Editor/UnrealEd/Private/Cooker/CookRequestCluster.h
+++ b/Engine/Source/Editor/UnrealEd/Private/Cooker/CookRequestCluster.h
@@ -80,6 +80,8 @@ public:
 		TMap<FPackageData*, TArray<FPackageData*>>& OutRequestGraph);
 
 	static TConstArrayView<FName> GetLocalizationReferences(FName PackageName, UCookOnTheFlyServer& InCOTFS);
+	static void IsRequestCookable(const ITargetPlatform* TargetPlatform, FPackageData& PackageData,
+		UCookOnTheFlyServer& COTFS, ESuppressCookReason& OutReason, bool& bOutCookable, bool& bOutExplorable);
 
 private:
 	struct FGraphSearch;
@@ -317,8 +319,7 @@ private:
 	static void IsRequestCookable(const ITargetPlatform* TargetPlatform, FName PackageName, FPackageData& PackageData,
 		FPackageDatas& InPackageDatas, FPackageTracker& InPackageTracker,
 		FStringView InDLCPath, bool bInErrorOnEngineContentUse, bool bInAllowUncookedAssetReferences,
-		bool bCanSkipEditorReferencedPackagesWhenCooking,
-		ESuppressCookReason& OutReason, bool& bOutCookable, bool& bOutExplorable);
+		bool bSkipOnlyEditorOnly, ESuppressCookReason& OutReason, bool& bOutCookable, bool& bOutExplorable);
 
 	TArray<FFilePlatformRequest> FilePlatformRequests;
 	TFastPointerMap<FPackageData*, ESuppressCookReason> OwnedPackageDatas;
diff --git a/Engine/Source/Editor/UnrealEd/Private/Cooker/CookTypes.cpp b/Engine/Source/Editor/UnrealEd/Private/Cooker/CookTypes.cpp
--- a/Engine/Source/Editor/UnrealEd/Private/Cooker/CookTypes.cpp
+++ b/Engine/Source/Editor/UnrealEd/Private/Cooker/CookTypes.cpp
@@ -859,3 +859,50 @@ bool LoadFromCompactBinary(FCbFieldView Field, FBeginCookContextForWorker& OutVa
 {
 	return LoadFromCompactBinary(Field, OutValue.PlatformContexts);
 }
+
+FMultiPackageReaderResults GetSaveExportsAndImports(UPackage* Package, UObject* Asset, FSavePackageArgs SaveArgs)
+{
+	uint32 OriginalPackageFlags = Package->GetPackageFlags();
+	ON_SCOPE_EXIT{ Package->SetPackageFlagsTo(OriginalPackageFlags); };
+
+	FPackageWriterToRecord* PackageWriter = new FPackageWriterToRecord(); // SavePackageContext deletes it
+	const ITargetPlatform* TargetPlatform = SaveArgs.SavePackageContext ? SaveArgs.SavePackageContext->TargetPlatform : nullptr;
+	FSavePackageContext SavePackageContext(TargetPlatform, PackageWriter);
+	SaveArgs.SavePackageContext = &SavePackageContext;
+
+	IPackageWriter::FBeginPackageInfo BeginInfo;
+	BeginInfo.PackageName = Package->GetFName();
+	PackageWriter->BeginPackage(BeginInfo);
+
+	FString FileName = Package->GetName();
+	FSavePackageResultStruct SaveResult = GEditor->Save(Package, Asset, *FileName, SaveArgs);
+	FMultiPackageReaderResults Result;
+	Result.Result = SaveResult.Result;
+	if (Result.Result != ESavePackageResult::Success)
+	{
+		return Result;
+	}
+
+	ICookedPackageWriter::FCommitPackageInfo Info;
+	Info.Status = IPackageWriter::ECommitStatus::Success;
+	Info.PackageName = Package->GetFName();
+	Info.WriteOptions = IPackageWriter::EWriteOptions::Write;
+	PackageWriter->CommitPackage(MoveTemp(Info));
+
+	for (int32 MultiOutputIndex = 0; MultiOutputIndex < PackageWriter->SavedRecord.Packages.Num(); ++MultiOutputIndex)
+	{
+		if (MultiOutputIndex > UE_ARRAY_COUNT(Result.Realms))
+		{
+			break;
+		}
+		FPackageReaderResults& Realm = Result.Realms[MultiOutputIndex];
+
+		FMemoryReaderView HeaderArchive(PackageWriter->SavedRecord.Packages[MultiOutputIndex].Buffer.GetView());
+		FPackageReader PackageReader;
+
+		Realm.bValid = PackageReader.OpenPackageFile(&HeaderArchive) &&
+			PackageReader.ReadLinkerObjects(Realm.Exports, Realm.Imports, Realm.SoftPackageReferences);
+	}
+
+	return Result;
+}
diff --git a/Engine/Source/Editor/UnrealEd/Private/Cooker/CookTypes.h b/Engine/Source/Editor/UnrealEd/Private/Cooker/CookTypes.h
--- a/Engine/Source/Editor/UnrealEd/Private/Cooker/CookTypes.h
+++ b/Engine/Source/Editor/UnrealEd/Private/Cooker/CookTypes.h
@@ -2,6 +2,7 @@
 
 #pragma once
 
+#include "AssetRegistry/PackageReader.h"
 #include "Containers/Map.h"
 #include "Containers/Set.h"
 #include "Containers/UnrealString.h"
@@ -579,4 +580,56 @@ LLM_DECLARE_TAG(Cooker);
 #define REMAPPED_PLUGINS TEXTVIEW("RemappedPlugins")
 extern float GCookProgressWarnBusyTime;
 
-inline constexpr float TickCookableObjectsFrameTime = .100f;
\ No newline at end of file
+inline constexpr float TickCookableObjectsFrameTime = .100f;
+
+/**
+ * Scoped struct to run a function when leaving the scope. The same purpose as ON_SCOPE_EXIT, 
+ * but it can also be triggered early.
+ */
+struct FOnScopeExit
+{
+public:
+	explicit FOnScopeExit(TUniqueFunction<void()>&& InExitFunction)
+		:ExitFunction(MoveTemp(InExitFunction))
+	{
+	}
+	~FOnScopeExit()
+	{
+		ExitEarly();
+	}
+	void ExitEarly()
+	{
+		if (ExitFunction)
+		{
+			ExitFunction();
+			ExitFunction.Reset();
+		}
+	}
+	void Abandon()
+	{
+		ExitFunction.Reset();
+	}
+private:
+	TUniqueFunction<void()> ExitFunction;
+};
+/**
+ * The linker results for a single realm of a package save
+ * (e.g. the main package or the optional package that extends the main package for optionally packaged data)
+ */
+struct FPackageReaderResults
+{
+	TMap<FSoftObjectPath, FPackageReader::FObjectData> Exports;
+	TMap<FSoftObjectPath, FPackageReader::FObjectData> Imports;
+	TMap<FName, bool> SoftPackageReferences;
+	bool bValid = false;
+};
+
+/** The linker results of saving a package. A SavePackage can have multiple outputs (for e.g. optional realm). */
+struct FMultiPackageReaderResults
+{
+	FPackageReaderResults Realms[2];
+	ESavePackageResult Result;
+};
+
+/** Save the package and read the LinkerTables of its saved data. */
+FMultiPackageReaderResults GetSaveExportsAndImports(UPackage* Package, UObject* Asset, FSavePackageArgs SaveArgs);
diff --git a/Engine/Source/Editor/UnrealEd/Private/Cooker/CookWorkerServer.cpp b/Engine/Source/Editor/UnrealEd/Private/Cooker/CookWorkerServer.cpp
--- a/Engine/Source/Editor/UnrealEd/Private/Cooker/CookWorkerServer.cpp
+++ b/Engine/Source/Editor/UnrealEd/Private/Cooker/CookWorkerServer.cpp
@@ -861,7 +861,7 @@ void FCookWorkerServer::QueueDiscoveredPackage(FDiscoveredPackageReplication&& D
 
 	TArray<const ITargetPlatform*, TInlineAllocator<ExpectedMaxNumPlatforms>> BufferPlatforms;
 	TConstArrayView<const ITargetPlatform*> DiscoveredPlatforms;
-	if (!COTFS.bCanSkipEditorReferencedPackagesWhenCooking)
+	if (!COTFS.bSkipOnlyEditorOnly)
 	{
 		DiscoveredPlatforms = OrderedSessionAndSpecialPlatforms;
 	}
diff --git a/Engine/Source/Runtime/CoreUObject/Internal/Serialization/PackageWriterToSharedBuffer.h b/Engine/Source/Runtime/CoreUObject/Internal/Serialization/PackageWriterToSharedBuffer.h
--- a/Engine/Source/Runtime/CoreUObject/Internal/Serialization/PackageWriterToSharedBuffer.h
+++ b/Engine/Source/Runtime/CoreUObject/Internal/Serialization/PackageWriterToSharedBuffer.h
@@ -160,3 +160,16 @@ protected:
 	FPackageWriterRecords Records;
 };
 
+/** A PackageWriter that saves the portions of the SavePackage output to sharedbuffers on a record structure. */
+class FPackageWriterToRecord : public TPackageWriterToSharedBuffer<IPackageWriter>
+{
+public:
+	FPackageWriterRecords::FPackage SavedRecord;
+
+protected:
+	virtual void CommitPackageInternal(FPackageWriterRecords::FPackage&& Record,
+		const FCommitPackageInfo& Info) override
+	{
+		SavedRecord = MoveTemp(Record);
+	}
+};
