Ensure that utrace files generated by multi-process cook don't collide.

diff --git a/Engine/Source/Editor/UnrealEd/Private/Cooker/CookDirector.cpp b/Engine/Source/Editor/UnrealEd/Private/Cooker/CookDirector.cpp
--- a/Engine/Source/Editor/UnrealEd/Private/Cooker/CookDirector.cpp
+++ b/Engine/Source/Editor/UnrealEd/Private/Cooker/CookDirector.cpp
@@ -1218,7 +1218,7 @@ FString FCookDirector::GetWorkerCommandLine(FWorkerId WorkerId, int32 ProfileId)
 	const TCHAR* ProjectName = FApp::GetProjectName();
 	checkf(ProjectName && ProjectName[0], TEXT("Expected UnrealEditor to be running with a non-empty project name"));
 	TArray<FString> Tokens;
-	UE::String::ParseTokensMultiple(CommandLine, { ' ', '\t', '\r', '\n' }, [&Tokens](FStringView Token)
+	UE::String::ParseTokensMultiple(CommandLine, { ' ', '\t', '\r', '\n' }, [&Tokens, ProfileId](FStringView Token)
 		{
 			if (Token.StartsWith(TEXT("-run=")) ||
 				Token == TEXT("-CookOnTheFly") ||
@@ -1237,6 +1237,21 @@ FString FCookDirector::GetWorkerCommandLine(FWorkerId WorkerId, int32 ProfileId)
 			{
 				return;
 			}
+			else if (Token.StartsWith(TEXT("-tracefile=")))
+			{
+				FString TraceFile;
+				FString TokenString(Token);
+				if (FParse::Value(*TokenString, TEXT("-tracefile="), TraceFile) && !TraceFile.IsEmpty())
+				{
+					FStringView BaseFilenameWithPath = FPathViews::GetBaseFilenameWithPath(TraceFile);
+					FStringView Extension = FPathViews::GetExtension(TraceFile, true /* bIncludeDot */);
+					Tokens.Add(FString::Printf(TEXT("-tracefile=\"%.*s_Worker%d%.*s\""),
+						BaseFilenameWithPath.Len(), BaseFilenameWithPath.GetData(),
+						ProfileId,
+						Extension.Len(), Extension.GetData()));
+					return;
+				}
+			}
 			Tokens.Add(FString(Token));
 		}, UE::String::EParseTokensOptions::SkipEmpty);
 	if (Tokens[0] != ProjectName)
