Add cooker argument to generate the editor thumbnail cache (only on WindowsClient platform and for public assets)

diff --git a/Engine/Source/Editor/UnrealEd/Classes/CookOnTheSide/CookOnTheFlyServer.h b/Engine/Source/Editor/UnrealEd/Classes/CookOnTheSide/CookOnTheFlyServer.h
--- a/Engine/Source/Editor/UnrealEd/Classes/CookOnTheSide/CookOnTheFlyServer.h
+++ b/Engine/Source/Editor/UnrealEd/Classes/CookOnTheSide/CookOnTheFlyServer.h
@@ -999,6 +999,9 @@ private:
 	/** Write a MapDependencyGraph to a metadata file in the sandbox for the given platform. */
 	void WriteMapDependencyGraph(const ITargetPlatform* TargetPlatform, TMap<FName, TSet<FName>>& MapDependencyGraph);
 
+	/** Generates the CachedEditorThumbnails.bin file */
+	void GenerateCachedEditorThumbnails();
+
 	void InitializeAllCulturesToCook(TConstArrayView<FString> CookCultures);
 	void CompileDLCLocalization(FBeginCookContext& BeginContext);
 	/** Find localization dependencies for all packages, used to add required localization files as soft references. */
@@ -1201,9 +1204,12 @@ private:
 	 * 
 	 * @return full path of the asset registry in the sandbox
 	 */
-	const FString GetSandboxAssetRegistryFilename();
+	FString GetSandboxAssetRegistryFilename();
+
+	FString GetCookedAssetRegistryFilename(const FString& PlatformName);
 
-	const FString GetCookedAssetRegistryFilename(const FString& PlatformName);
+	/* @return Full path of the CachedEditorThumbnails.bin file in the sandbox */
+	FString GetSandboxCachedEditorThumbnailsFilename();
 
 	/**
 	 * Get the sandbox root directory for that platform. Is effected by the CookingDlc settings.
diff --git a/Engine/Source/Editor/UnrealEd/Private/CookOnTheFlyServer.cpp b/Engine/Source/Editor/UnrealEd/Private/CookOnTheFlyServer.cpp
--- a/Engine/Source/Editor/UnrealEd/Private/CookOnTheFlyServer.cpp
+++ b/Engine/Source/Editor/UnrealEd/Private/CookOnTheFlyServer.cpp
@@ -47,13 +47,15 @@
 #include "Editor/UnrealEdEngine.h"
 #include "EditorDomain/EditorDomain.h"
 #include "EditorDomain/EditorDomainUtils.h"
+#include "ThumbnailExternalCache.h"
+#include "ObjectTools.h"
 #include "Engine/AssetManager.h"
 #include "Engine/Level.h"
 #include "Engine/LevelStreaming.h"
 #include "Engine/Texture.h"
 #include "Engine/TextureLODSettings.h"
 #include "Engine/WorldComposition.h"
-#include "EngineGlobals.h"
+#include "EngineGlobals.h "
 #include "FileServerMessages.h"
 #include "GameDelegates.h"
 #include "GenericPlatform/GenericPlatformCrashContext.h"
@@ -8542,27 +8544,38 @@ FString UCookOnTheFlyServer::ConvertToFullSandboxPath( const FString &FileName,
 	return Result;
 }
 
-const FString UCookOnTheFlyServer::GetSandboxAssetRegistryFilename()
+FString UCookOnTheFlyServer::GetSandboxAssetRegistryFilename()
 {
-	static const FString RegistryFilename = FPaths::ProjectDir() / GetAssetRegistryFilename();
-
 	if (IsCookingDLC())
 	{
 		check(IsDirectorCookByTheBook());
-		const FString DLCRegistryFilename = FPaths::Combine(*GetBaseDirectoryForDLC(), GetAssetRegistryFilename());
-		return ConvertToFullSandboxPath(*DLCRegistryFilename, true);
+		const FString Filename = FPaths::Combine(*GetBaseDirectoryForDLC(), GetAssetRegistryFilename());
+		return ConvertToFullSandboxPath(*Filename, true);
 	}
 
-	const FString SandboxRegistryFilename = ConvertToFullSandboxPath(*RegistryFilename, true);
-	return SandboxRegistryFilename;
+	static const FString RegistryFilename = FPaths::ProjectDir() / GetAssetRegistryFilename();
+	return ConvertToFullSandboxPath(*RegistryFilename, true);
 }
 
-const FString UCookOnTheFlyServer::GetCookedAssetRegistryFilename(const FString& PlatformName )
+FString UCookOnTheFlyServer::GetCookedAssetRegistryFilename(const FString& PlatformName )
 {
 	const FString CookedAssetRegistryFilename = GetSandboxAssetRegistryFilename().Replace(TEXT("[Platform]"), *PlatformName);
 	return CookedAssetRegistryFilename;
 }
 
+FString UCookOnTheFlyServer::GetSandboxCachedEditorThumbnailsFilename()
+{
+	if (IsCookingDLC())
+	{
+		check(IsDirectorCookByTheBook());
+		const FString Filename = FPaths::Combine(*GetBaseDirectoryForDLC(), FThumbnailExternalCache::GetCachedEditorThumbnailsFilename());
+		return ConvertToFullSandboxPath(*Filename, true);
+	}
+
+	static const FString CachedEditorThumbnailsFilename = FPaths::ProjectDir() / FThumbnailExternalCache::GetCachedEditorThumbnailsFilename();
+	return ConvertToFullSandboxPath(*CachedEditorThumbnailsFilename, true);
+}
+
 namespace UE::Cook
 {
 
@@ -9028,7 +9041,7 @@ void UCookOnTheFlyServer::CookByTheBookFinished()
 
 	{
 		// Save modified asset registry with all streaming chunk info generated during cook
-		const FString& SandboxRegistryFilename = GetSandboxAssetRegistryFilename();
+		const FString SandboxRegistryFilename = GetSandboxAssetRegistryFilename();
 
 		// previously shader library was saved at this spot, but it's too early to know the chunk assignments, we need to BuildChunkManifest in the asset registry first
 
@@ -9231,6 +9244,8 @@ void UCookOnTheFlyServer::CookByTheBookFinished()
 		}
 	}
 
+	GenerateCachedEditorThumbnails();
+
 	FinalizePackageStore();
 	ShutdownCookSession();
 
@@ -9323,7 +9338,6 @@ TMap<FName, TSet<FName>> UCookOnTheFlyServer::BuildMapDependencyGraph(const ITar
 	// assign chunks for all the map packages
 	for (const UE::Cook::FPackageData* const CookedPackage : PlatformCookedPackages)
 	{
-		TArray<FAssetData> PackageAssets;
 		FName Name = CookedPackage->GetPackageName();
 
 		if (!ContainsMap(Name))
@@ -9368,6 +9382,62 @@ void UCookOnTheFlyServer::WriteMapDependencyGraph(const ITargetPlatform* TargetP
 	FFileHelper::SaveStringToFile(DependencyString, *CookedMapDependencyGraphFilePlatform, FFileHelper::EEncodingOptions::ForceUnicode);
 }
 
+void UCookOnTheFlyServer::GenerateCachedEditorThumbnails()
+{
+	if (!FParse::Param(FCommandLine::Get(), TEXT("GenerateCachedEditorThumbnails")))
+	{
+		return;
+	}
+
+	UE_SCOPED_HIERARCHICAL_COOKTIMER(GenerateCachedEditorThumbnails);
+	for (const ITargetPlatform* Platform : PlatformManager->GetSessionPlatforms())
+	{
+		// CachedEditorThumbnails only make sense for data used in the editor
+		if (Platform->PlatformName() != TEXT("WindowsClient"))
+		{
+			continue;
+		}
+
+		const FString CachedEditorThumbnailsFilename = GetSandboxCachedEditorThumbnailsFilename();
+		UE_LOG(LogCook, Display, TEXT("Generating %s"), *CachedEditorThumbnailsFilename);
+
+		// Gather public assets
+		// We don't need thumbnails for private assets because they aren't visible in the editor
+		TArray<FAssetData> PublicAssets;
+		{
+			TArray<UE::Cook::FPackageData*> CookedPackages;
+			PackageDatas->GetCookedPackagesForPlatform(Platform, CookedPackages, CookedPackages);
+
+			for (const UE::Cook::FPackageData* const CookedPackage : CookedPackages)
+			{
+				TArray<FAssetData> Assets;
+				ensure(AssetRegistry->GetAssetsByPackageName(CookedPackage->GetPackageName(), Assets, /*bIncludeOnlyDiskAssets=*/true));
+
+				for (FAssetData& Asset : Assets)
+				{
+					if (!(Asset.PackageFlags & PKG_NotExternallyReferenceable))
+					{
+						UE_LOG(LogCook, Display, TEXT("Adding thumbnail for '%s'"), *Asset.GetObjectPathString());
+						PublicAssets.Add(MoveTemp(Asset));
+					}
+				}
+			}
+		}
+
+		// Write CachedEditorThumbnails.bin file
+		{
+			FThumbnailExternalCacheSettings Settings;
+			// Convert lossless thumbnails to lossy to save space
+			Settings.bRecompressLossless = true;
+			Settings.MaxImageSize = ThumbnailTools::DefaultThumbnailSize;
+
+			// Sort to be deterministic
+			FThumbnailExternalCache::SortAssetDatas(PublicAssets);
+			FThumbnailExternalCache::Get().SaveExternalCache(CachedEditorThumbnailsFilename, PublicAssets, Settings);
+		}
+	}
+}
+
 void UCookOnTheFlyServer::QueueCancelCookByTheBook()
 {
 	if (IsCookByTheBookMode() && IsInSession())
