Cooker: Interrupt PumpSaves/PumpLoads every 30 seconds to check for pollables like MPCook's CookWorkerClient.
#rnx

diff --git a/Engine/Source/Editor/UnrealEd/Private/CookOnTheFlyServer.cpp b/Engine/Source/Editor/UnrealEd/Private/CookOnTheFlyServer.cpp
--- a/Engine/Source/Editor/UnrealEd/Private/CookOnTheFlyServer.cpp
+++ b/Engine/Source/Editor/UnrealEd/Private/CookOnTheFlyServer.cpp
@@ -1286,7 +1286,7 @@ uint32 UCookOnTheFlyServer::TickCookByTheBook(const float TimeSlice, ECookTickFl
 
 	TickMainCookLoop(StackData);
 
-	CookByTheBookOptions->CookTime += StackData.Timer.GetTimeTillNow();
+	CookByTheBookOptions->CookTime += StackData.Timer.GetTickTimeTillNow();
 	// Make sure no UE_SCOPED_HIERARCHICAL_COOKTIMERs are around CookByTheBookFinishes or CancelCookByTheBook, as those functions delete memory for them
 	if (StackData.bCookCancelled)
 	{
@@ -1397,6 +1397,11 @@ void UCookOnTheFlyServer::TickMainCookLoop(UE::Cook::FTickStackData& StackData)
 	{
 		return;
 	}
+	// Set the soft time limit to spend pumping any action at 30s, so we periodically check for pollables
+	// This is useful on CookWorkers to poll the CookClientWorker and check for CookDirector shutdown.
+	constexpr float MaxActionTimeSlice = 30.f;
+	StackData.Timer.SetActionTimeSlice(FMath::Min(StackData.Timer.GetTickTimeSlice(), MaxActionTimeSlice));
+
 	UE_SCOPED_HIERARCHICAL_COOKTIMER(TickMainCookLoop);
 	bool bContinueTick = true;
 	while (bContinueTick && (!IsEngineExitRequested() || (IsCookByTheBookMode() && !IsCookingInEditor())))
@@ -1468,6 +1473,7 @@ void UCookOnTheFlyServer::TickCookStatus(UE::Cook::FTickStackData& StackData)
 
 	double CurrentTime = FPlatformTime::Seconds();
 	StackData.LoopStartTime = CurrentTime;
+	StackData.Timer.SetActionStartTime(CurrentTime);
 	if (LastCookableObjectTickTime + TickCookableObjectsFrameTime <= CurrentTime)
 	{
 		UE_SCOPED_COOKTIMER(TickCookableObjects);
@@ -1921,7 +1927,7 @@ void UCookOnTheFlyServer::PumpPollables(UE::Cook::FTickStackData& StackData, boo
 			Pollables.HeapPop(QueueKey, false /* bAllowShrinking */);
 			QueueKey.Pollable->RunDuringPump(StackData, CurrentTime, QueueKey.NextTimeSeconds);
 			PoppedQueueKeys.Add(MoveTemp(QueueKey));
-			if (StackData.Timer.IsTimeUp(CurrentTime))
+			if (StackData.Timer.IsActionTimeUp(CurrentTime))
 			{
 				break;
 			}
@@ -1952,7 +1958,7 @@ void UCookOnTheFlyServer::PumpPollables(UE::Cook::FTickStackData& StackData, boo
 			}
 			PollNextTimeSeconds = FMath::Min(QueueKey.NextTimeSeconds, PollNextTimeSeconds);
 			PollNextTimeIdleSeconds = FMath::Min(QueueKey.Pollable->NextTimeIdleSeconds, PollNextTimeIdleSeconds);
-			if (StackData.Timer.IsTimeUp(CurrentTime))
+			if (StackData.Timer.IsActionTimeUp(CurrentTime))
 			{
 				break;
 			}
@@ -2193,7 +2199,7 @@ void UCookOnTheFlyServer::WaitForAsync(UE::Cook::FTickStackData& StackData)
 	UE_SCOPED_HIERARCHICAL_COOKTIMER(WaitForAsync);
 	double CurrentTime = FPlatformTime::Seconds();
 	double SleepDuration = WaitForAsyncSleepSeconds;
-	SleepDuration = FMath::Min(SleepDuration, StackData.Timer.GetEndTimeSeconds() - CurrentTime);
+	SleepDuration = FMath::Min(SleepDuration, StackData.Timer.GetActionEndTimeSeconds() - CurrentTime);
 	SleepDuration = FMath::Min(SleepDuration, PollNextTimeIdleSeconds - CurrentTime);
 	SleepDuration = FMath::Min(SleepDuration, SaveBusyRetryTimeSeconds - CurrentTime);
 	SleepDuration = FMath::Min(SleepDuration, LoadBusyRetryTimeSeconds - CurrentTime);
@@ -2210,7 +2216,7 @@ UCookOnTheFlyServer::ECookAction UCookOnTheFlyServer::DecideNextCookAction(UE::C
 	}
 
 	double CurrentTime = StackData.LoopStartTime;
-	if (StackData.Timer.IsTimeUp(CurrentTime))
+	if (StackData.Timer.IsTickTimeUp(CurrentTime))
 	{
 		// Timeup does not impact idle status
 		return ECookAction::YieldTick;
@@ -2395,7 +2401,7 @@ void UCookOnTheFlyServer::PumpExternalRequests(const UE::Cook::FCookerTimer& Coo
 	TArray<UE::Cook::FFilePlatformRequest> BuildRequests;
 	TArray<UE::Cook::FSchedulerCallback> SchedulerCallbacks;
 	UE::Cook::EExternalRequestType RequestType;
-	while (!CookerTimer.IsTimeUp())
+	while (!CookerTimer.IsActionTimeUp())
 	{
 		BuildRequests.Reset();
 		SchedulerCallbacks.Reset();
@@ -2497,7 +2503,7 @@ void UCookOnTheFlyServer::PumpRequests(UE::Cook::FTickStackData& StackData, int3
 			break;
 		}
 
-		if (CookerTimer.IsTimeUp())
+		if (CookerTimer.IsActionTimeUp())
 		{
 			return;
 		}
@@ -2649,7 +2655,7 @@ void UCookOnTheFlyServer::PumpLoads(UE::Cook::FTickStackData& StackData, uint32
 	while (LoadReadyQueue.Num() + LoadPrepareQueue.Num() > static_cast<int32>(DesiredQueueLength) &&
 		OutNumPushed < LoadBatchSize)
 	{
-		if (StackData.Timer.IsTimeUp())
+		if (StackData.Timer.IsActionTimeUp())
 		{
 			return;
 		}
@@ -3763,7 +3769,7 @@ UE::Cook::EPollStatus UCookOnTheFlyServer::CallBeginCacheOnObjects(UE::Cook::FPa
 			PackageDatas->AddPendingCookedPlatformData(FPendingCookedPlatformData(Obj, TargetPlatform, PackageData, bNeedsResourceRelease, *this));
 		}
 
-		if (Timer.IsTimeUp())
+		if (Timer.IsActionTimeUp())
 		{
 #if DEBUG_COOKONTHEFLY
 			UE_LOG(LogCook, Display, TEXT("Object %s took too long to cache"), *Obj->GetFullName());
@@ -4290,29 +4296,29 @@ void UCookOnTheFlyServer::PumpSaves(UE::Cook::FTickStackData& StackData, uint32
 			continue;
 		}
 
-		bool bShouldFinishTick = false;
+		bool bShouldExitPump = false;
 		if (IsCookOnTheFlyMode())
 		{
 			if (IsUsingLegacyCookOnTheFlyScheduling() && !PackageData.GetIsUrgent())
 			{
 				if (WorkerRequests->HasExternalRequests() || PackageDatas->GetMonitor().GetNumUrgent() > 0)
 				{
-					bShouldFinishTick = true;
+					bShouldExitPump = true;
 				}
-				if (StackData.Timer.IsTimeUp())
+				if (StackData.Timer.IsActionTimeUp())
 				{
 					// our timeslice is up
-					bShouldFinishTick = true;
+					bShouldExitPump = true;
 				}
 			}
 			else
 			{
 				if (IsRealtimeMode())
 				{
-					if (StackData.Timer.IsTimeUp())
+					if (StackData.Timer.IsActionTimeUp())
 					{
 						// our timeslice is up
-						bShouldFinishTick = true;
+						bShouldExitPump = true;
 					}
 				}
 				else
@@ -4324,13 +4330,13 @@ void UCookOnTheFlyServer::PumpSaves(UE::Cook::FTickStackData& StackData, uint32
 		}
 		else // !IsCookOnTheFlyMode
 		{
-			if (StackData.Timer.IsTimeUp())
+			if (StackData.Timer.IsActionTimeUp())
 			{
 				// our timeslice is up
-				bShouldFinishTick = true;
+				bShouldExitPump = true;
 			}
 		}
-		if (bShouldFinishTick)
+		if (bShouldExitPump)
 		{
 			SaveQueue.AddFront(&PackageData);
 			return;
@@ -4395,7 +4401,7 @@ void UCookOnTheFlyServer::PumpSaves(UE::Cook::FTickStackData& StackData, uint32
 					// Poll the results again and check whether we are now done
 					PackageDatas->PollPendingCookedPlatformDatas(true, LastCookableObjectTickTime);
 					PrepareSaveStatus = PrepareSave(PackageData, StackData.Timer, false /* bPrecaching */);
-				} while (!StackData.Timer.IsTimeUp() && PrepareSaveStatus == EPollStatus::Incomplete);
+				} while (!StackData.Timer.IsActionTimeUp() && PrepareSaveStatus == EPollStatus::Incomplete);
 			}
 			// If we couldn't postpone or wait, then we need to exit and try again later
 			if (PrepareSaveStatus != EPollStatus::Success)
@@ -4427,7 +4433,7 @@ void UCookOnTheFlyServer::PumpSaves(UE::Cook::FTickStackData& StackData, uint32
 
 			// If we're in RealTimeMode, check whether the precaching overflowed our timer and if so exit before we do the potentially expensive SavePackage
 			// For non-realtime, overflowing the timer is not a critical issue.
-			if (IsRealtimeMode() && StackData.Timer.IsTimeUp())
+			if (IsRealtimeMode() && StackData.Timer.IsActionTimeUp())
 			{
 				SaveQueue.AddFront(&PackageData);
 				return;
@@ -4548,7 +4554,7 @@ void UCookOnTheFlyServer::TickPrecacheObjectsForPlatforms(const float TimeSlice,
 	}
 	++LastUpdateTick;
 
-	if (Timer.IsTimeUp())
+	if (Timer.IsActionTimeUp())
 	{
 		return;
 	}
@@ -4580,7 +4586,7 @@ void UCookOnTheFlyServer::TickPrecacheObjectsForPlatforms(const float TimeSlice,
 			}
 		}
 
-		if (Timer.IsTimeUp())
+		if (Timer.IsActionTimeUp())
 		{
 			return;
 		}
@@ -4619,7 +4625,7 @@ void UCookOnTheFlyServer::TickPrecacheObjectsForPlatforms(const float TimeSlice,
 				RouteBeginCacheForCookedPlatformData(PackageName, Texture, TargetPlatform);
 			}
 		}
-		if (Timer.IsTimeUp())
+		if (Timer.IsActionTimeUp())
 		{
 			return;
 		}
diff --git a/Engine/Source/Editor/UnrealEd/Private/Cooker/CookRequestCluster.cpp b/Engine/Source/Editor/UnrealEd/Private/Cooker/CookRequestCluster.cpp
--- a/Engine/Source/Editor/UnrealEd/Private/Cooker/CookRequestCluster.cpp
+++ b/Engine/Source/Editor/UnrealEd/Private/Cooker/CookRequestCluster.cpp
@@ -224,7 +224,7 @@ void FRequestCluster::FetchPackageNames(const FCookerTimer& CookerTimer, bool& b
 	int32 NextRequest = 0;
 	for (; NextRequest < FilePlatformRequests.Num(); ++NextRequest)
 	{
-		if ((NextRequest+1) % TimerCheckPeriod == 0 && CookerTimer.IsTimeUp())
+		if ((NextRequest+1) % TimerCheckPeriod == 0 && CookerTimer.IsActionTimeUp())
 		{
 			break;
 		}
@@ -475,7 +475,7 @@ void FRequestCluster::PumpExploration(const FCookerTimer& CookerTimer, bool& bOu
 	while (GraphSearch->TickExploration(bDone), !bDone)
 	{
 		GraphSearch->WaitForAsyncQueue(WaitTime);
-		if (CookerTimer.IsTimeUp())
+		if (CookerTimer.IsActionTimeUp())
 		{
 			bOutComplete = false;
 			return;
diff --git a/Engine/Source/Editor/UnrealEd/Private/Cooker/CookTypes.cpp b/Engine/Source/Editor/UnrealEd/Private/Cooker/CookTypes.cpp
--- a/Engine/Source/Editor/UnrealEd/Private/Cooker/CookTypes.cpp
+++ b/Engine/Source/Editor/UnrealEd/Private/Cooker/CookTypes.cpp
@@ -61,7 +61,8 @@ const TCHAR* LexToString(ESuppressCookReason Reason)
 
 
 FCookerTimer::FCookerTimer(float InTimeSlice)
-	: StartTime(FPlatformTime::Seconds()), TimeSlice(InTimeSlice)
+	: TickStartTime(FPlatformTime::Seconds()), ActionStartTime(TickStartTime)
+	, TickTimeSlice(InTimeSlice), ActionTimeSlice(InTimeSlice)
 {
 }
 
@@ -75,29 +76,82 @@ FCookerTimer::FCookerTimer(ENoWait)
 {
 }
 
-double FCookerTimer::GetTimeTillNow() const
+float FCookerTimer::GetTickTimeSlice() const
 {
-	return FPlatformTime::Seconds() - StartTime;
+	return TickTimeSlice;
 }
 
-double FCookerTimer::GetEndTimeSeconds() const
+double FCookerTimer::GetTickEndTimeSeconds() const
 {
-	return FMath::Min(StartTime + TimeSlice,  MAX_flt);
+	return FMath::Min(TickStartTime + TickTimeSlice, MAX_flt);
 }
 
-bool FCookerTimer::IsTimeUp() const
+bool FCookerTimer::IsTickTimeUp() const
 {
-	return IsTimeUp(FPlatformTime::Seconds());
+	return IsTickTimeUp(FPlatformTime::Seconds());
 }
 
-bool FCookerTimer::IsTimeUp(double CurrentTimeSeconds) const
+bool FCookerTimer::IsTickTimeUp(double CurrentTimeSeconds) const
 {
-	return CurrentTimeSeconds - StartTime > TimeSlice;
+	return CurrentTimeSeconds - TickStartTime > TickTimeSlice;
 }
 
-double FCookerTimer::GetTimeRemain() const
+double FCookerTimer::GetTickTimeRemain() const
 {
-	return TimeSlice - (FPlatformTime::Seconds() - StartTime);
+	return TickTimeSlice - (FPlatformTime::Seconds() - TickStartTime);
+}
+
+double FCookerTimer::GetTickTimeTillNow() const
+{
+	return FPlatformTime::Seconds() - TickStartTime;
+}
+
+float FCookerTimer::GetActionTimeSlice() const
+{
+	return ActionTimeSlice;
+}
+
+void FCookerTimer::SetActionTimeSlice(float InTimeSlice)
+{
+	double TickEndTime = GetTickEndTimeSeconds();
+	ActionTimeSlice = FMath::Min(InTimeSlice, FMath::Max(0, TickEndTime - ActionStartTime));
+}
+
+void FCookerTimer::SetActionStartTime()
+{
+	SetActionStartTime(FPlatformTime::Seconds());
+}
+
+void FCookerTimer::SetActionStartTime(double CurrentTimeSeconds)
+{
+	ActionStartTime = CurrentTimeSeconds;
+	double TickEndTime = GetTickEndTimeSeconds();
+	ActionTimeSlice = FMath::Min(ActionTimeSlice, FMath::Max(0, TickEndTime - ActionStartTime));
+}
+
+double FCookerTimer::GetActionEndTimeSeconds() const
+{
+	return FMath::Min(ActionStartTime + ActionTimeSlice, MAX_flt);
+}
+
+bool FCookerTimer::IsActionTimeUp() const
+{
+	return IsActionTimeUp(FPlatformTime::Seconds());
+}
+
+bool FCookerTimer::IsActionTimeUp(double CurrentTimeSeconds) const
+{
+	return CurrentTimeSeconds - ActionStartTime > ActionTimeSlice;
+}
+
+double FCookerTimer::GetActionTimeRemain() const
+{
+	return ActionTimeSlice - (FPlatformTime::Seconds() - ActionStartTime);
+}
+
+double FCookerTimer::GetActionTimeTillNow() const
+{
+	return FPlatformTime::Seconds() - ActionStartTime;
 }
 
 static uint32 SchedulerThreadTlsSlot = 0;
diff --git a/Engine/Source/Editor/UnrealEd/Private/Cooker/CookTypes.h b/Engine/Source/Editor/UnrealEd/Private/Cooker/CookTypes.h
--- a/Engine/Source/Editor/UnrealEd/Private/Cooker/CookTypes.h
+++ b/Engine/Source/Editor/UnrealEd/Private/Cooker/CookTypes.h
@@ -211,15 +211,28 @@ namespace UE::Cook
 		FCookerTimer(EForever);
 		FCookerTimer(ENoWait);
 
-		double GetTimeTillNow() const;
-		double GetEndTimeSeconds() const;
-		bool IsTimeUp() const;
-		bool IsTimeUp(double CurrentTimeSeconds) const;
-		double GetTimeRemain() const;
+		float GetTickTimeSlice() const;
+		double GetTickEndTimeSeconds() const;
+		bool IsTickTimeUp() const;
+		bool IsTickTimeUp(double CurrentTimeSeconds) const;
+		double GetTickTimeRemain() const;
+		double GetTickTimeTillNow() const;
+
+		float GetActionTimeSlice() const;
+		void SetActionTimeSlice(float InTimeSlice);
+		void SetActionStartTime();
+		void SetActionStartTime(double CurrentTimeSeconds);
+		double GetActionEndTimeSeconds() const;
+		bool IsActionTimeUp() const;
+		bool IsActionTimeUp(double CurrentTimeSeconds) const;
+		double GetActionTimeRemain() const;
+		double GetActionTimeTillNow() const;
 
 	public:
-		const double StartTime;
-		const float TimeSlice;
+		const double TickStartTime;
+		double ActionStartTime;
+		const float TickTimeSlice;
+		float ActionTimeSlice;
 	};
 
 	/** Temporary-lifetime data about the current tick of the cooker. */
