MPCook: Edit the CookWorker commandlines to list the CookProfileId near the start of the commandline so it can be found even when looking at truncated commandline displayed by the debugger.
#rnx

diff --git a/Engine/Source/Editor/UnrealEd/Private/Cooker/CookDirector.cpp b/Engine/Source/Editor/UnrealEd/Private/Cooker/CookDirector.cpp
--- a/Engine/Source/Editor/UnrealEd/Private/Cooker/CookDirector.cpp
+++ b/Engine/Source/Editor/UnrealEd/Private/Cooker/CookDirector.cpp
@@ -1260,10 +1260,10 @@ FString FCookDirector::GetWorkerCommandLine(FWorkerId WorkerId, int32 ProfileId)
 	}
 	Tokens.Insert(TEXT("-run=cook"), 1);
 	Tokens.Insert(TEXT("-cookworker"), 2);
+	Tokens.Insert(FString::Printf(TEXT("-CookProfileId=%d"), ProfileId), 3);
+	Tokens.Insert(FString::Printf(TEXT("-MultiprocessId=%d"), WorkerId.GetRemoteIndex() + 1), 4);
 	check(!WorkerConnectAuthority.IsEmpty()); // This should have been constructed in TryCreateWorkerConnectSocket before any CookWorkerServers could exist to call GetWorkerCommandLine
 	Tokens.Add(FString::Printf(TEXT("-CookDirectorHost=%s"), *WorkerConnectAuthority));
-	Tokens.Add(FString::Printf(TEXT("-MultiprocessId=%d"), WorkerId.GetRemoteIndex() + 1));
-	Tokens.Add(FString::Printf(TEXT("-CookProfileId=%d"), ProfileId));
 	Tokens.Add(TEXT("-unattended"));
 	Tokens.Add(FString::Printf(TEXT("-abslog=%s"), *GetWorkerLogFileName(ProfileId)));
 	if (CoreLimit > 0)
