[Cook] Don't generate CachedEditorThumbnail.bin if it's empty

diff --git a/Engine/Source/Editor/UnrealEd/Private/CookOnTheFlyServer.cpp b/Engine/Source/Editor/UnrealEd/Private/CookOnTheFlyServer.cpp
--- a/Engine/Source/Editor/UnrealEd/Private/CookOnTheFlyServer.cpp
+++ b/Engine/Source/Editor/UnrealEd/Private/CookOnTheFlyServer.cpp
@@ -9431,6 +9431,7 @@ void UCookOnTheFlyServer::GenerateCachedEditorThumbnails()
 		}
 
 		// Write CachedEditorThumbnails.bin file
+		if (!PublicAssets.IsEmpty())
 		{
 			FThumbnailExternalCacheSettings Settings;
 			// Convert lossless thumbnails to lossy to save space
