WorldPartitionFoliageBuilder: now supports a repair mode (with -repair) that will repair problematic foliage actors.

diff --git a/Engine/Source/Editor/UnrealEd/Private/WorldPartition/WorldPartitionFoliageBuilder.cpp b/Engine/Source/Editor/UnrealEd/Private/WorldPartition/WorldPartitionFoliageBuilder.cpp
--- a/Engine/Source/Editor/UnrealEd/Private/WorldPartition/WorldPartitionFoliageBuilder.cpp
+++ b/Engine/Source/Editor/UnrealEd/Private/WorldPartition/WorldPartitionFoliageBuilder.cpp
@@ -3,6 +3,8 @@
 #include "WorldPartition/WorldPartitionFoliageBuilder.h"
 
 #include "WorldPartition/ActorDescContainer.h"
+#include "WorldPartition/ActorPartition/PartitionActorDesc.h"
+#include "WorldPartition/DataLayer/DataLayerEditorContext.h"
 #include "WorldPartition/WorldPartition.h"
 #include "WorldPartition/WorldPartitionHelpers.h"
 #include "PackageSourceControlHelper.h"
@@ -20,8 +22,10 @@ DEFINE_LOG_CATEGORY_STATIC(LogWorldPartitionFoliageBuilder, All, All);
 UWorldPartitionFoliageBuilder::UWorldPartitionFoliageBuilder(const FObjectInitializer& ObjectInitializer)
 	: Super(ObjectInitializer)
 	, NewGridSize(0)
+	, bRepair(false)
 {
 	GetParamValue("NewGridSize=", NewGridSize);
+	bRepair = HasParam("Repair");
 }
 
 bool UWorldPartitionFoliageBuilder::RunInternal(UWorld* World, const FCellInfo& InCellInfo, FPackageSourceControlHelper& PackageHelper)
@@ -34,7 +38,15 @@ bool UWorldPartitionFoliageBuilder::RunInternal(UWorld* World, const FCellInfo&
 	}
 
 	// Validate New Grid Size
-	if (World->GetWorldSettings()->InstancedFoliageGridSize == NewGridSize)
+	if (bRepair)
+	{
+		if (NewGridSize)
+		{
+			UE_LOG(LogWorldPartitionFoliageBuilder, Error, TEXT("Cannot repair foliage actors while specifying NewGridSize."), NewGridSize);
+			return false;
+		}
+	}
+	else if (World->GetWorldSettings()->InstancedFoliageGridSize == NewGridSize)
 	{
 		UE_LOG(LogWorldPartitionFoliageBuilder, Display, TEXT("Foliage Grid Size already is %d."), NewGridSize);
 		return true;
@@ -45,16 +57,19 @@ bool UWorldPartitionFoliageBuilder::RunInternal(UWorld* World, const FCellInfo&
 		return false;
 	}
 
-	if (WorldPartition->GetActorDescContainerCount() > 1)
-	{ 
-		// @todo_ow: This commandlet is currently unsupported for worldpartitions having multiple containers.
-		// ForEachActorWithLoading below loops on all actors of all containers of a WorldPartition and saves new actors.
-		// The new actors are saved within the main container WorldPartition container. Newly saved actors coming from different actor desc container should be saved in the right container. 
-		UE_LOG(LogWorldPartitionFoliageBuilder, Error, TEXT("FoliageBuilder Commandlet is unsupported on WorldPartition using more than 1 ExternalActor Folder."));
-		return false;
-	}
+	if (!bRepair)
+	{
+		if (WorldPartition->GetActorDescContainerCount() > 1)
+		{ 
+			// @todo_ow: This commandlet is currently unsupported for worldpartitions having multiple containers.
+			// ForEachActorWithLoading below loops on all actors of all containers of a WorldPartition and saves new actors.
+			// The new actors are saved within the main container WorldPartition container. Newly saved actors coming from different actor desc container should be saved in the right container. 
+			UE_LOG(LogWorldPartitionFoliageBuilder, Error, TEXT("FoliageBuilder Commandlet is unsupported on WorldPartition using more than 1 ExternalActor Folder."));
+			return false;
+		}
 
-	UE_LOG(LogWorldPartitionFoliageBuilder, Display, TEXT("Changing Foliage Grid Size from %d to %d"), World->GetWorldSettings()->InstancedFoliageGridSize, NewGridSize);
+		UE_LOG(LogWorldPartitionFoliageBuilder, Display, TEXT("Changing Foliage Grid Size from %d to %d"), World->GetWorldSettings()->InstancedFoliageGridSize, NewGridSize);
+	}
 			
 	// Update Setting
 	World->GetWorldSettings()->Modify();
@@ -93,165 +108,274 @@ bool UWorldPartitionFoliageBuilder::RunInternal(UWorld* World, const FCellInfo&
 	{
 		TMap<UFoliageType*, TArray<FFoliageInstance>> FoliageToAdd;
 		FBox InstanceBounds(ForceInit);
-				
+
+		FPartitionActorDesc* PartitionActorDesc = (FPartitionActorDesc*)ActorDesc;
 		AInstancedFoliageActor* IFA = Cast<AInstancedFoliageActor>(ActorDesc->GetActor());
 		if (!IFA)
 		{
 			UE_LOG(LogWorldPartitionFoliageBuilder, Error, TEXT("Foliage actor failed to load: %s (%s)"), *ActorDesc->GetActorName().ToString(), *ActorDesc->GetActorPackage().ToString());
 			return false;
 		}
-	
-		check(IFA->GridSize != NewGridSize);
-        // Harvest Instances from existing IFA and build up instance bounds
-		UE_LOG(LogWorldPartitionFoliageBuilder, Display, TEXT("Processing existing foliage actor: %s (%s)"), *ActorDesc->GetActorName().ToString(), *ActorDesc->GetActorPackage().ToString());
-		
-        IFA->ForEachFoliageInfo([IFA, &FoliageToAdd, &InstanceBounds, &NumInstances](UFoliageType* FoliageType, FFoliageInfo& FoliageInfo)
-        {
-            check(FoliageType->GetTypedOuter<AInstancedFoliageActor>() == nullptr);
-            FoliageToAdd.FindOrAdd(FoliageType).Append(FoliageInfo.Instances);
-            for (const FFoliageInstance& Instance : FoliageInfo.Instances)
-            {
-                InstanceBounds += Instance.GetInstanceWorldTransform().GetLocation();
-            }
-
-			if(FoliageInfo.Instances.Num() > 0)
+
+		if (bRepair)
+		{
+			FPackageSourceControlHelper PackageHelper;
+
+			const FDataLayerEditorContext DataLayerEditorContext(World, ActorDesc->GetDataLayerInstanceNames());
+			const FActorPartitionIdentifier ActorPartitionIdentifier(ActorDesc->GetActorNativeClass(), IFA->GetGridGuid(), DataLayerEditorContext.GetHash());
+			const bool bShouldIncludeGridSizeInName = ActorPartitionIdentifier.GetClass()->GetDefaultObject<APartitionActor>()->ShouldIncludeGridSizeInName(World, ActorPartitionIdentifier);
+			const FString ExpectedActorName = APartitionActor::GetActorName(World, ActorPartitionIdentifier.GetClass(), ActorPartitionIdentifier.GetGridGuid(), ActorPartitionIdentifier, PartitionActorDesc->GridSize, 
+				PartitionActorDesc->GridIndexX, PartitionActorDesc->GridIndexY, PartitionActorDesc->GridIndexZ, DataLayerEditorContext.GetHash());
+
+			if (ActorDesc->GetActorName() != ExpectedActorName)
 			{
-				UE_LOG(LogWorldPartitionFoliageBuilder, Display, TEXT("    FoliageType: %s Count: %d"), *FoliageType->GetName(), FoliageInfo.Instances.Num());
+				UPackage* PackageToModify = nullptr;
+				UPackage* PackageToDelete = nullptr;
+				UPackage* PackageToAdd = nullptr;
+
+				UE_LOG(LogWorldPartitionFoliageBuilder, Display, TEXT("Repairing invalid foliage actor: %s->%s "), *ActorDesc->GetActorName().ToString(), *ExpectedActorName);
+
+				const FString ExpectedActorPath = FString::Printf(TEXT("%s.%s"), *World->PersistentLevel->GetPathName(), *ExpectedActorName);
+				
+				if (const FWorldPartitionActorDesc* DupActorDesc = WorldPartition->GetActorDescByName(ExpectedActorPath))
+				{
+					// Merge with existing
+					FWorldPartitionReference DupActorRef(WorldPartition, DupActorDesc->GetGuid());
+
+					AInstancedFoliageActor* DupIFA = Cast<AInstancedFoliageActor>(DupActorDesc->GetActor());
+					if (!DupIFA)
+					{
+						UE_LOG(LogWorldPartitionFoliageBuilder, Error, TEXT("Foliage actor failed to load: %s (%s)"), *DupActorDesc->GetActorName().ToString(), *DupActorDesc->GetActorPackage().ToString());
+						return false;
+					}
+
+					PackageToDelete = IFA->GetPackage();
+					IFA->ForEachFoliageInfo([DupIFA](UFoliageType* FoliageType, FFoliageInfo& FoliageInfo)
+					{
+						FFoliageInfo* NewFoliageInfo = nullptr;
+						UFoliageType* NewFoliageType = DupIFA->AddFoliageType(FoliageType, &NewFoliageInfo);
+
+						for (const FFoliageInstance& Instance : FoliageInfo.Instances)
+						{
+							NewFoliageInfo->AddInstance(NewFoliageType, Instance);
+						}
+
+						return true;
+					});
+					PackageToModify = DupIFA->GetPackage();
+				}
+				else
+				{
+					// Rename existing
+					PackageToDelete = IFA->GetPackage();
+					IFA->SetPackageExternal(false);
+					IFA->Rename(*ExpectedActorName, nullptr);
+					IFA->SetActorLabel(*ExpectedActorName);
+					IFA->SetPackageExternal(true);
+					PackageToAdd = IFA->GetPackage();
+				}
+
+				// Delete
+				if (PackageToDelete && !PackageHelper.Delete(PackageToDelete))
+				{
+					return false;
+				}
+
+				// Checkout
+				if (PackageToModify)
+				{
+					FString PackageFileName = SourceControlHelpers::PackageFilename(PackageToModify);
+					if (FPlatformFileManager::Get().GetPlatformFile().FileExists(*PackageFileName))
+					{
+						if (!PackageHelper.Checkout(PackageToModify))
+						{
+							return false;
+						}
+					}
+				}
+
+				// Save
+				auto SavePackage = [](UPackage* PackageToSave)
+				{
+					FString PackageFileName = SourceControlHelpers::PackageFilename(PackageToSave);
+					FSavePackageArgs SaveArgs;
+					SaveArgs.TopLevelFlags = RF_Standalone;
+					return UPackage::SavePackage(PackageToSave, nullptr, *PackageFileName, SaveArgs);
+				};
+
+				if (PackageToModify && !SavePackage(PackageToModify))
+				{
+					return false;
+				}
+
+				if (PackageToAdd && !SavePackage(PackageToAdd))
+				{
+					return false;
+				}
+
+				// Add
+				if (PackageToAdd && !PackageHelper.AddToSourceControl(PackageToAdd))
+				{
+					return false;
+				}
 			}
-			NumInstances += FoliageInfo.Instances.Num();
-
-            return true;
-        });
-
-        TMap<UActorPartitionSubsystem::FCellCoord, AInstancedFoliageActor*> IntersectingActors;
-
-        // Maintain a list of loaded 
-        TSet<FWorldPartitionReference> ActorReferences;
-
-        // Create/Reload new actors
-        FActorPartitionGridHelper::ForEachIntersectingCell(AInstancedFoliageActor::StaticClass(), InstanceBounds, World->PersistentLevel, [&](const UActorPartitionSubsystem::FCellCoord& InCellCoord, const FBox& InCellBounds)
-        {
-            // Create or Get (might have already been created and not GCed)
-            if (AInstancedFoliageActor* NewOrExistingActor = Cast<AInstancedFoliageActor>(ActorPartitionSubsystem->GetActor(AInstancedFoliageActor::StaticClass(), InCellCoord, /*bInCreate=*/true)))
-            {
-                IntersectingActors.Add(InCellCoord, NewOrExistingActor);
-            } 
-            else // If we haven't created or found the actor it means it was created/saved/unloaded and should be in the NewActorDescs so we can reload it
-            {
-                const FGuid& NewActorGuid = NewActors.FindChecked(InCellCoord);
-                FWorldPartitionReference Reference(WorldPartition, NewActorGuid);
-                ActorReferences.Add(Reference);
-                IntersectingActors.Add(InCellCoord, CastChecked<AInstancedFoliageActor>(Reference->GetActor()));
-            }
-
-            return true;
-        });
+		}
+		else
+		{	
+			check(IFA->GridSize != NewGridSize);
+			// Harvest Instances from existing IFA and build up instance bounds
+			UE_LOG(LogWorldPartitionFoliageBuilder, Display, TEXT("Processing existing foliage actor: %s (%s)"), *ActorDesc->GetActorName().ToString(), *ActorDesc->GetActorPackage().ToString());
+		
+			IFA->ForEachFoliageInfo([IFA, &FoliageToAdd, &InstanceBounds, &NumInstances](UFoliageType* FoliageType, FFoliageInfo& FoliageInfo)
+			{
+				check(FoliageType->GetTypedOuter<AInstancedFoliageActor>() == nullptr);
+				FoliageToAdd.FindOrAdd(FoliageType).Append(FoliageInfo.Instances);
+				for (const FFoliageInstance& Instance : FoliageInfo.Instances)
+				{
+					InstanceBounds += Instance.GetInstanceWorldTransform().GetLocation();
+				}
+
+				if(FoliageInfo.Instances.Num() > 0)
+				{
+					UE_LOG(LogWorldPartitionFoliageBuilder, Display, TEXT("    FoliageType: %s Count: %d"), *FoliageType->GetName(), FoliageInfo.Instances.Num());
+				}
+				NumInstances += FoliageInfo.Instances.Num();
+
+				return true;
+			});
+
+			TMap<UActorPartitionSubsystem::FCellCoord, AInstancedFoliageActor*> IntersectingActors;
+
+			// Maintain a list of loaded 
+			TSet<FWorldPartitionReference> ActorReferences;
+
+			// Create/Reload new actors
+			FActorPartitionGridHelper::ForEachIntersectingCell(AInstancedFoliageActor::StaticClass(), InstanceBounds, World->PersistentLevel, [&](const UActorPartitionSubsystem::FCellCoord& InCellCoord, const FBox& InCellBounds)
+			{
+				// Create or Get (might have already been created and not GCed)
+				if (AInstancedFoliageActor* NewOrExistingActor = Cast<AInstancedFoliageActor>(ActorPartitionSubsystem->GetActor(AInstancedFoliageActor::StaticClass(), InCellCoord, /*bInCreate=*/true)))
+				{
+					IntersectingActors.Add(InCellCoord, NewOrExistingActor);
+				} 
+				else // If we haven't created or found the actor it means it was created/saved/unloaded and should be in the NewActorDescs so we can reload it
+				{
+					const FGuid& NewActorGuid = NewActors.FindChecked(InCellCoord);
+					FWorldPartitionReference Reference(WorldPartition, NewActorGuid);
+					ActorReferences.Add(Reference);
+					IntersectingActors.Add(InCellCoord, CastChecked<AInstancedFoliageActor>(Reference->GetActor()));
+				}
+
+				return true;
+			});
                         
-        // Maintain list of modified IFAs so we don't save empty intersecting actors
-        TSet<AInstancedFoliageActor*> ModifiedIFAs;
+			// Maintain list of modified IFAs so we don't save empty intersecting actors
+			TSet<AInstancedFoliageActor*> ModifiedIFAs;
             
-        // Add Foliage to those actors
-        for (auto& InstancesPerFoliageType : FoliageToAdd)
-        {
-            for (const FFoliageInstance& Instance : InstancesPerFoliageType.Value)
-            {
-                // Here we should get back one of the IntersectingActors that was just created or reloaded so that is why we pass in false to bCreateIfNone
-                AInstancedFoliageActor* ModifiedIFA = AInstancedFoliageActor::Get(World, /*bCreateIfNone=*/false, World->PersistentLevel, Instance.Location);
-                check(ModifiedIFA);
-                ModifiedIFAs.Add(ModifiedIFA);
-                check(ModifiedIFA->GridSize == NewGridSize);
-                FFoliageInfo* NewFoliageInfo = nullptr;
-                UFoliageType* NewFoliageType = ModifiedIFA->AddFoliageType(InstancesPerFoliageType.Key, &NewFoliageInfo);
-                NewFoliageInfo->AddInstance(NewFoliageType, Instance);
-				NumInstancesProcessed++;
-            }
-        }
-
-        // Save only modified ones
-        for (auto Pair : IntersectingActors)
-        {
-            AInstancedFoliageActor* IntersectingActor = Pair.Value;
-            // No need to save actor as it hasn't changed
-            if (!ModifiedIFAs.Contains(IntersectingActor))
-            {
-                continue;
-            }
-
-            UPackage* Package = IntersectingActor->GetExternalPackage();
-            check(Package);
-
-            FString PackageFileName = SourceControlHelpers::PackageFilename(Package);
+			// Add Foliage to those actors
+			for (auto& InstancesPerFoliageType : FoliageToAdd)
+			{
+				for (const FFoliageInstance& Instance : InstancesPerFoliageType.Value)
+				{
+					// Here we should get back one of the IntersectingActors that was just created or reloaded so that is why we pass in false to bCreateIfNone
+					AInstancedFoliageActor* ModifiedIFA = AInstancedFoliageActor::Get(World, /*bCreateIfNone=*/false, World->PersistentLevel, Instance.Location);
+					check(ModifiedIFA);
+					ModifiedIFAs.Add(ModifiedIFA);
+					check(ModifiedIFA->GridSize == NewGridSize);
+					FFoliageInfo* NewFoliageInfo = nullptr;
+					UFoliageType* NewFoliageType = ModifiedIFA->AddFoliageType(InstancesPerFoliageType.Key, &NewFoliageInfo);
+					NewFoliageInfo->AddInstance(NewFoliageType, Instance);
+					NumInstancesProcessed++;
+				}
+			}
+
+			// Save only modified ones
+			for (auto Pair : IntersectingActors)
+			{
+				AInstancedFoliageActor* IntersectingActor = Pair.Value;
+				// No need to save actor as it hasn't changed
+				if (!ModifiedIFAs.Contains(IntersectingActor))
+				{
+					continue;
+				}
+
+				UPackage* Package = IntersectingActor->GetExternalPackage();
+				check(Package);
+
+				FString PackageFileName = SourceControlHelpers::PackageFilename(Package);
                                 
-            // Save package
-			FSavePackageArgs SaveArgs;
-			SaveArgs.TopLevelFlags = RF_Standalone;
-			if (!UPackage::SavePackage(Package, nullptr, *PackageFileName, SaveArgs))
-            {
-                UE_LOG(LogWorldPartitionFoliageBuilder, Error, TEXT("Error saving package %s."), *Package->GetName());
-                return false;
-            }
-
-            // Actor Desc didn't exist yet which means we need to add the file to SCC and add it to the new actors descs and reference it so it gets unloaded
-            if (!NewActors.Contains(Pair.Key))
-            {
-                if (!SCCHelper.AddToSourceControl(Package))
-                {
-                    // It is possible the resave can't checkout everything. Continue processing.
-                    UE_LOG(LogWorldPartitionFoliageBuilder, Error, TEXT("Error adding package to revision control %s."), *Package->GetName());
-                    return false;
-                }
-
-                FGuid NewActorGuid = IntersectingActor->GetActorGuid();
-                FWorldPartitionActorDesc* NewActorDesc = WorldPartition->GetActorDesc(NewActorGuid);
-                check(NewActorDesc);
-                NewActors.Add(Pair.Key, NewActorGuid);
-
-                // Add ref so actor gets unloaded 
-                FWorldPartitionReference Reference(WorldPartition, NewActorGuid);
-                ActorReferences.Add(Reference);
-            }
+				// Save package
+				FSavePackageArgs SaveArgs;
+				SaveArgs.TopLevelFlags = RF_Standalone;
+				if (!UPackage::SavePackage(Package, nullptr, *PackageFileName, SaveArgs))
+				{
+					UE_LOG(LogWorldPartitionFoliageBuilder, Error, TEXT("Error saving package %s."), *Package->GetName());
+					return false;
+				}
+
+				// Actor Desc didn't exist yet which means we need to add the file to SCC and add it to the new actors descs and reference it so it gets unloaded
+				if (!NewActors.Contains(Pair.Key))
+				{
+					if (!SCCHelper.AddToSourceControl(Package))
+					{
+						// It is possible the resave can't checkout everything. Continue processing.
+						UE_LOG(LogWorldPartitionFoliageBuilder, Error, TEXT("Error adding package to revision control %s."), *Package->GetName());
+						return false;
+					}
+
+					FGuid NewActorGuid = IntersectingActor->GetActorGuid();
+					FWorldPartitionActorDesc* NewActorDesc = WorldPartition->GetActorDesc(NewActorGuid);
+					check(NewActorDesc);
+					NewActors.Add(Pair.Key, NewActorGuid);
+
+					// Add ref so actor gets unloaded 
+					FWorldPartitionReference Reference(WorldPartition, NewActorGuid);
+					ActorReferences.Add(Reference);
+				}
                 
-            UE_LOG(LogWorldPartitionFoliageBuilder, Display, TEXT("Saved foliage actor %s (%s)."), *IntersectingActor->GetName(), *Package->GetName());
-        }
+				UE_LOG(LogWorldPartitionFoliageBuilder, Display, TEXT("Saved foliage actor %s (%s)."), *IntersectingActor->GetName(), *Package->GetName());
+			}
+		}
 					
 		return true;
 	}, ForEachActorWithLoadingParams);
 
-	check(NumInstances == NumInstancesProcessed);
-
-	// Checkout World for WorldSettings change
-	UPackage* WorldSettingsPackage = World->GetWorldSettings()->GetPackage();
-	if (!SCCHelper.Checkout(WorldSettingsPackage))
+	if (!bRepair)
 	{
-		UE_LOG(LogWorldPartitionFoliageBuilder, Error, TEXT("Error checking out package %s."), *WorldSettingsPackage->GetName());
-		return false;
-	}
+		check(NumInstances == NumInstancesProcessed);
 
-	// Save World 
-	{
-		FString PackageFileName = SourceControlHelpers::PackageFilename(WorldSettingsPackage);
-		FSavePackageArgs SaveArgs;
-		SaveArgs.TopLevelFlags = RF_Standalone;
-		if (!UPackage::SavePackage(World->GetPackage(), nullptr, *PackageFileName, SaveArgs))
+		// Checkout World for WorldSettings change
+		UPackage* WorldSettingsPackage = World->GetWorldSettings()->GetPackage();
+		if (!SCCHelper.Checkout(WorldSettingsPackage))
 		{
-			UE_LOG(LogWorldPartitionFoliageBuilder, Error, TEXT("Error saving package %s."), *WorldSettingsPackage->GetName());
+			UE_LOG(LogWorldPartitionFoliageBuilder, Error, TEXT("Error checking out package %s."), *WorldSettingsPackage->GetName());
 			return false;
 		}
-	}
 
-	// Delete old IFAs
-	for (const FGuid& ActorGuid : ExistingActors)
-	{
-		FWorldPartitionActorDesc* ActorDesc = WorldPartition->GetActorDesc(ActorGuid);
-		check(ActorDesc);
-		const FString PackageToDelete = ActorDesc->GetActorPackage().ToString();
-		if (!SCCHelper.Delete(PackageToDelete))
+		// Save World 
 		{
-			UE_LOG(LogWorldPartitionFoliageBuilder, Error, TEXT("Error deleting package %s."), *PackageToDelete);
-			return false;
+			FString PackageFileName = SourceControlHelpers::PackageFilename(WorldSettingsPackage);
+			FSavePackageArgs SaveArgs;
+			SaveArgs.TopLevelFlags = RF_Standalone;
+			if (!UPackage::SavePackage(World->GetPackage(), nullptr, *PackageFileName, SaveArgs))
+			{
+				UE_LOG(LogWorldPartitionFoliageBuilder, Error, TEXT("Error saving package %s."), *WorldSettingsPackage->GetName());
+				return false;
+			}
+		}
+
+		// Delete old IFAs
+		for (const FGuid& ActorGuid : ExistingActors)
+		{
+			FWorldPartitionActorDesc* ActorDesc = WorldPartition->GetActorDesc(ActorGuid);
+			check(ActorDesc);
+			const FString PackageToDelete = ActorDesc->GetActorPackage().ToString();
+			if (!SCCHelper.Delete(PackageToDelete))
+			{
+				UE_LOG(LogWorldPartitionFoliageBuilder, Error, TEXT("Error deleting package %s."), *PackageToDelete);
+				return false;
+			}
 		}
 	}
 
 	return true;
-}
-
+}
\ No newline at end of file
diff --git a/Engine/Source/Editor/UnrealEd/Public/WorldPartition/WorldPartitionFoliageBuilder.h b/Engine/Source/Editor/UnrealEd/Public/WorldPartition/WorldPartitionFoliageBuilder.h
--- a/Engine/Source/Editor/UnrealEd/Public/WorldPartition/WorldPartitionFoliageBuilder.h
+++ b/Engine/Source/Editor/UnrealEd/Public/WorldPartition/WorldPartitionFoliageBuilder.h
@@ -20,4 +20,5 @@ public:
 
 private:
 	uint32 NewGridSize;
+	bool bRepair;
 };
\ No newline at end of file
diff --git a/Engine/Source/Runtime/Engine/Private/ActorPartition/ActorPartitionSubsystem.cpp b/Engine/Source/Runtime/Engine/Private/ActorPartition/ActorPartitionSubsystem.cpp
--- a/Engine/Source/Runtime/Engine/Private/ActorPartition/ActorPartitionSubsystem.cpp
+++ b/Engine/Source/Runtime/Engine/Private/ActorPartition/ActorPartitionSubsystem.cpp
@@ -241,32 +241,11 @@ public:
 				
 		if (!FoundActor && bInCreate)
 		{
-			TStringBuilderWithBuffer<TCHAR, NAME_SIZE> ActorNameBuilder;
-
-			ActorNameBuilder += InActorPartitionId.GetClass()->GetName();
-			ActorNameBuilder += TEXT("_");
-
-			if (InActorPartitionId.GetGridGuid().IsValid())
-			{
-				ActorNameBuilder += InActorPartitionId.GetGridGuid().ToString(EGuidFormats::Base36Encoded);
-				ActorNameBuilder += TEXT("_");
-			}
-
-			if (InActorPartitionId.GetClass()->GetDefaultObject<APartitionActor>()->ShouldIncludeGridSizeInName(World, InActorPartitionId))
-			{
-				ActorNameBuilder += FString::Printf(TEXT("%d_"), InGridSize);
-			}
-
-			ActorNameBuilder += FString::Printf(TEXT("%d_%d_%d"), InCellCoord.X, InCellCoord.Y, InCellCoord.Z);
-
-			if (InActorPartitionId.GetDataLayerEditorContextHash() != FDataLayerEditorContext::EmptyHash)
-			{
-				ActorNameBuilder += FString::Printf(TEXT("_%X"), InActorPartitionId.GetDataLayerEditorContextHash());
-			}
+			const FString ActorName = APartitionActor::GetActorName(World, InActorPartitionId.GetClass(), InActorPartitionId.GetGridGuid(), InActorPartitionId, InGridSize, InCellCoord.X, InCellCoord.Y, InCellCoord.Z, InActorPartitionId.GetDataLayerEditorContextHash());
 
 			FActorSpawnParameters SpawnParams;
 			SpawnParams.OverrideLevel = InCellCoord.Level;
-			SpawnParams.Name = *ActorNameBuilder;
+			SpawnParams.Name = *ActorName;
 			SpawnParams.NameMode = FActorSpawnParameters::ESpawnActorNameMode::Required_Fatal;
 
 			// Handle the case where the actor already exists, but is in the undo stack (was deleted)
diff --git a/Engine/Source/Runtime/Engine/Private/ActorPartition/PartitionActor.cpp b/Engine/Source/Runtime/Engine/Private/ActorPartition/PartitionActor.cpp
--- a/Engine/Source/Runtime/Engine/Private/ActorPartition/PartitionActor.cpp
+++ b/Engine/Source/Runtime/Engine/Private/ActorPartition/PartitionActor.cpp
@@ -7,6 +7,7 @@
 #if WITH_EDITOR
 #include "Engine/Level.h"
 #include "WorldPartition/ActorPartition/PartitionActorDesc.h"
+#include "WorldPartition/DataLayer/DataLayerEditorContext.h"
 #include "GameFramework/WorldSettings.h"
 #include "Engine/World.h"
 #endif
@@ -50,6 +51,34 @@ bool APartitionActor::ShouldIncludeGridSizeInName(UWorld* InWorld, const FActorP
 {
 	return InWorld->GetWorldSettings()->bIncludeGridSizeInNameForPartitionedActors;
 }
+
+FString APartitionActor::GetActorName(UWorld* World, const UClass* Class, const FGuid& GridGuid, const FActorPartitionIdentifier& ActorPartitionId, uint32 GridSize, int32 CellCoordsX, int32 CellCoordsY, int32 CellCoordsZ, uint32 DataLayerEditorContext)
+{
+	TStringBuilderWithBuffer<TCHAR, NAME_SIZE> ActorNameBuilder;
+
+	ActorNameBuilder += Class->GetName();
+	ActorNameBuilder += TEXT("_");
+
+	if (GridGuid.IsValid())
+	{
+		ActorNameBuilder += GridGuid.ToString(EGuidFormats::Base36Encoded);
+		ActorNameBuilder += TEXT("_");
+	}
+
+	if (Class->GetDefaultObject<APartitionActor>()->ShouldIncludeGridSizeInName(World, ActorPartitionId))
+	{
+		ActorNameBuilder += FString::Printf(TEXT("%d_"), GridSize);
+	}
+
+	ActorNameBuilder += FString::Printf(TEXT("%d_%d_%d"), CellCoordsX, CellCoordsY, CellCoordsZ);
+
+	if (DataLayerEditorContext != FDataLayerEditorContext::EmptyHash)
+	{
+		ActorNameBuilder += FString::Printf(TEXT("_%X"), DataLayerEditorContext);
+	}
+
+	return ActorNameBuilder.ToString();
+}
 #endif
 
 #undef LOCTEXT_NAMESPACE
diff --git a/Engine/Source/Runtime/Engine/Public/ActorPartition/PartitionActor.h b/Engine/Source/Runtime/Engine/Public/ActorPartition/PartitionActor.h
--- a/Engine/Source/Runtime/Engine/Public/ActorPartition/PartitionActor.h
+++ b/Engine/Source/Runtime/Engine/Public/ActorPartition/PartitionActor.h
@@ -35,6 +35,8 @@ public:
 	/** The grid size this actors was generated for */
 	UPROPERTY()
 	uint32 GridSize;
+
+	static FString GetActorName(UWorld* World, const UClass* Class, const FGuid& GridGuid, const FActorPartitionIdentifier& ActorPartitionId, uint32 GridSize, int32 CellCoordsX, int32 CellCoordsY, int32 CellCoordsZ, uint32 DataLayerEditorContext);
 #endif
 };
 
