Remove GLog->Flush from UCookOnTheFlyServer::EvaluateGarbageCollectionResults to try and address a stall that is measured there in profiles.

diff --git a/Engine/Source/Editor/UnrealEd/Private/CookOnTheFlyServer.cpp b/Engine/Source/Editor/UnrealEd/Private/CookOnTheFlyServer.cpp
--- a/Engine/Source/Editor/UnrealEd/Private/CookOnTheFlyServer.cpp
+++ b/Engine/Source/Editor/UnrealEd/Private/CookOnTheFlyServer.cpp
@@ -5379,7 +5379,6 @@ void UCookOnTheFlyServer::EvaluateGarbageCollectionResults(bool bWasDueToOOM, bo
 		{
 			GLog->Logf(TEXT("LLM Tags are not displayed because llm is disabled. Run with -llm or -trace=memtag to see llm tags."));
 		}
-		GLog->Flush();
 	}
 	else
 	{
