Fix Editor.Usage.<AssetEditor> telemetry to be more schematizable.
* Edit is now called Editor.Usage.AssetEditorClosed
* The AssetEditor itself was moved into an attribute named "AssetEditor"
* The dots were removed from the other attributes to make them easier to dereference:
  * TotalDuration.Seconds -> TotalDurationSeconds
  * OpendInstances.Count -> OpendInstancesCount

diff --git a/Engine/Source/Editor/UnrealEd/Private/Subsystems/AssetEditorSubsystem.cpp b/Engine/Source/Editor/UnrealEd/Private/Subsystems/AssetEditorSubsystem.cpp
--- a/Engine/Source/Editor/UnrealEd/Private/Subsystems/AssetEditorSubsystem.cpp
+++ b/Engine/Source/Editor/UnrealEd/Private/Subsystems/AssetEditorSubsystem.cpp
@@ -92,16 +92,15 @@ void UAssetEditorSubsystem::Deinitialize()
 	if (FEngineAnalytics::IsAvailable())
 	{
 		TArray<FAnalyticsEventAttribute> EditorUsageAttribs;
-		EditorUsageAttribs.Empty(2);
+		EditorUsageAttribs.Empty(3);
 		for (auto Iter = EditorUsageAnalytics.CreateConstIterator(); Iter; ++Iter)
 		{
 			const FAssetEditorAnalyticInfo& Data = Iter.Value();
 			EditorUsageAttribs.Reset();
-			EditorUsageAttribs.Emplace(TEXT("TotalDuration.Seconds"), FString::Printf(TEXT("%.1f"), Data.SumDuration.GetTotalSeconds()));
-			EditorUsageAttribs.Emplace(TEXT("OpenedInstances.Count"), FString::Printf(TEXT("%d"), Data.NumTimesOpened));
-
-			const FString EventName = FString::Printf(TEXT("Editor.Usage.%s"), *Iter.Key().ToString());
-			FEngineAnalytics::GetProvider().RecordEvent(EventName, EditorUsageAttribs);
+			EditorUsageAttribs.Emplace(TEXT("TotalDurationSeconds"), Data.SumDuration.GetTotalSeconds());
+			EditorUsageAttribs.Emplace(TEXT("OpenedInstancesCount"), Data.NumTimesOpened);
+			EditorUsageAttribs.Emplace(TEXT("AssetEditor"), *Iter.Key().ToString());
+			FEngineAnalytics::GetProvider().RecordEvent(TEXT("Editor.Usage.AssetEditorClosed"), EditorUsageAttribs);
 		}
 	}
 
