Cooker and ShaderCompiler: Have the cooker pass in as long a time slice as it can to ProcessAsyncResults; this fixes the issue of an idle cooker sleeping for a second at a time and giving the shadercompiler only do 10ms of work for every 1s of cooker sleeping. Remove the call to ProcessAsyncResults during the frequent unforced PollPendingPlatformDatas; it is an unnecessary spend of time since we are processing it to completion every 60s or so; when unforced there is no need to minimize the latency for completing a package to less than one minute when there is still other work to do.
#rnx

diff --git a/Engine/Source/Editor/UnrealEd/Private/Commandlets/CookCommandlet.cpp b/Engine/Source/Editor/UnrealEd/Private/Commandlets/CookCommandlet.cpp
--- a/Engine/Source/Editor/UnrealEd/Private/Commandlets/CookCommandlet.cpp
+++ b/Engine/Source/Editor/UnrealEd/Private/Commandlets/CookCommandlet.cpp
@@ -628,7 +628,7 @@ void UCookCommandlet::ConditionalCollectGarbage(uint32 TickResults, UCookOnTheFl
 				}
 			}
 
-			GShaderCompilingManager->ProcessAsyncResults(false, false);
+			GShaderCompilingManager->ProcessAsyncResults(0.f, false);
 
 			FPlatformProcess::Sleep(0.05);
 
diff --git a/Engine/Source/Editor/UnrealEd/Private/CookOnTheFlyServer.cpp b/Engine/Source/Editor/UnrealEd/Private/CookOnTheFlyServer.cpp
--- a/Engine/Source/Editor/UnrealEd/Private/CookOnTheFlyServer.cpp
+++ b/Engine/Source/Editor/UnrealEd/Private/CookOnTheFlyServer.cpp
@@ -2127,10 +2127,10 @@ static void CBTBTickCommandletStats()
 	FStats::TickCommandletStats();
 }
 
-static void TickShaderCompilingManager()
+static void TickShaderCompilingManager(UE::Cook::FTickStackData& StackData)
 {
 	UE_SCOPED_COOKTIMER_AND_DURATION(CookByTheBook_ShaderProcessAsync, DetailedCookStats::TickLoopShaderProcessAsyncResultsTimeSec);
-	GShaderCompilingManager->ProcessAsyncResults(true, false);
+	GShaderCompilingManager->ProcessAsyncResults(StackData.Timer.GetActionTimeSlice(), false);
 }
 
 static void TickAssetRegistry()
@@ -2157,7 +2157,10 @@ void UCookOnTheFlyServer::InitializePollables()
 			Pollables.Emplace(MoveTemp(Pollable));
 		}
 		Pollables.Emplace(new FPollable(TEXT("ProcessDeferredCommands"), 60.f, 5.f, [this](FTickStackData&) { ProcessDeferredCommands(*this); }));
-		Pollables.Emplace(new FPollable(TEXT("ShaderCompilingManager"), 60.f, 5.f, [this](FTickStackData&) { TickShaderCompilingManager(); }));
+		Pollables.Emplace(new FPollable(TEXT("ShaderCompilingManager"), 60.f, 5.f, [this](FTickStackData& StackData)
+		{
+			TickShaderCompilingManager(StackData);
+		}));
 		Pollables.Emplace(new FPollable(TEXT("FlushRenderingCommands"), 60.f, 5.f, [this](FTickStackData&) { PollFlushRenderingCommands(); }));
 		if (TRefCountPtr<FPollable> Pollable = CreatePollableLLM())
 		{
diff --git a/Engine/Source/Editor/UnrealEd/Private/Cooker/CookPackageData.cpp b/Engine/Source/Editor/UnrealEd/Private/Cooker/CookPackageData.cpp
--- a/Engine/Source/Editor/UnrealEd/Private/Cooker/CookPackageData.cpp
+++ b/Engine/Source/Editor/UnrealEd/Private/Cooker/CookPackageData.cpp
@@ -3112,8 +3112,14 @@ void FPackageDatas::PollPendingCookedPlatformDatas(bool bForce, double& LastCook
 		return;
 	}
 
-	GShaderCompilingManager->ProcessAsyncResults(true /* bLimitExecutionTime */,
-		false /* bBlockOnGlobalShaderCompletion */);
+	if (bForce)
+	{
+		// When we are forced, because the caller has an urgent package to save, call ProcessAsyncResults
+		// with a small timeslice in case we need to process shaders to unblock the package
+		constexpr float TimeSlice = 0.01f;
+		GShaderCompilingManager->ProcessAsyncResults(TimeSlice, false /* bBlockOnGlobalShaderCompletion */);
+	}
+
 	FDelegateHandle EventHandle = FAssetCompilingManager::Get().OnPackageScopeEvent().AddLambda(
 		[this](UPackage* Package, bool bEntering)
 		{
diff --git a/Engine/Source/Runtime/Engine/Private/ShaderCompiler/ShaderCompiler.cpp b/Engine/Source/Runtime/Engine/Private/ShaderCompiler/ShaderCompiler.cpp
--- a/Engine/Source/Runtime/Engine/Private/ShaderCompiler/ShaderCompiler.cpp
+++ b/Engine/Source/Runtime/Engine/Private/ShaderCompiler/ShaderCompiler.cpp
@@ -5656,6 +5656,12 @@ void FShaderCompilingManager::FinishAllCompilation()
 }
 
 void FShaderCompilingManager::ProcessAsyncResults(bool bLimitExecutionTime, bool bBlockOnGlobalShaderCompletion)
+{
+	const float TimeSlice = bLimitExecutionTime ? ProcessGameThreadTargetTime : 0.f;
+	ProcessAsyncResults(TimeSlice, bBlockOnGlobalShaderCompletion);
+}
+
+void FShaderCompilingManager::ProcessAsyncResults(float TimeSlice, bool bBlockOnGlobalShaderCompletion)
 {
 	TRACE_CPUPROFILER_EVENT_SCOPE(FShaderCompilingManager::ProcessAsyncResults)
 
@@ -5729,13 +5735,13 @@ void FShaderCompilingManager::ProcessAsyncResults(bool bLimitExecutionTime, bool
 		} 
 		while (bRetry);
 
-		const float TimeBudget = bLimitExecutionTime ? ProcessGameThreadTargetTime : FLT_MAX;
+		const float TimeBudget = TimeSlice > 0 ? TimeSlice : FLT_MAX;
 		ProcessCompiledShaderMaps(PendingFinalizeShaderMaps, TimeBudget);
-		check(bLimitExecutionTime || PendingFinalizeShaderMaps.Num() == 0);
+		check(TimeSlice > 0 || PendingFinalizeShaderMaps.Num() == 0);
 	}
 
 
-	if (bBlockOnGlobalShaderCompletion && !bLimitExecutionTime && !IsRunningCookCommandlet())
+	if (bBlockOnGlobalShaderCompletion && TimeSlice <= 0 && !IsRunningCookCommandlet())
 	{
 		check(PendingFinalizeShaderMaps.Num() == 0);
 
diff --git a/Engine/Source/Runtime/Engine/Public/ShaderCompiler.h b/Engine/Source/Runtime/Engine/Public/ShaderCompiler.h
--- a/Engine/Source/Runtime/Engine/Public/ShaderCompiler.h
+++ b/Engine/Source/Runtime/Engine/Public/ShaderCompiler.h
@@ -920,11 +920,14 @@ public:
 
 	/** 
 	 * Processes completed asynchronous shader maps, and assigns them to relevant materials.
-	 * @param bLimitExecutionTime - When enabled, ProcessAsyncResults will be bandwidth throttled by ProcessGameThreadTargetTime, to limit hitching.
+	 * @param TimeSlice - When more than 0, ProcessAsyncResults will be bandwidth throttled by the given timeslice, to limit hitching.
 	 *		ProcessAsyncResults will then have to be called often to finish all shader maps (eg from Tick).  Otherwise, all compiled shader maps will be processed.
 	 * @param bBlockOnGlobalShaderCompletion - When enabled, ProcessAsyncResults will block until global shader maps are complete.
 	 *		This must be done before using global shaders for rendering.
 	 */
+	ENGINE_API void ProcessAsyncResults(float TimeSlice, bool bBlockOnGlobalShaderCompletion);
+
+	/** Version of ProcessAsyncResults that specifies use of ProcessGameThreadTargetTime for the timeslice. */
 	ENGINE_API void ProcessAsyncResults(bool bLimitExecutionTime, bool bBlockOnGlobalShaderCompletion);
 
 	/**
