If a SKEL class is not valid (e.g. is RF_Garbage), just set it's superstruct pointer to the stable REINST class if its super struct is recompiled. This can happen if client code loads temp blueprints and then marks them as RF_Garbage

diff --git a/Engine/Source/Editor/UnrealEd/Private/Kismet2/KismetReinstanceUtilities.cpp b/Engine/Source/Editor/UnrealEd/Private/Kismet2/KismetReinstanceUtilities.cpp
--- a/Engine/Source/Editor/UnrealEd/Private/Kismet2/KismetReinstanceUtilities.cpp
+++ b/Engine/Source/Editor/UnrealEd/Private/Kismet2/KismetReinstanceUtilities.cpp
@@ -1717,7 +1717,7 @@ void FBlueprintCompileReinstancer::GetSortedClassHierarchy(UClass* ClassToSearch
 	OutHierarchy.Sort([](UClass& A, UClass& B)->bool { return FBlueprintCompileReinstancer::ReinstancerOrderingFunction(&A, &B); });
 }
 
-void FBlueprintCompileReinstancer::MoveDependentSkelToReinst(UClass* const OwnerClass, TMap<UClass*, UClass*>& OldToNewMap)
+void FBlueprintCompileReinstancer::MoveDependentSkelToReinst(UClass* const OwnerClass, TMap<UClass*, UClass*>& NewSkeletonToOldSkeleton)
 {
 	// Gather the whole class hierarchy up the native class so that we can correctly create the REINST class parented to native
 	TArray<UClass*> ClassHierarchy;
@@ -1739,7 +1739,15 @@ void FBlueprintCompileReinstancer::MoveDependentSkelToReinst(UClass* const Owner
 		UObject* OldCDO = CurClass->ClassDefaultObject;
 		const FName ReinstanceName = MakeUniqueObjectName(GetTransientPackage(), CurClass->GetClass(), *(FString(TEXT("REINST_")) + *CurClass->GetName()));
 
-		checkf(IsValid(CurClass), TEXT("%s is invalid - will not duplicate successfully"), *(CurClass->GetName()));
+		if (!IsValid(CurClass))
+		{
+			if (UClass* const* NewSuper = NewSkeletonToOldSkeleton.Find(CurClass->GetSuperClass()))
+			{
+				CurClass->SetSuperStruct(*NewSuper);
+			}
+			continue;
+		}
+
 		UClass* ReinstClass = CastChecked<UClass>(StaticDuplicateObject(CurClass, GetTransientPackage(), ReinstanceName, ~RF_Transactional));
 		
 		ReinstClass->RemoveFromRoot();
@@ -1747,7 +1755,7 @@ void FBlueprintCompileReinstancer::MoveDependentSkelToReinst(UClass* const Owner
 		CurClass->ClassFlags &= ~CLASS_NewerVersionExists;
 		GIsDuplicatingClassForReinstancing = false;
 
-		UClass** OverridenParent = OldToNewMap.Find(ReinstClass->GetSuperClass());
+		UClass** OverridenParent = NewSkeletonToOldSkeleton.Find(ReinstClass->GetSuperClass());
 		if (OverridenParent && *OverridenParent)
 		{
 			ReinstClass->SetSuperStruct(*OverridenParent);
@@ -1757,7 +1765,7 @@ void FBlueprintCompileReinstancer::MoveDependentSkelToReinst(UClass* const Owner
 		ReinstClass->StaticLink(true);
 
 		// Map the old class to the new one
-		OldToNewMap.Add(CurClass, ReinstClass);
+		NewSkeletonToOldSkeleton.Add(CurClass, ReinstClass);
 
 		// Actually move the old CDO reference out of the way
 		if (OldCDO)
