Cooker: Change the implementation of Randomization to maintain leaf-to-root order of visitation, which is useful for performance; randomize only which edges are taken first from each vertex.
#rnx

diff --git a/Engine/Source/Editor/UnrealEd/Classes/CookOnTheSide/CookOnTheFlyServer.h b/Engine/Source/Editor/UnrealEd/Classes/CookOnTheSide/CookOnTheFlyServer.h
--- a/Engine/Source/Editor/UnrealEd/Classes/CookOnTheSide/CookOnTheFlyServer.h
+++ b/Engine/Source/Editor/UnrealEd/Classes/CookOnTheSide/CookOnTheFlyServer.h
@@ -1441,6 +1441,11 @@ private:
 	bool bSkipOnlyEditorOnly = false;
 	/** True if we're running cooklist; tweak output */
 	bool bCookListMode = false;
+	/**
+	 * True if we want to randomize cook order, for robustness validation or to avoid encountering the same DDC
+	 * build jobs when multiple machines are cooking at the same time.
+	 */
+	bool bRandomizeCookOrder = false;
 
 	/** Timers for tracking how long we have been busy, to manage retries and warnings of deadlock */
 	double SaveBusyStartTimeSeconds = MAX_flt;
diff --git a/Engine/Source/Editor/UnrealEd/Private/CookOnTheFlyServer.cpp b/Engine/Source/Editor/UnrealEd/Private/CookOnTheFlyServer.cpp
--- a/Engine/Source/Editor/UnrealEd/Private/CookOnTheFlyServer.cpp
+++ b/Engine/Source/Editor/UnrealEd/Private/CookOnTheFlyServer.cpp
@@ -6611,6 +6611,9 @@ void UCookOnTheFlyServer::SetInitializeConfigSettings(UE::Cook::FInitializeConfi
 				});
 		}
 	}
+
+	bRandomizeCookOrder = FParse::Param(FCommandLine::Get(), TEXT("RANDOMPACKAGEORDER")) ||
+		(FParse::Param(FCommandLine::Get(), TEXT("DIFFONLY")) && !FParse::Param(FCommandLine::Get(), TEXT("DIFFNORANDCOOK")));
 }
 
 bool UCookOnTheFlyServer::TryInitializeCookWorker()
@@ -10619,17 +10622,6 @@ void UCookOnTheFlyServer::GenerateInitialRequests(FBeginCookContext& BeginContex
 		LogCookerMessage(FString::Printf(TEXT("No files found to cook.")), EMessageSeverity::Warning);
 	}
 
-	if (FParse::Param(FCommandLine::Get(), TEXT("RANDOMPACKAGEORDER")) ||
-		(FParse::Param(FCommandLine::Get(), TEXT("DIFFONLY")) && !FParse::Param(FCommandLine::Get(), TEXT("DIFFNORANDCOOK"))))
-	{
-		UE_LOG(LogCook, Log, TEXT("Randomizing package order."));
-		//randomize the array, taking the Array_Shuffle approach, in order to help bring cooking determinism issues to the surface.
-		for (int32 FileIndex = 0; FileIndex < FilesInPath.Num(); ++FileIndex)
-		{
-			FilesInPath.Swap(FileIndex, FMath::RandRange(FileIndex, FilesInPath.Num() - 1));
-		}
-	}
-
 	{
 		UE_SCOPED_HIERARCHICAL_COOKTIMER(GenerateLongPackageName);
 		GenerateLongPackageNames(FilesInPath, FilesInPathInstigators);
diff --git a/Engine/Source/Editor/UnrealEd/Private/Cooker/CookRequestCluster.cpp b/Engine/Source/Editor/UnrealEd/Private/Cooker/CookRequestCluster.cpp
--- a/Engine/Source/Editor/UnrealEd/Private/Cooker/CookRequestCluster.cpp
+++ b/Engine/Source/Editor/UnrealEd/Private/Cooker/CookRequestCluster.cpp
@@ -23,6 +23,7 @@
 #include "HAL/PlatformProcess.h"
 #include "Interfaces/ITargetPlatform.h"
 #include "Misc/RedirectCollector.h"
+#include "Misc/ReverseIterate.h"
 #include "Misc/StringBuilder.h"
 #include "String/Find.h"
 #include "UObject/CoreRedirects.h"
@@ -501,6 +502,14 @@ void FRequestCluster::PumpExploration(const FCookerTimer& CookerTimer, bool& bOu
 	};
 
 	Algo::TopologicalSort(SortedPackages, GetElementDependencies, Algo::ETopologicalSort::AllowCycles);
+	if (COTFS.bRandomizeCookOrder)
+	{
+		RandomizeCookOrder(SortedPackages, GraphSearch->GetGraphEdges());
+		// Only shuffle the first cluster. We don't need to spend time shuffling later clusters because they
+		// will already occur at random times, and we want to suppress the log information about them since
+		// there can be a lot of small discovered clusters.
+		COTFS.bRandomizeCookOrder = false;
+	}
 	TMap<FPackageData*, int32> SortOrder;
 	int32 Counter = 0;
 	SortOrder.Reserve(SortedPackages.Num());
@@ -1561,4 +1570,226 @@ TConstArrayView<FName> FRequestCluster::GetLocalizationReferences(FName PackageN
 	return TConstArrayView<FName>();
 }
 
+template <typename T>
+static void ArrayShuffle(TArray<T>& Array)
+{
+	// iterate 0 to N-1, picking a random remaining vertex each loop
+	int32 N = Array.Num();
+	for (int32 I = 0; I < N; ++I)
+	{
+		Array.Swap(I, FMath::RandRange(I, N - 1));
+	}
+}
+
+template <typename T>
+static TArray<T> FindRootsFromLeafToRootOrderList(TConstArrayView<T> LeafToRootOrder, const TMap<T, TArray<T>>& Edges,
+	const TSet<T>& ValidVertices)
+{
+	// Iteratively
+	//    1) Add the leading rootward non-visited element to the root
+	//    2) Visit all elements reachable from that root
+	// This works because the input array is already sorted RootToLeaf, so we
+	// know the leading element has no incoming edges from anything later.
+	TArray<T> Roots;
+	TSet<T> Visited;
+	Visited.Reserve(LeafToRootOrder.Num());
+	struct FVisitEntry
+	{
+		T Vertex;
+		const TArray<T>* Edges;
+		int32 NextEdge;
+		void Set(T V, const TMap<T, TArray<T>>& AllEdges)
+		{
+			Vertex = V;
+			Edges = AllEdges.Find(V);
+			NextEdge = 0;
+		}
+	};
+	TArray<FVisitEntry> DFSStack;
+	int32 StackNum = 0;
+	auto Push = [&DFSStack, &Edges, &StackNum](T Vertex)
+	{
+		while (DFSStack.Num() <= StackNum)
+		{
+			DFSStack.Emplace();
+		}
+		DFSStack[StackNum++].Set(Vertex, Edges);
+	};
+	auto Pop = [&StackNum]()
+	{
+		--StackNum;
+	};
+
+	for (T Root : ReverseIterate(LeafToRootOrder))
+	{
+		bool bAlreadyExists;
+		Visited.Add(Root, &bAlreadyExists);
+		if (bAlreadyExists)
+		{
+			continue;
+		}
+		Roots.Add(Root);
+
+		Push(Root);
+		check(StackNum == 1);
+		while (StackNum > 0)
+		{
+			FVisitEntry& Entry = DFSStack[StackNum - 1];
+			bool bPushed = false;
+			while (Entry.Edges && Entry.NextEdge < Entry.Edges->Num())
+			{
+				T Target = (*Entry.Edges)[Entry.NextEdge++];
+				Visited.Add(Target, &bAlreadyExists);
+				if (!bAlreadyExists && ValidVertices.Contains(Target))
+				{
+					Push(Target);
+					bPushed = true;
+					break;
+				}
+			}
+			if (!bPushed)
+			{
+				Pop();
+			}
+		}
+	}
+	return Roots;
+}
+
+void FRequestCluster::RandomizeCookOrder(TArray<FPackageData*>& InOutLeafToRootOrder,
+	const TMap<FPackageData*, TArray<FPackageData*>>& Edges)
+{
+	// Notes on the ideal solution:
+	// In a graph without cycles, the visitation order of a DepthFirstSearch starting from the graph roots is a
+	// RootToLeaf ordering. We can randomize by randomly iterating the graph roots and by randomly iterating the
+	// edges when moving from each graph.
+	// In a graph with cycles, a RootToLeaf order is only defined in the condensed graph, and for
+	// each chain within a node of the condensed graph, we can randomize the vertices in the chain.
+	// But this requires creating the condensed graph, which is expensive.
+	// 
+	// Notes on the practical solution:
+	// Cycles are not supposed to be a large part of our graph, so we will not seek to do a good job of randomizing
+	// them. We do a DFS as we would in an acyclic graph, and when we encounter an already-visited vertex due to a cycle
+	// we just pretend that edge does not exist.
+	//
+	// Two DFS passes
+	// Pass 1, find all roots by iterating from root to leaf and DFSing each remaining head element.
+	// Pass 2: Iterate all roots in random order and DFS each one; append the leaf-to-root order of their search
+	// to the final leaf-to-root order.
+	if (InOutLeafToRootOrder.IsEmpty())
+	{
+		return;
+	}
+
+	struct FVisitEntry
+	{
+		FPackageData* Vertex;
+		TArray<FPackageData*> Edges;
+		int32 NextEdge;
+		void Set(FPackageData* V, const TMap<FPackageData*,TArray<FPackageData*>>& AllEdges)
+		{
+			Vertex = V;
+			Edges.Reset();
+			const TArray<FPackageData*>* EdgesFromV = AllEdges.Find(V);
+			if (EdgesFromV)
+			{
+				Edges.Append(*EdgesFromV);
+				ArrayShuffle(Edges);
+			}
+			NextEdge = 0;
+		}
+	};
+	
+	TArray<FVisitEntry> DFSStack;
+	int32 StackNum = 0;
+	auto Push = [&DFSStack, &Edges, &StackNum](FPackageData* Vertex)
+	{
+		while (DFSStack.Num() <= StackNum)
+		{
+			DFSStack.Emplace();
+		}
+		DFSStack[StackNum++].Set(Vertex, Edges);
+	};
+	auto Pop = [&StackNum]()
+	{
+		--StackNum;
+	};
+
+	TSet<FPackageData*> ValidVertices;
+	ValidVertices.Append(InOutLeafToRootOrder);
+
+	TArray<FPackageData*> Roots = FindRootsFromLeafToRootOrderList<FPackageData*>(InOutLeafToRootOrder, Edges, ValidVertices);
+	ArrayShuffle(Roots);
+
+	TSet<FPackageData*> Visited;
+	Visited.Reserve(InOutLeafToRootOrder.Num());
+	TArray<FPackageData*> OutputVisitOrder;
+
+	for (FPackageData* Root : Roots)
+	{
+		bool bAlreadyExists;
+		Visited.Add(Root, &bAlreadyExists);
+		if (bAlreadyExists)
+		{
+			continue;
+		}
+
+		Push(Root);
+		check(StackNum == 1);
+		while (StackNum > 0)
+		{
+			FVisitEntry& Entry = DFSStack[StackNum - 1];
+			bool bPushed = false;
+			while (Entry.NextEdge < Entry.Edges.Num())
+			{
+				FPackageData* Target = Entry.Edges[Entry.NextEdge++];
+				Visited.Add(Target, &bAlreadyExists);
+				if (!bAlreadyExists && ValidVertices.Contains(Target))
+				{
+					Push(Target);
+					bPushed = true;
+					break;
+				}
+			}
+			if (!bPushed)
+			{
+				// Add this vertex now as the most rootwards encountered; it is farther toward the root than any of the vertices it depends on
+				OutputVisitOrder.Add(Entry.Vertex);
+				Pop();
+			}
+		}
+	}
+	Visited.Empty();
+	Roots.Empty();
+
+	TMap<FPackageData*, int32> OriginalIndices;
+	OriginalIndices.Reserve(InOutLeafToRootOrder.Num());
+	for (int32 Index = 0; Index < InOutLeafToRootOrder.Num(); ++Index)
+	{
+		OriginalIndices.Add(InOutLeafToRootOrder[Index], Index);
+	}
+	check(OriginalIndices.Num() == InOutLeafToRootOrder.Num());
+	
+	// Copy the OutputVisitOrder over top of our InOut ordered parameter and report the diagnostic
+	// for the average shuffle distance.
+
+	double SumSquaredDistances = 0.0;
+
+	int32 WriteIndex = 0;
+	for (FPackageData* Vertex : OutputVisitOrder)
+	{
+		int32 OriginalIndex;
+		verify(OriginalIndices.RemoveAndCopyValue(Vertex, OriginalIndex));
+		SumSquaredDistances += FMath::Square(OriginalIndex - WriteIndex);
+
+		check(WriteIndex < InOutLeafToRootOrder.Num());
+		InOutLeafToRootOrder[WriteIndex++] = Vertex;
+	}
+	check(WriteIndex == InOutLeafToRootOrder.Num());
+
+	UE_LOG(LogCook, Display,
+		TEXT("RandomPackageOrder used, packages in the cluster were shuffled. %d elements in cluster, average shuffle distance == %0.1f."),
+		InOutLeafToRootOrder.Num(), (float) FMath::Sqrt(SumSquaredDistances/InOutLeafToRootOrder.Num()));
+}
+
 } // namespace UE::Cook
diff --git a/Engine/Source/Editor/UnrealEd/Private/Cooker/CookRequestCluster.h b/Engine/Source/Editor/UnrealEd/Private/Cooker/CookRequestCluster.h
--- a/Engine/Source/Editor/UnrealEd/Private/Cooker/CookRequestCluster.h
+++ b/Engine/Source/Editor/UnrealEd/Private/Cooker/CookRequestCluster.h
@@ -320,6 +320,8 @@ private:
 		FPackageDatas& InPackageDatas, FPackageTracker& InPackageTracker,
 		FStringView InDLCPath, bool bInErrorOnEngineContentUse, bool bInAllowUncookedAssetReferences,
 		bool bSkipOnlyEditorOnly, ESuppressCookReason& OutReason, bool& bOutCookable, bool& bOutExplorable);
+	static void RandomizeCookOrder(TArray<FPackageData*>& InOutLeafToRootOrder,
+		const TMap<FPackageData*, TArray<FPackageData*>>& Edges);
 
 	TArray<FFilePlatformRequest> FilePlatformRequests;
 	TFastPointerMap<FPackageData*, ESuppressCookReason> OwnedPackageDatas;
