Cooker: RandomPackageOrder: Fix diagnostic value broken by integer overflow.
#rnx

diff --git a/Engine/Source/Editor/UnrealEd/Private/Cooker/CookRequestCluster.cpp b/Engine/Source/Editor/UnrealEd/Private/Cooker/CookRequestCluster.cpp
--- a/Engine/Source/Editor/UnrealEd/Private/Cooker/CookRequestCluster.cpp
+++ b/Engine/Source/Editor/UnrealEd/Private/Cooker/CookRequestCluster.cpp
@@ -1780,7 +1780,8 @@ void FRequestCluster::RandomizeCookOrder(TArray<FPackageData*>& InOutLeafToRootO
 	{
 		int32 OriginalIndex;
 		verify(OriginalIndices.RemoveAndCopyValue(Vertex, OriginalIndex));
-		SumSquaredDistances += FMath::Square(OriginalIndex - WriteIndex);
+		double Distance = static_cast<double>(OriginalIndex - WriteIndex);
+		SumSquaredDistances += FMath::Square(Distance);
 
 		check(WriteIndex < InOutLeafToRootOrder.Num());
 		InOutLeafToRootOrder[WriteIndex++] = Vertex;
