Cooker: Add log message showing the packagename for packages that take long to save.
#rnx

diff --git a/Engine/Source/Editor/UnrealEd/Private/CookOnTheFlyServer.cpp b/Engine/Source/Editor/UnrealEd/Private/CookOnTheFlyServer.cpp
--- a/Engine/Source/Editor/UnrealEd/Private/CookOnTheFlyServer.cpp
+++ b/Engine/Source/Editor/UnrealEd/Private/CookOnTheFlyServer.cpp
@@ -5715,6 +5715,7 @@ void UCookOnTheFlyServer::SaveCookedPackage(UE::Cook::FSaveCookedPackageContext&
 
 	UE_SCOPED_HIERARCHICAL_CUSTOM_COOKTIMER_AND_DURATION(SaveCookedPackage, DetailedCookStats::TickCookOnTheSideSaveCookedPackageTimeSec)
 		UE_ADD_CUSTOM_COOKTIMER_META(SaveCookedPackage, PackageName, *WriteToString<256>(Context.PackageData.GetFileName()));
+	double SaveStartTime = FPlatformTime::Seconds();
 
 	UPackage* Package = Context.Package;
 	FOnScopeExit ScopedPackageFlags([Package, OriginalPackageFlags = Package->GetPackageFlags()]()
@@ -5787,6 +5788,12 @@ void UCookOnTheFlyServer::SaveCookedPackage(UE::Cook::FSaveCookedPackageContext&
 	// Need to restore flags before calling FinishPackage because it might need to save again
 	ScopedPackageFlags.ExitEarly();
 	Context.FinishPackage();
+
+	constexpr double SavePackageMinDurationLogTimeSeconds = 600.;
+	float SaveDurationSeconds = static_cast<float>(FPlatformTime::Seconds() - SaveStartTime);
+	UE_CLOG(SaveDurationSeconds >= SavePackageMinDurationLogTimeSeconds,
+		LogCook, Display, TEXT("SavePackagePerformance: Package %s took %.0fs to save."),
+		*WriteToString<256>(Package->GetFName()), SaveDurationSeconds);
 }
 
 bool UCookOnTheFlyServer::IsDebugRecordUnsolicited() const
