Add missing insight tag on long ComponentTypeRegistryData ticks

#rnx

diff --git a/Engine/Source/Editor/UnrealEd/Private/ComponentTypeRegistry.cpp b/Engine/Source/Editor/UnrealEd/Private/ComponentTypeRegistry.cpp
--- a/Engine/Source/Editor/UnrealEd/Private/ComponentTypeRegistry.cpp
+++ b/Engine/Source/Editor/UnrealEd/Private/ComponentTypeRegistry.cpp
@@ -671,6 +671,8 @@ void FComponentTypeRegistryData::ForceRefreshComponentList()
 
 void FComponentTypeRegistryData::Tick(float)
 {
+	TRACE_CPUPROFILER_EVENT_SCOPE(FComponentTypeRegistryData::Tick);
+
 	bool bRequiresRefresh = bNeedsRefreshNextTick;
 
 	if (PendingAssetData.Num() != 0)
