Fix for a potential crash during Blueprint asset sanitization on save in a restricted editor context.

#rnx

diff --git a/Engine/Source/Editor/UnrealEd/Private/Kismet2/BlueprintEditorUtils.cpp b/Engine/Source/Editor/UnrealEd/Private/Kismet2/BlueprintEditorUtils.cpp
--- a/Engine/Source/Editor/UnrealEd/Private/Kismet2/BlueprintEditorUtils.cpp
+++ b/Engine/Source/Editor/UnrealEd/Private/Kismet2/BlueprintEditorUtils.cpp
@@ -4,7 +4,6 @@
 #include "Algo/AnyOf.h"
 #include "Algo/Copy.h"
 #include "Algo/RemoveIf.h"
-#include "Algo/Sort.h"
 #include "Algo/Transform.h"
 #include "BlueprintCompilationManager.h"
 #include "UObject/Interface.h"
@@ -10716,34 +10715,35 @@ bool FBlueprintEditorUtils::HasRestrictedNodes(const UBlueprint* BP, TArray<UEdG
 	return false;
 }
 
-void FBlueprintEditorUtils::SanitizeRestrictedContentOnSave(UBlueprint* BP)
+void FBlueprintEditorUtils::SanitizeRestrictedContent(UBlueprint* BP)
 {
 	TArray<UEdGraphNode*> RestrictedNodes;
 	if (HasRestrictedNodes(BP, &RestrictedNodes))
 	{
-		// Front-load entry nodes so that their graphs can be removed first.
-		Algo::Sort(RestrictedNodes, [](const UEdGraphNode* A, const UEdGraphNode* B)
-		{
-			return A->IsA<UK2Node_Event>() || A->IsA<UK2Node_FunctionEntry>();
-		});
-
+		// Remove any restricted nodes and/or graphs (note: in restricted mode, it is implied that
+		// restricted items are not also user-facing, so any removals here should be opaque to users).
+		TSet<UEdGraph*> GraphsToRemove;
 		for (UEdGraphNode* Node : RestrictedNodes)
 		{
 			// Remove the entire graph if this node represents an entry point for an event/function.
 			if (Node->IsA<UK2Node_Event>()
 				|| Node->IsA<UK2Node_FunctionEntry>())
 			{
-				RemoveGraph(BP, Node->GetGraph(), EGraphRemoveFlags::MarkTransient);
-			}
-			else
-			{
-				// If the top-level graph is transient/removed, there's no need to remove the node here.
-				if (!GetTopLevelGraph(Node->GetGraph())->HasAnyFlags(RF_Transient))
-				{
-					RemoveNode(BP, Node, /*bDontRecompile =*/ true);
-				}
+				GraphsToRemove.Add(Node->GetGraph());
 			}
+			
+			// Defer updating the skeleton class since we're going to do a full recompile below.
+			RemoveNode(BP, Node, /*bDontRecompile =*/ true);
+		}
+
+		// Note: On removal, graphs will be renamed and outered to the package containing the Blueprint.
+		for (UEdGraph* Graph : GraphsToRemove)
+		{
+			RemoveGraph(BP, Graph, EGraphRemoveFlags::MarkTransient);
 		}
+
+		// Do a full recompile to bring everything back up-to-date and to sanitize any generated bytecode.
+		FKismetEditorUtilities::CompileBlueprint(BP);
 	}
 }
 
diff --git a/Engine/Source/Editor/UnrealEd/Public/Kismet2/BlueprintEditorUtils.h b/Engine/Source/Editor/UnrealEd/Public/Kismet2/BlueprintEditorUtils.h
--- a/Engine/Source/Editor/UnrealEd/Public/Kismet2/BlueprintEditorUtils.h
+++ b/Engine/Source/Editor/UnrealEd/Public/Kismet2/BlueprintEditorUtils.h
@@ -1934,9 +1934,11 @@ public:
 	static bool HasRestrictedNodes(const UBlueprint* BP, TArray<UEdGraphNode*>* OutRestrictedNodes = nullptr);
 
 	/**
-	 * Checks the given Blueprint for any restricted content within the current editor context and sanitizes it prior to saving.
+	 * Checks the given Blueprint for any restricted content within the current editor context and sanitizes it away.
+	 * 
+	 * Note: If any restricted content is removed, this will also recompile the Blueprint.
 	 */
-	static void SanitizeRestrictedContentOnSave(UBlueprint* BP);
+	static void SanitizeRestrictedContent(UBlueprint* BP);
 
 public:
 	static bool ShouldOpenWithDataOnlyEditor(const UBlueprint* Blueprint);
diff --git a/Engine/Source/Runtime/Engine/Private/Blueprint.cpp b/Engine/Source/Runtime/Engine/Private/Blueprint.cpp
--- a/Engine/Source/Runtime/Engine/Private/Blueprint.cpp
+++ b/Engine/Source/Runtime/Engine/Private/Blueprint.cpp
@@ -411,7 +411,7 @@ void UBlueprint::PreSave(FObjectPreSaveContext ObjectSaveContext)
 
 #if WITH_EDITOR
 	// Sanitize the Blueprint by clearing out any objects/types that are restricted within the current editor context.
-	FBlueprintEditorUtils::SanitizeRestrictedContentOnSave(this);
+	FBlueprintEditorUtils::SanitizeRestrictedContent(this);
 #endif	// WITH_EDITOR
 }
 
