Move around some asset thumbnail functionality (no change in logic)
Fixed EThumbnailTextureFlushMode::AlwaysFlush mode not fully working all the time (might be slower, but it's not used anywhere)

diff --git a/Engine/Source/Editor/UnrealEd/Private/AssetThumbnail.cpp b/Engine/Source/Editor/UnrealEd/Private/AssetThumbnail.cpp
--- a/Engine/Source/Editor/UnrealEd/Private/AssetThumbnail.cpp
+++ b/Engine/Source/Editor/UnrealEd/Private/AssetThumbnail.cpp
@@ -43,7 +43,7 @@
 
 namespace AssetThumbnailPool
 {
-	const FObjectThumbnail* LoadThumbnailFromPackage(const FAssetData& AssetData, FThumbnailMap& OutThumbnailMap)
+	const FObjectThumbnail* LoadThumbnailsFromPackage(const FAssetData& AssetData, FThumbnailMap& OutThumbnailMap)
 	{
 		FString PackageFilename;
 		if (FPackageName::DoesPackageExist(AssetData.PackageName.ToString(), &PackageFilename))
@@ -1203,7 +1203,7 @@ bool FAssetThumbnailPool::LoadThumbnail(TSharedRef<FThumbnailInfo> ThumbnailInfo
 	// rendering it without editor-only data might not give the same result as when it was rendered uncooked.
 	if (AssetData.PackageFlags & PKG_Cooked)
 	{
-		FoundThumbnail = AssetThumbnailPool::LoadThumbnailFromPackage(AssetData, ThumbnailMap);
+		FoundThumbnail = AssetThumbnailPool::LoadThumbnailsFromPackage(AssetData, ThumbnailMap);
 	}
 
 	if (!FoundThumbnail)
@@ -1303,7 +1303,7 @@ bool FAssetThumbnailPool::LoadThumbnail(TSharedRef<FThumbnailInfo> ThumbnailInfo
 	// If we don't have a thumbnail cached in memory, try to find it on disk
 	if (!FoundThumbnail && !(AssetData.PackageFlags & PKG_Cooked))
 	{
-		FoundThumbnail = AssetThumbnailPool::LoadThumbnailFromPackage(AssetData, ThumbnailMap);
+		FoundThumbnail = AssetThumbnailPool::LoadThumbnailsFromPackage(AssetData, ThumbnailMap);
 	}
 
 	if (FoundThumbnail)
diff --git a/Engine/Source/Editor/UnrealEd/Private/ObjectTools.cpp b/Engine/Source/Editor/UnrealEd/Private/ObjectTools.cpp
--- a/Engine/Source/Editor/UnrealEd/Private/ObjectTools.cpp
+++ b/Engine/Source/Editor/UnrealEd/Private/ObjectTools.cpp
@@ -4993,23 +4993,19 @@ namespace ThumbnailTools
 		// Get the rendering info for this object
 		FThumbnailRenderingInfo* RenderInfo = GUnrealEd ? GUnrealEd->GetThumbnailManager()->GetRenderingInfo( InObject ) : nullptr;
 
-		// Wait for all textures to be streamed in before we render the thumbnail
-		// @todo CB: This helps but doesn't result in 100%-streamed-in resources every time! :(
 		if( InFlushMode == EThumbnailTextureFlushMode::AlwaysFlush )
 		{
-			if (GShaderCompilingManager)
-			{
-				GShaderCompilingManager->ProcessAsyncResults(false, true);
-			}
+			// Wait for pending load requests.
+			FlushAsyncLoading();
 
-			if (UTexture* Texture = Cast<UTexture>(InObject))
-			{
-				FTextureCompilingManager::Get().FinishCompilation({Texture});
-			}
+			// Wait for shader and other asset compilation to finish.
+			FAssetCompilingManager::Get().FinishAllCompilation();
 
-			FlushAsyncLoading();
+			// Force all mips to load.
+			UTexture::ForceUpdateTextureStreaming();
 
-			IStreamingManager::Get().StreamAllResources( 100.0f );
+			// Force all streamed resources to finish.
+			IStreamingManager::Get().StreamAllResources();
 		}
 
 		// If this object's thumbnail will be rendered to a texture on the GPU.
@@ -5122,13 +5118,13 @@ namespace ThumbnailTools
 		if( RenderInfo != NULL && RenderInfo->Renderer != NULL )
 		{
 			// Set the size of cached thumbnails
-			const int32 ImageWidth = 	ThumbnailTools::DefaultThumbnailSize;
+			const int32 ImageWidth = ThumbnailTools::DefaultThumbnailSize;
 			const int32 ImageHeight = ThumbnailTools::DefaultThumbnailSize;
 
 			// For cached thumbnails we want to make sure that textures are fully streamed in so that the thumbnail we're saving won't have artifacts
 			// However, this can add 30s - 100s to editor load
 			//@todo - come up with a cleaner solution for this, preferably not blocking on texture streaming at all but updating when textures are fully streamed in
-			ThumbnailTools::EThumbnailTextureFlushMode::Type TextureFlushMode = ThumbnailTools::EThumbnailTextureFlushMode::NeverFlush;
+			const ThumbnailTools::EThumbnailTextureFlushMode::Type TextureFlushMode = ThumbnailTools::EThumbnailTextureFlushMode::NeverFlush;
 
 			if ( UTexture* Texture = Cast<UTexture>(InObject) )
 			{
@@ -5156,7 +5152,7 @@ namespace ThumbnailTools
 
 			// Generate the thumbnail
 			FObjectThumbnail NewThumbnail;
-			ThumbnailTools::RenderThumbnail(
+			RenderThumbnail(
 				InObject, ImageWidth, ImageHeight, TextureFlushMode, NULL,
 				&NewThumbnail );		// Out
 
@@ -5539,6 +5535,44 @@ namespace ThumbnailTools
 
 
 
+	bool GetPackageFilePathAndAssetFullName(const FAssetData& AssetData, FString& OutPackageFilePath, FName& OutAssetFullName)
+	{
+		// Determine package file path
+		if (FPackageName::DoesPackageExist(AssetData.PackageName.ToString(), &OutPackageFilePath))
+		{
+			// Determine asset fullname
+			FNameBuilder FullNameBuilder;
+			AssetData.GetFullName(FullNameBuilder);
+			OutAssetFullName = FName(FullNameBuilder);
+			return true;
+		}
+		return false;
+	}
+
+
+
+	bool LoadThumbnailFromPackage(const FAssetData& AssetData, FObjectThumbnail& OutThumbnail)
+	{
+		FString PackageFilePath;
+		FName AssetFullName;
+		if (GetPackageFilePathAndAssetFullName(AssetData, PackageFilePath, AssetFullName))
+		{
+			TSet<FName> AssetFullNames;
+			AssetFullNames.Add(AssetFullName);
+
+			FThumbnailMap ThumbnailMap;
+			LoadThumbnailsFromPackage(PackageFilePath, AssetFullNames, ThumbnailMap);
+			if (FObjectThumbnail* Found = ThumbnailMap.Find(AssetFullName))
+			{
+				OutThumbnail = MoveTemp(*Found);
+				return true;
+			}
+		}
+		return false;
+	}
+
+
+
 	/** Loads thumbnails from the specified package file name, try loading from external cache file if not found in package file */
 	bool LoadThumbnailsFromPackage( const FString& InPackageFileName, const TSet< FName >& InObjectFullNames, FThumbnailMap& InOutThumbnails )
 	{
@@ -5695,5 +5729,5 @@ namespace ThumbnailTools
 
 		return false;
 	}
-}
+		}
 
diff --git a/Engine/Source/Editor/UnrealEd/Private/ThumbnailExternalCache.cpp b/Engine/Source/Editor/UnrealEd/Private/ThumbnailExternalCache.cpp
--- a/Engine/Source/Editor/UnrealEd/Private/ThumbnailExternalCache.cpp
+++ b/Engine/Source/Editor/UnrealEd/Private/ThumbnailExternalCache.cpp
@@ -173,8 +173,7 @@ void FSaveThumbnailCache::Save(FArchive& Ar, const TArrayView<FAssetData> InAsse
 		ParallelFor(AssetsForSegment.Num(), [AssetsForSegment, &Tasks, &InSettings](int32 Index)
 		{
 			FSaveThumbnailCacheTask& Task = Tasks[Index];
-			Task.ObjectThumbnail = FThumbnailExternalCache::LoadThumbnailFromPackage(AssetsForSegment[Index]);
-			if (!Task.ObjectThumbnail.IsEmpty())
+			if (ThumbnailTools::LoadThumbnailFromPackage(AssetsForSegment[Index], Task.ObjectThumbnail) && !Task.ObjectThumbnail.IsEmpty())
 			{
 				{
 					FNameBuilder ObjectFullNameBuilder;
@@ -445,30 +444,6 @@ void FThumbnailExternalCache::SaveExternalCache(FArchive& Ar, const TArrayView<F
 	bIsSavingCache = false;
 }
 
-FObjectThumbnail FThumbnailExternalCache::LoadThumbnailFromPackage(const FAssetData& AssetData)
-{
-	// Determine filename
-	FString PackageFilename;
-	if (FPackageName::DoesPackageExist(AssetData.PackageName.ToString(), &PackageFilename))
-	{
-		// Thumbnails are identified in package with full object names
-		TSet<FName> ObjectFullNames;
-		FNameBuilder ObjectFullNameBuilder;
-		AssetData.GetFullName(ObjectFullNameBuilder);
-		const FName ObjectFullName(ObjectFullNameBuilder);
-		ObjectFullNames.Add(ObjectFullName);
-
-		FThumbnailMap ThumbnailMap;
-		ThumbnailTools::LoadThumbnailsFromPackage(PackageFilename, ObjectFullNames, ThumbnailMap);
-		if (FObjectThumbnail* Found = ThumbnailMap.Find(ObjectFullName))
-		{
-			return MoveTemp(*Found);
-		}
-	}
-
-	return FObjectThumbnail();
-}
-
 void FThumbnailExternalCache::OnContentPathMounted(const FString& InAssetPath, const FString& InFilesystemPath)
 {
 	if (TSharedPtr<IPlugin> FoundPlugin = IPluginManager::Get().FindPluginFromPath(InAssetPath))
diff --git a/Engine/Source/Editor/UnrealEd/Public/ObjectTools.h b/Engine/Source/Editor/UnrealEd/Public/ObjectTools.h
--- a/Engine/Source/Editor/UnrealEd/Public/ObjectTools.h
+++ b/Engine/Source/Editor/UnrealEd/Public/ObjectTools.h
@@ -722,7 +722,10 @@ namespace ThumbnailTools
 	/** Returns the thumbnail for the specified object or NULL if one doesn't exist yet */
 	UNREALED_API FObjectThumbnail* GetThumbnailForObject( UObject* InObject );
 
-	/** Loads thumbnails from the specified package file name */
+	/** Loads the thumbnail of an asset from the specified package file name (or from the external thumbnail cache file if it exists) */
+	UNREALED_API bool LoadThumbnailFromPackage(const FAssetData& AssetData, FObjectThumbnail& OutThumbnail);
+
+	/** Loads thumbnails from the specified package file name (or from the external thumbnail cache file if it exists) */
 	UNREALED_API bool LoadThumbnailsFromPackage( const FString& InPackageFileName, const TSet<FName>& InObjectFullNames, FThumbnailMap& InOutThumbnails );
 
 	/** Loads thumbnails from a package unless they're already cached in that package's thumbnail map */
@@ -730,6 +733,7 @@ namespace ThumbnailTools
 
 	/** Loads thumbnails for the specified objects (or copies them from a cache, if they're already loaded.) */
 	UNREALED_API bool ConditionallyLoadThumbnailsForObjects( const TArray< FName >& InObjectFullNames, FThumbnailMap& InOutThumbnails );
+
 	/** Standard thumbnail height setting used by generation */
 	inline const int32 DefaultThumbnailSize=256;
 
diff --git a/Engine/Source/Editor/UnrealEd/Public/ThumbnailExternalCache.h b/Engine/Source/Editor/UnrealEd/Public/ThumbnailExternalCache.h
--- a/Engine/Source/Editor/UnrealEd/Public/ThumbnailExternalCache.h
+++ b/Engine/Source/Editor/UnrealEd/Public/ThumbnailExternalCache.h
@@ -40,9 +40,6 @@ public:
 	/** Sort asset data list if determinism needed */
 	UNREALED_API static void SortAssetDatas(TArray<FAssetData>& InOutAssetDatas);
 
-	/** Try to load thumbnail for a given asset data */
-	static FObjectThumbnail LoadThumbnailFromPackage(const FAssetData& AssetData);
-
 	enum class EThumbnailExternalCacheHeaderFlags : uint64
 	{
 		None = 0,
