AssetEditorToolkit: Prevent calling OnRequestClose twice when going through FAssetEditorToolkit::CloseWindow.
The issue is that FAssetEditorToolkit::CloseWindow needs to call OnRequestClose to know if it can proceed. If it does proceed, the main tab gets destroyed and will also call OnRequestClose through CanCloseTab to know if it can proceed.

The solution here is to unbind OnRequestClose from the host tab after calling OnRequestClose from FAssetEditorToolkit::CloseWindow. This way the host tab won't call OnRequestClose when closing.

diff --git a/Engine/Source/Editor/UnrealEd/Private/Toolkits/AssetEditorToolkit.cpp b/Engine/Source/Editor/UnrealEd/Private/Toolkits/AssetEditorToolkit.cpp
--- a/Engine/Source/Editor/UnrealEd/Private/Toolkits/AssetEditorToolkit.cpp
+++ b/Engine/Source/Editor/UnrealEd/Private/Toolkits/AssetEditorToolkit.cpp
@@ -502,6 +502,12 @@ bool FAssetEditorToolkit::CloseWindow(EAssetEditorCloseReason InCloseReason)
 {
 	if (OnRequestClose(InCloseReason))
 	{
+		// We are closing, unbind OnRequestClose since we're past that point and we want to make sure we don't redo the request close process when closing the host tab
+		if (TSharedPtr<SStandaloneAssetEditorToolkitHost> StandaloneHostPtr = StandaloneHost.Pin())
+		{
+			StandaloneHostPtr->UnbindEditorCloseRequestFromHostTab();
+		}
+
 		OnClose();
 
 		// Close this toolkit
diff --git a/Engine/Source/Editor/UnrealEd/Private/Toolkits/SStandaloneAssetEditorToolkitHost.cpp b/Engine/Source/Editor/UnrealEd/Private/Toolkits/SStandaloneAssetEditorToolkitHost.cpp
--- a/Engine/Source/Editor/UnrealEd/Private/Toolkits/SStandaloneAssetEditorToolkitHost.cpp
+++ b/Engine/Source/Editor/UnrealEd/Private/Toolkits/SStandaloneAssetEditorToolkitHost.cpp
@@ -173,8 +173,9 @@ void SStandaloneAssetEditorToolkitHost::CreateDefaultStandaloneMenuBar(UToolMenu
 
 void SStandaloneAssetEditorToolkitHost::RestoreFromLayout( const TSharedRef<FTabManager::FLayout>& NewLayout )
 {
+	BindEditorCloseRequestToHostTab();
+
 	const TSharedRef<SDockTab> HostTab = HostTabPtr.Pin().ToSharedRef();
-	HostTab->SetCanCloseTab(EditorCloseRequest);
 	HostTab->SetOnTabClosed(SDockTab::FOnTabClosedCallback::CreateSP(this, &SStandaloneAssetEditorToolkitHost::OnTabClosed));
 
 	this->ChildSlot[SNullWidget::NullWidget];
@@ -277,6 +278,22 @@ FEditorModeTools& SStandaloneAssetEditorToolkitHost::GetEditorModeManager() cons
 	return HostedAssetEditorToolkit->GetEditorModeManager();
 }
 
+void SStandaloneAssetEditorToolkitHost::BindEditorCloseRequestToHostTab()
+{
+	if (TSharedPtr<SDockTab> HostTab = HostTabPtr.Pin())
+	{
+		HostTab->SetCanCloseTab(EditorCloseRequest);
+	}
+}
+
+void SStandaloneAssetEditorToolkitHost::UnbindEditorCloseRequestFromHostTab()
+{
+	if (TSharedPtr<SDockTab> HostTab = HostTabPtr.Pin())
+	{
+		HostTab->SetCanCloseTab(SDockTab::FCanCloseTab());
+	}
+}
+
 SStandaloneAssetEditorToolkitHost::~SStandaloneAssetEditorToolkitHost()
 {
 	ShutdownToolkitHost();
diff --git a/Engine/Source/Editor/UnrealEd/Private/Toolkits/SStandaloneAssetEditorToolkitHost.h b/Engine/Source/Editor/UnrealEd/Private/Toolkits/SStandaloneAssetEditorToolkitHost.h
--- a/Engine/Source/Editor/UnrealEd/Private/Toolkits/SStandaloneAssetEditorToolkitHost.h
+++ b/Engine/Source/Editor/UnrealEd/Private/Toolkits/SStandaloneAssetEditorToolkitHost.h
@@ -102,6 +102,16 @@ public:
 
 	virtual FEditorModeTools& GetEditorModeManager() const override;
 
+	/**
+	 * Binds EditorCloseRequest to the HostTab so that it triggers when the tab is trying to close.
+	 */
+	void BindEditorCloseRequestToHostTab();
+
+	/**
+	 * Unbinds EditorCloseRequest from the HostTab so that it doesn't trigger when the tab is trying to close.
+	 */
+	void UnbindEditorCloseRequestFromHostTab();
+
 private:
 	void OnTabClosed(TSharedRef<SDockTab> TabClosed) const;
 	void ShutdownToolkitHost();
