Use a TArray64 for scope events to handle huge traces

diff --git a/Engine/Source/Editor/UnrealEd/Private/Commandlets/SummarizeTraceCommandlet.cpp b/Engine/Source/Editor/UnrealEd/Private/Commandlets/SummarizeTraceCommandlet.cpp
--- a/Engine/Source/Editor/UnrealEd/Private/Commandlets/SummarizeTraceCommandlet.cpp
+++ b/Engine/Source/Editor/UnrealEd/Private/Commandlets/SummarizeTraceCommandlet.cpp
@@ -95,7 +95,7 @@ public:
 	virtual void OnCpuScopeExit(const FScope& Scope, const FString* ScopeName) {};
 
 	/** Invoked when a root event on the specified thread along with all child events down to the leaves are known. */
-	virtual void OnCpuScopeTree(uint32 ThreadId, const TArray<FSummarizeCpuScopeAnalyzer::FScopeEvent>& ScopeEvents, const TFunction<const FString*(uint32)>& ScopeLookupNameFn) {};
+	virtual void OnCpuScopeTree(uint32 ThreadId, const TArray64<FSummarizeCpuScopeAnalyzer::FScopeEvent>& ScopeEvents, const TFunction<const FString*(uint32)>& ScopeLookupNameFn) {};
 
 	/** Invoked when the trace stream has been fully consumed/processed. */
 	virtual void OnCpuScopeAnalysisEnd() {};
@@ -133,7 +133,7 @@ private:
 	struct FScopeTreeInfo
 	{
 		// Records the current root scope and its children to run analysis that needs to know the parent/child relationship.
-		TArray<FSummarizeCpuScopeAnalyzer::FScopeEvent> ScopeEvents;
+		TArray64<FSummarizeCpuScopeAnalyzer::FScopeEvent> ScopeEvents;
 
 		// Indicates if one of the scope in the current hierarchy is nameless. (Its names specs hasn't been received yet).
 		bool bHasNamelessScopes = false;
@@ -855,7 +855,7 @@ public:
 	 * @param ScopeEvents The scopes events containing one root event along with its hierarchy.
 	 * @param InScopeNameLookup Callback function to lookup scope names from scope ID.
 	 */
-	virtual void OnCpuScopeTree(uint32 ThreadId, const TArray<FSummarizeCpuScopeAnalyzer::FScopeEvent>& ScopeEvents, const TFunction<const FString*(uint32 /*ScopeId*/)>& InScopeNameLookup) override;
+	virtual void OnCpuScopeTree(uint32 ThreadId, const TArray64<FSummarizeCpuScopeAnalyzer::FScopeEvent>& ScopeEvents, const TFunction<const FString*(uint32 /*ScopeId*/)>& InScopeNameLookup) override;
 
 	/**
 	 * Invoked to notify that the trace session ended and that the analyzer can publish the statistics gathered.
@@ -886,7 +886,7 @@ FSummarizeCpuScopeHierarchyAnalyzer::FSummarizeCpuScopeHierarchyAnalyzer(const F
 	MatchedScopesSummary.Name = InAnalyzerName;
 }
 
-void FSummarizeCpuScopeHierarchyAnalyzer::OnCpuScopeTree(uint32 ThreadId, const TArray<FSummarizeCpuScopeAnalyzer::FScopeEvent>& ScopeEvents, const TFunction<const FString*(uint32 /*ScopeId*/)>& InScopeNameLookup)
+void FSummarizeCpuScopeHierarchyAnalyzer::OnCpuScopeTree(uint32 ThreadId, const TArray64<FSummarizeCpuScopeAnalyzer::FScopeEvent>& ScopeEvents, const TFunction<const FString*(uint32 /*ScopeId*/)>& InScopeNameLookup)
 {
 	// Scope matching the pattern.
 	struct FMatchScopeEnter
