Make the "Restore Open Asset Editor Tabs" prompt respect the recent assets filter

diff --git a/Engine/Source/Editor/MainFrame/Private/Frame/MainFrameHandler.cpp b/Engine/Source/Editor/MainFrame/Private/Frame/MainFrameHandler.cpp
--- a/Engine/Source/Editor/MainFrame/Private/Frame/MainFrameHandler.cpp
+++ b/Engine/Source/Editor/MainFrame/Private/Frame/MainFrameHandler.cpp
@@ -7,6 +7,72 @@
 #include "Framework/Application/SlateApplication.h"
 #include "Widgets/Docking/SDockTab.h"
 #include "HAL/PlatformApplicationMisc.h"
+#include "Subsystems/AssetEditorSubsystem.h"
+
+static bool ShowRestoreAssetsPromptOnStartup = true;
+FAutoConsoleVariableRef CVarShowRestoreAssetsPromptOnStartup(TEXT("Mainframe.ShowRestoreAssetsPromptOnStartup"), ShowRestoreAssetsPromptOnStartup, TEXT(""), ECVF_ReadOnly);
+
+void FMainFrameHandler::ShowMainFrameWindow(TSharedRef<SWindow> Window, const bool bStartImmersive, const bool bStartPIE) const
+{
+	// Make sure viewport windows are maximized/immersed if they need to be
+	FLevelEditorModule& LevelEditor = FModuleManager::GetModuleChecked< FLevelEditorModule >( TEXT( "LevelEditor" ) );
+
+	if( bStartPIE )
+	{
+		// Kick off an immersive PIE session immediately!
+
+		if( bStartImmersive )
+		{
+			// When in immersive play in editor, toggle game view on the active viewport
+			const bool bForceGameView = true;
+
+			// Start level viewport initially in immersive mode
+			LevelEditor.GoImmersiveWithActiveLevelViewport( bForceGameView );
+		}
+
+		LevelEditor.StartPlayInEditorSession();
+		Window->ShowWindow();
+		// Ensure the window is at the front or else we could end up capturing and locking the mouse to a window that isn't visible
+		bool bForceWindowToFront = true;
+		Window->BringToFront( bForceWindowToFront );
+
+		// Need to register after the window is shown or else we cant capture the mouse
+		TSharedPtr<IAssetViewport> Viewport = LevelEditor.GetFirstActiveViewport();
+		Viewport->RegisterGameViewportIfPIE();
+	}
+	else
+	{
+		if( bStartImmersive )
+		{
+			// When in immersive play in editor, toggle game view on the active viewport
+			const bool bForceGameView = true;
+
+			// Start level viewport initially in immersive mode
+			LevelEditor.GoImmersiveWithActiveLevelViewport( bForceGameView );
+		}
+
+		// Show the window!
+		Window->ShowWindow();
+
+		if( bStartImmersive )
+		{
+			// Ensure the window is at the front or else we could end up capturing and locking the mouse to a window that isn't visible
+			bool bForceWindowToFront = true;
+			Window->BringToFront( bForceWindowToFront );
+		}
+		else
+		{
+			// Focus the level editor viewport
+			LevelEditor.FocusViewport();
+
+			if(ShowRestoreAssetsPromptOnStartup)
+			{
+				// Restore any assets we had open. Note we don't do this on immersive PIE as its annoying to the user.
+				GEditor->GetEditorSubsystem<UAssetEditorSubsystem>()->RequestRestorePreviouslyOpenAssets();
+			}
+		}
+	}
+}
 
 void FMainFrameHandler::ShutDownEditor()
 {
diff --git a/Engine/Source/Editor/MainFrame/Private/Frame/MainFrameHandler.h b/Engine/Source/Editor/MainFrame/Private/Frame/MainFrameHandler.h
--- a/Engine/Source/Editor/MainFrame/Private/Frame/MainFrameHandler.h
+++ b/Engine/Source/Editor/MainFrame/Private/Frame/MainFrameHandler.h
@@ -211,64 +211,7 @@ public:
 	 * @param bStartImmersive True to force a main frame viewport into immersive mode
 	 * @param bStartPIE True to start a PIE session right away
 	 */
-	void ShowMainFrameWindow(TSharedRef<SWindow> Window, const bool bStartImmersive, const bool bStartPIE) const
-	{
-		// Make sure viewport windows are maximized/immersed if they need to be
-		FLevelEditorModule& LevelEditor = FModuleManager::GetModuleChecked< FLevelEditorModule >( TEXT( "LevelEditor" ) );
-
-		if( bStartPIE )
-		{
-			// Kick off an immersive PIE session immediately!
-
-			if( bStartImmersive )
-			{
-				// When in immersive play in editor, toggle game view on the active viewport
-				const bool bForceGameView = true;
-
-				// Start level viewport initially in immersive mode
-				LevelEditor.GoImmersiveWithActiveLevelViewport( bForceGameView );
-			}
-
-			LevelEditor.StartPlayInEditorSession();
-			Window->ShowWindow();
-			// Ensure the window is at the front or else we could end up capturing and locking the mouse to a window that isn't visible
-			bool bForceWindowToFront = true;
-			Window->BringToFront( bForceWindowToFront );
-
-			// Need to register after the window is shown or else we cant capture the mouse
-			TSharedPtr<IAssetViewport> Viewport = LevelEditor.GetFirstActiveViewport();
-			Viewport->RegisterGameViewportIfPIE();
-		}
-		else
-		{
-			if( bStartImmersive )
-			{
-				// When in immersive play in editor, toggle game view on the active viewport
-				const bool bForceGameView = true;
-
-				// Start level viewport initially in immersive mode
-				LevelEditor.GoImmersiveWithActiveLevelViewport( bForceGameView );
-			}
-
-			// Show the window!
-			Window->ShowWindow();
-
-			if( bStartImmersive )
-			{
-				// Ensure the window is at the front or else we could end up capturing and locking the mouse to a window that isn't visible
-				bool bForceWindowToFront = true;
-				Window->BringToFront( bForceWindowToFront );
-			}
-			else
-			{
-				// Focus the level editor viewport
-				LevelEditor.FocusViewport();
-
-				// Restore any assets we had open. Note we don't do this on immersive PIE as its annoying to the user.
-				GEditor->GetEditorSubsystem<UAssetEditorSubsystem>()->RequestRestorePreviouslyOpenAssets();
-			}
-		}
-	}
+	void ShowMainFrameWindow(TSharedRef<SWindow> Window, const bool bStartImmersive, const bool bStartPIE) const;
 
 	/** Gets the parent window of the mainframe */
 	TSharedPtr<SWindow> GetParentWindow() const
diff --git a/Engine/Source/Editor/UnrealEd/Classes/Settings/ContentBrowserSettings.h b/Engine/Source/Editor/UnrealEd/Classes/Settings/ContentBrowserSettings.h
--- a/Engine/Source/Editor/UnrealEd/Classes/Settings/ContentBrowserSettings.h
+++ b/Engine/Source/Editor/UnrealEd/Classes/Settings/ContentBrowserSettings.h
@@ -58,7 +58,7 @@ public:
 	bool bDisplayFriendlyNameForPluginFolders = true;
 
 	/** The number of objects to keep in the Content Browser Recently Opened filter */
-	UPROPERTY(meta = (DeprecatedProperty, DeprecationMessage = "Use NumRecentAssets in UEditorLoadingSavingSettings instead"))
+	UPROPERTY(meta = (DeprecatedProperty, DeprecationMessage = "The filter now always keeps track of your last 30 recent assets"))
 	int32 NumObjectsInRecentList_DEPRECATED;
 
 	/** Enables the rendering of Material Instance thumbnail previews */
diff --git a/Engine/Source/Editor/UnrealEd/Classes/Settings/EditorLoadingSavingSettings.h b/Engine/Source/Editor/UnrealEd/Classes/Settings/EditorLoadingSavingSettings.h
--- a/Engine/Source/Editor/UnrealEd/Classes/Settings/EditorLoadingSavingSettings.h
+++ b/Engine/Source/Editor/UnrealEd/Classes/Settings/EditorLoadingSavingSettings.h
@@ -117,10 +117,6 @@ public:
 	uint32 bRestoreOpenAssetTabsOnRestart_DEPRECATED:1;
 #endif
 
-	/** The number of recent assets the editor will keep track of */
-	UPROPERTY(EditAnywhere, config, Category = Startup, meta = (DisplayName = "Number of Recent Assets to keep track of", ClampMin = "1", ClampMax = "30"))
-	int32 NumRecentAssets = 20;
-	
 private:
 
 	UPROPERTY(config)
diff --git a/Engine/Source/Editor/UnrealEd/Private/Subsystems/AssetEditorSubsystem.cpp b/Engine/Source/Editor/UnrealEd/Private/Subsystems/AssetEditorSubsystem.cpp
--- a/Engine/Source/Editor/UnrealEd/Private/Subsystems/AssetEditorSubsystem.cpp
+++ b/Engine/Source/Editor/UnrealEd/Private/Subsystems/AssetEditorSubsystem.cpp
@@ -110,8 +110,8 @@ void UAssetEditorSubsystem::Deinitialize()
 
 void UAssetEditorSubsystem::InitializeRecentAssets()
 {
-	// The current max allowed is 30 assets, we keep track of 40 to be safe since we are ignoring levels when showing the list
-	RecentAssetsList = MakeUnique<FMainMRUFavoritesList>(TEXT("AssetEditorSubsystemRecents"), 40);
+	// The current max allowed is 30 assets
+	RecentAssetsList = MakeUnique<FMainMRUFavoritesList>(TEXT("AssetEditorSubsystemRecents"), 30);
 	RecentAssetsList->ReadFromINI();
 	
 	TArray<FString> RecentAssetEditors;
@@ -389,7 +389,10 @@ void UAssetEditorSubsystem::NotifyAssetOpened(UObject* Asset, IAssetEditorInstan
 
 	AssetOpenedInEditorEvent.Broadcast(Asset, InInstance);
 
-	SaveOpenAssetEditors(false);
+	if(InInstance->IncludeAssetInRestoreOpenAssetsPrompt())
+	{
+		SaveOpenAssetEditors(false);
+	}
 }
 
 
@@ -927,8 +930,25 @@ FOnModeUnregistered& UAssetEditorSubsystem::OnEditorModeUnregistered()
 
 void UAssetEditorSubsystem::RestorePreviouslyOpenAssets()
 {
-	TArray<FString> OpenAssets;
-	GConfig->GetArray(TEXT("AssetEditorSubsystem"), TEXT("OpenAssetsAtExit"), OpenAssets, GEditorPerProjectIni);
+	TArray<FString> AllOpenAssets;
+	TArray<FString> FilteredOpenAssets;
+	
+	GConfig->GetArray(TEXT("AssetEditorSubsystem"), TEXT("OpenAssetsAtExit"), AllOpenAssets, GEditorPerProjectIni);
+
+	if(!RecentAssetsFilter.IsBound())
+	{
+		FilteredOpenAssets = AllOpenAssets;
+	}
+	else
+	{
+		for(const FString& Asset : AllOpenAssets)
+		{
+			if(RecentAssetsFilter.Execute(Asset))
+			{
+				FilteredOpenAssets.Add(Asset);
+			}
+		}
+	}
 
 	bool bCleanShutdown = true;
 	GConfig->GetBool(TEXT("AssetEditorSubsystem"), TEXT("CleanShutdown"), bCleanShutdown, GEditorPerProjectIni);
@@ -944,7 +964,7 @@ void UAssetEditorSubsystem::RestorePreviouslyOpenAssets()
 	 */
 	bool bCrashedWithoutDebugger = !bDebuggerAttachedLastSession && !bCleanShutdown;
 
-	if (OpenAssets.Num() > 0)
+	if (FilteredOpenAssets.Num() > 0)
 	{
 		// This option overrides the saved setting
 		if(bAutoRestoreAndDisableSaving.IsSet())
@@ -952,7 +972,7 @@ void UAssetEditorSubsystem::RestorePreviouslyOpenAssets()
 			// If bAutoRestoreAndDisableSaving is true, we automatically restore the opened assets
 			if(bAutoRestoreAndDisableSaving.GetValue())
 			{
-				OpenEditorsForAssets(OpenAssets);
+				OpenEditorsForAssets(FilteredOpenAssets);
 			}
 			return;
 		}
@@ -962,7 +982,7 @@ void UAssetEditorSubsystem::RestorePreviouslyOpenAssets()
 		 */
 		if(bCrashedWithoutDebugger)
 		{
-			SpawnRestorePreviouslyOpenAssetsNotification(!bCrashedWithoutDebugger, OpenAssets);
+			SpawnRestorePreviouslyOpenAssetsNotification(!bCrashedWithoutDebugger, FilteredOpenAssets);
 			return;
 		}
 
@@ -972,7 +992,7 @@ void UAssetEditorSubsystem::RestorePreviouslyOpenAssets()
 		{
 		case ERestoreOpenAssetTabsMethod::AlwaysPrompt:
 			{
-				SpawnRestorePreviouslyOpenAssetsNotification(!bCrashedWithoutDebugger, OpenAssets);
+				SpawnRestorePreviouslyOpenAssetsNotification(!bCrashedWithoutDebugger, FilteredOpenAssets);
 				break;
 			}
 			
@@ -982,7 +1002,7 @@ void UAssetEditorSubsystem::RestorePreviouslyOpenAssets()
 		case ERestoreOpenAssetTabsMethod::AlwaysRestore:
 			{
 				// Pretend that we showed the notification and that the user clicked "Restore Now"
-				OpenEditorsForAssets(OpenAssets);
+				OpenEditorsForAssets(FilteredOpenAssets);
 				break;
 			}
 		}
@@ -1009,6 +1029,16 @@ TOptional<bool> UAssetEditorSubsystem::GetAutoRestoreAndDisableSavingOverride()
 	return bAutoRestoreAndDisableSaving;
 }
 
+void UAssetEditorSubsystem::SetRecentAssetsFilter(const FMainMRUFavoritesList::FDoesMRUFavoritesItemPassFilter& InFilter)
+{
+	RecentAssetsFilter = InFilter;
+
+	if(RecentAssetsList)
+	{
+		RecentAssetsList->RegisterDoesMRUFavoritesItemPassFilterDelegate(InFilter);
+	}
+}
+
 void UAssetEditorSubsystem::SpawnRestorePreviouslyOpenAssetsNotification(const bool bCleanShutdown, const TArray<FString>& AssetsToOpen)
 {
 	FText NotificationMessage = bCleanShutdown
@@ -1146,9 +1176,7 @@ bool UAssetEditorSubsystem::ShouldShowRecentAssetsMenu(const FName& InAssetEdito
 		return false;
 	}
 
-	const int32 MaxRecentAssets = GetDefault<UEditorLoadingSavingSettings>()->NumRecentAssets;
-
-	for ( int32 CurRecentIndex = 0; CurRecentIndex < RecentAssetsList->GetNumItems() && CurRecentIndex < MaxRecentAssets; ++CurRecentIndex )
+	for ( int32 CurRecentIndex = 0; CurRecentIndex < RecentAssetsList->GetNumItems() && CurRecentIndex < MaxRecentAssetsToShowInMenu; ++CurRecentIndex )
 	{
 		const FString& CurRecent = RecentAssetsList->GetMRUItem(CurRecentIndex);
 
@@ -1164,12 +1192,10 @@ bool UAssetEditorSubsystem::ShouldShowRecentAssetsMenu(const FName& InAssetEdito
 
 void UAssetEditorSubsystem::CreateRecentAssetsMenu(UToolMenu* InMenu, const FName InAssetEditorName)
 {
-	const int32 MaxRecentAssets = GetDefault<UEditorLoadingSavingSettings>()->NumRecentAssets;
-
 	FToolMenuSection& Section = InMenu->FindOrAddSection("Recents");
 	
 	// Keep adding assets until we reach the end of the MRU list, or we reach the max allowed assets
-	for ( int32 CurRecentIndex = 0; CurRecentIndex < RecentAssetsList->GetNumItems() && CurRecentIndex < MaxRecentAssets; ++CurRecentIndex )
+	for ( int32 CurRecentIndex = 0; CurRecentIndex < RecentAssetsList->GetNumItems() && CurRecentIndex < MaxRecentAssetsToShowInMenu; ++CurRecentIndex )
 	{
 		const FString& CurRecent = RecentAssetsList->GetMRUItem(CurRecentIndex);
 
@@ -1215,7 +1241,7 @@ void UAssetEditorSubsystem::RegisterLevelEditorMenuExtensions()
 		InSection.AddSubMenu(
 			"RecentAssetsSubmenu",
 			LOCTEXT("RecentAssetsSubmenu_Label", "Recent Assets"),
-			LOCTEXT("RecentAssetsSubMenu_ToolTip", "Access your recently opened assets"),
+			FText::Format(LOCTEXT("RecentAssetsSubMenu_ToolTip", "Access your last {0} recently opened assets"), MaxRecentAssetsToShowInMenu),
 			FNewToolMenuDelegate::CreateUObject(this, &UAssetEditorSubsystem::CreateRecentAssetsMenu, AssetEditorName),
 			false,
 			FSlateIcon(FAppStyle::GetAppStyleSetName(), "Icons.RecentAssets"));
@@ -1247,7 +1273,7 @@ void UAssetEditorSubsystem::CreateRecentAssetsMenuForEditor(const IAssetEditorIn
 			InSection.AddSubMenu(
 			"RecentAssetsSubmenu",
 			LOCTEXT("RecentAssetsSubmenu_Label", "Recent Assets"),
-			LOCTEXT("RecentAssetsSubMenu_ToolTip", "Access your recently opened assets"),
+			FText::Format(LOCTEXT("RecentAssetsSubMenu_ToolTip", "Access your last {0} recently opened assets"), MaxRecentAssetsToShowInMenu),
 			FNewToolMenuDelegate::CreateUObject(this, &UAssetEditorSubsystem::CreateRecentAssetsMenu, AssetEditorName),
 			false,
 			FSlateIcon(FAppStyle::GetAppStyleSetName(), "Icons.RecentAssets")
@@ -1262,7 +1288,7 @@ void UAssetEditorSubsystem::CreateRecentAssetsMenuForEditor(const IAssetEditorIn
 			InSection.AddSubMenu(
 				"RecentAssetEditorAssetsSubmenu",
 				FText::Format(LOCTEXT("RecentAssetEditorAssetsSubmenu_Label", "{0}"), FText::FromName(RecentAssetsMenuName)),
-				LOCTEXT("RecentAssetEditorAssetsSubMenu_ToolTip", "Access your recently opened assets"),
+				FText::Format(LOCTEXT("RecentAssetEditorAssetsSubmenu_Label", "Access your recently opened {0} assets"), FText::FromName(EditingAssetTypeName)),
 				FNewToolMenuDelegate::CreateUObject(this, &UAssetEditorSubsystem::CreateRecentAssetsMenu, AssetEditorName),
 				false,
 				FSlateIcon(FAppStyle::GetAppStyleSetName(), "Icons.RecentAssets")
@@ -1277,7 +1303,8 @@ void UAssetEditorSubsystem::SaveOpenAssetEditors(const bool bOnShutdown)
 {
 	SaveRecentAssets(bOnShutdown);
 
-	// Saving is disabled if bAutoRestoreAndDisableSaving is set
+	// We are already saving the open asset editors manually before bSavingOnShutdown is true. This is to avoid saving that there are no open asset editors b/c the editor is shutting down
+	// If we are restoring the layout, bAutoRestoreAndDisable saving is true
 	if (!bSavingOnShutdown && !bAutoRestoreAndDisableSaving.IsSet())
 	{
 		TArray<FString> OpenAssets;
@@ -1285,14 +1312,14 @@ void UAssetEditorSubsystem::SaveOpenAssetEditors(const bool bOnShutdown)
 		for (const TPair<IAssetEditorInstance*, FAssetEntry>& EditorPair : OpenedEditors)
 		{
 			IAssetEditorInstance* Editor = EditorPair.Key;
-			if (Editor != nullptr)
+			if (Editor != nullptr && Editor->IncludeAssetInRestoreOpenAssetsPrompt())
 			{
 				UObject* EditedObject = EditorPair.Value.ObjectPtr.Get();
 				if (EditedObject != nullptr)
 				{
 					// only record assets that have a valid saved package
 					UPackage* Package = EditedObject->GetOutermost();
-					if (Package != nullptr && Package->GetFileSize() != 0)
+					if (Package != nullptr && Package->GetFileSize() != 0 )
 					{
 						OpenAssets.Add(EditedObject->GetPathName());
 					}
diff --git a/Engine/Source/Editor/UnrealEd/Public/Subsystems/AssetEditorSubsystem.h b/Engine/Source/Editor/UnrealEd/Public/Subsystems/AssetEditorSubsystem.h
--- a/Engine/Source/Editor/UnrealEd/Public/Subsystems/AssetEditorSubsystem.h
+++ b/Engine/Source/Editor/UnrealEd/Public/Subsystems/AssetEditorSubsystem.h
@@ -67,6 +67,8 @@ public:
 		PRAGMA_ENABLE_DEPRECATION_WARNINGS
 	}
 
+	/** If false, the asset being edited will not be included in reopen assets prompt on restart */
+	virtual bool IncludeAssetInRestoreOpenAssetsPrompt() const { return true; }
 	virtual bool IsPrimaryEditor() const = 0;
 	virtual void InvokeTab(const struct FTabId& TabId) = 0;
 	UE_DEPRECATED(5.0, "Toolbar tab no longer exists and tab ID will return None; do not add it to layouts")
@@ -307,6 +309,8 @@ public:
 	/** Get the current value of bAutoRestoreAndDisableSaving */
 	TOptional<bool> GetAutoRestoreAndDisableSavingOverride() const;
 
+	void SetRecentAssetsFilter(const FMainMRUFavoritesList::FDoesMRUFavoritesItemPassFilter& InFilter);
+
 private:
 
 	/** Handles a package being reloaded */
@@ -453,4 +457,10 @@ private:
 	
 	/** Map keeping track of the asset editor each recent asset was opened in */
 	TMap<FString, FString> RecentAssetToAssetEditorMap;
+
+	/** The filter run through any assets before they are shown in the recents menu or restore prompt */
+	FMainMRUFavoritesList::FDoesMRUFavoritesItemPassFilter RecentAssetsFilter;
+
+	/** The max number of recent assets to show in the menu */
+	int32 MaxRecentAssetsToShowInMenu = 20;
 };
