Cooker: Update -cookfirst to include cooking the generated packages if the commandline packages are package generators.
#rnx

diff --git a/Engine/Source/Editor/UnrealEd/Classes/CookOnTheSide/CookOnTheFlyServer.h b/Engine/Source/Editor/UnrealEd/Classes/CookOnTheSide/CookOnTheFlyServer.h
--- a/Engine/Source/Editor/UnrealEd/Classes/CookOnTheSide/CookOnTheFlyServer.h
+++ b/Engine/Source/Editor/UnrealEd/Classes/CookOnTheSide/CookOnTheFlyServer.h
@@ -526,9 +526,9 @@ private:
 	 * @param PackageData			The PackageData to be considered for saving.
 	 */
 	void QueueDiscoveredPackage(UE::Cook::FPackageData& PackageData, UE::Cook::FInstigator&& Instigator, 
-		UE::Cook::FDiscoveredPlatformSet&& ReachablePlatforms);
+		UE::Cook::FDiscoveredPlatformSet&& ReachablePlatforms, bool bUrgent=false);
 	void QueueDiscoveredPackageOnDirector(UE::Cook::FPackageData& PackageData, UE::Cook::FInstigator&& Instigator,
-		UE::Cook::FDiscoveredPlatformSet&& ReachablePlatforms);
+		UE::Cook::FDiscoveredPlatformSet&& ReachablePlatforms, bool bUrgent);
 
 	/** Called when a package is cancelled and returned to idle. Notifies CookDirector when on a CookWorker. */
 	void DemoteToIdle(UE::Cook::FPackageData& PackageData, UE::Cook::ESendFlags SendFlags, UE::Cook::ESuppressCookReason Reason);
diff --git a/Engine/Source/Editor/UnrealEd/Private/CookOnTheFlyServer.cpp b/Engine/Source/Editor/UnrealEd/Private/CookOnTheFlyServer.cpp
--- a/Engine/Source/Editor/UnrealEd/Private/CookOnTheFlyServer.cpp
+++ b/Engine/Source/Editor/UnrealEd/Private/CookOnTheFlyServer.cpp
@@ -2874,7 +2874,7 @@ void UCookOnTheFlyServer::RejectPackageToLoad(UE::Cook::FPackageData& PackageDat
 }
 
 void UCookOnTheFlyServer::QueueDiscoveredPackage(UE::Cook::FPackageData& PackageData,
-	UE::Cook::FInstigator&& Instigator, UE::Cook::FDiscoveredPlatformSet&& ReachablePlatforms)
+	UE::Cook::FInstigator&& Instigator, UE::Cook::FDiscoveredPlatformSet&& ReachablePlatforms, bool bUrgent)
 {
 	using namespace UE::Cook;
 
@@ -2901,11 +2901,11 @@ void UCookOnTheFlyServer::QueueDiscoveredPackage(UE::Cook::FPackageData& Package
 	{
 		OnDiscoveredPackageDebug(PackageData.GetPackageName(), Instigator);
 	}
-	WorkerRequests->QueueDiscoveredPackage(*this, PackageData, MoveTemp(Instigator), MoveTemp(ReachablePlatforms));
+	WorkerRequests->QueueDiscoveredPackage(*this, PackageData, MoveTemp(Instigator), MoveTemp(ReachablePlatforms), bUrgent);
 }
 
 void UCookOnTheFlyServer::QueueDiscoveredPackageOnDirector(UE::Cook::FPackageData& PackageData,
-	UE::Cook::FInstigator&& Instigator, UE::Cook::FDiscoveredPlatformSet&& ReachablePlatforms)
+	UE::Cook::FInstigator&& Instigator, UE::Cook::FDiscoveredPlatformSet&& ReachablePlatforms, bool bUrgent)
 {
 	using namespace UE::Cook;
 
@@ -2924,7 +2924,7 @@ void UCookOnTheFlyServer::QueueDiscoveredPackageOnDirector(UE::Cook::FPackageDat
 	if (!CookByTheBookOptions->bSkipHardReferences ||
 		(Instigator.Category == EInstigator::GeneratedPackage))
 	{
-		PackageData.QueueAsDiscovered(MoveTemp(Instigator), MoveTemp(ReachablePlatforms));
+		PackageData.QueueAsDiscovered(MoveTemp(Instigator), MoveTemp(ReachablePlatforms), bUrgent);
 	}
 }
 
@@ -3135,8 +3135,13 @@ UE::Cook::EPollStatus UCookOnTheFlyServer::QueueGeneratedPackages(UE::Cook::FGen
 			FPackageData* ChildPackageData = ChildInfo.PackageData;
 			// Set the Instigator now rather than delaying it until the discovery queue is processed.
 			ChildPackageData->SetInstigator(Generator, FInstigator(EInstigator::GeneratedPackage, OwnerName));
+			// When running -cookfirst, generated packages should also be cooked first, so set the urgency of the
+			// generated packages to match the urgency of the generator
+			bool bUrgent = PackageData.GetIsUrgent();
+
+			// Queue the package for cooking
 			QueueDiscoveredPackage(*ChildPackageData, FInstigator(ChildPackageData->GetInstigator()),
-				EDiscoveredPlatformSet::CopyFromInstigator);
+				EDiscoveredPlatformSet::CopyFromInstigator, bUrgent);
 		}
 		Info.SetSaveStateComplete(FCookGenerationInfo::ESaveState::QueueGeneratedPackages);
 	}
diff --git a/Engine/Source/Editor/UnrealEd/Private/Cooker/CookPackageData.cpp b/Engine/Source/Editor/UnrealEd/Private/Cooker/CookPackageData.cpp
--- a/Engine/Source/Editor/UnrealEd/Private/Cooker/CookPackageData.cpp
+++ b/Engine/Source/Editor/UnrealEd/Private/Cooker/CookPackageData.cpp
@@ -221,17 +221,17 @@ void FPackageData::AddReachablePlatformsInternal(FPackageData& PackageData,
 	}
 }
 
-void FPackageData::QueueAsDiscovered(FInstigator&& InInstigator, FDiscoveredPlatformSet&& ReachablePlatforms)
+void FPackageData::QueueAsDiscovered(FInstigator&& InInstigator, FDiscoveredPlatformSet&& ReachablePlatforms, bool bUrgent)
 {
-	QueueAsDiscoveredInternal(*this, MoveTemp(InInstigator), MoveTemp(ReachablePlatforms));
+	QueueAsDiscoveredInternal(*this, MoveTemp(InInstigator), MoveTemp(ReachablePlatforms), bUrgent);
 }
 
 void FPackageData::QueueAsDiscoveredInternal(FPackageData& PackageData, FInstigator&& InInstigator,
-	FDiscoveredPlatformSet&& ReachablePlatforms)
+	FDiscoveredPlatformSet&& ReachablePlatforms, bool bUrgent)
 {
 	// This is a static helper function to make it impossible to make a typo and use this->Instigator instead of InInstigator
 	TRingBuffer<FDiscoveryQueueElement>& Queue = PackageData.PackageDatas.GetRequestQueue().GetDiscoveryQueue();
-	Queue.Add(FDiscoveryQueueElement{ &PackageData, MoveTemp(InInstigator), MoveTemp(ReachablePlatforms) });
+	Queue.Add(FDiscoveryQueueElement{ &PackageData, MoveTemp(InInstigator), MoveTemp(ReachablePlatforms), bUrgent });
 }
 
 void FPackageData::SetIsUrgent(bool Value)
diff --git a/Engine/Source/Editor/UnrealEd/Private/Cooker/CookPackageData.h b/Engine/Source/Editor/UnrealEd/Private/Cooker/CookPackageData.h
--- a/Engine/Source/Editor/UnrealEd/Private/Cooker/CookPackageData.h
+++ b/Engine/Source/Editor/UnrealEd/Private/Cooker/CookPackageData.h
@@ -279,7 +279,7 @@ public:
 		FInstigator&& InInstigator);
 
 	/** Add the given reachable platforms to this PackageData and send it back to Request state for exploration. */
-	void QueueAsDiscovered(FInstigator&& InInstigator, FDiscoveredPlatformSet&& ReachablePlatforms);
+	void QueueAsDiscovered(FInstigator&& InInstigator, FDiscoveredPlatformSet&& ReachablePlatforms, bool bUrgent);
 
 	/**
 	 * Clear all the inprogress variables from the current PackageData. It is invalid to call this except when
@@ -566,8 +566,10 @@ private:
 
 	void SetIsUrgent(bool Value);
 	void SetInstigatorInternal(FInstigator&& InInstigator);
-	static void AddReachablePlatformsInternal(FPackageData& PackageData, TConstArrayView<const ITargetPlatform*> Platforms, FInstigator&& InInstigator);
-	static void QueueAsDiscoveredInternal(FPackageData& PackageData, FInstigator&& InInstigator, FDiscoveredPlatformSet&& ReachablePlatforms);
+	static void AddReachablePlatformsInternal(FPackageData& PackageData, TConstArrayView<const ITargetPlatform*> Platforms,
+		FInstigator&& InInstigator);
+	static void QueueAsDiscoveredInternal(FPackageData& PackageData, FInstigator&& InInstigator,
+		FDiscoveredPlatformSet&& ReachablePlatforms, bool bUrgent);
 
 	/**
 	 * Set the FileName of the file that contains the package. This member is private because FPackageDatas
@@ -1032,6 +1034,7 @@ struct FDiscoveryQueueElement
 	FPackageData* PackageData;
 	FInstigator Instigator;
 	FDiscoveredPlatformSet ReachablePlatforms;
+	bool bUrgent;
 };
 
 /**
diff --git a/Engine/Source/Editor/UnrealEd/Private/Cooker/CookRequestCluster.cpp b/Engine/Source/Editor/UnrealEd/Private/Cooker/CookRequestCluster.cpp
--- a/Engine/Source/Editor/UnrealEd/Private/Cooker/CookRequestCluster.cpp
+++ b/Engine/Source/Editor/UnrealEd/Private/Cooker/CookRequestCluster.cpp
@@ -192,6 +192,7 @@ FRequestCluster::FRequestCluster(UCookOnTheFlyServer& InCOTFS, TRingBuffer<FDisc
 		// Send it to the Request state if it's not already there, remove it from its old container
 		// and add it to this cluster.
 		PackageData.SendToState(EPackageState::Request, ESendFlags::QueueRemove);
+		PackageData.AddUrgency(Discovery->bUrgent, false /* bAllowUpdateState */);
 		OwnedPackageDatas.Add(&PackageData, ESuppressCookReason::NotSuppressed);
 	}
 }
diff --git a/Engine/Source/Editor/UnrealEd/Private/Cooker/CookWorkerServer.cpp b/Engine/Source/Editor/UnrealEd/Private/Cooker/CookWorkerServer.cpp
--- a/Engine/Source/Editor/UnrealEd/Private/Cooker/CookWorkerServer.cpp
+++ b/Engine/Source/Editor/UnrealEd/Private/Cooker/CookWorkerServer.cpp
@@ -883,7 +883,7 @@ void FCookWorkerServer::QueueDiscoveredPackage(FDiscoveredPackageReplication&& D
 	}
 	Director.ResetFinalIdleHeartbeatFence();
 	Platforms.ConvertFromBitfield(OrderedSessionAndSpecialPlatforms);
-	COTFS.QueueDiscoveredPackageOnDirector(PackageData, MoveTemp(Instigator), MoveTemp(Platforms));
+	COTFS.QueueDiscoveredPackageOnDirector(PackageData, MoveTemp(Instigator), MoveTemp(Platforms), false /* bUrgent */);
 }
 
 FCookWorkerServer::FTickState::FTickState()
diff --git a/Engine/Source/Editor/UnrealEd/Private/Cooker/IWorkerRequests.h b/Engine/Source/Editor/UnrealEd/Private/Cooker/IWorkerRequests.h
--- a/Engine/Source/Editor/UnrealEd/Private/Cooker/IWorkerRequests.h
+++ b/Engine/Source/Editor/UnrealEd/Private/Cooker/IWorkerRequests.h
@@ -50,7 +50,7 @@ public:
 
 	// Writing Packages during Cook
 	virtual void QueueDiscoveredPackage(UCookOnTheFlyServer& COTFS, FPackageData& PackageData,
-		FInstigator&& Instigator, FDiscoveredPlatformSet&& ReachablePlatforms) = 0;
+		FInstigator&& Instigator, FDiscoveredPlatformSet&& ReachablePlatforms, bool bUrgent) = 0;
 
 	// Writing Packages from CookByTheBook
 	virtual void AddStartCookByTheBookRequest(FFilePlatformRequest&& Request) = 0;
diff --git a/Engine/Source/Editor/UnrealEd/Private/Cooker/WorkerRequestsLocal.cpp b/Engine/Source/Editor/UnrealEd/Private/Cooker/WorkerRequestsLocal.cpp
--- a/Engine/Source/Editor/UnrealEd/Private/Cooker/WorkerRequestsLocal.cpp
+++ b/Engine/Source/Editor/UnrealEd/Private/Cooker/WorkerRequestsLocal.cpp
@@ -43,9 +43,9 @@ void FWorkerRequestsLocal::DequeueAllExternal(TArray<FSchedulerCallback>& OutCal
 }
 
 void FWorkerRequestsLocal::QueueDiscoveredPackage(UCookOnTheFlyServer& COTFS, FPackageData& PackageData,
-	FInstigator&& Instigator, FDiscoveredPlatformSet&& ReachablePlatforms)
+	FInstigator&& Instigator, FDiscoveredPlatformSet&& ReachablePlatforms, bool bUrgent)
 {
-	COTFS.QueueDiscoveredPackageOnDirector(PackageData, MoveTemp(Instigator), MoveTemp(ReachablePlatforms));
+	COTFS.QueueDiscoveredPackageOnDirector(PackageData, MoveTemp(Instigator), MoveTemp(ReachablePlatforms), bUrgent);
 }
 
 void FWorkerRequestsLocal::AddStartCookByTheBookRequest(FFilePlatformRequest&& Request)
diff --git a/Engine/Source/Editor/UnrealEd/Private/Cooker/WorkerRequestsLocal.h b/Engine/Source/Editor/UnrealEd/Private/Cooker/WorkerRequestsLocal.h
--- a/Engine/Source/Editor/UnrealEd/Private/Cooker/WorkerRequestsLocal.h
+++ b/Engine/Source/Editor/UnrealEd/Private/Cooker/WorkerRequestsLocal.h
@@ -33,7 +33,7 @@ public:
 	virtual void RemapTargetPlatforms(const TMap<ITargetPlatform*, ITargetPlatform*>& Remap) override;
 	virtual void OnRemoveSessionPlatform(const ITargetPlatform* TargetPlatform) override;
 	virtual void QueueDiscoveredPackage(UCookOnTheFlyServer& COTFS, FPackageData& PackageData,
-		FInstigator&& Instigator, FDiscoveredPlatformSet&& ReachablePlatforms) override;
+		FInstigator&& Instigator, FDiscoveredPlatformSet&& ReachablePlatforms, bool bUrgent) override;
 	virtual void AddStartCookByTheBookRequest(FFilePlatformRequest&& Request) override;
 	virtual void InitializeCookOnTheFly() override;
 	virtual void AddCookOnTheFlyRequest(FFilePlatformRequest&& Request) override;
diff --git a/Engine/Source/Editor/UnrealEd/Private/Cooker/WorkerRequestsRemote.cpp b/Engine/Source/Editor/UnrealEd/Private/Cooker/WorkerRequestsRemote.cpp
--- a/Engine/Source/Editor/UnrealEd/Private/Cooker/WorkerRequestsRemote.cpp
+++ b/Engine/Source/Editor/UnrealEd/Private/Cooker/WorkerRequestsRemote.cpp
@@ -54,8 +54,9 @@ void FWorkerRequestsRemote::DequeueAllExternal(TArray<FSchedulerCallback>& OutCa
 }
 
 void FWorkerRequestsRemote::QueueDiscoveredPackage(UCookOnTheFlyServer& COTFS, FPackageData& PackageData,
-	FInstigator&& Instigator, FDiscoveredPlatformSet&& ReachablePlatforms)
+	FInstigator&& Instigator, FDiscoveredPlatformSet&& ReachablePlatforms, bool bUrgent)
 {
+	(void)bUrgent; // Tracking urgency on CookWorkers is not supported
 	CookWorkerClient.ReportDiscoveredPackage(PackageData, MoveTemp(Instigator), MoveTemp(ReachablePlatforms));
 }
 
diff --git a/Engine/Source/Editor/UnrealEd/Private/Cooker/WorkerRequestsRemote.h b/Engine/Source/Editor/UnrealEd/Private/Cooker/WorkerRequestsRemote.h
--- a/Engine/Source/Editor/UnrealEd/Private/Cooker/WorkerRequestsRemote.h
+++ b/Engine/Source/Editor/UnrealEd/Private/Cooker/WorkerRequestsRemote.h
@@ -38,7 +38,7 @@ public:
 	virtual void RemapTargetPlatforms(const TMap<ITargetPlatform*, ITargetPlatform*>& Remap) override;
 	virtual void OnRemoveSessionPlatform(const ITargetPlatform* TargetPlatform) override;
 	virtual void QueueDiscoveredPackage(UCookOnTheFlyServer& COTFS, FPackageData& PackageData,
-		FInstigator&& Instigator, FDiscoveredPlatformSet&& ReachablePlatforms) override;
+		FInstigator&& Instigator, FDiscoveredPlatformSet&& ReachablePlatforms, bool bUrgent) override;
 	virtual void AddStartCookByTheBookRequest(FFilePlatformRequest&& Request) override;
 	virtual void InitializeCookOnTheFly() override;
 	virtual void AddCookOnTheFlyRequest(FFilePlatformRequest&& Request) override;
