MPCook: Downgrade followup log messages for cookworker crashes from Error to Warning, to avoid making them seem like separate issues.
#rnx

diff --git a/Engine/Source/Editor/UnrealEd/Private/Cooker/CookDirector.cpp b/Engine/Source/Editor/UnrealEd/Private/Cooker/CookDirector.cpp
--- a/Engine/Source/Editor/UnrealEd/Private/Cooker/CookDirector.cpp
+++ b/Engine/Source/Editor/UnrealEd/Private/Cooker/CookDirector.cpp
@@ -317,7 +317,7 @@ void FCookDirector::AssignRequests(TArray<FWorkerId>&& InWorkers, TArray<TRefCou
 			FWorkerId WorkerIdConstraint = Request->GetWorkerAssignmentConstraint();
 			if (WorkerIdConstraint.IsValid() && WorkerIdConstraint != WorkerId)
 			{
-				UE_LOG(LogCook, Error, TEXT("Package %s can only be cooked by a now-disconnected CookWorker. The package can not be cooked."),
+				UE_LOG(LogCook, Warning, TEXT("Package %s can only be cooked by a now-disconnected CookWorker. The package can not be cooked."),
 					*Request->GetPackageName().ToString());
 				Assignment = FWorkerId::Invalid();
 				RemovedRequests.Add(Request);
diff --git a/Engine/Source/Editor/UnrealEd/Private/Cooker/CookWorkerServer.cpp b/Engine/Source/Editor/UnrealEd/Private/Cooker/CookWorkerServer.cpp
--- a/Engine/Source/Editor/UnrealEd/Private/Cooker/CookWorkerServer.cpp
+++ b/Engine/Source/Editor/UnrealEd/Private/Cooker/CookWorkerServer.cpp
@@ -139,6 +139,7 @@ void FCookWorkerServer::SendCrashDiagnostics()
 	UE_LOG(LogCook, Display, TEXT("LostConnection to CookWorker %d. Log messages written after communication loss:"), ProfileId);
 	FString LogText;
 	int32 ReadFlags = FILEREAD_AllowWrite; // To be able to open a file for read that might be open for write from another process, we have to specify FILEREAD_AllowWrite
+	bool bLoggedErrorMessage = false;
 	if (!FFileHelper::LoadFileToString(LogText, *LogFileName, FFileHelper::EHashOptions::None, ReadFlags))
 	{
 		UE_LOG(LogCook, Warning, TEXT("No log file found for CookWorker %d."), ProfileId);
@@ -189,13 +190,23 @@ void FCookWorkerServer::SendCrashDiagnostics()
 			{
 				Verbosity = ELogVerbosity::Error;
 			}
+			bLoggedErrorMessage |= Verbosity == ELogVerbosity::Error;
 			FMsg::Logf(__FILE__, __LINE__, Category, Verbosity, TEXT("[CookWorker %d]: %.*s"),
 				ProfileId, Message.Len(), Message.GetData());
 		}
 	}
 	if (!CrashDiagnosticsError.IsEmpty())
 	{
-		UE_LOG(LogCook, Error, TEXT("%s"), *CrashDiagnosticsError);
+		if (!bLoggedErrorMessage)
+		{
+			UE_LOG(LogCook, Error, TEXT("%s"), *CrashDiagnosticsError);
+		}
+		else
+		{
+			// When we already logged an error from the crashed worker, log the what-went-wrong as a warning rather than an error,
+			// to avoid making it seem like a separate issue.
+			UE_LOG(LogCook, Warning, TEXT("%s"), *CrashDiagnosticsError);
+		}
 	}
 
 	bNeedCrashDiagnostics = false;
