CookProfiling: Warn instead of crashing if -CookStatsFile is unwritable. Create separate files for each CookWorker.
#rnx

diff --git a/Engine/Source/Editor/UnrealEd/Private/Cooker/CookProfiling.cpp b/Engine/Source/Editor/UnrealEd/Private/Cooker/CookProfiling.cpp
--- a/Engine/Source/Editor/UnrealEd/Private/Cooker/CookProfiling.cpp
+++ b/Engine/Source/Editor/UnrealEd/Private/Cooker/CookProfiling.cpp
@@ -15,6 +15,7 @@
 #include "HAL/FileManager.h"
 #include "Misc/CommandLine.h"
 #include "Misc/OutputDevice.h"
+#include "Misc/PathViews.h"
 #include "Misc/StringBuilder.h"
 #include "PackageBuildDependencyTracker.h"
 #include "Policies/CondensedJsonPrintPolicy.h"
@@ -1068,6 +1069,16 @@ void LogCookStats(ECookMode::Type CookMode)
 		FString CookStatsFileName;
 		if (FParse::Value(FCommandLine::Get(), TEXT("-CookStatsFile="), CookStatsFileName))
 		{
+			uint32 MultiprocessId = 0;
+			FParse::Value(FCommandLine::Get(), TEXT("-MultiprocessId="), MultiprocessId);
+			if (MultiprocessId != 0)
+			{
+				FStringView BaseName = FPathViews::GetBaseFilenameWithPath(CookStatsFileName);
+				FStringView Extension = FPathViews::GetExtension(CookStatsFileName, true);
+				CookStatsFileName = FString::Printf(TEXT("%.*s_CookWorker%u%.*s"),
+					BaseName.Len(), BaseName.GetData(), MultiprocessId, Extension.Len(), Extension.GetData());
+			}
+
 			FString JsonString;
 			TSharedRef<TJsonWriter<TCHAR, TCondensedJsonPrintPolicy<TCHAR>>> JsonWriter = TJsonWriterFactory<TCHAR, TCondensedJsonPrintPolicy<TCHAR> >::Create(&JsonString);
 			JsonWriter->WriteObjectStart();
@@ -1081,8 +1092,15 @@ void LogCookStats(ECookMode::Type CookMode)
 			JsonWriter->WriteObjectEnd();
 			JsonWriter->Close();
 			TUniquePtr<FArchive> JsonFile(IFileManager::Get().CreateFileWriter(*CookStatsFileName));
-			JsonFile->Serialize(TCHAR_TO_ANSI(*JsonString), JsonString.Len());
-			JsonFile->Close();
+			if (!JsonFile)
+			{
+				UE_LOG(LogCook, Warning, TEXT("Could not write to CookStatsFile %s."), *CookStatsFileName);
+			}
+			else
+			{
+				JsonFile->Serialize(TCHAR_TO_ANSI(*JsonString), JsonString.Len());
+				JsonFile->Close();
+			}
 		}
 
 	}
