EditorFactories: Allow specialized text importers to skip objects by name

diff --git a/Engine/Source/Editor/UnrealEd/Private/Factories/EditorFactories.cpp b/Engine/Source/Editor/UnrealEd/Private/Factories/EditorFactories.cpp
--- a/Engine/Source/Editor/UnrealEd/Private/Factories/EditorFactories.cpp
+++ b/Engine/Source/Editor/UnrealEd/Private/Factories/EditorFactories.cpp
@@ -5518,6 +5518,12 @@ void FCustomizableTextObjectFactory::ProcessBuffer(UObject* InParent, EObjectFla
 
 				UObject* ObjectParent = InParent ? InParent : GetParentForNewObject(ObjClass);
 
+				// Make sure we are allowed to create this object
+				if(!CanCreateObject(ObjectParent, ObjClass, ObjName))
+				{
+					continue;
+				}
+
 				// Make sure this name is not used by anything else. Will rename other stuff if necessary
 				UpdateObjectName(ObjClass, ObjectParent, ObjName);
 				ClearObjectNameUsage(ObjectParent, ObjName);
@@ -5646,7 +5652,13 @@ bool FCustomizableTextObjectFactory::CanCreateClass(UClass* ObjectClass, bool& b
 {
 	return false;
 }
-		
+
+bool FCustomizableTextObjectFactory::CanCreateObject(UObject* InParent, UClass* ObjectClass,
+	const FName& InDesiredName) const
+{
+	return true;
+}
+
 /** This is called on each created object after PreEditChange and the property text is imported, but before PostEditChange */
 void FCustomizableTextObjectFactory::ProcessConstructedObject(UObject* CreatedObject)
 {
diff --git a/Engine/Source/Editor/UnrealEd/Public/Factories.h b/Engine/Source/Editor/UnrealEd/Public/Factories.h
--- a/Engine/Source/Editor/UnrealEd/Public/Factories.h
+++ b/Engine/Source/Editor/UnrealEd/Public/Factories.h
@@ -50,6 +50,9 @@ protected:
 	/** Return true if the an object of type ObjectClass is allowed to be created; If false is returned, the object and subobjects will be ignored. */
 	virtual bool CanCreateClass(UClass* ObjectClass, bool& bOmitSubObjs) const;
 
+	/** Return true if the an object of type ObjectClass is allowed to be created with a given name under the provided parent; If false is returned, the object and subobjects will be ignored. */
+	virtual bool CanCreateObject(UObject* InParent, UClass* ObjectClass, const FName& InDesiredName) const;
+
 	/** This is called on each created object after the property text is imported */
 	virtual void ProcessConstructedObject(UObject* CreatedObject);
 
