ThumbnailExternalCache.LogEntries cvar

diff --git a/Engine/Source/Editor/UnrealEd/Private/ThumbnailExternalCache.cpp b/Engine/Source/Editor/UnrealEd/Private/ThumbnailExternalCache.cpp
--- a/Engine/Source/Editor/UnrealEd/Private/ThumbnailExternalCache.cpp
+++ b/Engine/Source/Editor/UnrealEd/Private/ThumbnailExternalCache.cpp
@@ -22,6 +22,9 @@ DEFINE_LOG_CATEGORY_STATIC(LogThumbnailExternalCache, Log, All);
 
 namespace ThumbnailExternalCache
 {
+	bool bLogEntries = false;
+	FAutoConsoleVariableRef CVarLogEntries(TEXT("ThumbnailExternalCache.LogEntries"), bLogEntries, TEXT("Whether log entries added to a thumbnail cache"));
+
 	const int64 LatestVersion = 0;
 	const uint64 ExpectedHeaderId = 0x424d5548545f4555; // "UE_THUMB"
 	const FString ThumbnailImageFormatName(TEXT(""));
@@ -130,6 +133,8 @@ void FSaveThumbnailCache::Save(FArchive& Ar, const TArrayView<FAssetData> InAsse
 
 	const int32 NumAssetDatas = InAssetDatas.Num();
 
+	UE_LOG(LogThumbnailExternalCache, Log, TEXT("Saving thumbnails for %d assets to %s"), NumAssetDatas, *Ar.GetArchiveName());
+
 	FText StatusText = LOCTEXT("SaveStatus", "Saving Thumbnails: {0}");
 	FScopedSlowTask SlowTask( (float)NumAssetDatas / (float)TaskBatchSize, FText::Format(StatusText, FText::AsNumber(NumAssetDatas)));
 	SlowTask.MakeDialog(/*bShowCancelButton*/ false);
@@ -237,14 +242,20 @@ void FSaveThumbnailCache::Save(FArchive& Ar, const TArrayView<FAssetData> InAsse
 
 	int64 NumThumbnails = PackageThumbnailRecords.Num();
 	Ar << NumThumbnails;
-
-	FString ThumbnailNameString;
-	for (FPackageThumbnailRecord& PackageThumbnailRecord : PackageThumbnailRecords)
 	{
-		ThumbnailNameString.Reset();
-		PackageThumbnailRecord.Name.AppendString(ThumbnailNameString);
-		Ar << ThumbnailNameString;
-		Ar << PackageThumbnailRecord.Offset;
+		FString ThumbnailNameString;
+		int64 Index = 0;
+		for (FPackageThumbnailRecord& PackageThumbnailRecord : PackageThumbnailRecords)
+		{
+			ThumbnailNameString.Reset();
+			PackageThumbnailRecord.Name.AppendString(ThumbnailNameString);
+			if (ThumbnailExternalCache::bLogEntries)
+			{
+				UE_LOG(LogThumbnailExternalCache, Log, TEXT("\t[%d] %s"), Index++, *ThumbnailNameString);
+			}
+			Ar << ThumbnailNameString;
+			Ar << PackageThumbnailRecord.Offset;
+		}
 	}
 
 	// Modify top of archive to know where table of contents is located
