CookProfiling: Suppress the CookStatsFile on CookWorkers rather than writing to a renamed file; the renamed file is not used by callers and is wasteful. In the long run we want to replicate the profile data back to the cookdirector and merge it in.
#rnx

diff --git a/Engine/Source/Editor/UnrealEd/Private/Cooker/CookProfiling.cpp b/Engine/Source/Editor/UnrealEd/Private/Cooker/CookProfiling.cpp
--- a/Engine/Source/Editor/UnrealEd/Private/Cooker/CookProfiling.cpp
+++ b/Engine/Source/Editor/UnrealEd/Private/Cooker/CookProfiling.cpp
@@ -1073,12 +1073,14 @@ void LogCookStats(ECookMode::Type CookMode)
 			FParse::Value(FCommandLine::Get(), TEXT("-MultiprocessId="), MultiprocessId);
 			if (MultiprocessId != 0)
 			{
-				FStringView BaseName = FPathViews::GetBaseFilenameWithPath(CookStatsFileName);
-				FStringView Extension = FPathViews::GetExtension(CookStatsFileName, true);
-				CookStatsFileName = FString::Printf(TEXT("%.*s_CookWorker%u%.*s"),
-					BaseName.Len(), BaseName.GetData(), MultiprocessId, Extension.Len(), Extension.GetData());
+				// Suppress the file creation on CookWorkers
+				// TODO: Replicate the information back to the CookDirector instead, UE-185774
+				CookStatsFileName.Empty();
 			}
+		}
 
+		if (!CookStatsFileName.IsEmpty())
+		{
 			FString JsonString;
 			TSharedRef<TJsonWriter<TCHAR, TCondensedJsonPrintPolicy<TCHAR>>> JsonWriter = TJsonWriterFactory<TCHAR, TCondensedJsonPrintPolicy<TCHAR> >::Create(&JsonString);
 			JsonWriter->WriteObjectStart();
