Fixed issue with 'Reimport Scene With File' menu

diff --git a/Engine/Plugins/Interchange/Editor/Source/InterchangeEditor/Private/AssetDefinition_InterchangeSceneImportAsset.cpp b/Engine/Plugins/Interchange/Editor/Source/InterchangeEditor/Private/AssetDefinition_InterchangeSceneImportAsset.cpp
--- a/Engine/Plugins/Interchange/Editor/Source/InterchangeEditor/Private/AssetDefinition_InterchangeSceneImportAsset.cpp
+++ b/Engine/Plugins/Interchange/Editor/Source/InterchangeEditor/Private/AssetDefinition_InterchangeSceneImportAsset.cpp
@@ -76,15 +76,20 @@ namespace MenuExtension_InterchangeSceneImportAsset
 			return;
 		}
 		
+#if WITH_EDITORONLY_DATA
 		UInterchangeSceneImportAsset* SceneImportAsset = SceneImportAssets[0];
+		if (!SceneImportAsset->AssetImportData)
+		{
+			return;
+		}
 
-		TArray<FString> OpenFilenames;
+		TArray<FString> OpenFilenames = SceneImportAsset->AssetImportData->ExtractFilenames();
 		FReimportManager::Instance()->GetNewReimportPath(SceneImportAsset, OpenFilenames);
 		if (OpenFilenames.Num() == 1 && !OpenFilenames[0].IsEmpty())
 		{
 			ExecuteReimportOneAsset(SceneImportAsset, OpenFilenames[0]);
 		}
-
+#endif
 	}
 
 	bool CanExecuteReimport(const FToolMenuContext& InContext)
diff --git a/Engine/Source/Editor/UnrealEd/Private/Editor.cpp b/Engine/Source/Editor/UnrealEd/Private/Editor.cpp
--- a/Engine/Source/Editor/UnrealEd/Private/Editor.cpp
+++ b/Engine/Source/Editor/UnrealEd/Private/Editor.cpp
@@ -778,31 +778,8 @@ void FReimportManager::GetNewReimportPath(UObject* Obj, TArray<FString>& InOutFi
 		InOutFilenames[RealSourceFileIndex].Empty();
 	}
 
-	// Get the list of valid factories
-	for( TObjectIterator<UClass> It ; It ; ++It )
-	{
-		UClass* CurrentClass = (*It);
-
-		if( CurrentClass->IsChildOf(UFactory::StaticClass()) && !(CurrentClass->HasAnyClassFlags(CLASS_Abstract)) )
-		{
-			UFactory* Factory = Cast<UFactory>( CurrentClass->GetDefaultObject() );
-			if( Factory->bEditorImport && Factory->DoesSupportClass(Obj->GetClass()) )
-			{
-				Factories.Add( Factory );
-			}
-		}
-	}
-
-	if ( Factories.Num() <= 0 )
-	{
-		// No matching factories for this asset, fail
-		return;
-	}
-
+	// Append the Interchange supported translator formats for this object
 	TMultiMap<uint32, UFactory*> DummyFilterIndexToFactory;
-
-	IAssetTools& AssetTools = FModuleManager::LoadModuleChecked<FAssetToolsModule>(TEXT("AssetTools")).Get();
-	//Append the interchange supported translator formats for this object
 	if (UInterchangeManager::IsInterchangeImportEnabled())
 	{
 		//Get the extension interchange can translate for this object
@@ -810,15 +787,34 @@ void FReimportManager::GetNewReimportPath(UObject* Obj, TArray<FString>& InOutFi
 		ObjectTools::AppendFormatsFileExtensions(TranslatorFormats, FileTypes, AllExtensions, DummyFilterIndexToFactory);
 	}
 
-	//If this object was not import with interchange add the legacy formats
-	if(AllExtensions.IsEmpty())
+	// Interchange is either disabled or do not support the given object, check with the legacy factories
+	if (AllExtensions.IsEmpty())
 	{
+		// Get the list of valid factories
+		for (TObjectIterator<UClass> It; It; ++It)
+		{
+			UClass* CurrentClass = (*It);
+
+			if (CurrentClass->IsChildOf(UFactory::StaticClass()) && !(CurrentClass->HasAnyClassFlags(CLASS_Abstract)))
+			{
+				UFactory* Factory = Cast<UFactory>(CurrentClass->GetDefaultObject());
+				if (Factory->bEditorImport && Factory->DoesSupportClass(Obj->GetClass()))
+				{
+					Factories.Add(Factory);
+				}
+			}
+		}
+
+		if (Factories.Num() <= 0)
+		{
+			// No matching factories for this asset, fail
+			return;
+		}
+
 		// Generate the file types and extensions represented by the selected factories
 		ObjectTools::GenerateFactoryFileExtensions(Factories, FileTypes, AllExtensions, DummyFilterIndexToFactory);
 	}
 
-	
-
 	FString DefaultFolder;
 	FString DefaultFile;
 	
