Run GenerateCachedEditorThumbnails only on package that were actually cooked (not just any cooked package; i.e. from engine/game when cooking a dlc)

diff --git a/Engine/Source/Editor/UnrealEd/Private/CookOnTheFlyServer.cpp b/Engine/Source/Editor/UnrealEd/Private/CookOnTheFlyServer.cpp
--- a/Engine/Source/Editor/UnrealEd/Private/CookOnTheFlyServer.cpp
+++ b/Engine/Source/Editor/UnrealEd/Private/CookOnTheFlyServer.cpp
@@ -9440,19 +9440,25 @@ void UCookOnTheFlyServer::GenerateCachedEditorThumbnails()
 		TArray<FAssetData> PublicAssets;
 		{
 			TArray<UE::Cook::FPackageData*> CookedPackages;
-			PackageDatas->GetCookedPackagesForPlatform(Platform, CookedPackages, CookedPackages);
+			{
+				TArray<UE::Cook::FPackageData*> FailedPackages;
+				PackageDatas->GetCookedPackagesForPlatform(Platform, CookedPackages, FailedPackages);
+			}
 
 			for (const UE::Cook::FPackageData* const CookedPackage : CookedPackages)
 			{
-				TArray<FAssetData> Assets;
-				ensure(AssetRegistry->GetAssetsByPackageName(CookedPackage->GetPackageName(), Assets, /*bIncludeOnlyDiskAssets=*/true));
-
-				for (FAssetData& Asset : Assets)
+				if (CookedPackage->GetWasCookedThisSession())
 				{
-					if (!(Asset.PackageFlags & PKG_NotExternallyReferenceable))
+					TArray<FAssetData> Assets;
+					ensure(AssetRegistry->GetAssetsByPackageName(CookedPackage->GetPackageName(), Assets, /*bIncludeOnlyDiskAssets=*/true));
+
+					for (FAssetData& Asset : Assets)
 					{
-						UE_LOG(LogCook, Display, TEXT("Adding thumbnail for '%s'"), *Asset.GetObjectPathString());
-						PublicAssets.Add(MoveTemp(Asset));
+						if (!(Asset.PackageFlags & PKG_NotExternallyReferenceable))
+						{
+							UE_LOG(LogCook, Verbose, TEXT("Adding thumbnail for '%s'"), *Asset.GetObjectPathString());
+							PublicAssets.Add(MoveTemp(Asset));
+						}
 					}
 				}
 			}
@@ -10852,7 +10858,7 @@ void UCookOnTheFlyServer::RecordDLCPackagesFromBaseGame(FBeginCookContext& Begin
 		{
 			if (UE::Cook::FPackageData* PackageData = PackageDatas->TryAddPackageDataByFileName(FName(*AssetPath)))
 			{
-				PackageData->SetPlatformsCooked(BeginContext.TargetPlatforms, UE::Cook::ECookResult::Succeeded);
+				PackageData->SetPlatformsCooked(BeginContext.TargetPlatforms, UE::Cook::ECookResult::Succeeded, /*bWasCookedThisSession=*/false);
 				++PackageDataFromBaseGameNum;
 			}
 			else
diff --git a/Engine/Source/Editor/UnrealEd/Private/Cooker/CookPackageData.cpp b/Engine/Source/Editor/UnrealEd/Private/Cooker/CookPackageData.cpp
--- a/Engine/Source/Editor/UnrealEd/Private/Cooker/CookPackageData.cpp
+++ b/Engine/Source/Editor/UnrealEd/Private/Cooker/CookPackageData.cpp
@@ -87,6 +87,7 @@ FPackageData::FPackageData(FPackageDatas& PackageDatas, const FName& InPackageNa
 	, bCookedPlatformDataStarted(0), bCookedPlatformDataCalled(0), bCookedPlatformDataComplete(0)
 	, MonitorCookResult((uint8)ECookResult::NotAttempted)
 	, bInitializedGeneratorSave(0), bCompletedGeneration(0), bGenerated(0), bKeepReferencedDuringGC(0)
+	, bWasCookedThisSession(0)
 {
 	SetState(EPackageState::Idle);
 	SendToState(EPackageState::Idle, ESendFlags::QueueAdd);
@@ -288,27 +289,36 @@ void FPackageData::ClearInProgressData()
 	CompletionCallback = FCompletionCallback();
 }
 
-void FPackageData::SetPlatformsCooked(const TConstArrayView<const ITargetPlatform*> TargetPlatforms,
-	const TConstArrayView<ECookResult> Result)
+void FPackageData::SetPlatformsCooked(
+	const TConstArrayView<const ITargetPlatform*> TargetPlatforms,
+	const TConstArrayView<ECookResult> Result,
+	const bool bInWasCookedThisSession)
 {
 	check(TargetPlatforms.Num() == Result.Num());
 	for (int32 n = 0; n < TargetPlatforms.Num(); ++n)
 	{
-		SetPlatformCooked(TargetPlatforms[n], Result[n]);
+		SetPlatformCooked(TargetPlatforms[n], Result[n], bInWasCookedThisSession);
 	}
 }
 
-void FPackageData::SetPlatformsCooked(const TConstArrayView<const ITargetPlatform*> TargetPlatforms,
-	ECookResult Result)
+void FPackageData::SetPlatformsCooked(
+	const TConstArrayView<const ITargetPlatform*> TargetPlatforms, 
+	ECookResult Result,
+	const bool bInWasCookedThisSession)
 {
 	for (const ITargetPlatform* TargetPlatform : TargetPlatforms)
 	{
-		SetPlatformCooked(TargetPlatform, Result);
+		SetPlatformCooked(TargetPlatform, Result, bInWasCookedThisSession);
 	}
 }
 
-void FPackageData::SetPlatformCooked(const ITargetPlatform* TargetPlatform, ECookResult CookResult)
+void FPackageData::SetPlatformCooked(
+	const ITargetPlatform* TargetPlatform, 
+	ECookResult CookResult, 
+	const bool bInWasCookedThisSession)
 {
+	bWasCookedThisSession |= bInWasCookedThisSession && (CookResult == ECookResult::Succeeded);
+
 	bool bNewCookAttemptedValue = (CookResult != ECookResult::NotAttempted);
 	bool bModifiedCookAttempted = false;
 	bool bHasAnyOtherCookAttempted = false;
@@ -328,6 +338,7 @@ void FPackageData::SetPlatformCooked(const ITargetPlatform* TargetPlatform, ECoo
 			bHasAnyOtherCookAttempted = bHasAnyOtherCookAttempted | Pair.Value.IsCookAttempted();
 		}
 	}
+
 	if (!bExists && bNewCookAttemptedValue)
 	{
 		FPackagePlatformData& Value = PlatformDatas.FindOrAdd(TargetPlatform);
@@ -335,6 +346,7 @@ void FPackageData::SetPlatformCooked(const ITargetPlatform* TargetPlatform, ECoo
 		Value.SetSaveTimedOut(false);
 		bModifiedCookAttempted = true;
 	}
+
 	if (bModifiedCookAttempted && !bHasAnyOtherCookAttempted)
 	{
 		if (bNewCookAttemptedValue)
@@ -343,6 +355,7 @@ void FPackageData::SetPlatformCooked(const ITargetPlatform* TargetPlatform, ECoo
 		}
 		else
 		{
+			bWasCookedThisSession = false;
 			PackageDatas.GetMonitor().OnLastCookedPlatformRemoved(*this);
 		}
 	}
@@ -375,6 +388,7 @@ void FPackageData::ClearCookResults()
 	}
 	if (bModifiedCookAttempted)
 	{
+		bWasCookedThisSession = false;
 		PackageDatas.GetMonitor().OnLastCookedPlatformRemoved(*this);
 	}
 }
@@ -398,6 +412,7 @@ void FPackageData::ClearCookResults(const ITargetPlatform* TargetPlatform)
 	}
 	if (bModifiedCookAttempted && !bHasAnyOthers)
 	{
+		bWasCookedThisSession = false;
 		PackageDatas.GetMonitor().OnLastCookedPlatformRemoved(*this);
 	}
 }
@@ -2931,7 +2946,7 @@ void FPackageDatas::AddExistingPackageDatasForPlatform(TConstArrayView<FConstruc
 
 			PackageData = NewPackageData;
 		}
-		PackageData->SetPlatformCooked(TargetPlatform, ECookResult::Succeeded);
+		PackageData->SetPlatformCooked(TargetPlatform, ECookResult::Succeeded, /*bWasCookedThisSession=*/false);
 	}
 	OutPackageDataFromBaseGameNum += ExistingPackages.Num();
 }
diff --git a/Engine/Source/Editor/UnrealEd/Private/Cooker/CookPackageData.h b/Engine/Source/Editor/UnrealEd/Private/Cooker/CookPackageData.h
--- a/Engine/Source/Editor/UnrealEd/Private/Cooker/CookPackageData.h
+++ b/Engine/Source/Editor/UnrealEd/Private/Cooker/CookPackageData.h
@@ -291,10 +291,9 @@ public:
 	 * FindOrAdd each TargetPlatform and set its flags: CookAttempted=true, Succeeded=<given>.
 	 * In version that takes two arrays, TargetPlatforms and Succeeded must be the same length.
 	 */
-	void SetPlatformsCooked(const TConstArrayView<const ITargetPlatform*> TargetPlatforms,
-		const TConstArrayView<ECookResult> Succeeded);
-	void SetPlatformsCooked(const TConstArrayView<const ITargetPlatform*> TargetPlatforms, ECookResult Result);
-	void SetPlatformCooked(const ITargetPlatform* TargetPlatform, ECookResult Result);
+	void SetPlatformsCooked(const TConstArrayView<const ITargetPlatform*> TargetPlatforms, const TConstArrayView<ECookResult> Succeeded, bool bInWasCookedThisSession = true);
+	void SetPlatformsCooked(const TConstArrayView<const ITargetPlatform*> TargetPlatforms, ECookResult Result, bool bInWasCookedThisSession = true);
+	void SetPlatformCooked(const ITargetPlatform* TargetPlatform, ECookResult Result, bool bInWasCookedThisSession = true);
 	/**
 	 * FindOrAdd each TargetPlatform and set its flags: CookAttempted=false.
 	 * In Version that takes no TargetPlatform, CookAttempted is cleared from all existing platforms.
@@ -533,6 +532,9 @@ public:
 	/** Set whether COTFS is keeping this package referenced referenced during GC. */
 	void SetKeepReferencedDuringGC(bool Value) { bKeepReferencedDuringGC = Value != 0; }
 
+	/** Return whether the package was cooked during this session */
+	bool GetWasCookedThisSession() const { return static_cast<bool>(bWasCookedThisSession); }
+
 	/** For MultiProcessCooks, Get the id of the worker this Package is assigned to; InvalidId means owned by local. */
 	FWorkerId GetWorkerAssignment() const { return WorkerAssignment; }
 	/**
@@ -688,6 +690,7 @@ private:
 	uint32 bCompletedGeneration : 1;
 	uint32 bGenerated : 1;
 	uint32 bKeepReferencedDuringGC : 1;
+	uint32 bWasCookedThisSession : 1;
 };
 
 /** A single object in athe save of a package that might have had BeginCacheForCookedPlatformData called already */
diff --git a/Engine/Source/Editor/UnrealEd/Private/ThumbnailExternalCache.cpp b/Engine/Source/Editor/UnrealEd/Private/ThumbnailExternalCache.cpp
--- a/Engine/Source/Editor/UnrealEd/Private/ThumbnailExternalCache.cpp
+++ b/Engine/Source/Editor/UnrealEd/Private/ThumbnailExternalCache.cpp
@@ -22,9 +22,6 @@ DEFINE_LOG_CATEGORY_STATIC(LogThumbnailExternalCache, Log, All);
 
 namespace ThumbnailExternalCache
 {
-	bool bLogEntries = false;
-	FAutoConsoleVariableRef CVarLogEntries(TEXT("ThumbnailExternalCache.LogEntries"), bLogEntries, TEXT("Whether log entries added to a thumbnail cache"));
-
 	const int64 LatestVersion = 0;
 	const uint64 ExpectedHeaderId = 0x424d5548545f4555; // "UE_THUMB"
 	const FString ThumbnailImageFormatName(TEXT(""));
@@ -249,10 +246,7 @@ void FSaveThumbnailCache::Save(FArchive& Ar, const TArrayView<FAssetData> InAsse
 		{
 			ThumbnailNameString.Reset();
 			PackageThumbnailRecord.Name.AppendString(ThumbnailNameString);
-			if (ThumbnailExternalCache::bLogEntries)
-			{
-				UE_LOG(LogThumbnailExternalCache, Log, TEXT("\t[%d] %s"), Index++, *ThumbnailNameString);
-			}
+			UE_LOG(LogThumbnailExternalCache, Verbose, TEXT("\t[%d] %s"), Index++, *ThumbnailNameString);
 			Ar << ThumbnailNameString;
 			Ar << PackageThumbnailRecord.Offset;
 		}
