Rename newly introduced cooker arg GenerateCachedEditorThumbnails to CachedEditorThumbnails
We want to specify the argument on UGC cookers but there was a bug with the original implementation and we don't want it to run on older builds

diff --git a/Engine/Source/Editor/UnrealEd/Private/CookOnTheFlyServer.cpp b/Engine/Source/Editor/UnrealEd/Private/CookOnTheFlyServer.cpp
--- a/Engine/Source/Editor/UnrealEd/Private/CookOnTheFlyServer.cpp
+++ b/Engine/Source/Editor/UnrealEd/Private/CookOnTheFlyServer.cpp
@@ -9418,7 +9418,7 @@ void UCookOnTheFlyServer::WriteMapDependencyGraph(const ITargetPlatform* TargetP
 
 void UCookOnTheFlyServer::GenerateCachedEditorThumbnails()
 {
-	if (!FParse::Param(FCommandLine::Get(), TEXT("GenerateCachedEditorThumbnails")))
+	if (!FParse::Param(FCommandLine::Get(), TEXT("CachedEditorThumbnails")))
 	{
 		return;
 	}
