Cooker: Implement -CookIncludeClass=Class1+Class2+.... This allows cooking only packages containing an instance of the given class or classes. This relies on assetregistry dependencies to find all packages that would ordinarily be cooked. Packages found through dependencies that are discovered only during load of their referencer are not guaranteed to be cooked.
#rn Minor,Cooking

diff --git a/Engine/Source/Editor/UnrealEd/Classes/CookOnTheSide/CookOnTheFlyServer.h b/Engine/Source/Editor/UnrealEd/Classes/CookOnTheSide/CookOnTheFlyServer.h
--- a/Engine/Source/Editor/UnrealEd/Classes/CookOnTheSide/CookOnTheFlyServer.h
+++ b/Engine/Source/Editor/UnrealEd/Classes/CookOnTheSide/CookOnTheFlyServer.h
@@ -869,6 +869,9 @@ public:
 	/** Is the local CookOnTheFlyServer cooking a DLC plugin rather than a Project+Engine+EmbeddedPlugins? */
 	bool IsCookingDLC() const;
 
+	/** Return whether EDLCookInfo verification has NOT been rendered useless by settings such as CookFilter. */
+	bool ShouldVerifyEDLCookInfo() const;
+
 protected:
 	// FExec interface used in the editor
 	virtual bool Exec_Editor(class UWorld* InWorld, const TCHAR* Cmd, FOutputDevice& Ar) override;
@@ -1061,6 +1064,7 @@ private:
 	/** Set parameters that rely on config settings from startup and don't depend on StartCook* functions. */
 	void LoadInitializeConfigSettings(const FString& InOutputDirectoryOverride);
 	void SetInitializeConfigSettings(UE::Cook::FInitializeConfigSettings&& Settings);
+	void ParseCookFilters();
 
 	/** Set parameters that rely on config and CookByTheBook settings. */
 	void LoadBeginCookConfigSettings(FBeginCookContext& BeginContext);
@@ -1410,6 +1414,8 @@ private:
 
 	/** Package -> Package dependencies that were hidden and discovered when a package loaded or ran system-specific code. */
 	TMap<FName, TArray<FName>> DiscoveredDependencies;
+	/** Classes (and all subclasses) that were listed as the only classes that should be cooked in the filter settings. */
+	TSet<FName> CookFilterIncludedClasses;
 
 	/** True when PumpLoads has detected it is blocked on async work and CookOnTheFlyServer should do work elsewhere. */
 	bool bLoadBusy = false;
@@ -1446,6 +1452,8 @@ private:
 	 * build jobs when multiple machines are cooking at the same time.
 	 */
 	bool bRandomizeCookOrder = false;
+	/** True if commandline arguments specified that we suppress the cook of packages based on filter criteria. */
+	bool bCookFilter = false;
 
 	/** Timers for tracking how long we have been busy, to manage retries and warnings of deadlock */
 	double SaveBusyStartTimeSeconds = MAX_flt;
diff --git a/Engine/Source/Editor/UnrealEd/Private/Commandlets/CookCommandlet.cpp b/Engine/Source/Editor/UnrealEd/Private/Commandlets/CookCommandlet.cpp
--- a/Engine/Source/Editor/UnrealEd/Private/Commandlets/CookCommandlet.cpp
+++ b/Engine/Source/Editor/UnrealEd/Private/Commandlets/CookCommandlet.cpp
@@ -546,7 +546,7 @@ void UCookCommandlet::RunCookByTheBookCook(UCookOnTheFlyServer* CookOnTheFlyServ
 	FLowLevelMemTracker::Get().UpdateStatsPerFrame();
 #endif
 
-	if (StartupOptions.DLCName.IsEmpty())
+	if (CookOnTheFlyServer->ShouldVerifyEDLCookInfo())
 	{
 		bool bFullReferencesExpected = !(CookOptions & ECookByTheBookOptions::SkipHardReferences);
 		UE::SavePackageUtilities::VerifyEDLCookInfo(bFullReferencesExpected);
diff --git a/Engine/Source/Editor/UnrealEd/Private/CookOnTheFlyServer.cpp b/Engine/Source/Editor/UnrealEd/Private/CookOnTheFlyServer.cpp
--- a/Engine/Source/Editor/UnrealEd/Private/CookOnTheFlyServer.cpp
+++ b/Engine/Source/Editor/UnrealEd/Private/CookOnTheFlyServer.cpp
@@ -6637,6 +6637,70 @@ void UCookOnTheFlyServer::SetInitializeConfigSettings(UE::Cook::FInitializeConfi
 
 	bRandomizeCookOrder = FParse::Param(FCommandLine::Get(), TEXT("RANDOMPACKAGEORDER")) ||
 		(FParse::Param(FCommandLine::Get(), TEXT("DIFFONLY")) && !FParse::Param(FCommandLine::Get(), TEXT("DIFFNORANDCOOK")));
+
+	ParseCookFilters();
+}
+
+void UCookOnTheFlyServer::ParseCookFilters()
+{
+	CookFilterIncludedClasses.Empty();
+	bCookFilter = false;
+	if (!IsCookByTheBookMode() || IsCookingInEditor())
+	{
+		return;
+	}
+
+	FString IncludeClassesString;
+	if (FParse::Value(FCommandLine::Get(), TEXT("-cookincludeclass="), IncludeClassesString))
+	{
+		TArray<FString> IncludeClasses;
+		const TCHAR* Delimiters[] = { TEXT(","), TEXT("+"), TEXT(";")};
+		IncludeClassesString.ParseIntoArray(IncludeClasses, Delimiters,
+			UE_ARRAY_COUNT(Delimiters), true /* bCullEmpty */);
+		TArray<FTopLevelAssetPath> RootNames;
+		for (FString& IncludeClassString : IncludeClasses)
+		{
+			FTopLevelAssetPath ClassPath;
+			if (!UClass::IsShortTypeName(IncludeClassString))
+			{
+				ClassPath.TrySetPath(IncludeClassString);
+			}
+			else
+			{
+				ClassPath = UClass::TryConvertShortTypeNameToPathName<UClass>(IncludeClassString);
+			}
+			if (!ClassPath.IsValid())
+			{
+				UE_LOG(LogCook, Error, TEXT("CookIncludeClass: Could not convert string '%s' into a class path. Ignoring it."),
+					*IncludeClassString);
+				continue;
+			}
+			UClass* IncludedClass = FindObject<UClass>(nullptr, *ClassPath.ToString());
+			if (!IncludedClass)
+			{
+				UE_LOG(LogCook, Error, TEXT("CookIncludeClass: Could not find class with ClassPath '%s'. Ignoring it."),
+					*IncludeClassString);
+				continue;
+			}
+			FTopLevelAssetPath NormalizedClassPath(IncludedClass);
+			RootNames.Add(NormalizedClassPath);
+			CookFilterIncludedClasses.Add(FName(*NormalizedClassPath.ToString()));
+		}
+		if (!RootNames.IsEmpty())
+		{
+			TSet<FTopLevelAssetPath> DerivedClassNames;
+			AssetRegistry->GetDerivedClassNames(RootNames, TSet<FTopLevelAssetPath>() /* ExcludedClassNames */,
+				DerivedClassNames);
+			for (const FTopLevelAssetPath& NormalizedClassPath : DerivedClassNames)
+			{
+				CookFilterIncludedClasses.Add(FName(*NormalizedClassPath.ToString()));
+			}
+
+			UE_LOG(LogCook, Display, TEXT("CookIncludeClass: Only cooking packages that contains objects with class in { %s }"),
+				*FString::JoinBy(RootNames, TEXT(", "), [](const FTopLevelAssetPath& P) { return P.ToString(); }));
+			bCookFilter = true;
+		}
+	}
 }
 
 bool UCookOnTheFlyServer::TryInitializeCookWorker()
@@ -10967,6 +11031,11 @@ void FEDLMPCollector::ServerReceiveMessage(FMPCollectorServerMessageContext& Con
 
 }
 
+bool UCookOnTheFlyServer::ShouldVerifyEDLCookInfo() const
+{
+	return CookByTheBookOptions->DlcName.IsEmpty() && !bCookFilter;
+}
+
 void UCookOnTheFlyServer::BeginCookEDLCookInfo(FBeginCookContext& BeginContext)
 {
 	if (IsCookingInEditor())
diff --git a/Engine/Source/Editor/UnrealEd/Private/Cooker/CookRequestCluster.cpp b/Engine/Source/Editor/UnrealEd/Private/Cooker/CookRequestCluster.cpp
--- a/Engine/Source/Editor/UnrealEd/Private/Cooker/CookRequestCluster.cpp
+++ b/Engine/Source/Editor/UnrealEd/Private/Cooker/CookRequestCluster.cpp
@@ -40,21 +40,19 @@ FRequestCluster::FRequestCluster(UCookOnTheFlyServer& InCOTFS)
 	, PackageTracker(*InCOTFS.PackageTracker)
 	, BuildDefinitions(*InCOTFS.BuildDefinitions)
 {
-	if (!COTFS.IsCookOnTheFlyMode())
-	{
-		UE::Cook::FCookByTheBookOptions& Options = *COTFS.CookByTheBookOptions;
-		bAllowHardDependencies = !Options.bSkipHardReferences;
-		bAllowSoftDependencies = !Options.bSkipSoftReferences;
-		bErrorOnEngineContentUse = Options.bErrorOnEngineContentUse;
-		bAllowUncookedAssetReferences = Options.bAllowUncookedAssetReferences;
-	}
-	else
+	// CookByTheBookOptions is always available; in other modes it is set to the default values
+	UE::Cook::FCookByTheBookOptions& Options = *COTFS.CookByTheBookOptions;
+	bAllowHardDependencies = !Options.bSkipHardReferences;
+	bAllowSoftDependencies = !Options.bSkipSoftReferences;
+	bErrorOnEngineContentUse = Options.bErrorOnEngineContentUse;
+	if (COTFS.IsCookOnTheFlyMode())
 	{
 		// Do not queue soft-dependencies during CookOnTheFly; wait for them to be requested
 		// TODO: Report soft dependencies separately, and mark them as normal priority,
 		// and mark all hard dependencies as high priority in cook on the fly.
 		bAllowSoftDependencies = false;
 	}
+
 	bHybridIterativeEnabled = COTFS.bHybridIterativeEnabled;
 	if (bErrorOnEngineContentUse)
 	{
@@ -1456,38 +1454,25 @@ void FRequestCluster::IsRequestCookable(const ITargetPlatform* Platform, FPackag
 	UCookOnTheFlyServer& COTFS, ESuppressCookReason& OutReason, bool& bOutCookable, bool& bOutExplorable)
 {
 	FString LocalDLCPath;
-	bool bLocalErrorOnEngineContentUse = false;
-	bool bLocalAllowUncookedAssetReferences = false;
-	if (!COTFS.IsCookOnTheFlyMode())
-	{
-		UE::Cook::FCookByTheBookOptions& Options = *COTFS.CookByTheBookOptions;
-		bLocalErrorOnEngineContentUse = Options.bErrorOnEngineContentUse;
-		bLocalAllowUncookedAssetReferences = Options.bAllowUncookedAssetReferences;
-	}
-	if (bLocalErrorOnEngineContentUse)
+	if (COTFS.CookByTheBookOptions->bErrorOnEngineContentUse)
 	{
 		LocalDLCPath = FPaths::Combine(*COTFS.GetBaseDirectoryForDLC(), TEXT("Content"));
 		FPaths::MakeStandardFilename(LocalDLCPath);
 	}
 
-	IsRequestCookable(Platform, PackageData.GetPackageName(), PackageData, *COTFS.PackageDatas,
-		*COTFS.PackageTracker, LocalDLCPath, bLocalErrorOnEngineContentUse, bLocalAllowUncookedAssetReferences,
-		COTFS.bSkipOnlyEditorOnly, OutReason, bOutCookable, bOutExplorable);
+	IsRequestCookable(Platform, PackageData.GetPackageName(), PackageData, COTFS,
+		LocalDLCPath, OutReason, bOutCookable, bOutExplorable);
 }
 
 void FRequestCluster::IsRequestCookable(const ITargetPlatform* Platform, FName PackageName, FPackageData& PackageData,
 	ESuppressCookReason& OutReason, bool& bOutCookable, bool& bOutExplorable)
 {
-	return IsRequestCookable(Platform, PackageName, PackageData, PackageDatas, PackageTracker,
-		DLCPath, bErrorOnEngineContentUse, bAllowUncookedAssetReferences,
-		COTFS.bSkipOnlyEditorOnly,
-		OutReason, bOutCookable, bOutExplorable);
+	return IsRequestCookable(Platform, PackageName, PackageData, COTFS,
+		DLCPath, OutReason, bOutCookable, bOutExplorable);
 }
 
 void FRequestCluster::IsRequestCookable(const ITargetPlatform* Platform, FName PackageName, FPackageData& PackageData,
-	FPackageDatas& InPackageDatas, FPackageTracker& InPackageTracker,
-	FStringView InDLCPath, bool bInErrorOnEngineContentUse, bool bInAllowUncookedAssetReferences,
-	bool bSkipOnlyEditorOnly, ESuppressCookReason& OutReason, bool& bOutCookable, bool& bOutExplorable)
+	UCookOnTheFlyServer& InCOTFS, FStringView InDLCPath, ESuppressCookReason& OutReason, bool& bOutCookable, bool& bOutExplorable)
 {
 	check(Platform != CookerLoadingPlatformKey); // IsRequestCookable should not be called for The CookerLoadingPlatform; it has different rules
 
@@ -1504,7 +1489,7 @@ void FRequestCluster::IsRequestCookable(const ITargetPlatform* Platform, FName P
 	}
 
 	FName FileName = PackageData.GetFileName();
-	if (InPackageTracker.NeverCookPackageList.Contains(FileName))
+	if (InCOTFS.PackageTracker->NeverCookPackageList.Contains(FileName))
 	{
 		if (INDEX_NONE != UE::String::FindFirst(NameBuffer, ULevel::GetExternalActorsFolderName(), ESearchCase::IgnoreCase))
 		{
@@ -1513,7 +1498,7 @@ void FRequestCluster::IsRequestCookable(const ITargetPlatform* Platform, FName P
 			OutReason = ESuppressCookReason::NeverCook;
 			bOutCookable = false;
 			bOutExplorable = true;
-			if (bSkipOnlyEditorOnly)
+			if (InCOTFS.bSkipOnlyEditorOnly)
 			{
 				bOutExplorable = true;
 			}
@@ -1536,7 +1521,7 @@ void FRequestCluster::IsRequestCookable(const ITargetPlatform* Platform, FName P
 		return;
 	}
 
-	if (bInErrorOnEngineContentUse && !InDLCPath.IsEmpty())
+	if (InCOTFS.CookByTheBookOptions->bErrorOnEngineContentUse && !InDLCPath.IsEmpty())
 	{
 		FileName.ToString(NameBuffer);
 		if (!FStringView(NameBuffer).StartsWith(InDLCPath))
@@ -1545,7 +1530,7 @@ void FRequestCluster::IsRequestCookable(const ITargetPlatform* Platform, FName P
 			{
 				// AllowUncookedAssetReferences should only be used when the DLC plugin to cook is going to be mounted where uncooked packages are available.
 				// This will allow a DLC plugin to be recooked continually and mounted in an uncooked editor which is useful for CI.
-				if (!bInAllowUncookedAssetReferences)
+				if (!InCOTFS.CookByTheBookOptions->bAllowUncookedAssetReferences)
 				{
 					UE_LOG(LogCook, Error, TEXT("Uncooked Engine or Game content %s is being referenced by DLC!"), *NameBuffer);
 				}
@@ -1557,9 +1542,34 @@ void FRequestCluster::IsRequestCookable(const ITargetPlatform* Platform, FName P
 		}
 	}
 
+	// The package is ordinarily cookable and explorable. In some cases we filter out for testing
+	// packages that are ordinarily cookable; set bOutCookable to false if so.
+	bOutExplorable = true;
+	if (InCOTFS.bCookFilter && !InCOTFS.CookFilterIncludedClasses.IsEmpty())
+	{
+		TOptional<FAssetPackageData> AssetData = IAssetRegistry::GetChecked().GetAssetPackageDataCopy(PackageName);
+		if (AssetData)
+		{
+			bool bIncluded = false;
+			for (FName ClassName : AssetData->ImportedClasses)
+			{
+				if (InCOTFS.CookFilterIncludedClasses.Contains(ClassName))
+				{
+					bIncluded = true;
+					break;
+				}
+			}
+			if (!bIncluded)
+			{
+				OutReason = ESuppressCookReason::CookFilter;
+				bOutCookable = false;
+				return;
+			}
+		}
+	}
+
 	OutReason = ESuppressCookReason::NotSuppressed;
 	bOutCookable = true;
-	bOutExplorable = true;
 }
 
 TConstArrayView<FName> FRequestCluster::GetLocalizationReferences(FName PackageName, UCookOnTheFlyServer& InCOTFS)
diff --git a/Engine/Source/Editor/UnrealEd/Private/Cooker/CookRequestCluster.h b/Engine/Source/Editor/UnrealEd/Private/Cooker/CookRequestCluster.h
--- a/Engine/Source/Editor/UnrealEd/Private/Cooker/CookRequestCluster.h
+++ b/Engine/Source/Editor/UnrealEd/Private/Cooker/CookRequestCluster.h
@@ -318,9 +318,8 @@ private:
 	void IsRequestCookable(const ITargetPlatform* TargetPlatform, FName PackageName, FPackageData& PackageData,
 		ESuppressCookReason& OutReason, bool& bOutCookable, bool& bOutExplorable);
 	static void IsRequestCookable(const ITargetPlatform* TargetPlatform, FName PackageName, FPackageData& PackageData,
-		FPackageDatas& InPackageDatas, FPackageTracker& InPackageTracker,
-		FStringView InDLCPath, bool bInErrorOnEngineContentUse, bool bInAllowUncookedAssetReferences,
-		bool bSkipOnlyEditorOnly, ESuppressCookReason& OutReason, bool& bOutCookable, bool& bOutExplorable);
+		UCookOnTheFlyServer& InCOTFS, FStringView InDLCPath, ESuppressCookReason& OutReason, bool& bOutCookable,
+		bool& bOutExplorable);
 	static void RandomizeCookOrder(TArray<FPackageData*>& InOutLeafToRootOrder,
 		const TMap<FPackageData*, TArray<FPackageData*>>& Edges);
 
@@ -338,7 +337,6 @@ private:
 	bool bAllowSoftDependencies = true;
 	bool bHybridIterativeEnabled = true;
 	bool bErrorOnEngineContentUse = false;
-	bool bAllowUncookedAssetReferences = false;
 	bool bPackageNamesComplete = false;
 	bool bDependenciesComplete = false;
 	bool bStartAsyncComplete = false;
diff --git a/Engine/Source/Editor/UnrealEd/Private/Cooker/CookTypes.cpp b/Engine/Source/Editor/UnrealEd/Private/Cooker/CookTypes.cpp
--- a/Engine/Source/Editor/UnrealEd/Private/Cooker/CookTypes.cpp
+++ b/Engine/Source/Editor/UnrealEd/Private/Cooker/CookTypes.cpp
@@ -55,6 +55,7 @@ const TCHAR* LexToString(ESuppressCookReason Reason)
 	case ESuppressCookReason::CookCanceled: return TEXT("CookCanceled");
 	case ESuppressCookReason::MultiprocessAssignmentError: return TEXT("MultiprocessAssignmentError");
 	case ESuppressCookReason::RetractedByCookDirector: return TEXT("RetractedByCookDirector");
+	case ESuppressCookReason::CookFilter: return TEXT("CookFilter");
 	default: return TEXT("Invalid");
 	}
 }
diff --git a/Engine/Source/Editor/UnrealEd/Private/Cooker/CookTypes.h b/Engine/Source/Editor/UnrealEd/Private/Cooker/CookTypes.h
--- a/Engine/Source/Editor/UnrealEd/Private/Cooker/CookTypes.h
+++ b/Engine/Source/Editor/UnrealEd/Private/Cooker/CookTypes.h
@@ -160,6 +160,7 @@ namespace UE::Cook
 		CookCanceled,
 		MultiprocessAssignmentError,
 		RetractedByCookDirector,
+		CookFilter,
 	};
 	const TCHAR* LexToString(UE::Cook::ESuppressCookReason Reason);
 
