AssetEditorSubsystem: Fix duplicate LOCTEXT warning

diff --git a/Engine/Source/Editor/UnrealEd/Private/Subsystems/AssetEditorSubsystem.cpp b/Engine/Source/Editor/UnrealEd/Private/Subsystems/AssetEditorSubsystem.cpp
--- a/Engine/Source/Editor/UnrealEd/Private/Subsystems/AssetEditorSubsystem.cpp
+++ b/Engine/Source/Editor/UnrealEd/Private/Subsystems/AssetEditorSubsystem.cpp
@@ -1288,7 +1288,7 @@ void UAssetEditorSubsystem::CreateRecentAssetsMenuForEditor(const IAssetEditorIn
 			InSection.AddSubMenu(
 				"RecentAssetEditorAssetsSubmenu",
 				FText::Format(LOCTEXT("RecentAssetEditorAssetsSubmenu_Label", "{0}"), FText::FromName(RecentAssetsMenuName)),
-				FText::Format(LOCTEXT("RecentAssetEditorAssetsSubmenu_Label", "Access your recently opened {0} assets"), FText::FromName(EditingAssetTypeName)),
+				FText::Format(LOCTEXT("RecentAssetEditorAssetsSubmenu_Tooltip", "Access your recently opened {0} assets"), FText::FromName(EditingAssetTypeName)),
 				FNewToolMenuDelegate::CreateUObject(this, &UAssetEditorSubsystem::CreateRecentAssetsMenu, AssetEditorName),
 				false,
 				FSlateIcon(FAppStyle::GetAppStyleSetName(), "Icons.RecentAssets")
