Structured error logging in more places in the engine

diff --git a/Engine/Source/Editor/UnrealEd/Private/Kismet2/CompilerResultsLog.cpp b/Engine/Source/Editor/UnrealEd/Private/Kismet2/CompilerResultsLog.cpp
--- a/Engine/Source/Editor/UnrealEd/Private/Kismet2/CompilerResultsLog.cpp
+++ b/Engine/Source/Editor/UnrealEd/Private/Kismet2/CompilerResultsLog.cpp
@@ -7,6 +7,7 @@
 #include "Editor/EditorPerProjectUserSettings.h"
 #include "MessageLogModule.h"
 #include "Logging/MessageLog.h"
+#include "Logging/StructuredLog.h"
 #include "Misc/UObjectToken.h"
 #include "SourceCodeNavigation.h"
 #include "IHotReload.h"
@@ -442,18 +443,19 @@ void FCompilerResultsLog::InternalLogMessage(FName MessageID, const TSharedRef<F
 		if (Severity == EMessageSeverity::Error)
 		{
 			if (IsRunningCommandlet())
-			{
-				UE_ASSET_LOG(LogBlueprint, Error, *SourcePath, TEXT("[Compiler] %s from Source: %s"), *Message->ToText().ToString(), *SourcePath);
+			{				
+				UE_LOGFMT_NSLOC(LogBlueprint, Error, "Blueprint", "CompilerError", "[Compiler] {ErrorMessage} from Source: {SourceFile}",
+												("ErrorMessage", Message->ToText().ToString()), ("SourceFile", SourcePath));
 			}
 			else
 			{
 				// in editor the compiler log is 'rich' and we don't need to annotate with the full blueprint path, just the name:
-				UE_ASSET_LOG(LogBlueprint, Error, *SourcePath, TEXT("[Compiler] %s"), *Message->ToText().ToString());
+				UE_ASSET_LOG(LogBlueprint, Error, *SourcePath, TEXT("[Compiler] %s"), *Message->ToText().ToString());				
 			}
 		}
 		else if (Severity == EMessageSeverity::Warning || Severity == EMessageSeverity::PerformanceWarning)
 		{
-			UE_ASSET_LOG(LogBlueprint, Warning, *SourcePath, TEXT("[Compiler] %s"), *Message->ToText().ToString());
+			UE_ASSET_LOG(LogBlueprint, Warning, *SourcePath, TEXT("[Compiler] %s"), *Message->ToText().ToString());			
 		}
 		else
 		{
