Fix collections in editor not showing up
#rnx

diff --git a/Engine/Source/Editor/UnrealEd/Classes/Settings/ContentBrowserSettings.h b/Engine/Source/Editor/UnrealEd/Classes/Settings/ContentBrowserSettings.h
--- a/Engine/Source/Editor/UnrealEd/Classes/Settings/ContentBrowserSettings.h
+++ b/Engine/Source/Editor/UnrealEd/Classes/Settings/ContentBrowserSettings.h
@@ -257,8 +257,7 @@ private:
 	UPROPERTY(config)
 	bool DisplayFavorites;
 
-	UPROPERTY(config)
-	bool DockCollections;
+	bool DockCollections = true;
 
 	UPROPERTY(config)
 	bool DisplayCppFolders;
