WorldPartitionConvertCommandlet - Batch SCC operations (add/checkout/delete) to greatly speed up processing
* Converting a 5000 actors level: before 54m 28s, after 3m 24s

diff --git a/Engine/Source/Editor/UnrealEd/Private/Commandlets/WorldPartitionConvertCommandlet.cpp b/Engine/Source/Editor/UnrealEd/Private/Commandlets/WorldPartitionConvertCommandlet.cpp
--- a/Engine/Source/Editor/UnrealEd/Private/Commandlets/WorldPartitionConvertCommandlet.cpp
+++ b/Engine/Source/Editor/UnrealEd/Private/Commandlets/WorldPartitionConvertCommandlet.cpp
@@ -816,32 +816,35 @@ int32 UWorldPartitionConvertCommandlet::Main(const FString& Params)
 		// Append DataLayer assets folder
 		CleanupPaths.Add(DataLayerAssetFolder);
 
+		TArray<FString> FilesToDelete;
+
 		for (const FString& CleanupPath : CleanupPaths)
 		{
 			FString Directory = FPackageName::LongPackageNameToFilename(CleanupPath);
 			if (IFileManager::Get().DirectoryExists(*Directory))
 			{
-				bool bResult = IFileManager::Get().IterateDirectoryRecursively(*Directory, [this, &PackageHelper](const TCHAR* FilenameOrDirectory, bool bIsDirectory)
+				IFileManager::Get().IterateDirectoryRecursively(*Directory, [this, &FilesToDelete](const TCHAR* FilenameOrDirectory, bool bIsDirectory)
 				{
 					if (!bIsDirectory)
 					{
 						FString Filename(FilenameOrDirectory);
 						if (Filename.EndsWith(FPackageName::GetAssetPackageExtension()))
 						{
-							return PackageHelper.Delete(Filename);
+							FilesToDelete.Emplace(MoveTemp(Filename));
 						}
 					}
 					return true;
 				});
-
-				if (!bResult)
-				{
-					UE_LOG(LogWorldPartitionConvertCommandlet, Error, TEXT("Failed to delete previous conversion package(s)"));
-					return 1;
-				}
 			}
 		}
 
+		bool bResult = PackageHelper.Delete(FilesToDelete);
+		if (!bResult)
+		{
+			UE_LOG(LogWorldPartitionConvertCommandlet, Error, TEXT("Failed to delete previous conversion package(s)"));
+			return 1;
+		}
+
 		if (FPackageName::SearchForPackageOnDisk(OldLevelName, &OldLevelName))
 		{
 			if (!PackageHelper.Delete(OldLevelName))
@@ -1581,12 +1584,10 @@ int32 UWorldPartitionConvertCommandlet::Main(const FString& Params)
 		{
 			TRACE_CPUPROFILER_EVENT_SCOPE(DeleteSourceLevels);
 
-			for (UPackage* Package : PackagesToDelete)
+			if (!PackageHelper.Delete(PackagesToDelete))
 			{
-				if (!PackageHelper.Delete(Package))
-				{
-					return 1;
-				}
+				UE_LOG(LogWorldPartitionConvertCommandlet, Error, TEXT("Failed to delete source level package(s)"));
+				return 1;
 			}
 		}
 
@@ -1595,17 +1596,24 @@ int32 UWorldPartitionConvertCommandlet::Main(const FString& Params)
 			TRACE_CPUPROFILER_EVENT_SCOPE(CheckoutPackages);
 
 			UE_LOG(LogWorldPartitionConvertCommandlet, Log, TEXT("Checking out %d packages."), PackagesToSave.Num());
+
+			TArray<FString> FilesToCheckout;
+			FilesToCheckout.Reserve(PackagesToSave.Num());
+
 			for(UPackage* Package: PackagesToSave)
 			{
 				FString PackageFileName = SourceControlHelpers::PackageFilename(Package);
 				if (FPlatformFileManager::Get().GetPlatformFile().FileExists(*PackageFileName))
 				{
-					if (!PackageHelper.Checkout(Package))
-					{
-						return 1;
-					}
+					FilesToCheckout.Emplace(MoveTemp(PackageFileName));
 				}
 			}
+
+			if (!PackageHelper.Checkout(FilesToCheckout))
+			{
+				UE_LOG(LogWorldPartitionConvertCommandlet, Error, TEXT("Failed to checkout package(s)"));
+				return 1;
+			}
 		}
 
 		SET_WARN_COLOR(COLOR_YELLOW);
@@ -1675,12 +1683,9 @@ int32 UWorldPartitionConvertCommandlet::Main(const FString& Params)
 			TRACE_CPUPROFILER_EVENT_SCOPE(AddPackagesToSourceControl);
 
 			// Add new packages to source control
-			for(UPackage* PackageToSave: PackagesToSave)
+			if(!PackageHelper.AddToSourceControl(PackagesToSave))
 			{
-				if(!PackageHelper.AddToSourceControl(PackageToSave))
-				{
-					return 1;
-				}
+				return 1;
 			}
 		}
 
diff --git a/Engine/Source/Editor/UnrealEd/Private/PackageSourceControlHelper.cpp b/Engine/Source/Editor/UnrealEd/Private/PackageSourceControlHelper.cpp
--- a/Engine/Source/Editor/UnrealEd/Private/PackageSourceControlHelper.cpp
+++ b/Engine/Source/Editor/UnrealEd/Private/PackageSourceControlHelper.cpp
@@ -53,6 +53,11 @@ bool FPackageSourceControlHelper::Delete(const FString& PackageName) const
 
 bool FPackageSourceControlHelper::Delete(const TArray<FString>& PackageNames, bool bErrorsAsWarnings) const
 {
+	if (PackageNames.IsEmpty())
+	{
+		return true;
+	}
+
 	bool bSuccess = true;
 	
 	// Early out when not using source control
@@ -233,7 +238,7 @@ bool FPackageSourceControlHelper::AddToSourceControl(UPackage* Package) const
 
 bool FPackageSourceControlHelper::AddToSourceControl(const TArray<UPackage*>& Packages, bool bErrorsAsWarnings) const
 {
-	if (!UseSourceControl())
+	if (!UseSourceControl() || PackageNames.IsEmpty())
 	{
 		return true;
 	}
@@ -244,14 +249,14 @@ bool FPackageSourceControlHelper::AddToSourceControl(const TArray<UPackage*>& Pa
 	for (const UPackage* Package : Packages)
 	{
 		FString PackageFilename = SourceControlHelpers::PackageFilename(Package);
-		PackageFilenames.Add(MoveTemp(PackageFilename));
+		PackageFilenames.Emplace(MoveTemp(PackageFilename));
 	}
 	return AddToSourceControl(PackageFilenames, bErrorsAsWarnings);
 }
 
 bool FPackageSourceControlHelper::AddToSourceControl(const TArray<FString>& PackageNames, bool bErrorsAsWarnings) const
 {
-	if (!UseSourceControl())
+	if (!UseSourceControl() || PackageNames.IsEmpty())
 	{
 		return true;
 	}
@@ -322,6 +327,11 @@ bool FPackageSourceControlHelper::Checkout(UPackage* Package) const
 
 bool FPackageSourceControlHelper::Checkout(const TArray<FString>& PackageNames, bool bErrorsAsWarnings) const
 {
+	if (PackageNames.IsEmpty())
+	{
+		return true;
+	}
+
 	const bool bUseSourceControl = UseSourceControl();
 
 	// Convert package names to package filenames
