Fix bad merge in CL 25488579
#rnx

diff --git a/Engine/Source/Editor/UnrealEd/Private/PackageSourceControlHelper.cpp b/Engine/Source/Editor/UnrealEd/Private/PackageSourceControlHelper.cpp
--- a/Engine/Source/Editor/UnrealEd/Private/PackageSourceControlHelper.cpp
+++ b/Engine/Source/Editor/UnrealEd/Private/PackageSourceControlHelper.cpp
@@ -238,7 +238,7 @@ bool FPackageSourceControlHelper::AddToSourceControl(UPackage* Package) const
 
 bool FPackageSourceControlHelper::AddToSourceControl(const TArray<UPackage*>& Packages, bool bErrorsAsWarnings) const
 {
-	if (!UseSourceControl() || PackageNames.IsEmpty())
+	if (!UseSourceControl() || Packages.IsEmpty())
 	{
 		return true;
 	}
