Added extra info to an MRU assert

#rnx

diff --git a/Engine/Source/Editor/UnrealEd/Private/MRUList.cpp b/Engine/Source/Editor/UnrealEd/Private/MRUList.cpp
--- a/Engine/Source/Editor/UnrealEd/Private/MRUList.cpp
+++ b/Engine/Source/Editor/UnrealEd/Private/MRUList.cpp
@@ -75,7 +75,7 @@ void FMRUList::MoveToTop(int32 InItem)
 
 void FMRUList::AddMRUItem(const FString& InItem)
 {
-	check(FPackageName::IsValidLongPackageName(InItem));
+	checkf(FPackageName::IsValidLongPackageName(InItem), TEXT("FMRUList::AddMRUItem called with an invalid package name: %s"), *InItem);
 
 	// See if the item already exists in the list.  If so,
 	// move it to the top of the list and leave.
