Mark transient package as bIsCookedForEditor when duplicating a class from a package that is bIsCookedForEditor, this prevents assymetrical stream reading

diff --git a/Engine/Source/Editor/UnrealEd/Private/Kismet2/KismetReinstanceUtilities.cpp b/Engine/Source/Editor/UnrealEd/Private/Kismet2/KismetReinstanceUtilities.cpp
--- a/Engine/Source/Editor/UnrealEd/Private/Kismet2/KismetReinstanceUtilities.cpp
+++ b/Engine/Source/Editor/UnrealEd/Private/Kismet2/KismetReinstanceUtilities.cpp
@@ -1813,7 +1813,15 @@ UClass* FBlueprintCompileReinstancer::MoveCDOToNewClass(UClass* OwnerClass, cons
 	const FName ReinstanceName = MakeUniqueObjectName(GetTransientPackage(), OwnerClass->GetClass(), *(FString(TEXT("REINST_")) + *OwnerClass->GetName()));
 
 	checkf(IsValid(OwnerClass), TEXT("%s is invalid - will not duplicate successfully"), *(OwnerClass->GetName()));
+
+	// UField metadata serialization assumes bIsCookedForEditor is stable:
+	UPackage* TransientPackage = GetTransientPackage();
+	UPackage* SourcePackage = OwnerClass->GetOutermost();
+	const bool bIsSourceCookedForEditor = SourcePackage->bIsCookedForEditor;
+	const bool bIsTransientCookedForEditor = TransientPackage->bIsCookedForEditor;
+	TransientPackage->bIsCookedForEditor = bIsSourceCookedForEditor;
 	UClass* CopyOfOwnerClass = CastChecked<UClass>(StaticDuplicateObject(OwnerClass, GetTransientPackage(), ReinstanceName, ~RF_Transactional));
+	TransientPackage->bIsCookedForEditor = bIsTransientCookedForEditor;
 
 	CopyOfOwnerClass->RemoveFromRoot();
 	OwnerClass->ClassFlags &= ~CLASS_NewerVersionExists;
