Limit the package build dependency tracker to exclude references to classes and references to CompiledIn packages.

diff --git a/Engine/Source/Editor/UnrealEd/Private/Cooker/PackageBuildDependencyTracker.cpp b/Engine/Source/Editor/UnrealEd/Private/Cooker/PackageBuildDependencyTracker.cpp
--- a/Engine/Source/Editor/UnrealEd/Private/Cooker/PackageBuildDependencyTracker.cpp
+++ b/Engine/Source/Editor/UnrealEd/Private/Cooker/PackageBuildDependencyTracker.cpp
@@ -5,6 +5,7 @@
 #include "HAL/Platform.h"
 #include "Logging/LogMacros.h"
 #include "Misc/PackageAccessTrackingOps.h"
+#include "UObject/Class.h"
 #include "UObject/Package.h"
 
 #if UE_WITH_PACKAGE_ACCESS_TRACKING
@@ -85,10 +86,18 @@ FPackageBuildDependencyTracker::~FPackageBuildDependencyTracker()
 	Disable();
 }
 
+static bool ShouldSkipDependency(const UObject* Object)
+{
+	return !Object ||
+		!GUObjectArray.IsValidIndex(Object) ||
+		!Object->HasAnyFlags(RF_Public) ||
+		(Object->GetClass() == UClass::StaticClass());
+}
+
 void FPackageBuildDependencyTracker::StaticOnObjectHandleRead(TArrayView<const UObject* const> Objects)
 {
 	int Count = Objects.Num();
-	if(Count == 0 || (Count == 1 && (!Objects[0] || !Objects[0]->HasAnyFlags(RF_Public))))
+	if(Count == 0 || (Count == 1 && ShouldSkipDependency(Objects[0])))
 	{
 		return;
 	}
@@ -106,19 +115,15 @@ void FPackageBuildDependencyTracker::StaticOnObjectHandleRead(TArrayView<const U
 
 	for (const UObject* ReadObject : Objects)
 	{
-		if (!ReadObject)
-		{
-			continue;
-		}
-
-		if (!ReadObject->HasAnyFlags(RF_Public))
+		if (ShouldSkipDependency(ReadObject))
 		{
 			continue;
 		}
 
 		FName Referencer = AccumulatedScopeData->PackageName;
-		FName Referenced = ReadObject->GetOutermost()->GetFName();
-		if (Referencer == Referenced)
+		UPackage* ReferencedPackage = ReadObject->GetOutermost();
+		FName Referenced = ReferencedPackage->GetFName();
+		if ((Referencer == Referenced) || ReferencedPackage->HasAnyPackageFlags(PKG_CompiledIn))
 		{
 			continue;
 		}
