Temporarily disable FPackageBuildDependencyTracker use during cook to measure impact on farm cook duration.

diff --git a/Engine/Source/Editor/UnrealEd/Private/UnrealEdMisc.cpp b/Engine/Source/Editor/UnrealEd/Private/UnrealEdMisc.cpp
--- a/Engine/Source/Editor/UnrealEd/Private/UnrealEdMisc.cpp
+++ b/Engine/Source/Editor/UnrealEd/Private/UnrealEdMisc.cpp
@@ -220,7 +220,7 @@ FUnrealEdMisc::FUnrealEdMisc() :
 {
 	//This is an early entry-point into the UnrealEd module to perform some editor-specific configuration
 #if UE_WITH_PACKAGE_ACCESS_TRACKING
-	const bool bBuildDependencyTrackingNeeded = GIsEditor && (IsRunningCookCommandlet() || !GetDefault<UEditorExperimentalSettings>()->bDisableCookInEditor);
+	const bool bBuildDependencyTrackingNeeded = false;//GIsEditor && (IsRunningCookCommandlet() || !GetDefault<UEditorExperimentalSettings>()->bDisableCookInEditor);
 	if (!bBuildDependencyTrackingNeeded)
 	{
 		FPackageBuildDependencyTracker::Get().Disable();
