Fixed rare crash at shutdown in SActorEditorContext

diff --git a/Engine/Source/Editor/UnrealEd/Private/SActorEditorContext.cpp b/Engine/Source/Editor/UnrealEd/Private/SActorEditorContext.cpp
--- a/Engine/Source/Editor/UnrealEd/Private/SActorEditorContext.cpp
+++ b/Engine/Source/Editor/UnrealEd/Private/SActorEditorContext.cpp
@@ -201,7 +201,7 @@ void SActorEditorContext::Rebuild()
 
 bool SActorEditorContext::IsVisible(UWorld* InWorld)
 {
-	if (InWorld && GEditor)
+	if (InWorld && GEditor && !IsEngineExitRequested())
 	{
 		return UActorEditorContextSubsystem::Get()->GetDisplayableClients().Num() > 0;
 	}
