ResavePackages: Avoid loading and saving external actors/objects individually. These already get handled when their level is resaved.

diff --git a/Engine/Source/Editor/UnrealEd/Private/Commandlets/ContentCommandlets.cpp b/Engine/Source/Editor/UnrealEd/Private/Commandlets/ContentCommandlets.cpp
--- a/Engine/Source/Editor/UnrealEd/Private/Commandlets/ContentCommandlets.cpp
+++ b/Engine/Source/Editor/UnrealEd/Private/Commandlets/ContentCommandlets.cpp
@@ -4,6 +4,7 @@
 	UCContentCommandlets.cpp: Various commmandlets.
 =============================================================================*/
 
+#include "Algo/RemoveIf.h"
 #include "AssetCompilingManager.h"
 #include "AssetRegistry/AssetData.h"
 #include "CollectionManagerModule.h"
@@ -329,6 +330,12 @@ int32 UResavePackagesCommandlet::InitializeResaveParameters( const TArray<FStrin
 		}
 	}
 
+	// Filtering out external actors/objects here avoids any issues that could happen with loading and saving them individually. They are instead handled within the context of their owning level in UResavePackagesCommandlet::PerformAdditionalOperations. For specifically resaving a World Partition map and its actors, see WorldPartitionResaveActorsBuilder.
+	PackageNames.SetNum(Algo::RemoveIf(PackageNames, [](const FString& PackageName)
+	{
+		return PackageName.Contains(FPackagePath::GetExternalActorsFolderName()) || PackageName.Contains(FPackagePath::GetExternalObjectsFolderName());
+	}));
+
 	// Check for a max package limit
 	MaxPackagesToResave = -1;
 	for ( int32 SwitchIdx = 0; SwitchIdx < Switches.Num(); SwitchIdx++ )
