Fbx exporter: add an option to export the highest LOD source data.

#rnx

diff --git a/Engine/Source/Editor/UnrealEd/Classes/Exporters/FbxExportOption.h b/Engine/Source/Editor/UnrealEd/Classes/Exporters/FbxExportOption.h
--- a/Engine/Source/Editor/UnrealEd/Classes/Exporters/FbxExportOption.h
+++ b/Engine/Source/Editor/UnrealEd/Classes/Exporters/FbxExportOption.h
@@ -50,9 +50,18 @@ public:
 	uint32 LevelOfDetail : 1;
 
 	/** If enabled, export collision */
-	UPROPERTY(EditAnywhere, BlueprintReadWrite, config, category = StaticMesh)
+	UPROPERTY(EditAnywhere, BlueprintReadWrite, config, category = StaticMesh, meta = (EditCondition = "!bExportSourceMesh"))
 	uint32 Collision : 1;
 
+	/*
+	 * If enabled, export the highest LOD source data instead of the render data.
+	 * Note:
+	 *     - No LOD will be exported for static meshes. (Level Of Detail option will be ignored)
+	 *     - No Collision will be exported for static meshes. (Collision option will be ignore)
+	 */
+	UPROPERTY(EditAnywhere, BlueprintReadWrite, config, category = StaticMesh)
+	uint32 bExportSourceMesh : 1;
+
 	/** If enabled, export the morph targets */
 	UPROPERTY(EditAnywhere, BlueprintReadWrite, config, category = SkeletalMesh)
 	uint32 bExportMorphTargets : 1;
diff --git a/Engine/Source/Editor/UnrealEd/Private/Fbx/FbxExportOption.cpp b/Engine/Source/Editor/UnrealEd/Private/Fbx/FbxExportOption.cpp
--- a/Engine/Source/Editor/UnrealEd/Private/Fbx/FbxExportOption.cpp
+++ b/Engine/Source/Editor/UnrealEd/Private/Fbx/FbxExportOption.cpp
@@ -23,6 +23,7 @@ UFbxExportOption::UFbxExportOption(const FObjectInitializer& ObjectInitializer)
 	bForceFrontXAxis = false;
 	LevelOfDetail = true;
 	Collision = true;
+	bExportSourceMesh = false;
 	bExportMorphTargets = true;
 	VertexColor = true;
 	MapSkeletalMotionToRoot = false;
diff --git a/Engine/Source/Editor/UnrealEd/Private/Fbx/FbxMainExport.cpp b/Engine/Source/Editor/UnrealEd/Private/Fbx/FbxMainExport.cpp
--- a/Engine/Source/Editor/UnrealEd/Private/Fbx/FbxMainExport.cpp
+++ b/Engine/Source/Editor/UnrealEd/Private/Fbx/FbxMainExport.cpp
@@ -211,6 +211,14 @@ void FFbxExporter::FillExportOptions(bool BatchMode, bool bShowOptionDialog, con
 	);
 
 	FSlateApplication::Get().AddModalWindow(Window, ParentWindow, false);
+
+	//Respect the edit condition
+	if(ExportOptionsUI->bExportSourceMesh)
+	{
+		ExportOptionsUI->Collision = false;
+		ExportOptionsUI->LevelOfDetail = false;
+	}
+
 	ExportOptionsUI->SaveOptions();
 
 	if (FbxOptionWindow->ShouldExport())
@@ -3846,223 +3854,197 @@ void FFbxExporter::ExportObjectMetadata(const UObject* ObjectToExport, FbxNode*
 	}
 }
 
-/**
- * Exports a static mesh
- * @param StaticMesh	The static mesh to export
- * @param MeshName		The name of the mesh for the FBX file
- * @param FbxActor		The fbx node representing the mesh
- * @param ExportLOD		The LOD of the mesh to export
- * @param LightmapUVChannel If set, performs a "lightmap export" and exports only the single given UV channel
- * @param ColorBuffer	Vertex color overrides to export
- * @param MaterialOrderOverride	Optional ordering of materials to set up correct material ID's across multiple meshes being export such as BSP surfaces which share common materials. Should be used sparingly
- */
-FbxNode* FFbxExporter::ExportStaticMeshToFbx(const UStaticMesh* StaticMesh, int32 ExportLOD, const TCHAR* MeshName, FbxNode* FbxActor, int32 LightmapUVChannel /*= -1*/, const FColorVertexBuffer* ColorBuffer /*= NULL*/, const TArray<FStaticMaterial>* MaterialOrderOverride /*= NULL*/, const TArray<UMaterialInterface*>* OverrideMaterials /*= NULL*/)
+bool FFbxExporter::ExportStaticMeshFromMeshDescription(FbxMesh* Mesh
+	, const UStaticMesh* StaticMesh
+	, const FMeshDescription* MeshDescription
+	, FbxNode* FbxActor
+	, int32 LightmapUVChannel
+	, const TArray<FStaticMaterial>* MaterialOrderOverride
+	, const TArray<UMaterialInterface*>* OverrideMaterials)
 {
-	FbxMesh* Mesh = nullptr;
-	if ((ExportLOD == 0 || ExportLOD == -1) && LightmapUVChannel == -1 && ColorBuffer == nullptr && MaterialOrderOverride == nullptr)
-	{
-		Mesh = FbxMeshes.FindRef(StaticMesh);
-	}
 
-	if (!Mesh)
+	if (MeshDescription->IsEmpty() || MeshDescription->Vertices().Num() == 0)
 	{
-		Mesh = FbxMesh::Create(Scene, TCHAR_TO_UTF8(MeshName));
-
-		const FStaticMeshLODResources& RenderMesh = StaticMesh->GetLODForExport(ExportLOD);
-
-		// Verify the integrity of the static mesh.
-		if (RenderMesh.VertexBuffers.StaticMeshVertexBuffer.GetNumVertices() == 0)
-		{
-			return nullptr;
-		}
+		return false;
+	}
 
-		if (RenderMesh.Sections.Num() == 0)
-		{
-			return nullptr;
-		}
+	FStaticMeshConstAttributes Attributes(*MeshDescription);
+	TVertexAttributesConstRef<FVector3f> VertexPositions = Attributes.GetVertexPositions();
+	TVertexInstanceAttributesConstRef<FVector3f> VertexInstanceNormals = Attributes.GetVertexInstanceNormals();
+	TVertexInstanceAttributesConstRef<FVector3f> VertexInstanceTangents = Attributes.GetVertexInstanceTangents();
+	TVertexInstanceAttributesConstRef<float> VertexInstanceBinormalSigns = Attributes.GetVertexInstanceBinormalSigns();
+	TVertexInstanceAttributesConstRef<FVector2f> VertexInstanceUVs = Attributes.GetVertexInstanceUVs();
+	TVertexInstanceAttributesConstRef<FVector4f> VertexInstanceColors = Attributes.GetVertexInstanceColors();
+	TPolygonGroupAttributesConstRef<FName> PolygonGroupMaterialSlotNames = Attributes.GetPolygonGroupMaterialSlotNames();
+	TEdgeAttributesConstRef<bool> EdgeHardnesses = Attributes.GetEdgeHardnesses();
 
-		// Remaps an Unreal vert to final reduced vertex list
-		TArray<int32> VertRemap;
-		TArray<int32> UniqueVerts;
+	const int32 VertexCount = MeshDescription->Vertices().Num();
+	const int32 VertexInstanceCount = MeshDescription->VertexInstances().Num();
 
-		// Weld verts
-		DetermineVertsToWeld(VertRemap, UniqueVerts, RenderMesh);
+	Mesh->InitControlPoints(VertexCount);
 
-		// Create and fill in the vertex position data source.
-		// The position vertices are duplicated, for some reason, retrieve only the first half vertices.
-		const int32 VertexCount = VertRemap.Num();
-		const int32 PolygonsCount = RenderMesh.Sections.Num();
+	FbxVector4* ControlPoints = Mesh->GetControlPoints();
+	for (int32 PosIndex = 0; PosIndex < VertexCount; ++PosIndex)
+	{
+		FVector Position = (FVector)VertexPositions[FVertexID(PosIndex)];
+		ControlPoints[PosIndex] = FbxVector4(Position.X, -Position.Y, Position.Z);
+	}
 
-		Mesh->InitControlPoints(UniqueVerts.Num());
+	// Set the normals on Layer 0.
+	FbxLayer* Layer = Mesh->GetLayer(0);
+	if (Layer == nullptr)
+	{
+		Mesh->CreateLayer();
+		Layer = Mesh->GetLayer(0);
+	}
 
-		FbxVector4* ControlPoints = Mesh->GetControlPoints();
-		for (int32 PosIndex = 0; PosIndex < UniqueVerts.Num(); ++PosIndex)
+	TArray<uint32> Indices;
+	Indices.Reserve(VertexInstanceCount);
+	for (const FPolygonGroupID& PolygonGroupID : MeshDescription->PolygonGroups().GetElementIDs())
+	{
+		for(const FTriangleID& TriangleID : MeshDescription->GetPolygonGroupTriangles(PolygonGroupID))
 		{
-			int32 UnrealPosIndex = UniqueVerts[PosIndex];
-			FVector Position = (FVector)RenderMesh.VertexBuffers.PositionVertexBuffer.VertexPosition(UnrealPosIndex);
-			ControlPoints[PosIndex] = FbxVector4(Position.X, -Position.Y, Position.Z);
+			TArrayView<const FVertexInstanceID> TriangleVertexInstanceIDs = MeshDescription->GetTriangleVertexInstances(TriangleID);
+			for(const FVertexInstanceID& VertexInstanceID : TriangleVertexInstanceIDs)
+			{
+				Indices.Add(MeshDescription->GetVertexInstanceVertex(VertexInstanceID).GetValue());
+			}
 		}
+	}
 
-		// Set the normals on Layer 0.
-		FbxLayer* Layer = Mesh->GetLayer(0);
-		if (Layer == nullptr)
-		{
-			Mesh->CreateLayer();
-			Layer = Mesh->GetLayer(0);
-		}
+	// Create and fill in the per-face-vertex normal data source.
+	// We extract the Z-tangent and the X/Y-tangents which are also stored in the render mesh.
+	FbxLayerElementNormal* LayerElementNormal = FbxLayerElementNormal::Create(Mesh, "");
+	FbxLayerElementTangent* LayerElementTangent = FbxLayerElementTangent::Create(Mesh, "");
+	FbxLayerElementBinormal* LayerElementBinormal = FbxLayerElementBinormal::Create(Mesh, "");
 
-		// Build list of Indices re-used multiple times to lookup Normals, UVs, other per face vertex information
-		TArray<uint32> Indices;
-		for (int32 PolygonsIndex = 0; PolygonsIndex < PolygonsCount; ++PolygonsIndex)
+	// Set 3 NTBs per triangle instead of storing on positional control points
+	LayerElementNormal->SetMappingMode(FbxLayerElement::eByPolygonVertex);
+	LayerElementTangent->SetMappingMode(FbxLayerElement::eByPolygonVertex);
+	LayerElementBinormal->SetMappingMode(FbxLayerElement::eByPolygonVertex);
+
+	// Set the NTBs values for every polygon vertex.
+	LayerElementNormal->SetReferenceMode(FbxLayerElement::eDirect);
+	LayerElementTangent->SetReferenceMode(FbxLayerElement::eDirect);
+	LayerElementBinormal->SetReferenceMode(FbxLayerElement::eDirect);
+
+	//Extract the tangent space
+	{
+		for (const FPolygonGroupID& PolygonGroupID : MeshDescription->PolygonGroups().GetElementIDs())
 		{
-			FIndexArrayView RawIndices = RenderMesh.IndexBuffer.GetArrayView();
-			const FStaticMeshSection& Polygons = RenderMesh.Sections[PolygonsIndex];
-			const uint32 TriangleCount = Polygons.NumTriangles;
-			for (uint32 TriangleIndex = 0; TriangleIndex < TriangleCount; ++TriangleIndex)
+			for (const FTriangleID& TriangleID : MeshDescription->GetPolygonGroupTriangles(PolygonGroupID))
 			{
-				for (uint32 PointIndex = 0; PointIndex < 3; PointIndex++)
+				TArrayView<const FVertexInstanceID> TriangleVertexInstanceIDs = MeshDescription->GetTriangleVertexInstances(TriangleID);
+				for (const FVertexInstanceID& VertexInstanceID : TriangleVertexInstanceIDs)
 				{
-					uint32 UnrealVertIndex = RawIndices[Polygons.FirstIndex + ((TriangleIndex * 3) + PointIndex)];
-					Indices.Add(UnrealVertIndex);
+					//Query the tangent space
+					FVector3f Normal, Tangent, Binormal;
+					Tangent = VertexInstanceTangents[VertexInstanceID];
+					Normal = VertexInstanceNormals[VertexInstanceID];
+					float BinormalSign = VertexInstanceBinormalSigns[VertexInstanceID];
+					Binormal = (FVector3f::CrossProduct(Normal, Tangent).GetSafeNormal() * BinormalSign);
+
+					FbxVector4 FbxNormal = FbxVector4(Normal.X, -Normal.Y, Normal.Z);
+					FbxNormal.Normalize();
+					LayerElementNormal->GetDirectArray().Add(FbxNormal);
+
+					FbxVector4 FbxTangent = FbxVector4(Tangent.X, -Tangent.Y, Tangent.Z);
+					FbxTangent.Normalize();
+					LayerElementTangent->GetDirectArray().Add(FbxTangent);
+
+					FbxVector4 FbxBinormal = FbxVector4(-Binormal.X, Binormal.Y, -Binormal.Z);
+					FbxBinormal.Normalize();
+					LayerElementBinormal->GetDirectArray().Add(FbxBinormal);
 				}
 			}
 		}
+	}
 
-		// Create and fill in the per-face-vertex normal data source.
-		// We extract the Z-tangent and the X/Y-tangents which are also stored in the render mesh.
-		FbxLayerElementNormal* LayerElementNormal = FbxLayerElementNormal::Create(Mesh, "");
-		FbxLayerElementTangent* LayerElementTangent = FbxLayerElementTangent::Create(Mesh, "");
-		FbxLayerElementBinormal* LayerElementBinormal = FbxLayerElementBinormal::Create(Mesh, "");
-
-		// Set 3 NTBs per triangle instead of storing on positional control points
-		LayerElementNormal->SetMappingMode(FbxLayerElement::eByPolygonVertex);
-		LayerElementTangent->SetMappingMode(FbxLayerElement::eByPolygonVertex);
-		LayerElementBinormal->SetMappingMode(FbxLayerElement::eByPolygonVertex);
-
-		// Set the NTBs values for every polygon vertex.
-		LayerElementNormal->SetReferenceMode(FbxLayerElement::eDirect);
-		LayerElementTangent->SetReferenceMode(FbxLayerElement::eDirect);
-		LayerElementBinormal->SetReferenceMode(FbxLayerElement::eDirect);
+	Layer->SetNormals(LayerElementNormal);
+	Layer->SetTangents(LayerElementTangent);
+	Layer->SetBinormals(LayerElementBinormal);
 
-		TArray<FbxVector4> FbxNormals;
-		TArray<FbxVector4> FbxTangents;
-		TArray<FbxVector4> FbxBinormals;
-		
-		FbxNormals.AddUninitialized(VertexCount);
-		FbxTangents.AddUninitialized(VertexCount);
-		FbxBinormals.AddUninitialized(VertexCount);
-		
-		for (int32 NTBIndex = 0; NTBIndex < VertexCount; ++NTBIndex)
+	// Create and fill in the per-face-vertex texture coordinate data source(s).
+	// Create UV for Diffuse channel.
+	int32 TexCoordSourceCount = (LightmapUVChannel == -1) ? MeshDescription->GetNumUVElementChannels() : LightmapUVChannel + 1;
+	int32 TexCoordSourceIndex = (LightmapUVChannel == -1) ? 0 : LightmapUVChannel;
+	for (; TexCoordSourceIndex < TexCoordSourceCount; ++TexCoordSourceIndex)
+	{
+		FbxLayer* UVsLayer = (LightmapUVChannel == -1) ? Mesh->GetLayer(TexCoordSourceIndex) : Mesh->GetLayer(0);
+		if (UVsLayer == NULL)
 		{
-			FVector3f Normal = (FVector3f)(RenderMesh.VertexBuffers.StaticMeshVertexBuffer.VertexTangentZ(NTBIndex));
-			FbxVector4& FbxNormal = FbxNormals[NTBIndex];
-			FbxNormal = FbxVector4(Normal.X, -Normal.Y, Normal.Z);
-			FbxNormal.Normalize();
-
-			FVector3f Tangent = (FVector3f)(RenderMesh.VertexBuffers.StaticMeshVertexBuffer.VertexTangentX(NTBIndex));
-			FbxVector4& FbxTangent = FbxTangents[NTBIndex];
-			FbxTangent = FbxVector4(Tangent.X, -Tangent.Y, Tangent.Z);
-			FbxTangent.Normalize();
-
-			FVector3f Binormal = -(FVector3f)(RenderMesh.VertexBuffers.StaticMeshVertexBuffer.VertexTangentY(NTBIndex));
-			FbxVector4& FbxBinormal = FbxBinormals[NTBIndex];
-			FbxBinormal = FbxVector4(Binormal.X, -Binormal.Y, Binormal.Z);
-			FbxBinormal.Normalize();
+			Mesh->CreateLayer();
+			UVsLayer = (LightmapUVChannel == -1) ? Mesh->GetLayer(TexCoordSourceIndex) : Mesh->GetLayer(0);
 		}
+		check(UVsLayer);
 
-		// Add one normal per each face index (3 per triangle)
-		for (int32 FbxVertIndex = 0; FbxVertIndex < Indices.Num(); FbxVertIndex++)
+		FString UVChannelNameBuilder = TEXT("UVmap_") + FString::FromInt(TexCoordSourceIndex);
+		const auto UVChannelNameUTF8 = TStringConversion<FTCHARToUTF8_Convert>(*UVChannelNameBuilder); // Do not inline it! The lifetime of this object needs to extend over the usage of the converted buffer.
+		const char* UVChannelName = UVChannelNameUTF8.Get(); // actually UTF8 as required by Fbx, but can't use UE's UTF8CHAR type because that's a uint8 aka *unsigned* char
+		if ((LightmapUVChannel >= 0) || ((LightmapUVChannel == -1) && (TexCoordSourceIndex == StaticMesh->GetLightMapCoordinateIndex())))
 		{
-			uint32 UnrealVertIndex = Indices[FbxVertIndex];
-			LayerElementNormal->GetDirectArray().Add(FbxNormals[UnrealVertIndex]);
-			LayerElementTangent->GetDirectArray().Add(FbxTangents[UnrealVertIndex]);
-			LayerElementBinormal->GetDirectArray().Add(FbxBinormals[UnrealVertIndex]);
+			UVChannelName = "LightMapUV";
 		}
 		
-		Layer->SetNormals(LayerElementNormal);
-		Layer->SetTangents(LayerElementTangent);
-		Layer->SetBinormals(LayerElementBinormal);
-		
-		FbxNormals.Empty();
-		FbxTangents.Empty();
-		FbxBinormals.Empty();
-
-		// Create and fill in the per-face-vertex texture coordinate data source(s).
-		// Create UV for Diffuse channel.
-		int32 TexCoordSourceCount = (LightmapUVChannel == -1) ? RenderMesh.VertexBuffers.StaticMeshVertexBuffer.GetNumTexCoords() : LightmapUVChannel + 1;
-		int32 TexCoordSourceIndex = (LightmapUVChannel == -1) ? 0 : LightmapUVChannel;
-		for (; TexCoordSourceIndex < TexCoordSourceCount; ++TexCoordSourceIndex)
-		{
-			FbxLayer* UVsLayer = (LightmapUVChannel == -1) ? Mesh->GetLayer(TexCoordSourceIndex) : Mesh->GetLayer(0);
-			if (UVsLayer == NULL)
-			{
-				Mesh->CreateLayer();
-				UVsLayer = (LightmapUVChannel == -1) ? Mesh->GetLayer(TexCoordSourceIndex) : Mesh->GetLayer(0);
-			}
-			check(UVsLayer);
-
-			FString UVChannelNameBuilder = TEXT("UVmap_") + FString::FromInt(TexCoordSourceIndex);
-			const auto UVChannelNameUTF8 = TStringConversion<FTCHARToUTF8_Convert>(*UVChannelNameBuilder); // Do not inline it! The lifetime of this object needs to extend over the usage of the converted buffer.
-			const char* UVChannelName = UVChannelNameUTF8.Get(); // actually UTF8 as required by Fbx, but can't use UE's UTF8CHAR type because that's a uint8 aka *unsigned* char
-			if ((LightmapUVChannel >= 0) || ((LightmapUVChannel == -1) && (TexCoordSourceIndex == StaticMesh->GetLightMapCoordinateIndex())))
-			{
-				UVChannelName = "LightMapUV";
-			}
-
-			FbxLayerElementUV* UVDiffuseLayer = FbxLayerElementUV::Create(Mesh, UVChannelName);
-
-			// Note: when eINDEX_TO_DIRECT is used, IndexArray must be 3xTriangle count, DirectArray can be smaller
-			UVDiffuseLayer->SetMappingMode(FbxLayerElement::eByPolygonVertex);
-			UVDiffuseLayer->SetReferenceMode(FbxLayerElement::eIndexToDirect);
-
-			TArray<int32> UvsRemap;
-			TArray<int32> UniqueUVs;
-			// Weld UVs
-			DetermineUVsToWeld(UvsRemap, UniqueUVs, RenderMesh.VertexBuffers.StaticMeshVertexBuffer, TexCoordSourceIndex);
+		FbxLayerElementUV* UVDiffuseLayer = FbxLayerElementUV::Create(Mesh, UVChannelName);
 
-			// Create the texture coordinate data source.
-			for (int32 FbxVertIndex = 0; FbxVertIndex < UniqueUVs.Num(); FbxVertIndex++)
-			{
-				int32 UnrealVertIndex = UniqueUVs[FbxVertIndex];
-				const FVector2f& TexCoord = RenderMesh.VertexBuffers.StaticMeshVertexBuffer.GetVertexUV(UnrealVertIndex, TexCoordSourceIndex);
-				UVDiffuseLayer->GetDirectArray().Add(FbxVector2(TexCoord.X, -TexCoord.Y + 1.0));
-			}
+		UVDiffuseLayer->SetMappingMode(FbxLayerElement::eByPolygonVertex);
+		UVDiffuseLayer->SetReferenceMode(FbxLayerElement::eDirect);
 
-			// For each face index, point to a texture uv
-			UVDiffuseLayer->GetIndexArray().SetCount(Indices.Num());
-			for (int32 FbxVertIndex = 0; FbxVertIndex < Indices.Num(); FbxVertIndex++)
+		for (const FPolygonGroupID& PolygonGroupID : MeshDescription->PolygonGroups().GetElementIDs())
+		{
+			for (const FTriangleID& TriangleID : MeshDescription->GetPolygonGroupTriangles(PolygonGroupID))
 			{
-				uint32 UnrealVertIndex = Indices[FbxVertIndex];
-				int32 NewVertIndex = UvsRemap[UnrealVertIndex];
-				UVDiffuseLayer->GetIndexArray().SetAt(FbxVertIndex, NewVertIndex);
+				TArrayView<const FVertexInstanceID> TriangleVertexInstanceIDs = MeshDescription->GetTriangleVertexInstances(TriangleID);
+				for (const FVertexInstanceID& VertexInstanceID : TriangleVertexInstanceIDs)
+				{
+					const FVector2f& VertexUV = VertexInstanceUVs.Get(VertexInstanceID, TexCoordSourceIndex);
+					UVDiffuseLayer->GetDirectArray().Add(FbxVector2(VertexUV.X, -VertexUV.Y + 1.0));
+				}
 			}
-
-			UVsLayer->SetUVs(UVDiffuseLayer, FbxLayerElement::eTextureDiffuse);
 		}
+		UVsLayer->SetUVs(UVDiffuseLayer, FbxLayerElement::eTextureDiffuse);
+	}
+
+	FbxLayerElementMaterial* MatLayer = FbxLayerElementMaterial::Create(Mesh, "");
+	MatLayer->SetMappingMode(FbxLayerElement::eByPolygon);
+	MatLayer->SetReferenceMode(FbxLayerElement::eIndexToDirect);
+	Layer->SetMaterials(MatLayer);
 
-		FbxLayerElementMaterial* MatLayer = FbxLayerElementMaterial::Create(Mesh, "");
-		MatLayer->SetMappingMode(FbxLayerElement::eByPolygon);
-		MatLayer->SetReferenceMode(FbxLayerElement::eIndexToDirect);
-		Layer->SetMaterials(MatLayer);
+	FbxLayerElementSmoothing* SmoothingInfoLayer = FbxLayerElementSmoothing::Create(Mesh, ""); 
+	SmoothingInfoLayer->SetMappingMode(FbxLayerElement::eByEdge);
+	SmoothingInfoLayer->SetReferenceMode(FbxLayerElement::eDirect);
+	Layer->SetSmoothing(SmoothingInfoLayer);
 
-		// Keep track of the number of tri's we export
-		uint32 AccountedTriangles = 0;
-		for (int32 PolygonsIndex = 0; PolygonsIndex < PolygonsCount; ++PolygonsIndex)
+	//Create the polygon with the correct material
+	{
+		int32 IndiceIndex = 0;
+		TSet<FEdgeID> ProcessEdges;
+		ProcessEdges.Reserve(Indices.Num());
+		for(const FPolygonGroupID& PolygonGroupID : MeshDescription->PolygonGroups().GetElementIDs())
 		{
-			const FStaticMeshSection& Polygons = RenderMesh.Sections[PolygonsIndex];
-			FIndexArrayView RawIndices = RenderMesh.IndexBuffer.GetArrayView();
 			UMaterialInterface* Material = nullptr;
-			
-			if (OverrideMaterials && OverrideMaterials->IsValidIndex(Polygons.MaterialIndex))
+
+			if (OverrideMaterials && OverrideMaterials->IsValidIndex(PolygonGroupID.GetValue()))
 			{
-				Material = (*OverrideMaterials)[Polygons.MaterialIndex];
+				Material = (*OverrideMaterials)[PolygonGroupID.GetValue()];
 			}
 			else
 			{
-				Material = StaticMesh->GetMaterial(Polygons.MaterialIndex);
+				FName CurrentMaterialSlotName = PolygonGroupMaterialSlotNames[PolygonGroupID];
+				int32 MaterialIndex = StaticMesh->GetMaterialIndexFromImportedMaterialSlotName(CurrentMaterialSlotName);
+				if(MaterialIndex == INDEX_NONE)
+				{
+					MaterialIndex = PolygonGroupID.GetValue();
+				}
+				if(!StaticMesh->GetStaticMaterials().IsValidIndex(MaterialIndex))
+				{
+					MaterialIndex = 0;
+				}
+				Material = StaticMesh->GetMaterial(MaterialIndex);
 			}
 
-			FbxSurfaceMaterial* FbxMaterial = Material ? ExportMaterial(Material) : NULL;
+			FbxSurfaceMaterial* FbxMaterial = Material ? ExportMaterial(Material) : nullptr;
 			if (!FbxMaterial)
 			{
 				FbxMaterial = CreateDefaultMaterial();
@@ -4076,88 +4058,459 @@ FbxNode* FFbxExporter::ExportStaticMeshToFbx(const UStaticMesh* StaticMesh, int3
 			{
 				ActualMatIndex = MaterialOrderOverride->Find(Material);
 			}
-			// Static meshes contain one triangle list per element.
-			// [GLAFORTE] Could it occasionally contain triangle strips? How do I know?
-			uint32 TriangleCount = Polygons.NumTriangles;
 
-			// Copy over the index buffer into the FBX polygons set.
-			for (uint32 TriangleIndex = 0; TriangleIndex < TriangleCount; ++TriangleIndex)
+			//Create the triangles
 			{
-				Mesh->BeginPolygon(ActualMatIndex);
-				for (uint32 PointIndex = 0; PointIndex < 3; PointIndex++)
+				for (const FTriangleID& TriangleID : MeshDescription->GetPolygonGroupTriangles(PolygonGroupID))
 				{
-					uint32 OriginalUnrealVertIndex = RawIndices[Polygons.FirstIndex + ((TriangleIndex * 3) + PointIndex)];
-					int32 RemappedVertIndex = VertRemap[OriginalUnrealVertIndex];
-					Mesh->AddPolygon(RemappedVertIndex);
+					FVertexID Corner[3];
+					int32 CornerIndex = 0;
+					Mesh->BeginPolygon(ActualMatIndex);
+					TArrayView<const FVertexInstanceID> TriangleVertexInstanceIDs = MeshDescription->GetTriangleVertexInstances(TriangleID);
+					for (const FVertexInstanceID& VertexInstanceID : TriangleVertexInstanceIDs)
+					{
+						Corner[CornerIndex] = MeshDescription->GetVertexInstanceVertex(VertexInstanceID);
+						Mesh->AddPolygon(Indices[IndiceIndex]);
+						IndiceIndex++;
+						CornerIndex++;
+					}
+					Mesh->EndPolygon();
 				}
-				Mesh->EndPolygon();
 			}
+		}
+		//Build the edge so we can set the edge hardness
+		Mesh->BuildMeshEdgeArray();
+		for (const FPolygonGroupID& PolygonGroupID : MeshDescription->PolygonGroups().GetElementIDs())
+		{
+			for (const FTriangleID& TriangleID : MeshDescription->GetPolygonGroupTriangles(PolygonGroupID))
+			{
+				FVertexID Corner[3];
+				int32 CornerIndex = 0;
+				TArrayView<const FVertexInstanceID> TriangleVertexInstanceIDs = MeshDescription->GetTriangleVertexInstances(TriangleID);
+				for (const FVertexInstanceID& VertexInstanceID : TriangleVertexInstanceIDs)
+				{
+					Corner[CornerIndex] = MeshDescription->GetVertexInstanceVertex(VertexInstanceID);
+					CornerIndex++;
+				}
+
+				//Add the smoothing group for the triangle
+				for (CornerIndex = 0; CornerIndex < 3; ++CornerIndex)
+				{
+					FVertexID EdgeStart = Corner[CornerIndex];
+					FVertexID EdgeEnd = Corner[(CornerIndex + 1) % 3];
+					FEdgeID MatchEdgeId = MeshDescription->GetVertexPairEdge(EdgeStart, EdgeEnd);
+					if (ProcessEdges.Contains(MatchEdgeId))
+					{
+						continue;
+					}
+					ProcessEdges.Add(MatchEdgeId);
 
-			AccountedTriangles += TriangleCount;
+					bool ReverseEdge = false;
+					int32 FbxEdgeIndex = Mesh->GetMeshEdgeIndex(EdgeStart.GetValue(), EdgeEnd.GetValue(), ReverseEdge);
+					if (FbxEdgeIndex == -1)
+					{
+						continue;
+					}
+					int32 EdgeHardnessValue = 1;
+					if (MatchEdgeId != INDEX_NONE)
+					{
+						EdgeHardnessValue = EdgeHardnesses[MatchEdgeId] ? 0 : 1;
+					}
+					int32 LayerAddIndex = SmoothingInfoLayer->GetDirectArray().Add(EdgeHardnessValue);
+					ensure(LayerAddIndex == FbxEdgeIndex);
+				}
+			}
 		}
+	}
 
-#ifdef TODO_FBX
-		// Throw a warning if this is a lightmap export and the exported poly count does not match the raw triangle data count
-		if (LightmapUVChannel != -1 && AccountedTriangles != RenderMesh.RawTriangles.GetElementCount())
+	// Create and fill in the vertex color data source.
+	uint32 ColorVertexCount = MeshDescription->VertexInstanceAttributes().HasAttribute(MeshAttribute::VertexInstance::Color) ? MeshDescription->VertexInstances().Num() : 0;
+
+	// Only export vertex colors if they exist
+	if (GetExportOptions()->VertexColor && ColorVertexCount > 0)
+	{
+		FbxLayerElementVertexColor* VertexColor = FbxLayerElementVertexColor::Create(Mesh, "");
+		VertexColor->SetMappingMode(FbxLayerElement::eByPolygonVertex);
+		VertexColor->SetReferenceMode(FbxLayerElement::eDirect);
+		FbxLayerElementArrayTemplate<FbxColor>& VertexColorArray = VertexColor->GetDirectArray();
+		Layer->SetVertexColors(VertexColor);
+
+		for (const FPolygonGroupID& PolygonGroupID : MeshDescription->PolygonGroups().GetElementIDs())
 		{
-			FMessageDialog::Open( EAppMsgType::Ok, NSLOCTEXT("UnrealEd", "StaticMeshEditor_LightmapExportFewerTriangles", "Fewer polygons have been exported than the raw triangle count.  This Lightmapped UV mesh may contain fewer triangles than the destination mesh on import.") );
+			for (const FTriangleID& TriangleID : MeshDescription->GetPolygonGroupTriangles(PolygonGroupID))
+			{
+				TArrayView<const FVertexInstanceID> TriangleVertexInstanceIDs = MeshDescription->GetTriangleVertexInstances(TriangleID);
+				for (const FVertexInstanceID& VertexInstanceID : TriangleVertexInstanceIDs)
+				{
+					const FVector4f& SourceVertexColor = VertexInstanceColors[VertexInstanceID];
+					FLinearColor LinearColor(SourceVertexColor);
+					//Convert to sRGB
+					FColor Color = LinearColor.ToFColor(true);
+					VertexColorArray.Add(FbxColor(static_cast<float>(Color.R)/255.0f, static_cast<float>(Color.G) / 255.0f, static_cast<float>(Color.B) / 255.0f, static_cast<float>(Color.A) / 255.0f));
+				}
+			}
 		}
+	}
+	return true;
+}
 
-		// Create and fill in the smoothing data source.
-		FbxLayerElementSmoothing* SmoothingInfo = FbxLayerElementSmoothing::Create(Mesh, "");
-		SmoothingInfo->SetMappingMode(FbxLayerElement::eByPolygon);
-		SmoothingInfo->SetReferenceMode(FbxLayerElement::eDirect);
-		FbxLayerElementArrayTemplate<int>& SmoothingArray = SmoothingInfo->GetDirectArray();
-		Layer->SetSmoothing(SmoothingInfo);
+bool FFbxExporter::ExportStaticMeshFromRenderData(FbxMesh* Mesh
+	, const UStaticMesh* StaticMesh
+	, const FStaticMeshLODResources& RenderMesh
+	, FbxNode* FbxActor
+	, int32 LightmapUVChannel
+	, const FColorVertexBuffer* ColorBuffer
+	, const TArray<FStaticMaterial>* MaterialOrderOverride
+	, const TArray<UMaterialInterface*>* OverrideMaterials)
+{
+	// Verify the integrity of the static mesh.
+	if (RenderMesh.VertexBuffers.StaticMeshVertexBuffer.GetNumVertices() == 0)
+	{
+		return false;
+	}
 
-		// This is broken. We are exporting the render mesh but providing smoothing
-		// information from the source mesh. The render triangles are not in the
-		// same order. Therefore we should export the raw mesh or not export
-		// smoothing group information!
-		int32 TriangleCount = RenderMesh.RawTriangles.GetElementCount();
-		FStaticMeshTriangle* RawTriangleData = (FStaticMeshTriangle*)RenderMesh.RawTriangles.Lock(LOCK_READ_ONLY);
-		for (int32 TriangleIndex = 0; TriangleIndex < TriangleCount; TriangleIndex++)
+	if (RenderMesh.Sections.Num() == 0)
+	{
+		return false;
+	}
+
+	// Remaps an Unreal vert to final reduced vertex list
+	TArray<int32> VertRemap;
+	TArray<int32> UniqueVerts;
+
+	// Weld verts
+	DetermineVertsToWeld(VertRemap, UniqueVerts, RenderMesh);
+
+	// Create and fill in the vertex position data source.
+	// The position vertices are duplicated, for some reason, retrieve only the first half vertices.
+	const int32 VertexCount = VertRemap.Num();
+	const int32 PolygonsCount = RenderMesh.Sections.Num();
+
+	Mesh->InitControlPoints(UniqueVerts.Num());
+
+	FbxVector4* ControlPoints = Mesh->GetControlPoints();
+	for (int32 PosIndex = 0; PosIndex < UniqueVerts.Num(); ++PosIndex)
+	{
+		int32 UnrealPosIndex = UniqueVerts[PosIndex];
+		FVector Position = (FVector)RenderMesh.VertexBuffers.PositionVertexBuffer.VertexPosition(UnrealPosIndex);
+		ControlPoints[PosIndex] = FbxVector4(Position.X, -Position.Y, Position.Z);
+	}
+
+	// Set the normals on Layer 0.
+	FbxLayer* Layer = Mesh->GetLayer(0);
+	if (Layer == nullptr)
+	{
+		Mesh->CreateLayer();
+		Layer = Mesh->GetLayer(0);
+	}
+
+	// Build list of Indices re-used multiple times to lookup Normals, UVs, other per face vertex information
+	TArray<uint32> Indices;
+	for (int32 PolygonsIndex = 0; PolygonsIndex < PolygonsCount; ++PolygonsIndex)
+	{
+		FIndexArrayView RawIndices = RenderMesh.IndexBuffer.GetArrayView();
+		const FStaticMeshSection& Polygons = RenderMesh.Sections[PolygonsIndex];
+		const uint32 TriangleCount = Polygons.NumTriangles;
+		for (uint32 TriangleIndex = 0; TriangleIndex < TriangleCount; ++TriangleIndex)
+		{
+			for (uint32 PointIndex = 0; PointIndex < 3; PointIndex++)
+			{
+				uint32 UnrealVertIndex = RawIndices[Polygons.FirstIndex + ((TriangleIndex * 3) + PointIndex)];
+				Indices.Add(UnrealVertIndex);
+			}
+		}
+	}
+
+	// Create and fill in the per-face-vertex normal data source.
+	// We extract the Z-tangent and the X/Y-tangents which are also stored in the render mesh.
+	FbxLayerElementNormal* LayerElementNormal = FbxLayerElementNormal::Create(Mesh, "");
+	FbxLayerElementTangent* LayerElementTangent = FbxLayerElementTangent::Create(Mesh, "");
+	FbxLayerElementBinormal* LayerElementBinormal = FbxLayerElementBinormal::Create(Mesh, "");
+
+	// Set 3 NTBs per triangle instead of storing on positional control points
+	LayerElementNormal->SetMappingMode(FbxLayerElement::eByPolygonVertex);
+	LayerElementTangent->SetMappingMode(FbxLayerElement::eByPolygonVertex);
+	LayerElementBinormal->SetMappingMode(FbxLayerElement::eByPolygonVertex);
+
+	// Set the NTBs values for every polygon vertex.
+	LayerElementNormal->SetReferenceMode(FbxLayerElement::eDirect);
+	LayerElementTangent->SetReferenceMode(FbxLayerElement::eDirect);
+	LayerElementBinormal->SetReferenceMode(FbxLayerElement::eDirect);
+
+	TArray<FbxVector4> FbxNormals;
+	TArray<FbxVector4> FbxTangents;
+	TArray<FbxVector4> FbxBinormals;
+
+	FbxNormals.AddUninitialized(VertexCount);
+	FbxTangents.AddUninitialized(VertexCount);
+	FbxBinormals.AddUninitialized(VertexCount);
+
+	for (int32 NTBIndex = 0; NTBIndex < VertexCount; ++NTBIndex)
+	{
+		FVector3f Normal = (FVector3f)(RenderMesh.VertexBuffers.StaticMeshVertexBuffer.VertexTangentZ(NTBIndex));
+		FbxVector4& FbxNormal = FbxNormals[NTBIndex];
+		FbxNormal = FbxVector4(Normal.X, -Normal.Y, Normal.Z);
+		FbxNormal.Normalize();
+
+		FVector3f Tangent = (FVector3f)(RenderMesh.VertexBuffers.StaticMeshVertexBuffer.VertexTangentX(NTBIndex));
+		FbxVector4& FbxTangent = FbxTangents[NTBIndex];
+		FbxTangent = FbxVector4(Tangent.X, -Tangent.Y, Tangent.Z);
+		FbxTangent.Normalize();
+
+		FVector3f Binormal = -(FVector3f)(RenderMesh.VertexBuffers.StaticMeshVertexBuffer.VertexTangentY(NTBIndex));
+		FbxVector4& FbxBinormal = FbxBinormals[NTBIndex];
+		FbxBinormal = FbxVector4(Binormal.X, -Binormal.Y, Binormal.Z);
+		FbxBinormal.Normalize();
+	}
+
+	// Add one normal per each face index (3 per triangle)
+	for (int32 FbxVertIndex = 0; FbxVertIndex < Indices.Num(); FbxVertIndex++)
+	{
+		uint32 UnrealVertIndex = Indices[FbxVertIndex];
+		LayerElementNormal->GetDirectArray().Add(FbxNormals[UnrealVertIndex]);
+		LayerElementTangent->GetDirectArray().Add(FbxTangents[UnrealVertIndex]);
+		LayerElementBinormal->GetDirectArray().Add(FbxBinormals[UnrealVertIndex]);
+	}
+
+	Layer->SetNormals(LayerElementNormal);
+	Layer->SetTangents(LayerElementTangent);
+	Layer->SetBinormals(LayerElementBinormal);
+
+	FbxNormals.Empty();
+	FbxTangents.Empty();
+	FbxBinormals.Empty();
+
+	// Create and fill in the per-face-vertex texture coordinate data source(s).
+	// Create UV for Diffuse channel.
+	int32 TexCoordSourceCount = (LightmapUVChannel == -1) ? RenderMesh.VertexBuffers.StaticMeshVertexBuffer.GetNumTexCoords() : LightmapUVChannel + 1;
+	int32 TexCoordSourceIndex = (LightmapUVChannel == -1) ? 0 : LightmapUVChannel;
+	for (; TexCoordSourceIndex < TexCoordSourceCount; ++TexCoordSourceIndex)
+	{
+		FbxLayer* UVsLayer = (LightmapUVChannel == -1) ? Mesh->GetLayer(TexCoordSourceIndex) : Mesh->GetLayer(0);
+		if (UVsLayer == NULL)
+		{
+			Mesh->CreateLayer();
+			UVsLayer = (LightmapUVChannel == -1) ? Mesh->GetLayer(TexCoordSourceIndex) : Mesh->GetLayer(0);
+		}
+		check(UVsLayer);
+
+		FString UVChannelNameBuilder = TEXT("UVmap_") + FString::FromInt(TexCoordSourceIndex);
+		const auto UVChannelNameUTF8 = TStringConversion<FTCHARToUTF8_Convert>(*UVChannelNameBuilder); // Do not inline it! The lifetime of this object needs to extend over the usage of the converted buffer.
+		const char* UVChannelName = UVChannelNameUTF8.Get(); // actually UTF8 as required by Fbx, but can't use UE's UTF8CHAR type because that's a uint8 aka *unsigned* char
+		if ((LightmapUVChannel >= 0) || ((LightmapUVChannel == -1) && (TexCoordSourceIndex == StaticMesh->GetLightMapCoordinateIndex())))
 		{
-			FStaticMeshTriangle* Triangle = (RawTriangleData++);
+			UVChannelName = "LightMapUV";
+		}
+
+		FbxLayerElementUV* UVDiffuseLayer = FbxLayerElementUV::Create(Mesh, UVChannelName);
+
+		// Note: when eINDEX_TO_DIRECT is used, IndexArray must be 3xTriangle count, DirectArray can be smaller
+		UVDiffuseLayer->SetMappingMode(FbxLayerElement::eByPolygonVertex);
+		UVDiffuseLayer->SetReferenceMode(FbxLayerElement::eIndexToDirect);
 
-			SmoothingArray.Add(Triangle->SmoothingMask);
+		TArray<int32> UvsRemap;
+		TArray<int32> UniqueUVs;
+		// Weld UVs
+		DetermineUVsToWeld(UvsRemap, UniqueUVs, RenderMesh.VertexBuffers.StaticMeshVertexBuffer, TexCoordSourceIndex);
+
+		// Create the texture coordinate data source.
+		for (int32 FbxVertIndex = 0; FbxVertIndex < UniqueUVs.Num(); FbxVertIndex++)
+		{
+			int32 UnrealVertIndex = UniqueUVs[FbxVertIndex];
+			const FVector2f& TexCoord = RenderMesh.VertexBuffers.StaticMeshVertexBuffer.GetVertexUV(UnrealVertIndex, TexCoordSourceIndex);
+			UVDiffuseLayer->GetDirectArray().Add(FbxVector2(TexCoord.X, -TexCoord.Y + 1.0));
 		}
-		RenderMesh.RawTriangles.Unlock();
+
+		// For each face index, point to a texture uv
+		UVDiffuseLayer->GetIndexArray().SetCount(Indices.Num());
+		for (int32 FbxVertIndex = 0; FbxVertIndex < Indices.Num(); FbxVertIndex++)
+		{
+			uint32 UnrealVertIndex = Indices[FbxVertIndex];
+			int32 NewVertIndex = UvsRemap[UnrealVertIndex];
+			UVDiffuseLayer->GetIndexArray().SetAt(FbxVertIndex, NewVertIndex);
+		}
+
+		UVsLayer->SetUVs(UVDiffuseLayer, FbxLayerElement::eTextureDiffuse);
+	}
+
+	FbxLayerElementMaterial* MatLayer = FbxLayerElementMaterial::Create(Mesh, "");
+	MatLayer->SetMappingMode(FbxLayerElement::eByPolygon);
+	MatLayer->SetReferenceMode(FbxLayerElement::eIndexToDirect);
+	Layer->SetMaterials(MatLayer);
+
+	// Keep track of the number of tri's we export
+	uint32 AccountedTriangles = 0;
+	for (int32 PolygonsIndex = 0; PolygonsIndex < PolygonsCount; ++PolygonsIndex)
+	{
+		const FStaticMeshSection& Polygons = RenderMesh.Sections[PolygonsIndex];
+		FIndexArrayView RawIndices = RenderMesh.IndexBuffer.GetArrayView();
+		UMaterialInterface* Material = nullptr;
+
+		if (OverrideMaterials && OverrideMaterials->IsValidIndex(Polygons.MaterialIndex))
+		{
+			Material = (*OverrideMaterials)[Polygons.MaterialIndex];
+		}
+		else
+		{
+			Material = StaticMesh->GetMaterial(Polygons.MaterialIndex);
+		}
+
+		FbxSurfaceMaterial* FbxMaterial = Material ? ExportMaterial(Material) : NULL;
+		if (!FbxMaterial)
+		{
+			FbxMaterial = CreateDefaultMaterial();
+		}
+		int32 MatIndex = FbxActor->AddMaterial(FbxMaterial);
+
+		// Determine the actual material index
+		int32 ActualMatIndex = MatIndex;
+
+		if (MaterialOrderOverride)
+		{
+			ActualMatIndex = MaterialOrderOverride->Find(Material);
+		}
+		// Static meshes contain one triangle list per element.
+		// [GLAFORTE] Could it occasionally contain triangle strips? How do I know?
+		uint32 TriangleCount = Polygons.NumTriangles;
+
+		// Copy over the index buffer into the FBX polygons set.
+		for (uint32 TriangleIndex = 0; TriangleIndex < TriangleCount; ++TriangleIndex)
+		{
+			Mesh->BeginPolygon(ActualMatIndex);
+			for (uint32 PointIndex = 0; PointIndex < 3; PointIndex++)
+			{
+				uint32 OriginalUnrealVertIndex = RawIndices[Polygons.FirstIndex + ((TriangleIndex * 3) + PointIndex)];
+				int32 RemappedVertIndex = VertRemap[OriginalUnrealVertIndex];
+				Mesh->AddPolygon(RemappedVertIndex);
+			}
+			Mesh->EndPolygon();
+		}
+
+		AccountedTriangles += TriangleCount;
+	}
+
+#ifdef TODO_FBX
+	// Throw a warning if this is a lightmap export and the exported poly count does not match the raw triangle data count
+	if (LightmapUVChannel != -1 && AccountedTriangles != RenderMesh.RawTriangles.GetElementCount())
+	{
+		FMessageDialog::Open(EAppMsgType::Ok, NSLOCTEXT("UnrealEd", "StaticMeshEditor_LightmapExportFewerTriangles", "Fewer polygons have been exported than the raw triangle count.  This Lightmapped UV mesh may contain fewer triangles than the destination mesh on import."));
+	}
+
+	// Create and fill in the smoothing data source.
+	FbxLayerElementSmoothing* SmoothingInfo = FbxLayerElementSmoothing::Create(Mesh, "");
+	SmoothingInfo->SetMappingMode(FbxLayerElement::eByPolygon);
+	SmoothingInfo->SetReferenceMode(FbxLayerElement::eDirect);
+	FbxLayerElementArrayTemplate<int>& SmoothingArray = SmoothingInfo->GetDirectArray();
+	Layer->SetSmoothing(SmoothingInfo);
+
+	// This is broken. We are exporting the render mesh but providing smoothing
+	// information from the source mesh. The render triangles are not in the
+	// same order. Therefore we should export the raw mesh or not export
+	// smoothing group information!
+	int32 TriangleCount = RenderMesh.RawTriangles.GetElementCount();
+	FStaticMeshTriangle* RawTriangleData = (FStaticMeshTriangle*)RenderMesh.RawTriangles.Lock(LOCK_READ_ONLY);
+	for (int32 TriangleIndex = 0; TriangleIndex < TriangleCount; TriangleIndex++)
+	{
+		FStaticMeshTriangle* Triangle = (RawTriangleData++);
+
+		SmoothingArray.Add(Triangle->SmoothingMask);
+	}
+	RenderMesh.RawTriangles.Unlock();
 #endif // #if TODO_FBX
 
-		// Create and fill in the vertex color data source.
-		const FColorVertexBuffer* ColorBufferToUse = ColorBuffer ? ColorBuffer : &RenderMesh.VertexBuffers.ColorVertexBuffer;
-		uint32 ColorVertexCount = ColorBufferToUse->GetNumVertices();
+	// Create and fill in the vertex color data source.
+	const FColorVertexBuffer* ColorBufferToUse = ColorBuffer ? ColorBuffer : &RenderMesh.VertexBuffers.ColorVertexBuffer;
+	uint32 ColorVertexCount = ColorBufferToUse->GetNumVertices();
+
+	// Only export vertex colors if they exist
+	if (GetExportOptions()->VertexColor && ColorVertexCount > 0)
+	{
+		FbxLayerElementVertexColor* VertexColor = FbxLayerElementVertexColor::Create(Mesh, "");
+		VertexColor->SetMappingMode(FbxLayerElement::eByPolygonVertex);
+		VertexColor->SetReferenceMode(FbxLayerElement::eIndexToDirect);
+		FbxLayerElementArrayTemplate<FbxColor>& VertexColorArray = VertexColor->GetDirectArray();
+		Layer->SetVertexColors(VertexColor);
 
-		// Only export vertex colors if they exist
-		if (GetExportOptions()->VertexColor && ColorVertexCount > 0)
+		for (int32 FbxVertIndex = 0; FbxVertIndex < Indices.Num(); FbxVertIndex++)
 		{
-			FbxLayerElementVertexColor* VertexColor = FbxLayerElementVertexColor::Create(Mesh, "");
-			VertexColor->SetMappingMode(FbxLayerElement::eByPolygonVertex);
-			VertexColor->SetReferenceMode(FbxLayerElement::eIndexToDirect);
-			FbxLayerElementArrayTemplate<FbxColor>& VertexColorArray = VertexColor->GetDirectArray();
-			Layer->SetVertexColors(VertexColor);
+			FLinearColor VertColor(1.0f, 1.0f, 1.0f);
+			uint32 UnrealVertIndex = Indices[FbxVertIndex];
+			if (UnrealVertIndex < ColorVertexCount)
+			{
+				VertColor = ColorBufferToUse->VertexColor(UnrealVertIndex).ReinterpretAsLinear();
+			}
+
+			VertexColorArray.Add(FbxColor(VertColor.R, VertColor.G, VertColor.B, VertColor.A));
+		}
+
+		VertexColor->GetIndexArray().SetCount(Indices.Num());
+		for (int32 FbxVertIndex = 0; FbxVertIndex < Indices.Num(); FbxVertIndex++)
+		{
+			VertexColor->GetIndexArray().SetAt(FbxVertIndex, FbxVertIndex);
+		}
+	}
+	return true;
+}
+
+/**
+ * Exports a static mesh
+ * @param StaticMesh	The static mesh to export
+ * @param MeshName		The name of the mesh for the FBX file
+ * @param FbxActor		The fbx node representing the mesh
+ * @param ExportLOD		The LOD of the mesh to export
+ * @param LightmapUVChannel If set, performs a "lightmap export" and exports only the single given UV channel
+ * @param ColorBuffer	Vertex color overrides to export
+ * @param MaterialOrderOverride	Optional ordering of materials to set up correct material ID's across multiple meshes being export such as BSP surfaces which share common materials. Should be used sparingly
+ */
+FbxNode* FFbxExporter::ExportStaticMeshToFbx(const UStaticMesh* StaticMesh, int32 ExportLOD, const TCHAR* MeshName, FbxNode* FbxActor, int32 LightmapUVChannel /*= -1*/, const FColorVertexBuffer* ColorBuffer /*= NULL*/, const TArray<FStaticMaterial>* MaterialOrderOverride /*= NULL*/, const TArray<UMaterialInterface*>* OverrideMaterials /*= NULL*/)
+{
+	FbxMesh* Mesh = nullptr;
+	if ((ExportLOD == 0 || ExportLOD == -1) && LightmapUVChannel == -1 && ColorBuffer == nullptr && MaterialOrderOverride == nullptr)
+	{
+		Mesh = FbxMeshes.FindRef(StaticMesh);
+	}
+
+	if (!Mesh)
+	{
+		Mesh = FbxMesh::Create(Scene, TCHAR_TO_UTF8(MeshName));
+
+		int32 LodIndex = ExportLOD == INDEX_NONE ? 0 : ExportLOD;
+		bool bExportSourceMesh = GetExportOptions()->bExportSourceMesh && !GetExportOptions()->LevelOfDetail && !GetExportOptions()->Collision;
+		const bool bUseNaniteData = LodIndex == 0 && StaticMesh->IsNaniteEnabled() && StaticMesh->IsHiResMeshDescriptionValid();
+		const bool bUseLodData = !bUseNaniteData && StaticMesh->IsMeshDescriptionValid(LodIndex);
+		bExportSourceMesh &= (bUseNaniteData || bUseLodData);
 
-			for (int32 FbxVertIndex = 0; FbxVertIndex < Indices.Num(); FbxVertIndex++)
+		if (bExportSourceMesh)
+		{
+			if (bUseNaniteData)
 			{
-				FLinearColor VertColor(1.0f, 1.0f, 1.0f);
-				uint32 UnrealVertIndex = Indices[FbxVertIndex];
-				if (UnrealVertIndex < ColorVertexCount)
+				//Export the nanite mesh description
+				if (!ExportStaticMeshFromMeshDescription(Mesh, StaticMesh, StaticMesh->GetHiResMeshDescription(), FbxActor, LightmapUVChannel, MaterialOrderOverride, OverrideMaterials))
 				{
-					VertColor = ColorBufferToUse->VertexColor(UnrealVertIndex).ReinterpretAsLinear();
+					return nullptr;
 				}
-
-				VertexColorArray.Add(FbxColor(VertColor.R, VertColor.G, VertColor.B, VertColor.A));
 			}
-
-			VertexColor->GetIndexArray().SetCount(Indices.Num());
-			for (int32 FbxVertIndex = 0; FbxVertIndex < Indices.Num(); FbxVertIndex++)
+			else
 			{
-				VertexColor->GetIndexArray().SetAt(FbxVertIndex, FbxVertIndex);
+				ensure(bUseLodData);
+				//Export the lod mesh description
+				if(!ExportStaticMeshFromMeshDescription(Mesh, StaticMesh, StaticMesh->GetMeshDescription(LodIndex), FbxActor, LightmapUVChannel, MaterialOrderOverride, OverrideMaterials))
+				{
+					return nullptr;
+				}
 			}
 		}
-
-		if ((ExportLOD == 0 || ExportLOD == -1) && LightmapUVChannel == -1 && ColorBuffer == nullptr && MaterialOrderOverride == nullptr)
+		else
+		{
+			//Export the render data
+			const FStaticMeshLODResources& RenderMesh = StaticMesh->GetLODForExport(LodIndex);
+			if (!ExportStaticMeshFromRenderData(Mesh, StaticMesh, RenderMesh, FbxActor, LightmapUVChannel, ColorBuffer, MaterialOrderOverride, OverrideMaterials))
+			{
+				return nullptr;
+			}
+		}
+		
+		if (LodIndex == 0 && LightmapUVChannel == -1 && ColorBuffer == nullptr && MaterialOrderOverride == nullptr)
 		{
 			FbxMeshes.Add(StaticMesh, Mesh);
 		}
diff --git a/Engine/Source/Editor/UnrealEd/Private/FbxExporter.h b/Engine/Source/Editor/UnrealEd/Private/FbxExporter.h
--- a/Engine/Source/Editor/UnrealEd/Private/FbxExporter.h
+++ b/Engine/Source/Editor/UnrealEd/Private/FbxExporter.h
@@ -292,6 +292,23 @@ private:
 	 */
 	FbxNode* ExportStaticMeshToFbx(const UStaticMesh* StaticMesh, int32 ExportLOD, const TCHAR* MeshName, FbxNode* FbxActor, int32 LightmapUVChannel = -1, const FColorVertexBuffer* ColorBuffer = nullptr, const TArray<FStaticMaterial>* MaterialOrderOverride = nullptr, const TArray<UMaterialInterface*>* OverrideMaterials = nullptr);
 
+	bool ExportStaticMeshFromMeshDescription(FbxMesh* Mesh
+		, const UStaticMesh* StaticMesh
+		, const FMeshDescription* MeshDescription
+		, FbxNode* FbxActor
+		, int32 LightmapUVChannel
+		, const TArray<FStaticMaterial>* MaterialOrderOverride
+		, const TArray<UMaterialInterface*>* OverrideMaterials);
+
+	bool ExportStaticMeshFromRenderData(FbxMesh* Mesh
+		, const UStaticMesh* StaticMesh
+		, const FStaticMeshLODResources& RenderMesh
+		, FbxNode* FbxActor
+		, int32 LightmapUVChannel
+		, const FColorVertexBuffer* ColorBuffer
+		, const TArray<FStaticMaterial>* MaterialOrderOverride
+		, const TArray<UMaterialInterface*>* OverrideMaterials);
+
 	/**
 	 * Exports a spline mesh
 	 * @param SplineMeshComp	The spline mesh component to export
