Don't apply editor context to child actor components.

diff --git a/Engine/Source/Editor/UnrealEd/Private/Subsystems/ActorEditorContextSubsystem.cpp b/Engine/Source/Editor/UnrealEd/Private/Subsystems/ActorEditorContextSubsystem.cpp
--- a/Engine/Source/Editor/UnrealEd/Private/Subsystems/ActorEditorContextSubsystem.cpp
+++ b/Engine/Source/Editor/UnrealEd/Private/Subsystems/ActorEditorContextSubsystem.cpp
@@ -51,7 +51,7 @@ void UActorEditorContextSubsystem::UnregisterClient(IActorEditorContextClient* C
 void UActorEditorContextSubsystem::ApplyContext(AActor* InActor)
 {
 	UWorld* World = GetWorld();
-	if (Clients.IsEmpty() || !World || !InActor || (InActor->GetWorld() != World))
+	if (Clients.IsEmpty() || !World || !InActor || (InActor->GetWorld() != World) || InActor->IsChildActor())
 	{
 		return;
 	}
