Fixing reimport item when editing an imported material

diff --git a/Engine/Source/Editor/UnrealEd/Private/Toolkits/AssetEditorToolkit.cpp b/Engine/Source/Editor/UnrealEd/Private/Toolkits/AssetEditorToolkit.cpp
--- a/Engine/Source/Editor/UnrealEd/Private/Toolkits/AssetEditorToolkit.cpp
+++ b/Engine/Source/Editor/UnrealEd/Private/Toolkits/AssetEditorToolkit.cpp
@@ -951,6 +951,15 @@ bool FAssetEditorToolkit::CanReimport( UObject* EditingObject ) const
 	// Don't allow user to perform certain actions on objects that aren't actually assets (e.g. Level Script blueprint objects)
 	if( EditingObject != NULL && EditingObject->IsAsset() )
 	{
+		// Apply the same logic as Reimport from the Context Menu, see FAssetFileContextMenu::AreImportedAssetActionsVisible
+		FAssetToolsModule& AssetToolsModule = FModuleManager::Get().LoadModuleChecked<FAssetToolsModule>("AssetTools");
+		UClass* EditingClass = EditingObject->GetClass();
+		auto AssetTypeActions = AssetToolsModule.Get().GetAssetTypeActionsForClass(EditingClass).Pin();
+		if(!AssetTypeActions.IsValid() || !AssetTypeActions->IsImportedAsset())
+		{
+			return false;
+		}
+
 		if ( FReimportManager::Instance()->CanReimport( EditingObject ) )
 		{
 			return true;
