Fixed cooking telemetry. The cooking callback originally in place did not fire in when cooking by the book was started from a commandlet. The CTBB Started and Ended events are now static so can be hooked up at any point. It is assumed that there will only ever be one instance of a UCookOnTheFlyServer per process.

diff --git a/Engine/Source/Editor/UnrealEd/Classes/CookOnTheSide/CookOnTheFlyServer.h b/Engine/Source/Editor/UnrealEd/Classes/CookOnTheSide/CookOnTheFlyServer.h
--- a/Engine/Source/Editor/UnrealEd/Classes/CookOnTheSide/CookOnTheFlyServer.h
+++ b/Engine/Source/Editor/UnrealEd/Classes/CookOnTheSide/CookOnTheFlyServer.h
@@ -831,8 +831,11 @@ public:
 	void OnObjectUpdated( UObject *Object );
 	void OnObjectSaved( UObject *ObjectSaved, FObjectPreSaveContext SaveContext );
 
+	DECLARE_MULTICAST_DELEGATE(FOnCookByTheBookStarted);
+	static FOnCookByTheBookStarted& OnCookByTheBookStarted() { return CookByTheBookStartedEvent; };
+
 	DECLARE_MULTICAST_DELEGATE(FOnCookByTheBookFinished);
-	FOnCookByTheBookFinished& OnCookByTheBookFinished() { return CookByTheBookFinishedEvent; };
+	static FOnCookByTheBookFinished& OnCookByTheBookFinished() { return CookByTheBookFinishedEvent; };
 	/**
 	* Marks a package as dirty for cook
 	* causes package to be recooked on next request (and all dependent packages which are currently cooked)
@@ -1466,7 +1469,9 @@ private:
 	double LastCookableObjectTickTime = 0.;
 
 	// Cook events that can be listenned to
-	FOnCookByTheBookFinished CookByTheBookFinishedEvent;
+	static FOnCookByTheBookStarted CookByTheBookStartedEvent;
+	static FOnCookByTheBookFinished CookByTheBookFinishedEvent;
+
 
 	// These structs are TUniquePtr rather than inline members so we can keep their headers private.
 	// See class header comments for their purpose.
diff --git a/Engine/Source/Editor/UnrealEd/Private/CookOnTheFlyServer.cpp b/Engine/Source/Editor/UnrealEd/Private/CookOnTheFlyServer.cpp
--- a/Engine/Source/Editor/UnrealEd/Private/CookOnTheFlyServer.cpp
+++ b/Engine/Source/Editor/UnrealEd/Private/CookOnTheFlyServer.cpp
@@ -210,6 +210,9 @@ static FAutoConsoleVariableRef CVarCookDisplayWarnBusyTime(
 /// Cook on the fly server
 ///////////////////////////////////////////////////////////////
 UCookOnTheFlyServer* UCookOnTheFlyServer::ActiveCOTFS = nullptr;
+UCookOnTheFlyServer::FOnCookByTheBookStarted UCookOnTheFlyServer::CookByTheBookStartedEvent;
+UCookOnTheFlyServer::FOnCookByTheBookFinished UCookOnTheFlyServer::CookByTheBookFinishedEvent;
+
 
 namespace UE
 {
@@ -10298,6 +10301,8 @@ void UCookOnTheFlyServer::StartCookByTheBook( const FCookByTheBookStartupOptions
 	InitializePollables();
 	RecordDLCPackagesFromBaseGame(BeginContext);
 	RegisterCookByTheBookDelegates();
+
+	CookByTheBookStartedEvent.Broadcast();
 }
 
 const UCookOnTheFlyServer::FCookByTheBookStartupOptions& UCookOnTheFlyServer::BlockOnPrebootCookGate(bool& bOutAbortCook,
