Editor: Add ability to check if given mode is the only active mode.

diff --git a/Engine/Source/Editor/UnrealEd/Private/EditorModeManager.cpp b/Engine/Source/Editor/UnrealEd/Private/EditorModeManager.cpp
--- a/Engine/Source/Editor/UnrealEd/Private/EditorModeManager.cpp
+++ b/Engine/Source/Editor/UnrealEd/Private/EditorModeManager.cpp
@@ -527,6 +527,11 @@ bool FEditorModeTools::IsOnlyVisibleActiveMode(FEditorModeID InMode) const
 	return !bFoundAnotherVisibleMode;
 }
 
+bool FEditorModeTools::IsOnlyActiveMode(FEditorModeID InMode) const
+{
+	return ActiveScriptableModes.Num() == 1 && ActiveScriptableModes[0]->GetID() == InMode;
+}
+
 void FEditorModeTools::OnEditorSelectionChanged(UObject* NewSelection)
 {
 	if (NewSelection == GetSelectedActors())
diff --git a/Engine/Source/Editor/UnrealEd/Public/EditorModeManager.h b/Engine/Source/Editor/UnrealEd/Public/EditorModeManager.h
--- a/Engine/Source/Editor/UnrealEd/Public/EditorModeManager.h
+++ b/Engine/Source/Editor/UnrealEd/Public/EditorModeManager.h
@@ -512,6 +512,7 @@ public:
 	void ActorPropChangeNotify();
 	void UpdateInternalData();
 	bool IsOnlyVisibleActiveMode(FEditorModeID InMode) const;
+	bool IsOnlyActiveMode(FEditorModeID InMode) const;
 
 	/*
 	* Sets the active Modes ToolBar Palette Tab to the named Palette
