Skip auto-saving map packages that are in /Temp when using the BackupAndOverwrite method

This makes its behavior consistent with the BackupAndRestore method, which has never saved /Temp maps

#rnx

diff --git a/Engine/Source/Editor/UnrealEd/Private/PackageAutoSaver.cpp b/Engine/Source/Editor/UnrealEd/Private/PackageAutoSaver.cpp
--- a/Engine/Source/Editor/UnrealEd/Private/PackageAutoSaver.cpp
+++ b/Engine/Source/Editor/UnrealEd/Private/PackageAutoSaver.cpp
@@ -304,19 +304,44 @@ void FPackageAutoSaver::AttemptAutoSave()
 					}
 				}
 
-				// Note: The in-place save does a regular save of dirty packages, so it doesn't set the bIsAutoSaving flag since it functions like a user-initiated save
-				const bool bAutoSaveMaps = LoadingSavingSettings->bAutoSaveMaps && DirtyMapsForAutoSave.Num() > 0;
-				const bool bAutoSaveContent = LoadingSavingSettings->bAutoSaveContent && DirtyContentForAutoSave.Num() > 0;
-
-				const bool bPromptUserForSave = false;
-				const bool bFastSave = false;
-				const bool bNotifyNoPackagesSaved = false;
-				const bool bCanBeDeclined = false;
-
-				bool bPackagesNeededSaving = false;
-				const bool bSuccess = FEditorFileUtils::SaveDirtyPackages(bPromptUserForSave, bAutoSaveMaps, bAutoSaveContent, bFastSave, bNotifyNoPackagesSaved, bCanBeDeclined, &bPackagesNeededSaving);
-				if (bPackagesNeededSaving)
+				// Build the complete list of packages to save
+				TArray<UPackage*> PackagesToSave;
 				{
+					auto AppendPackagesToSave = [&PackagesToSave](const TSet<TWeakObjectPtr<UPackage>, TWeakObjectPtrSetKeyFuncs<TWeakObjectPtr<UPackage>>>& PotentialPackagesToSave)
+					{
+						FNameBuilder PkgName;
+						for (const TWeakObjectPtr<UPackage>& PotentialPackageToSave : PotentialPackagesToSave)
+						{
+							if (UPackage* Pkg = PotentialPackageToSave.Get())
+							{
+								PkgName.Reset();
+								Pkg->GetFName().AppendString(PkgName);
+
+								// Skip packages in read-only roots (like /Temp)
+								if (FPackageName::IsValidLongPackageName(PkgName.ToView(), /*bIncludeReadOnlyRoots*/false))
+								{
+									PackagesToSave.Add(Pkg);
+								}
+							}
+						}
+					};
+
+					if (LoadingSavingSettings->bAutoSaveMaps)
+					{
+						AppendPackagesToSave(DirtyMapsForAutoSave);
+					}
+					if (LoadingSavingSettings->bAutoSaveContent)
+					{
+						AppendPackagesToSave(DirtyContentForAutoSave);
+					}
+				}
+
+				if (PackagesToSave.Num() > 0)
+				{
+					// Note: The in-place save does a regular save of dirty packages, so it doesn't set the bIsAutoSaving flag since it functions like a user-initiated save
+					const bool bSuccess = UEditorLoadingAndSavingUtils::SavePackages(PackagesToSave, false);
+
+					// Note: We don't update DirtyMapsForAutoSave/DirtyContentForAutoSave manually post-save, as any packages that were actually saved will have been removed from the list via the save callback
 					if (bSuccess)
 					{
 						MapsSaveResults = EAutosaveContentPackagesResult::Success;
