Sanitize inputs for UncompressCopyFile

diff --git a/Engine/Source/Editor/UnrealEd/Private/Factories/PackFactory.cpp b/Engine/Source/Editor/UnrealEd/Private/Factories/PackFactory.cpp
--- a/Engine/Source/Editor/UnrealEd/Private/Factories/PackFactory.cpp
+++ b/Engine/Source/Editor/UnrealEd/Private/Factories/PackFactory.cpp
@@ -6,6 +6,7 @@
 
 #include "Factories/PackFactory.h"
 #include "HAL/PlatformFileManager.h"
+#include "Math/GuardedInt.h"
 #include "Misc/MessageDialog.h"
 #include "HAL/FileManager.h"
 #include "Misc/FileHelper.h"
@@ -87,28 +88,79 @@ namespace PackFactoryHelper
 	// Utility function to uncompress and copy a single pak entry out of the Source archive and in to the Destination archive using PersistentBuffer as temporary space
 	bool UncompressCopyFile(FArchive& DestAr, FArchive& Source, const FPakEntry& Entry, TArray<uint8>& PersistentBuffer, const FPakFile& PakFile)
 	{
-		if (Entry.UncompressedSize == 0)
+		// Entry is untrusted data.
+		if (Entry.UncompressedSize <= 0)
 		{
 			return false;
 		}
 
-		int32 WorkingSize = Entry.CompressionBlockSize;
-		uint8 CompressionMethodIndex = IntCastChecked<uint8>(Entry.CompressionMethodIndex);
-		int32 MaxCompressionBlockSize = FCompression::GetMaximumCompressedSize(PakFile.GetInfo().GetCompressionMethod(CompressionMethodIndex), WorkingSize);
-		WorkingSize += MaxCompressionBlockSize;
+		TOptional<FName> CompressionMethod = PakFile.GetInfo().TryGetCompressionMethod(Entry.CompressionMethodIndex);
+		if (CompressionMethod.IsSet() == false)
+		{
+			return false;
+		}
+
+		FGuardedInt32 GuardedWorkingSize = FGuardedInt32(Entry.CompressionBlockSize);
+		int32 MaxCompressionBlockSize = FCompression::GetMaximumCompressedSize(CompressionMethod.GetValue(), GuardedWorkingSize.Get(0));
+		if (MaxCompressionBlockSize < 0)
+		{
+			return false;
+		}
+
+		GuardedWorkingSize += MaxCompressionBlockSize;
+		int32 WorkingSize = GuardedWorkingSize.Get(0);
+		if (WorkingSize <= 0)
+		{
+			return false;
+		}
+
+		// WorkingSize is now a sanitized int32 value.
 		if (PersistentBuffer.Num() < WorkingSize)
 		{
 			PersistentBuffer.SetNumUninitialized(WorkingSize);
 		}
 
-		uint8* UncompressedBuffer = PersistentBuffer.GetData() + MaxCompressionBlockSize;
-
 		for (uint32 BlockIndex = 0, BlockIndexNum = Entry.CompressionBlocks.Num(); BlockIndex < BlockIndexNum; ++BlockIndex)
 		{
-			int32 CompressedBlockSize = IntCastChecked<int32>(Entry.CompressionBlocks[BlockIndex].CompressedEnd - Entry.CompressionBlocks[BlockIndex].CompressedStart);
-			uint32 UncompressedBlockSize = (uint32)FMath::Min<int64>(Entry.UncompressedSize - Entry.CompressionBlockSize*BlockIndex, Entry.CompressionBlockSize);
-			Source.Seek(Entry.CompressionBlocks[BlockIndex].CompressedStart + (PakFile.GetInfo().HasRelativeCompressedChunkOffsets() ? Entry.Offset : 0));
+			FGuardedInt64 CompressedBlockSize64 = FGuardedInt64(Entry.CompressionBlocks[BlockIndex].CompressedEnd) - Entry.CompressionBlocks[BlockIndex].CompressedStart;
+			if (CompressedBlockSize64.ValidAndGreaterThan(0) == false ||
+				IntFitsIn<int32>(CompressedBlockSize64.Get(0)) == false)
+			{
+				return false;
+			}
+
+			int32 CompressedBlockSize = static_cast<int32>(CompressedBlockSize64.Get(0));
+			// CompressedBlockSize now sanitized
+
+
+			FGuardedInt64 UncompressedBlockSize64 = FGuardedInt64(Entry.UncompressedSize) - FGuardedInt64(Entry.CompressionBlockSize) * BlockIndex;
+			if (UncompressedBlockSize64.ValidAndGreaterThan(0) == false ||
+				IntFitsIn<uint32>(UncompressedBlockSize64.Get(0)) == false)
+			{
+				return false;
+			}
+
+			// CompressionBlockSize is guaranteed to fit in int32 from earlier, so we know after the Min() we can fit in uint32
+			uint32 UncompressedBlockSize = (uint32)FMath::Min<int64>(UncompressedBlockSize64.Get(0), Entry.CompressionBlockSize);
+
+			if (Entry.Offset < 0)
+			{
+				return false;
+			}
+
+			FGuardedInt64 OffsetInPak = FGuardedInt64(Entry.CompressionBlocks[BlockIndex].CompressedStart) + (PakFile.GetInfo().HasRelativeCompressedChunkOffsets() ? Entry.Offset : 0);
+			if (OffsetInPak.IsValid() == false)
+			{
+				return false;
+			}
+
+			Source.Seek(OffsetInPak.Get(0));
 			int32 SizeToRead = Entry.IsEncrypted() ? Align(CompressedBlockSize, FAES::AESBlockSize) : CompressedBlockSize;
+			if (SizeToRead > PersistentBuffer.Num())
+			{
+				// Shouldn't ever happen but I think this is possible if we are uncompressed and for some reason the block size isn't aligned to AESBlockSize.
+				return false;
+			}
 			Source.Serialize(PersistentBuffer.GetData(), SizeToRead);
 
 			if (Entry.IsEncrypted())
@@ -119,7 +171,8 @@ namespace PackFactoryHelper
 				FAES::DecryptData(PersistentBuffer.GetData(), SizeToRead, Key);
 			}
 
-			if (!FCompression::UncompressMemory(PakFile.GetInfo().GetCompressionMethod(CompressionMethodIndex), UncompressedBuffer, UncompressedBlockSize, PersistentBuffer.GetData(), CompressedBlockSize))
+			uint8* UncompressedBuffer = PersistentBuffer.GetData() + MaxCompressionBlockSize;
+			if (!FCompression::UncompressMemory(CompressionMethod.GetValue(), UncompressedBuffer, UncompressedBlockSize, PersistentBuffer.GetData(), CompressedBlockSize))
 			{
 				return false;
 			}
diff --git a/Engine/Source/Runtime/PakFile/Public/IPlatformFilePak.h b/Engine/Source/Runtime/PakFile/Public/IPlatformFilePak.h
--- a/Engine/Source/Runtime/PakFile/Public/IPlatformFilePak.h
+++ b/Engine/Source/Runtime/PakFile/Public/IPlatformFilePak.h
@@ -322,6 +322,15 @@ struct FPakInfo
 	{
 		return CompressionMethods[Index];
 	}
+	
+	TOptional<FName> TryGetCompressionMethod(uint32 Index) const
+	{
+		if (IntFitsIn<int32>(Index) == false)
+		{
+			return TOptional<FName>();
+		}
+		return (int32)Index >= CompressionMethods.Num() ? TOptional<FName>() : TOptional<FName>(CompressionMethods[(int32)Index]);
+	}
 };
 
 /**
