Fix for Cook Stats not being calculated anymore in CookOnTheFlyServer

diff --git a/Engine/Source/Editor/UnrealEd/Private/CookOnTheFlyServer.cpp b/Engine/Source/Editor/UnrealEd/Private/CookOnTheFlyServer.cpp
--- a/Engine/Source/Editor/UnrealEd/Private/CookOnTheFlyServer.cpp
+++ b/Engine/Source/Editor/UnrealEd/Private/CookOnTheFlyServer.cpp
@@ -9373,8 +9373,8 @@ void UCookOnTheFlyServer::ShutdownCookSession()
 		UnregisterCookByTheBookDelegates();
 
 		PrintFinishStats();
-
 		OutputHierarchyTimers();
+		PrintDetailedCookStats();
 	}
 	CookByTheBookOptions->ClearSessionData();
 	PlatformManager->ClearSessionPlatforms(*this);
