Fixed regressions when undoing a deleted OFPA actor post-save

UnloadPackages had been updated to always clear the transaction buffer by default, which resulted in CleanupAfterSuccessfulDelete always clearing the transaction buffer even when it had previously decided not to. CleanupAfterSuccessfulDelete now sets bResetTransBuffer to false when calling UnloadPackages, as it already handles clearing the transaction buffer when needed.

UWorldPartition keeps a list of PinnedActors (which a newly placed actor would be placed in), but would never remove an actor from this list when its FWorldPartitionActorDesc was deleted. This resulted in a newly placed, saved, then deleted, saved, and then undone, actor being unloaded once it was saved after the undo, as PinnedActors thought it was already tracking the actor so would skip taking a new reference which allowed the actor ref-count to reach zero and be unloaded.

diff --git a/Engine/Source/Editor/UnrealEd/Private/ObjectTools.cpp b/Engine/Source/Editor/UnrealEd/Private/ObjectTools.cpp
--- a/Engine/Source/Editor/UnrealEd/Private/ObjectTools.cpp
+++ b/Engine/Source/Editor/UnrealEd/Private/ObjectTools.cpp
@@ -2410,7 +2410,9 @@ namespace ObjectTools
 			AllPackagesToUnload.Append(PackagesToDelete);
 			AllPackagesToUnload.Append(PackagesToUnload);
 
-			UPackageTools::UnloadPackages(AllPackagesToUnload);
+			UPackageTools::FUnloadPackageParams UnloadParams(AllPackagesToUnload);
+			UnloadParams.bResetTransBuffer = false; // Don't reset the transaction buffer, as we handled that above if needed
+			UPackageTools::UnloadPackages(UnloadParams);
 		}
 		CollectGarbage( GARBAGE_COLLECTION_KEEPFLAGS );
 
diff --git a/Engine/Source/Runtime/Engine/Private/WorldPartition/WorldPartition.cpp b/Engine/Source/Runtime/Engine/Private/WorldPartition/WorldPartition.cpp
--- a/Engine/Source/Runtime/Engine/Private/WorldPartition/WorldPartition.cpp
+++ b/Engine/Source/Runtime/Engine/Private/WorldPartition/WorldPartition.cpp
@@ -1192,7 +1192,12 @@ void UWorldPartition::OnActorDescAdded(FWorldPartitionActorDesc* NewActorDesc)
 }
 
 void UWorldPartition::OnActorDescRemoved(FWorldPartitionActorDesc* ActorDesc)
-{	
+{
+	if (PinnedActors)
+	{
+		PinnedActors->RemoveActors({ FWorldPartitionHandle(ActorDesc->GetContainer(), ActorDesc->GetGuid()) });
+	}
+
 	UnhashActorDesc(ActorDesc);
 
 	if (ForceLoadedActors)
