Don't preemptively mark lingering OFPA packages as dirty

This was initially done so that undoing the delete of a OFPA actor would allow the restored actor to be saved, but this dirtying is now better handled by the transaction buffer (which considers a save as a fence point, and will dirty any object changes that undo beyond a fence).

#rnx

diff --git a/Engine/Source/Editor/UnrealEd/Private/ObjectTools.cpp b/Engine/Source/Editor/UnrealEd/Private/ObjectTools.cpp
--- a/Engine/Source/Editor/UnrealEd/Private/ObjectTools.cpp
+++ b/Engine/Source/Editor/UnrealEd/Private/ObjectTools.cpp
@@ -2539,7 +2539,6 @@ namespace ObjectTools
 					if (UPackage* Package = FindPackage(nullptr, *PackageName))
 					{
 						Package->MarkAsNewlyCreated();
-						Package->SetDirtyFlag(true);
 					}
 				}
 			}
