Basic selection state is mirrored to Typed Element Data Storage for actors

diff --git a/Engine/Source/Editor/UnrealEd/Private/Elements/Actor/ActorElementEditorSelectionInterface.cpp b/Engine/Source/Editor/UnrealEd/Private/Elements/Actor/ActorElementEditorSelectionInterface.cpp
--- a/Engine/Source/Editor/UnrealEd/Private/Elements/Actor/ActorElementEditorSelectionInterface.cpp
+++ b/Engine/Source/Editor/UnrealEd/Private/Elements/Actor/ActorElementEditorSelectionInterface.cpp
@@ -1,11 +1,13 @@
 // Copyright Epic Games, Inc. All Rights Reserved.
-
+ 
 #include "Elements/Actor/ActorElementEditorSelectionInterface.h"
-
+ 
 #include "Elements/Actor/ActorElementData.h"
+#include "Elements/Columns/TypedElementSelectionColumns.h"
 #include "Elements/Framework/EngineElementsLibrary.h"
 #include "Elements/Framework/TypedElementHandle.h"
 #include "Elements/Framework/TypedElementList.h"
+#include "Elements/Framework/TypedElementRegistry.h"
 #include "Elements/Interfaces/TypedElementSelectionInterface.h"
 #include "Elements/Object/ObjectElementEditorSelectionInterface.h"
 #include "GameFramework/Actor.h"
@@ -13,7 +15,7 @@
 #include "Templates/SharedPointer.h"
 #include "UObject/WeakObjectPtr.h"
 #include "UObject/WeakObjectPtrTemplates.h"
-
+ 
 class FActorElementTransactedElement : public ITypedElementTransactedElement
 {
 private:
@@ -21,7 +23,7 @@ private:
 	{
 		return MakeUnique<FActorElementTransactedElement>(*this);
 	}
-
+ 
 	virtual FTypedElementHandle GetElementImpl() const override
 	{
 		const AActor* Actor = ActorPtr.Get(/*bEvenIfPendingKill*/true);
@@ -29,45 +31,84 @@ private:
 			? UEngineElementsLibrary::AcquireEditorActorElementHandle(Actor)
 			: FTypedElementHandle();
 	}
-
+ 
 	virtual void SetElementImpl(const FTypedElementHandle& InElementHandle) override
 	{
 		AActor* Actor = ActorElementDataUtil::GetActorFromHandleChecked(InElementHandle);
 		ActorPtr = Actor;
 	}
-
+ 
 	virtual void SerializeImpl(FArchive& InArchive) override
 	{
 		InArchive << ActorPtr;
 	}
-
+ 
 	TWeakObjectPtr<AActor> ActorPtr;
 };
-
+ 
+bool UActorElementEditorSelectionInterface::SelectElement(const FTypedElementHandle& InElementHandle, const FTypedElementListPtr& InSelectionSet, const FTypedElementSelectionOptions& InSelectionOptions)
+{
+	const AActor* Actor = ActorElementDataUtil::GetActorFromHandle(InElementHandle);
+	UTypedElementRegistry* Registry = InSelectionSet->GetRegistry();
+	// Add a selection column in TEDS
+	if (const ITypedElementDataStorageCompatibilityInterface* Compatibility = Registry->GetDataStorageCompatibility())
+	{
+		TypedElementRowHandle Row = Compatibility->FindRowWithCompatibleObject(Actor);
+		if (Row != TypedElementInvalidRowHandle)
+		{
+			if (ITypedElementDataStorageInterface* DataStorage = UTypedElementRegistry::GetInstance()->GetMutableDataStorage())
+			{
+				DataStorage->AddColumn(Row, FTypedElementSelectionColumn::StaticStruct());
+			}
+		}
+	}
+	
+	return Super::SelectElement(InElementHandle, InSelectionSet, InSelectionOptions);
+}
+ 
+bool UActorElementEditorSelectionInterface::DeselectElement(const FTypedElementHandle& InElementHandle, const FTypedElementListPtr& InSelectionSet, const FTypedElementSelectionOptions& InSelectionOptions)
+{
+	const AActor* Actor = ActorElementDataUtil::GetActorFromHandle(InElementHandle);
+	UTypedElementRegistry* Registry = InSelectionSet->GetRegistry();
+	if (const ITypedElementDataStorageCompatibilityInterface* Compatibility = Registry->GetDataStorageCompatibility())
+	{
+		TypedElementRowHandle Row = Compatibility->FindRowWithCompatibleObject(Actor);
+		if (Row != TypedElementInvalidRowHandle)
+		{
+			if (ITypedElementDataStorageInterface* DataStorage = UTypedElementRegistry::GetInstance()->GetMutableDataStorage())
+            {
+            	DataStorage->RemoveColumn(Row, FTypedElementSelectionColumn::StaticStruct());
+            }
+		}
+	}
+	
+	return Super::DeselectElement(InElementHandle, InSelectionSet, InSelectionOptions);
+}
+ 
 bool UActorElementEditorSelectionInterface::IsElementSelected(const FTypedElementHandle& InElementHandle, const FTypedElementListConstPtr& SelectionSetPtr, const FTypedElementIsSelectedOptions& InSelectionOptions)
 {
 	const AActor* Actor = ActorElementDataUtil::GetActorFromHandle(InElementHandle);
 	return SelectionSetPtr && Actor && IsActorSelected(Actor, SelectionSetPtr.ToSharedRef(), InSelectionOptions);
 }
-
+ 
 bool UActorElementEditorSelectionInterface::ShouldPreventTransactions(const FTypedElementHandle& InElementHandle)
 {
 	const AActor* Actor = ActorElementDataUtil::GetActorFromHandle(InElementHandle);
 	return Actor && UObjectElementEditorSelectionInterface::ShouldObjectPreventTransactions(Actor);
 }
-
+ 
 TUniquePtr<ITypedElementTransactedElement> UActorElementEditorSelectionInterface::CreateTransactedElementImpl()
 {
 	return MakeUnique<FActorElementTransactedElement>();
 }
-
+ 
 bool UActorElementEditorSelectionInterface::IsActorSelected(const AActor* InActor, FTypedElementListConstRef InSelectionSet, const FTypedElementIsSelectedOptions& InSelectionOptions)
 {
 	if (InSelectionSet->Num() == 0)
 	{
 		return false;
 	}
-
+ 
 	if (FTypedElementHandle ActorElement = UEngineElementsLibrary::AcquireEditorActorElementHandle(InActor, /*bAllowCreate*/false))
 	{
 		if (InSelectionSet->Contains(ActorElement))
@@ -75,7 +116,7 @@ bool UActorElementEditorSelectionInterface::IsActorSelected(const AActor* InActo
 			return true;
 		}
 	}
-
+ 
 	if (InSelectionOptions.AllowIndirect())
 	{
 		if (const AActor* RootSelectionActor = InActor->GetRootSelectionParent())
@@ -86,6 +127,6 @@ bool UActorElementEditorSelectionInterface::IsActorSelected(const AActor* InActo
 			}
 		}
 	}
-
+ 
 	return false;
-}
+}
\ No newline at end of file
diff --git a/Engine/Source/Editor/UnrealEd/Public/Elements/Actor/ActorElementEditorSelectionInterface.h b/Engine/Source/Editor/UnrealEd/Public/Elements/Actor/ActorElementEditorSelectionInterface.h
--- a/Engine/Source/Editor/UnrealEd/Public/Elements/Actor/ActorElementEditorSelectionInterface.h
+++ b/Engine/Source/Editor/UnrealEd/Public/Elements/Actor/ActorElementEditorSelectionInterface.h
@@ -22,6 +22,8 @@ class UNREALED_API UActorElementEditorSelectionInterface : public UActorElementS
 	GENERATED_BODY()
 
 public:
+	virtual bool SelectElement(const FTypedElementHandle& InElementHandle, const FTypedElementListPtr& InSelectionSet, const FTypedElementSelectionOptions& InSelectionOptions) override;
+	virtual bool DeselectElement(const FTypedElementHandle& InElementHandle, const FTypedElementListPtr& InSelectionSet, const FTypedElementSelectionOptions& InSelectionOptions) override;
 	virtual bool IsElementSelected(const FTypedElementHandle& InElementHandle, const FTypedElementListConstPtr& SelectionSetPtr, const FTypedElementIsSelectedOptions& InSelectionOptions) override;
 	virtual bool ShouldPreventTransactions(const FTypedElementHandle& InElementHandle) override;
 	virtual TUniquePtr<ITypedElementTransactedElement> CreateTransactedElementImpl() override;
diff --git a/Engine/Source/Runtime/TypedElementFramework/Public/Elements/Columns/TypedElementSelectionColumns.h b/Engine/Source/Runtime/TypedElementFramework/Public/Elements/Columns/TypedElementSelectionColumns.h
new file mode 100644
--- /dev/null
+++ b/Engine/Source/Runtime/TypedElementFramework/Public/Elements/Columns/TypedElementSelectionColumns.h
@@ -0,0 +1,17 @@
+﻿// Copyright Epic Games, Inc. All Rights Reserved.
+
+#pragma once
+
+#include "Elements/Interfaces/TypedElementDataStorageInterface.h"
+#include "UObject/ObjectMacros.h"
+
+#include "TypedElementSelectionColumns.generated.h"
+
+/**
+ * Column to represent that a row is selected
+ */
+USTRUCT(meta = (DisplayName = "Selected"))
+struct FTypedElementSelectionColumn final : public FTypedElementDataStorageColumn
+{
+	GENERATED_BODY()
+};
