Fix HotReload and Live Coding to not bother serializing delegates and bulk data.  Delegate serialization was causing asserts with new concurrency checks and could have resulted in crashes in the past.

diff --git a/Engine/Source/Editor/UnrealEd/Private/Kismet2/ReloadUtilities.cpp b/Engine/Source/Editor/UnrealEd/Private/Kismet2/ReloadUtilities.cpp
--- a/Engine/Source/Editor/UnrealEd/Private/Kismet2/ReloadUtilities.cpp
+++ b/Engine/Source/Editor/UnrealEd/Private/Kismet2/ReloadUtilities.cpp
@@ -1244,6 +1244,7 @@ void FReload::ReplaceReferencesToReconstructedCDOs()
 					{
 						ArIsObjectReferenceCollector = true;
 						ArIgnoreOuterRef = true;
+						ArShouldSkipBulkData = true;
 					}
 
 					virtual FString GetArchiveName() const override
@@ -1266,6 +1267,11 @@ void FReload::ReplaceReferencesToReconstructedCDOs()
 						return *this;
 					}
 
+					virtual bool ShouldSkipProperty(const FProperty* InProperty) const
+					{
+						return InProperty->GetClass()->HasAnyCastFlags(CASTCLASS_FDelegateProperty | CASTCLASS_FMulticastDelegateProperty | CASTCLASS_FMulticastInlineDelegateProperty | CASTCLASS_FMulticastSparseDelegateProperty);
+					}
+
 					const TMap<UObject*, UObject*>& ReconstructedCDOsMap;
 					UObject* PotentialReferencer;
 				};
