[CrashFix] FPropertySoftPath::Resolve crashing when resolving paths to properties owned  by FInstancedStructs

diff --git a/Engine/Source/Editor/UnrealEd/Private/DiffUtils.cpp b/Engine/Source/Editor/UnrealEd/Private/DiffUtils.cpp
--- a/Engine/Source/Editor/UnrealEd/Private/DiffUtils.cpp
+++ b/Engine/Source/Editor/UnrealEd/Private/DiffUtils.cpp
@@ -110,6 +110,14 @@ FResolvedProperty FPropertySoftPath::Resolve(const UStruct* Struct, const void*
 	{
 		CurrentBlock = NextBlock;
 		const FProperty* NextProperty = UEDiffUtils_Private::Resolve(NextClass, PropertyChain[i].PropertyName);
+		if (!NextProperty)
+		{
+			if (const FStructProperty* StructProperty = CastField<FStructProperty>(Property))
+            {
+            	StructProperty->FindInnerPropertyInstance(PropertyChain[i].PropertyName, CurrentBlock, NextProperty, NextBlock);
+            }
+		}
+		CurrentBlock = NextBlock;
 
 		// if an index was provided, resolve it
 		const int32 PropertyIndex = TryReadIndex(PropertyChain, i);
