Honor native class data invariants for REINST classes, we do occasionally PostLoad these, and don't want native code to be missing it's companion class data
#rnx

diff --git a/Engine/Source/Editor/UnrealEd/Private/Kismet2/KismetReinstanceUtilities.cpp b/Engine/Source/Editor/UnrealEd/Private/Kismet2/KismetReinstanceUtilities.cpp
--- a/Engine/Source/Editor/UnrealEd/Private/Kismet2/KismetReinstanceUtilities.cpp
+++ b/Engine/Source/Editor/UnrealEd/Private/Kismet2/KismetReinstanceUtilities.cpp
@@ -1858,6 +1858,16 @@ UClass* FBlueprintCompileReinstancer::MoveCDOToNewClass(UClass* OwnerClass, cons
 	CopyOfOwnerClass->Bind();
 	CopyOfOwnerClass->StaticLink(true);
 
+	// make sure we've bound to our native sparse data - this should have been done in
+	// link but don't want to destabilize early adopters:
+	if (!CopyOfOwnerClass->GetSparseClassDataStruct())
+	{
+		if (UScriptStruct* SparseClassDataStructArchetype = CopyOfOwnerClass->GetSparseClassDataArchetypeStruct())
+		{
+			CopyOfOwnerClass->SetSparseClassDataStruct(SparseClassDataStructArchetype);
+		}
+	}
+
 	if(OldCDO)
 	{
 		// @todo: #dano, rename bAvoidCDODuplication because it's really a flag to move the CDO aside not 'prevent duplication':
