First pass at the Entity Prefab Editor

diff --git a/Engine/Source/Editor/UnrealEd/Public/Tools/BaseAssetToolkit.h b/Engine/Source/Editor/UnrealEd/Public/Tools/BaseAssetToolkit.h
--- a/Engine/Source/Editor/UnrealEd/Public/Tools/BaseAssetToolkit.h
+++ b/Engine/Source/Editor/UnrealEd/Public/Tools/BaseAssetToolkit.h
@@ -58,7 +58,7 @@ protected:
 	virtual AssetEditorViewportFactoryFunction GetViewportDelegate();
 	virtual TSharedPtr<FEditorViewportClient> CreateEditorViewportClient() const;
 	TSharedRef<SDockTab> SpawnTab_Viewport(const FSpawnTabArgs& Args);
-	TSharedRef<SDockTab> SpawnTab_Details(const FSpawnTabArgs& Args);
+	virtual TSharedRef<SDockTab> SpawnTab_Details(const FSpawnTabArgs& Args);
 
 protected:
 	/** Property View */
