Added Interchange to the import/reimport for Nanite

diff --git a/Engine/Source/Editor/UnrealEd/Private/FbxMeshUtils.cpp b/Engine/Source/Editor/UnrealEd/Private/FbxMeshUtils.cpp
--- a/Engine/Source/Editor/UnrealEd/Private/FbxMeshUtils.cpp
+++ b/Engine/Source/Editor/UnrealEd/Private/FbxMeshUtils.cpp
@@ -87,6 +87,47 @@ namespace FbxMeshUtils
 
 			ImportOptions->bAutoComputeLodDistances = true; //Setting auto compute distance to true will avoid changing the staticmesh flag
 		}
+
+		bool CopyHighResMeshDescription(UStaticMesh* SrcStaticMesh, UStaticMesh* BaseStaticMesh)
+		{
+			if (!SrcStaticMesh || !SrcStaticMesh->IsSourceModelValid(0))
+			{
+				return false;
+			}
+
+			{
+				BaseStaticMesh->Modify();
+
+				FMeshDescription* HiResMeshDescription = BaseStaticMesh->GetHiResMeshDescription();
+				if (HiResMeshDescription == nullptr)
+				{
+					HiResMeshDescription = BaseStaticMesh->CreateHiResMeshDescription();
+				}
+
+				check(HiResMeshDescription);
+
+				BaseStaticMesh->ModifyHiResMeshDescription();
+
+				FMeshDescription* TempLOD0MeshDescription = SrcStaticMesh->GetMeshDescription(0);
+				check(TempLOD0MeshDescription);
+
+				if (FMeshDescription* BaseMeshDescription = BaseStaticMesh->GetMeshDescription(0))
+				{
+					FString MaterialNameConflictMsg = TEXT("[Asset ") + BaseStaticMesh->GetPathName() + TEXT("] Nanite hi - res import have some material name that differ from the LOD 0 material name.Your nanite hi - res should use the same material names the LOD 0 use to ensure we can remap the section in the same order.");
+					FString MaterialCountConflictMsg = TEXT("[Asset ") + BaseStaticMesh->GetPathName() + TEXT("] Nanite hi-res import dont have the same material count then LOD 0. Your nanite hi-res should have equal number of material.");
+					FStaticMeshOperations::ReorderMeshDescriptionPolygonGroups(*BaseMeshDescription, *TempLOD0MeshDescription, MaterialNameConflictMsg, MaterialCountConflictMsg);
+				}
+
+				*HiResMeshDescription = MoveTemp(*TempLOD0MeshDescription);
+
+				BaseStaticMesh->CommitHiResMeshDescription();
+
+				BaseStaticMesh->PostEditChange();
+				BaseStaticMesh->MarkPackageDirty();
+
+				return true;
+			}
+		}
 	}
 
 	/** Helper function used for retrieving data required for importing static mesh LODs */
@@ -312,6 +353,51 @@ namespace FbxMeshUtils
 			return false;
 		}
 
+		//We will use interchange only if interchange is enabled and the mesh we want to add a LOD was imported with interchange
+		const UInterchangeAssetImportData* SelectedInterchangeAssetImportData = Cast<UInterchangeAssetImportData>(BaseStaticMesh->GetAssetImportData());
+		if (UInterchangeManager::IsInterchangeImportEnabled() && SelectedInterchangeAssetImportData)
+		{
+			UStaticMesh* TempStaticMesh = NewObject<UStaticMesh>(GetTransientPackage(), NAME_None, RF_Transient | RF_Public | RF_Standalone);
+			TempStaticMesh->AddSourceModel();
+
+			// Since it is Async, any action on the mesh should be locked
+
+			UInterchangeSourceData* SourceData = UInterchangeManager::GetInterchangeManager().CreateSourceData(Filename);
+			//Call interchange mesh utilities to import custom LOD
+			UInterchangeMeshUtilities::ImportCustomLodAsync(TempStaticMesh, 0, SourceData).Then([BaseStaticMesh,TempStaticMesh](TFuture<bool> Result)
+			{
+				bool bResult = Result.Get();
+				Async(EAsyncExecution::TaskGraphMainThread, [BaseStaticMesh, TempStaticMesh, bResult]()
+				{
+					// Copy high res mesh from temporary static mesh to targeted one
+					if (bResult && Private::CopyHighResMeshDescription(TempStaticMesh, BaseStaticMesh))
+					{
+						// Notification of success
+						FNotificationInfo NotificationInfo(FText::GetEmpty());
+						NotificationInfo.Text = NSLOCTEXT("UnrealEd", "ImportStaticMeshHiResSourceModelSuccessful", "High res mesh imported successfully!");
+						NotificationInfo.ExpireDuration = 5.0f;
+						FSlateNotificationManager::Get().AddNotification(NotificationInfo);
+					}
+					else
+					{
+						// Notification of failure
+						FNotificationInfo NotificationInfo(FText::GetEmpty());
+						NotificationInfo.Text = NSLOCTEXT("UnrealEd", "ImportStaticMeshHiResSourceModelFail", "Failed to import high res mesh!");
+						NotificationInfo.ExpireDuration = 5.0f;
+						FSlateNotificationManager::Get().AddNotification(NotificationInfo);
+					}
+
+					if(TempStaticMesh)
+					{
+						TempStaticMesh->ClearFlags(RF_Public | RF_Standalone);
+						TempStaticMesh->MarkAsGarbage();
+					}
+				});
+			});
+
+			return true;
+		}
+
 		bool bSuccess = false;
 
 		UE_LOG(LogExportMeshUtils, Log, TEXT("Fbx Mesh loading"));
@@ -365,36 +451,8 @@ namespace FbxMeshUtils
 				}
 
 				// Add the imported mesh to the existing model
-				if (TempStaticMesh && TempStaticMesh->IsSourceModelValid(0))
+				if (Private::CopyHighResMeshDescription(TempStaticMesh, BaseStaticMesh))
 				{
-					BaseStaticMesh->Modify();
-
-					FMeshDescription* HiResMeshDescription = BaseStaticMesh->GetHiResMeshDescription();
-					if (HiResMeshDescription == nullptr)
-					{
-						HiResMeshDescription = BaseStaticMesh->CreateHiResMeshDescription();
-					}
-					check(HiResMeshDescription);
-
-					BaseStaticMesh->ModifyHiResMeshDescription();
-
-					FMeshDescription* TempLOD0MeshDescription = TempStaticMesh->GetMeshDescription(0);
-					check(TempLOD0MeshDescription);
-
-					if (FMeshDescription* BaseMeshDescription = BaseStaticMesh->GetMeshDescription(0))
-					{
-						FString MaterialNameConflictMsg = TEXT("[Asset ") + BaseStaticMesh->GetPathName() + TEXT("] Nanite hi - res import have some material name that differ from the LOD 0 material name.Your nanite hi - res should use the same material names the LOD 0 use to ensure we can remap the section in the same order.");
-						FString MaterialCountConflictMsg = TEXT("[Asset ") + BaseStaticMesh->GetPathName() + TEXT("] Nanite hi-res import dont have the same material count then LOD 0. Your nanite hi-res should have equal number of material.");
-						FStaticMeshOperations::ReorderMeshDescriptionPolygonGroups(*BaseMeshDescription, *TempLOD0MeshDescription, MaterialNameConflictMsg, MaterialCountConflictMsg);
-					}
-					
-					*HiResMeshDescription = MoveTemp(*TempLOD0MeshDescription);
-
-					BaseStaticMesh->CommitHiResMeshDescription();
-
-					BaseStaticMesh->PostEditChange();
-					BaseStaticMesh->MarkPackageDirty();
-
 					FNotificationInfo NotificationInfo(FText::GetEmpty());
 					NotificationInfo.Text = FText::Format(LOCTEXT("HiResMeshImportSuccessful", "High res mesh imported successfully!"), FText::AsNumber(0));
 					NotificationInfo.ExpireDuration = 5.0f;
