Cooker: Fix incorrect call to VerifyEDLCookInfo because ShouldVerifyEDLCookInfo function is using the string DlcName which is cleared during ShutdownCookSession. We need to call ShouldVerifyEDLCookInfo while the cook session is still in progress.
#rnx

diff --git a/Engine/Source/Editor/UnrealEd/Private/Commandlets/CookCommandlet.cpp b/Engine/Source/Editor/UnrealEd/Private/Commandlets/CookCommandlet.cpp
--- a/Engine/Source/Editor/UnrealEd/Private/Commandlets/CookCommandlet.cpp
+++ b/Engine/Source/Editor/UnrealEd/Private/Commandlets/CookCommandlet.cpp
@@ -527,11 +527,13 @@ void UCookCommandlet::RunCookByTheBookCook(UCookOnTheFlyServer* CookOnTheFlyServ
 #if ENABLE_LOW_LEVEL_MEM_TRACKER
 	FLowLevelMemTracker::Get().UpdateStatsPerFrame();
 #endif
+	bool bShouldVerifyEDLCookInfo = false;
 	do
 	{
 		{
 			COOK_STAT(FScopedDurationTimer StartCookByTheBookTimer(DetailedCookStats::StartCookByTheBookTimeSec));
 			CookOnTheFlyServer->StartCookByTheBook(StartupOptions);
+			bShouldVerifyEDLCookInfo = CookOnTheFlyServer->ShouldVerifyEDLCookInfo();
 		}
 		DECLARE_SCOPE_CYCLE_COUNTER(TEXT("CookByTheBook.MainLoop"), STAT_CookByTheBook_MainLoop, STATGROUP_LoadTime);
 		while (CookOnTheFlyServer->IsInSession())
@@ -546,7 +548,7 @@ void UCookCommandlet::RunCookByTheBookCook(UCookOnTheFlyServer* CookOnTheFlyServ
 	FLowLevelMemTracker::Get().UpdateStatsPerFrame();
 #endif
 
-	if (CookOnTheFlyServer->ShouldVerifyEDLCookInfo())
+	if (bShouldVerifyEDLCookInfo)
 	{
 		bool bFullReferencesExpected = !(CookOptions & ECookByTheBookOptions::SkipHardReferences);
 		UE::SavePackageUtilities::VerifyEDLCookInfo(bFullReferencesExpected);
