World Partition Builder - Allow builders to process non partitioned world by overriding the CanProcessNonPartitionedWorlds() method

diff --git a/Engine/Source/Editor/UnrealEd/Private/WorldPartition/WorldPartitionBuilder.cpp b/Engine/Source/Editor/UnrealEd/Private/WorldPartition/WorldPartitionBuilder.cpp
--- a/Engine/Source/Editor/UnrealEd/Private/WorldPartition/WorldPartitionBuilder.cpp
+++ b/Engine/Source/Editor/UnrealEd/Private/WorldPartition/WorldPartitionBuilder.cpp
@@ -104,13 +104,9 @@ bool UWorldPartitionBuilder::RunBuilder(UWorld* World)
 			EditorWorld = MakeUnique<FScopedEditorWorld>(World, GetWorldInitializationValues());
 		}
 
-		// Make sure the world is partitioned
-		if (UWorld::IsPartitionedWorld(World))
+		// Make sure the world is partitioned if required
+		if (UWorld::IsPartitionedWorld(World) || CanProcessNonPartitionedWorlds())
 		{
-			// Ensure the world has a valid world partition.
-			UWorldPartition* WorldPartition = World->GetWorldPartition();
-			check(WorldPartition);
-
 			FWorldContext& WorldContext = GEditor->GetEditorWorldContext(true /*bEnsureIsGWorld*/);
 			WorldContext.SetCurrentWorld(World);
 			UWorld* PrevGWorld = GWorld;
@@ -168,49 +164,20 @@ FWorldBuilderCellCoord FCellInfo::GetCellCount(const FBox& InBounds, const int32
 
 bool UWorldPartitionBuilder::Run(UWorld* World, FPackageSourceControlHelper& PackageHelper)
 {
+	UWorldPartition* WorldPartition = World->GetWorldPartition();
+	check(WorldPartition || CanProcessNonPartitionedWorlds());
+
 	// Notify derived classes that partition building process starts
 	bool bResult = PreRun(World, PackageHelper);
 
-	UWorldPartition* WorldPartition = World->GetWorldPartition();
-
-	// Properly Setup DataLayers for Builder
-	UDataLayerManager* DataLayerManager = WorldPartition->GetDataLayerManager();
+	// Load data layers
+	LoadDataLayers(World);
 	
-	// Load Data Layers
-	bool bUpdateEditorCells = false;
-	DataLayerManager->ForEachDataLayerInstance([&bUpdateEditorCells, this](UDataLayerInstance* DataLayer)
-	{
-		const FName DataLayerShortName(DataLayer->GetDataLayerShortName());
-
-		// Load all Non Excluded Data Layers + Non DynamicallyLoaded Data Layers + Initially Active Data Layers + Data Layers provided by builder
-		const bool bLoadedInEditor = !ExcludedDataLayerShortNames.Contains(DataLayerShortName) && ((bLoadNonDynamicDataLayers && !DataLayer->IsRuntime()) ||
-										(bLoadInitiallyActiveDataLayers && DataLayer->GetInitialRuntimeState() == EDataLayerRuntimeState::Activated) ||
-									DataLayerShortNames.Contains(DataLayerShortName));
-		if (DataLayer->IsLoadedInEditor() != bLoadedInEditor)
-		{
-			bUpdateEditorCells = true;
-			DataLayer->SetIsLoadedInEditor(bLoadedInEditor, /*bFromUserChange*/false);
-			if (RequiresCommandletRendering() && bLoadedInEditor)
-			{
-				DataLayer->SetIsInitiallyVisible(true);
-			}
-		}
-			
-		UE_LOG(LogWorldPartitionBuilder, Display, TEXT("DataLayer '%s' Loaded: %d"), *UDataLayerInstance::GetDataLayerText(DataLayer).ToString(), bLoadedInEditor ? 1 : 0);
-			
-		return true;
-	});
-	
-	if (bUpdateEditorCells)
-	{
-		UE_LOG(LogWorldPartitionBuilder, Display, TEXT("DataLayer load state changed refreshing editor cells"));
-		FDataLayersEditorBroadcast::StaticOnActorDataLayersEditorLoadingStateChanged(false);
-	}
-
-	const ELoadingMode LoadingMode = GetLoadingMode();
+	static const FBox BoxEntireWorld = FBox(FVector(-HALF_WORLD_MAX, -HALF_WORLD_MAX, -HALF_WORLD_MAX), FVector(HALF_WORLD_MAX, HALF_WORLD_MAX, HALF_WORLD_MAX));
+	const ELoadingMode LoadingMode = WorldPartition ? GetLoadingMode() : ELoadingMode::Custom;
+		
 	FCellInfo CellInfo;
-
-	CellInfo.EditorBounds = IterativeWorldBounds.IsValid ? IterativeWorldBounds : WorldPartition->GetEditorWorldBounds();
+	CellInfo.EditorBounds = IterativeWorldBounds.IsValid ? IterativeWorldBounds : WorldPartition ? WorldPartition->GetEditorWorldBounds() : BoxEntireWorld;
 	CellInfo.IterativeCellSize = IterativeCellSize;
 
 	if ((LoadingMode == ELoadingMode::IterativeCells) || (LoadingMode == ELoadingMode::IterativeCells2D))
@@ -295,32 +262,70 @@ bool UWorldPartitionBuilder::Run(UWorld* World, FPackageSourceControlHelper& Pac
 			}
 		}
 	}
-	else
+	else if (LoadingMode == ELoadingMode::EntireWorld)
 	{
-		TUniquePtr<FLoaderAdapterShape> LoaderAdapterShape;
-		
-		if (LoadingMode == ELoadingMode::EntireWorld)
-		{
-			CellInfo.Bounds = FBox(FVector(-HALF_WORLD_MAX, -HALF_WORLD_MAX, -HALF_WORLD_MAX), FVector(HALF_WORLD_MAX, HALF_WORLD_MAX, HALF_WORLD_MAX));
-			LoaderAdapterShape = MakeUnique<FLoaderAdapterShape>(World, CellInfo.Bounds, TEXT("Loaded Region"));
-			LoaderAdapterShape->Load();
-		}
-		else
-		{
-			CellInfo.Bounds.Init();
-		}
+		CellInfo.Bounds = BoxEntireWorld;
+
+		TUniquePtr<FLoaderAdapterShape> LoaderAdapterShape = MakeUnique<FLoaderAdapterShape>(World, CellInfo.Bounds, TEXT("Loaded Region"));
+		LoaderAdapterShape->Load();
 
 		bResult = RunInternal(World, CellInfo, PackageHelper);
 
-		if (LoaderAdapterShape)
-		{
-			LoaderAdapterShape.Reset();
-		}
+		LoaderAdapterShape.Reset();
+	}
+	else if (LoadingMode == ELoadingMode::Custom)
+	{
+		CellInfo.Bounds.Init();
+		bResult = RunInternal(World, CellInfo, PackageHelper);
+	}
+	else
+	{
+		checkNoEntry();
 	}
 
 	return PostRun(World, PackageHelper, bResult);
 }
 
+void UWorldPartitionBuilder::LoadDataLayers(UWorld* InWorld)
+{
+	if (UWorldPartition* WorldPartition = InWorld->GetWorldPartition())
+	{
+		// Properly Setup DataLayers for Builder
+		UDataLayerManager* DataLayerManager = WorldPartition->GetDataLayerManager();
+	
+		// Load Data Layers
+		bool bUpdateEditorCells = false;
+		DataLayerManager->ForEachDataLayerInstance([&bUpdateEditorCells, this](UDataLayerInstance* DataLayer)
+		{
+			const FName DataLayerShortName(DataLayer->GetDataLayerShortName());
+
+			// Load all Non Excluded Data Layers + Non DynamicallyLoaded Data Layers + Initially Active Data Layers + Data Layers provided by builder
+			const bool bLoadedInEditor = !ExcludedDataLayerShortNames.Contains(DataLayerShortName) && ((bLoadNonDynamicDataLayers && !DataLayer->IsRuntime()) ||
+											(bLoadInitiallyActiveDataLayers && DataLayer->GetInitialRuntimeState() == EDataLayerRuntimeState::Activated) ||
+										DataLayerShortNames.Contains(DataLayerShortName));
+			if (DataLayer->IsLoadedInEditor() != bLoadedInEditor)
+			{
+				bUpdateEditorCells = true;
+				DataLayer->SetIsLoadedInEditor(bLoadedInEditor, /*bFromUserChange*/false);
+				if (RequiresCommandletRendering() && bLoadedInEditor)
+				{
+					DataLayer->SetIsInitiallyVisible(true);
+				}
+			}
+			
+			UE_LOG(LogWorldPartitionBuilder, Display, TEXT("DataLayer '%s' Loaded: %d"), *UDataLayerInstance::GetDataLayerText(DataLayer).ToString(), bLoadedInEditor ? 1 : 0);
+			
+			return true;
+		});
+	
+		if (bUpdateEditorCells)
+		{
+			UE_LOG(LogWorldPartitionBuilder, Display, TEXT("DataLayer load state changed refreshing editor cells"));
+			FDataLayersEditorBroadcast::StaticOnActorDataLayersEditorLoadingStateChanged(false);
+		}
+	}
+}
+
 bool UWorldPartitionBuilder::SavePackages(const TArray<UPackage*>& Packages, FPackageSourceControlHelper& PackageHelper, bool bErrorsAsWarnings)
 {
 	if (Packages.IsEmpty())
diff --git a/Engine/Source/Editor/UnrealEd/Public/WorldPartition/WorldPartitionBuilder.h b/Engine/Source/Editor/UnrealEd/Public/WorldPartition/WorldPartitionBuilder.h
--- a/Engine/Source/Editor/UnrealEd/Public/WorldPartition/WorldPartitionBuilder.h
+++ b/Engine/Source/Editor/UnrealEd/Public/WorldPartition/WorldPartitionBuilder.h
@@ -89,6 +89,11 @@ protected:
 	 */
 	virtual bool ShouldSkipCell(const FWorldBuilderCellCoord& CellCoord) const { return false; }
 
+	/** 
+	 * Some builders may have the ability to process non partitioned worlds.
+	 */
+	virtual bool CanProcessNonPartitionedWorlds() const { return false; }
+
 	bool OnFilesModified(const TArray<FString>& InModifiedFiles, const FString& InChangelistDescription) const;
 	bool OnPackagesModified(const TArray<UPackage*>& InModifiedPackages, const FString& InChangelistDescription) const;
 
@@ -145,6 +150,8 @@ protected:
 	FModifiedFilesHandler ModifiedFilesHandler;
 
 private:
+	void LoadDataLayers(UWorld* InWorld);
+
 	friend struct FWorldPartitionBuilderArgsScope;
 	static FString Args;
 };
