Allow FFuzzyPathMatcher to be used by localization gather commandlets outside of UnrealEd

#rnx

diff --git a/Engine/Source/Editor/UnrealEd/Classes/Commandlets/GatherTextCommandletBase.h b/Engine/Source/Editor/UnrealEd/Classes/Commandlets/GatherTextCommandletBase.h
--- a/Engine/Source/Editor/UnrealEd/Classes/Commandlets/GatherTextCommandletBase.h
+++ b/Engine/Source/Editor/UnrealEd/Classes/Commandlets/GatherTextCommandletBase.h
@@ -18,7 +18,7 @@ struct UNREALED_API FGatherTextDelegates
 };
 
 /** Performs fuzzy path matching against a set of include and exclude paths */
-class FFuzzyPathMatcher
+class UNREALED_API FFuzzyPathMatcher
 {
 public:
 	enum class EPathMatch
