Make sure sparse data is duplicated when cloning a blueprint class

diff --git a/Engine/Source/Editor/UnrealEd/Private/Kismet2/BlueprintEditorUtils.cpp b/Engine/Source/Editor/UnrealEd/Private/Kismet2/BlueprintEditorUtils.cpp
--- a/Engine/Source/Editor/UnrealEd/Private/Kismet2/BlueprintEditorUtils.cpp
+++ b/Engine/Source/Editor/UnrealEd/Private/Kismet2/BlueprintEditorUtils.cpp
@@ -2076,6 +2076,24 @@ void FBlueprintEditorUtils::PostDuplicateBlueprint(UBlueprint* Blueprint, bool b
 			check(NewCDO != nullptr);
 			UEditorEngine::CopyPropertiesForUnrelatedObjects(OldCDO, NewCDO);
 
+			// copy sparse data over to the new class sparse data, if any:
+			const TObjectPtr<UScriptStruct> SparseData = Blueprint->GeneratedClass->GetSparseClassDataStruct();
+			if (SparseData && OldBPGC->bIsSparseClassDataSerializable)
+			{
+				void* SparseDataInstance = Blueprint->GeneratedClass->GetOrCreateSparseClassData();
+
+				const TObjectPtr<UScriptStruct> OldSparseData = OldBPGC->GetSparseClassDataStruct();
+				if (ensure(OldSparseData == SparseData))
+				{
+					const void* OldSparseDataInstance = OldBPGC->GetSparseClassData(EGetSparseClassDataMethod::ReturnIfNull);
+					if (OldSparseDataInstance && ensure(OldSparseDataInstance != SparseDataInstance))
+					{
+						SparseData->CopyScriptStruct(SparseDataInstance, OldSparseDataInstance);
+					}
+				}
+				NewBPGC->bIsSparseClassDataSerializable = true; // match the object we're being duplicated from
+			}
+
 			FBlueprintEditorUtils::ReconstructAllNodes(Blueprint);
 
 			if (!FBlueprintDuplicationScopeFlags::HasAnyFlag(FBlueprintDuplicationScopeFlags::NoExtraCompilation))
