Record the root struct type that a property soft path is constructed with, so we can disambiguate between CDOs and sparse class data members. This allows the diff tool to properly display changes to properties in a sparse class data instance even when there is a deprecated property on the CDO with the same name

diff --git a/Engine/Source/Editor/UnrealEd/Private/DiffUtils.cpp b/Engine/Source/Editor/UnrealEd/Private/DiffUtils.cpp
--- a/Engine/Source/Editor/UnrealEd/Private/DiffUtils.cpp
+++ b/Engine/Source/Editor/UnrealEd/Private/DiffUtils.cpp
@@ -66,23 +66,65 @@ namespace UEDiffUtils_Private
 	}
 }
 
-FResolvedProperty FPropertySoftPath::Resolve(const UObject* Object) const
+FPropertySoftPath::FPropertySoftPath()
+	: RootTypeHint(nullptr)
 {
-	FResolvedProperty ResolvedProperty = Resolve(Object->GetClass(), Object);
-	if (!ResolvedProperty.Property && Object->HasAnyFlags(RF_ClassDefaultObject))
+}
+
+FPropertySoftPath::FPropertySoftPath(TArray<FName> InPropertyChain)
+	: RootTypeHint(nullptr)
+{
+	for (FName PropertyName : InPropertyChain)
+	{
+		PropertyChain.Push(FChainElement(PropertyName));
+	}
+}
+
+FPropertySoftPath::FPropertySoftPath(FPropertyPath InPropertyPath)
+	: RootTypeHint(nullptr)
+{
+	for (int32 PropertyIndex = 0, end = InPropertyPath.GetNumProperties(); PropertyIndex != end; ++PropertyIndex)
 	{
-		if (const UScriptStruct* SparseClassDataStruct = Object->GetClass()->GetSparseClassDataStruct())
+		const FPropertyInfo& Info = InPropertyPath.GetPropertyInfo(PropertyIndex);
+		if (Info.ArrayIndex != INDEX_NONE)
 		{
-			if (const void* SparseClassData = Object->GetClass()->GetSparseClassData(EGetSparseClassDataMethod::ReturnIfNull))
-			{
-				ResolvedProperty = Resolve(SparseClassDataStruct, SparseClassData);
-			}
+			PropertyChain.Push(FName(*FString::FromInt(Info.ArrayIndex)));
+		}
+		else
+		{
+			PropertyChain.Push(FChainElement(Info.Property.Get()));
+		}
+	}
+
+	if (InPropertyPath.GetNumProperties() > 0)
+	{
+		const FProperty* RootProperty = InPropertyPath.GetPropertyInfo(0).Property.Get();
+		if(RootProperty)
+		{
+			RootTypeHint = RootProperty->GetOwnerStruct();
 		}
 	}
-	return ResolvedProperty;
 }
 
+FPropertySoftPath::FPropertySoftPath(const FPropertySoftPath& SubPropertyPath, const FProperty* LeafProperty)
+	: PropertyChain(SubPropertyPath.PropertyChain)
+	, RootTypeHint(SubPropertyPath.RootTypeHint)
+{
+	PropertyChain.Push(FChainElement(LeafProperty));
+}
+
+FPropertySoftPath::FPropertySoftPath(const FPropertySoftPath& SubPropertyPath, int32 ContainerIndex)
+	: PropertyChain(SubPropertyPath.PropertyChain)
+	, RootTypeHint(SubPropertyPath.RootTypeHint)
+{
+	PropertyChain.Push(FName(*FString::FromInt(ContainerIndex)));
+}
 
+FResolvedProperty FPropertySoftPath::Resolve(const UObject* Object) const
+{
+	FResolvedProperty ResolvedProperty = Resolve(Object->GetClass(), Object);
+	return ResolvedProperty;
+}
 
 int32 FPropertySoftPath::TryReadIndex(const TArray<FChainElement>& LocalPropertyChain, int32& OutIndex)
 {
@@ -100,6 +142,21 @@ int32 FPropertySoftPath::TryReadIndex(const TArray<FChainElement>& LocalProperty
 
 FResolvedProperty FPropertySoftPath::Resolve(const UStruct* Struct, const void* StructData) const
 {
+	if (RootTypeHint && RootTypeHint != Struct)
+	{
+		if (const UClass* AsClass = Cast<UClass>(Struct))
+		{
+			const UScriptStruct* SparseClassDataStruct = AsClass->GetSparseClassDataStruct();
+			if (SparseClassDataStruct && SparseClassDataStruct->IsChildOf(RootTypeHint))
+			{
+				if (const void* SparseClassData = const_cast<UClass*>(AsClass)->GetSparseClassData(EGetSparseClassDataMethod::ReturnIfNull))
+				{
+					return Resolve(RootTypeHint, SparseClassData);
+				}
+			}
+		}
+	}
+
 	// dig into the object, finding nested objects, etc:
 	const void* CurrentBlock = StructData;
 	const UStruct* NextClass = Struct;
diff --git a/Engine/Source/Editor/UnrealEd/Public/DiffUtils.h b/Engine/Source/Editor/UnrealEd/Public/DiffUtils.h
--- a/Engine/Source/Editor/UnrealEd/Public/DiffUtils.h
+++ b/Engine/Source/Editor/UnrealEd/Public/DiffUtils.h
@@ -62,45 +62,11 @@ struct FResolvedProperty
  */
 struct FPropertySoftPath
 {
-	FPropertySoftPath()
-	{
-	}
-
-	FPropertySoftPath(TArray<FName> InPropertyChain)
-	{
-		for (FName PropertyName : InPropertyChain)
-		{
-			PropertyChain.Push(FChainElement(PropertyName));
-		}
-	}
-
-	FPropertySoftPath( FPropertyPath InPropertyPath )
-	{
-		for (int32 PropertyIndex = 0, end = InPropertyPath.GetNumProperties(); PropertyIndex != end; ++PropertyIndex)
-		{
-			const FPropertyInfo& Info = InPropertyPath.GetPropertyInfo(PropertyIndex);
-			if (Info.ArrayIndex != INDEX_NONE)
-			{
-				PropertyChain.Push(FName(*FString::FromInt(Info.ArrayIndex)));
-			}
-			else
-			{
-				PropertyChain.Push(FChainElement(Info.Property.Get()));
-			}
-		}
-	}
-
-	FPropertySoftPath(const FPropertySoftPath& SubPropertyPath, const FProperty* LeafProperty)
-		: PropertyChain(SubPropertyPath.PropertyChain)
-	{
-		PropertyChain.Push(FChainElement(LeafProperty));
-	}
-
-	FPropertySoftPath(const FPropertySoftPath& SubPropertyPath, int32 ContainerIndex)
-		: PropertyChain(SubPropertyPath.PropertyChain)
-	{
-		PropertyChain.Push(FName(*FString::FromInt(ContainerIndex)));
-	}
+	UNREALED_API FPropertySoftPath();
+	UNREALED_API FPropertySoftPath(TArray<FName> InPropertyChain);
+	UNREALED_API FPropertySoftPath(FPropertyPath InPropertyPath);
+	UNREALED_API FPropertySoftPath(const FPropertySoftPath& SubPropertyPath, const FProperty* LeafProperty);
+	UNREALED_API FPropertySoftPath(const FPropertySoftPath& SubPropertyPath, int32 ContainerIndex);
 
 	UNREALED_API FResolvedProperty Resolve(const UObject* Object) const;
 	UNREALED_API FResolvedProperty Resolve(const UStruct* Struct, const void* StructData) const;
@@ -176,6 +142,7 @@ private:
 
 	friend uint32 GetTypeHash( FPropertySoftPath const& Path );
 	TArray<FChainElement> PropertyChain;
+	const UStruct* RootTypeHint;
 };
 
 struct FSCSIdentifier
