deprecate TiffLoader

#rnx

diff --git a/Engine/Source/Editor/UnrealEd/Private/Factories/TIFFLoader.cpp b/Engine/Source/Editor/UnrealEd/Private/Factories/TIFFLoader.cpp
--- a/Engine/Source/Editor/UnrealEd/Private/Factories/TIFFLoader.cpp
+++ b/Engine/Source/Editor/UnrealEd/Private/Factories/TIFFLoader.cpp
@@ -4,6 +4,11 @@
 
 DEFINE_LOG_CATEGORY_STATIC(LogTIFFLoader, Log, All);
 
+
+// disable warnings about myself :
+PRAGMA_DISABLE_DEPRECATION_WARNINGS
+
+
 #if WITH_FREEIMAGE_LIB
 
 #include "Misc/Paths.h"
@@ -385,3 +390,5 @@ bool FTiffLoadHelper::IsValid()
 }
 
 #endif // WITH_FREEIMAGE_LIB
+
+PRAGMA_ENABLE_DEPRECATION_WARNINGS
diff --git a/Engine/Source/Editor/UnrealEd/Private/Factories/TIFFLoader.h b/Engine/Source/Editor/UnrealEd/Private/Factories/TIFFLoader.h
--- a/Engine/Source/Editor/UnrealEd/Private/Factories/TIFFLoader.h
+++ b/Engine/Source/Editor/UnrealEd/Private/Factories/TIFFLoader.h
@@ -7,14 +7,14 @@
 
 // DEPRECATED
 // use Tiff ImageWrapper instead
-// delete me in UE 5.3
+// delete me in UE 5.4
 
 class UTexture2D;
 
 struct FIBITMAP;
 struct FIMEMORY;
 
-class FTiffLoadHelper
+class UE_DEPRECATED(5.3,"Use Tiff ImageWrapper instead") FTiffLoadHelper
 {
 public:
 
