Broke circular dependencies:
AutomationController -> UnrealEd -> AutomationController
SubobjectEditor -> UnrealEd -> SubobjectEditor

diff --git a/Engine/Plugins/Experimental/AnimNext/Source/AnimNextEditor/AnimNextEditor.Build.cs b/Engine/Plugins/Experimental/AnimNext/Source/AnimNextEditor/AnimNextEditor.Build.cs
--- a/Engine/Plugins/Experimental/AnimNext/Source/AnimNextEditor/AnimNextEditor.Build.cs
+++ b/Engine/Plugins/Experimental/AnimNext/Source/AnimNextEditor/AnimNextEditor.Build.cs
@@ -36,6 +36,7 @@ namespace UnrealBuildTool.Rules
 					"BlueprintGraph",	// For K2 Schema
 					"DesktopWidgets",
 					"ContentBrowserFileDataSource",
+					"SubobjectEditor",
 				}
 			);
 		}
diff --git a/Engine/Plugins/Runtime/Metasound/Source/MetasoundEditor/MetasoundEditor.Build.cs b/Engine/Plugins/Runtime/Metasound/Source/MetasoundEditor/MetasoundEditor.Build.cs
--- a/Engine/Plugins/Runtime/Metasound/Source/MetasoundEditor/MetasoundEditor.Build.cs
+++ b/Engine/Plugins/Runtime/Metasound/Source/MetasoundEditor/MetasoundEditor.Build.cs
@@ -31,6 +31,7 @@ namespace UnrealBuildTool.Rules
 					"MetasoundGenerator",
 					"MetasoundGraphCore",
 					"SignalProcessing",
+					"SubobjectEditor",
 					"ToolMenus",
 					"ToolWidgets",
 					"WaveTable",
diff --git a/Engine/Source/Editor/Blutility/Blutility.Build.cs b/Engine/Source/Editor/Blutility/Blutility.Build.cs
--- a/Engine/Source/Editor/Blutility/Blutility.Build.cs
+++ b/Engine/Source/Editor/Blutility/Blutility.Build.cs
@@ -15,6 +15,7 @@ public class Blutility : ModuleRules
 
 		PrivateDependencyModuleNames.AddRange(
 			new string[] {
+				"AutomationController",
 				"Core",
 				"CoreUObject",
 				"Engine",
diff --git a/Engine/Source/Editor/UnrealEd/UnrealEd.Build.cs b/Engine/Source/Editor/UnrealEd/UnrealEd.Build.cs
--- a/Engine/Source/Editor/UnrealEd/UnrealEd.Build.cs
+++ b/Engine/Source/Editor/UnrealEd/UnrealEd.Build.cs
@@ -75,7 +75,6 @@ public class UnrealEd : ModuleRules
 				"BlueprintGraph",
 				"HTTP",
 				"FunctionalTesting",
-				"AutomationController",
 				"Localization",
 				"AudioEditor",
 				"NetworkFileSystem",
@@ -95,7 +94,6 @@ public class UnrealEd : ModuleRules
 				"InterchangeEngine",
 				"DeveloperToolSettings",
 				"SubobjectDataInterface",
-				"SubobjectEditor",
 				"PhysicsUtilities",
 				"ToolWidgets",
 				"AssetTools", 
@@ -194,6 +192,7 @@ public class UnrealEd : ModuleRules
 				"CookOnTheFlyNetServer",
 				"Zen",
 				"BuildSettings",
+				"SubobjectEditor",
 			}
 		);
 
