Fix regression from 25489627.

diff --git a/Engine/Source/Editor/UnrealEd/Private/Commandlets/WorldPartitionConvertCommandlet.cpp b/Engine/Source/Editor/UnrealEd/Private/Commandlets/WorldPartitionConvertCommandlet.cpp
--- a/Engine/Source/Editor/UnrealEd/Private/Commandlets/WorldPartitionConvertCommandlet.cpp
+++ b/Engine/Source/Editor/UnrealEd/Private/Commandlets/WorldPartitionConvertCommandlet.cpp
@@ -1206,7 +1206,7 @@ int32 UWorldPartitionConvertCommandlet::Main(const FString& Params)
 	// Prepare levels for conversion
 	DetachDependantLevelPackages(MainLevel);
 
-	if (PrepareLevelActors(MainLevel, MutableView(MainLevel->Actors), true)) 
+	if (!PrepareLevelActors(MainLevel, MutableView(MainLevel->Actors), true)) 
 	{
 		return 1;
 	}
@@ -1372,7 +1372,10 @@ int32 UWorldPartitionConvertCommandlet::Main(const FString& Params)
 
 		UE_LOG(LogWorldPartitionConvertCommandlet, Log, TEXT("Converting %s"), *SubWorld->GetName());
 
-		PrepareLevelActors(SubLevel, ActorsToConvert, false);
+		if (!PrepareLevelActors(SubLevel, ActorsToConvert, false))
+		{
+			return 1;
+		}
 
 		for(AActor* Actor: ActorsToConvert)
 		{
