Removed dependencies that aren't needed by the modules that declare them. This is to help fix bad circular dependencies in the engine.

diff --git a/Engine/Source/Editor/AnimationBlueprintLibrary/AnimationBlueprintLibrary.build.cs b/Engine/Source/Editor/AnimationBlueprintLibrary/AnimationBlueprintLibrary.build.cs
--- a/Engine/Source/Editor/AnimationBlueprintLibrary/AnimationBlueprintLibrary.build.cs
+++ b/Engine/Source/Editor/AnimationBlueprintLibrary/AnimationBlueprintLibrary.build.cs
@@ -11,7 +11,6 @@ public class AnimationBlueprintLibrary : ModuleRules
 			{
 				"Core",
 				"CoreUObject",
-				"Kismet",
 				"TimeManagement"
 			}
 		);
@@ -20,7 +19,6 @@ public class AnimationBlueprintLibrary : ModuleRules
 			new string[]
 			{
 				"AnimGraph",
-				"KismetCompiler",
 				"Engine",
 				"BlueprintGraph",
 				"UnrealEd",
diff --git a/Engine/Source/Editor/ContentBrowserData/ContentBrowserData.Build.cs b/Engine/Source/Editor/ContentBrowserData/ContentBrowserData.Build.cs
--- a/Engine/Source/Editor/ContentBrowserData/ContentBrowserData.Build.cs
+++ b/Engine/Source/Editor/ContentBrowserData/ContentBrowserData.Build.cs
@@ -14,7 +14,6 @@ public class ContentBrowserData : ModuleRules
 				"CoreUObject",
 				"Engine",
 				"EditorSubsystem",
-				"EditorFramework",
 				"UnrealEd",
 				"Projects",
 			}
@@ -30,7 +29,6 @@ public class ContentBrowserData : ModuleRules
 			new string[] {
 				"Slate",
 				"SlateCore",
-				"InputCore",
 			}
 		);
 	}
diff --git a/Engine/Source/Editor/EditorWidgets/EditorWidgets.Build.cs b/Engine/Source/Editor/EditorWidgets/EditorWidgets.Build.cs
--- a/Engine/Source/Editor/EditorWidgets/EditorWidgets.Build.cs
+++ b/Engine/Source/Editor/EditorWidgets/EditorWidgets.Build.cs
@@ -15,7 +15,6 @@ public class EditorWidgets : ModuleRules
 				"InputCore",
 				"Slate",
 				"SlateCore",
-				"EditorFramework",
 				"UnrealEd",
 				"CoreUObject",
 				"ToolWidgets",
@@ -26,4 +25,4 @@ public class EditorWidgets : ModuleRules
 		
 		DynamicallyLoadedModuleNames.Add("AssetRegistry");
 	}
-}
\ No newline at end of file
+}
diff --git a/Engine/Source/Editor/LevelEditor/LevelEditor.Build.cs b/Engine/Source/Editor/LevelEditor/LevelEditor.Build.cs
--- a/Engine/Source/Editor/LevelEditor/LevelEditor.Build.cs
+++ b/Engine/Source/Editor/LevelEditor/LevelEditor.Build.cs
@@ -31,7 +31,6 @@ public class LevelEditor : ModuleRules
 		PublicIncludePathModuleNames.AddRange(
 			new string[] {
 				"CommonMenuExtensions",
-				"HeadMountedDisplay",
 				"Settings",
 				"ToolWidgets",
 				"UnrealEd",
@@ -57,22 +56,17 @@ public class LevelEditor : ModuleRules
 				"StatsViewer",
 				"EditorFramework",
 				"UnrealEd", 
-				"DeveloperSettings",
-				"RenderCore",
 				"DeviceProfileServices",
 				"ContentBrowser",
 				"SceneOutliner",
 				"ActorPickerMode",
 				"RHI",
 				"Projects",
-				"TargetPlatform",
 				"TypedElementFramework",
 				"TypedElementRuntime",
 				"EngineSettings",
 				"PropertyEditor",
-				"Kismet",
 				"KismetWidgets",
-				"Sequencer",
 				"Foliage",
 				"HierarchicalLODOutliner",
 				"HierarchicalLODUtilities",
@@ -92,9 +86,10 @@ public class LevelEditor : ModuleRules
 				"DerivedDataEditor",
 				"EditorWidgets",
 				"ToolWidgets",
-				"VirtualizationEditor",
 				"UnsavedAssetsTracker",
-				"UncontrolledChangelists"
+				"UncontrolledChangelists",
+				"RenderCore",
+				"DeveloperSettings"
 			}
 		);
 
diff --git a/Engine/Source/Editor/PropertyEditor/PropertyEditor.Build.cs b/Engine/Source/Editor/PropertyEditor/PropertyEditor.Build.cs
--- a/Engine/Source/Editor/PropertyEditor/PropertyEditor.Build.cs
+++ b/Engine/Source/Editor/PropertyEditor/PropertyEditor.Build.cs
@@ -15,12 +15,6 @@ public class PropertyEditor : ModuleRules
                 	"SceneDepthPickerMode",
 			}
 		);
-		
-        PublicIncludePathModuleNames.AddRange(
-            new string[] {
-				"EditorFramework",
-			}
-        );
 
 		PrivateIncludePathModuleNames.AddRange(
 			new string[] {
@@ -46,7 +40,6 @@ public class PropertyEditor : ModuleRules
 				"SlateCore",
 				"EditorWidgets",
 				"Documentation",
-                	"RHI",
 				"ConfigEditor",
                 	"SceneOutliner",
 				"DesktopPlatform",
diff --git a/Engine/Source/Editor/UnrealEd/UnrealEd.Build.cs b/Engine/Source/Editor/UnrealEd/UnrealEd.Build.cs
--- a/Engine/Source/Editor/UnrealEd/UnrealEd.Build.cs
+++ b/Engine/Source/Editor/UnrealEd/UnrealEd.Build.cs
@@ -18,30 +18,21 @@ public class UnrealEd : ModuleRules
 		PrivateIncludePathModuleNames.AddRange(
 			new string[]
 			{
-				"BehaviorTreeEditor",
 				"StructViewer",
 				"MainFrame",
 				"TurnkeySupport",
-				"MergeActors",
 				"MessagingCommon",
 				"MovieSceneCapture",
-				"RenderResourceViewer",
 				"PlacementMode",
 				"Settings",
 				"SettingsEditor",
 				"ViewportSnapping",
 				"SourceCodeAccess",
 				"OutputLog",
-				"MessagingRpc",
-				"PortalRpc",
 				"PortalServices",
-				"Persona",
 				"PhysicsAssetEditor",
-				"ClothingSystemEditorInterface",
 				"Media",
 				"VirtualTexturingEditor",
-				"TextureBuild",
-				"CSVtoSVG",
 				"HotReload",
 				"StaticMeshEditor",
 				"WorkspaceMenuStructure",
@@ -71,7 +62,6 @@ public class UnrealEd : ModuleRules
 				"SourceControl",
 				"UncontrolledChangelists",
 				"UnrealEdMessages",
-				"GameplayDebugger",
 				"BlueprintGraph",
 				"HTTP",
 				"FunctionalTesting",
@@ -113,7 +103,6 @@ public class UnrealEd : ModuleRules
 				"AppFramework",
 				"CinematicCamera",
 				"CurveEditor",
-				"ContentBrowserData",
 				"DataLayerEditor",
 				"DerivedDataCache",
 				"Zen",
@@ -124,7 +113,6 @@ public class UnrealEd : ModuleRules
 				"ImageWrapper",
 				"ImageWriteQueue",
 				"InputCore",
-				"InputBindingEditor",
 				"LauncherServices",
 				"MaterialEditor",
 				"MessageLog",
@@ -137,7 +125,6 @@ public class UnrealEd : ModuleRules
 				"RenderCore",
 				"RHI",
 				"Sockets",
-				"SourceControlWindows",
 				"StatsViewer",
 				"SwarmInterface",
 				"TargetPlatform",
@@ -145,13 +132,11 @@ public class UnrealEd : ModuleRules
 				"EditorWidgets",
 				"GraphEditor",
 				"Kismet",
-				"InternationalizationSettings",
 				"JsonUtilities",
 				"Landscape",
 				"MeshPaint",
 				"Foliage",
 				"FoliageEdit",
-				"VectorVM",
 				"MaterialUtilities",
 				"LocalizationService",
 				"LevelEditor",
@@ -160,21 +145,17 @@ public class UnrealEd : ModuleRules
 				"HierarchicalLODUtilities",
 				"Analytics",
 				"AnalyticsET",
-				"PluginWarden",
 				"PixelInspectorModule",
 				"MovieScene",
 				"MovieSceneTracks",
-				"Sequencer",
 				"ViewportInteraction",
 				"VREditor",
 				"ClothingSystemEditor",
                 "ClothingSystemRuntimeInterface",
                 "ClothingSystemRuntimeCommon",
-                "ClothingSystemRuntimeNv",
 				"PIEPreviewDeviceProfileSelector",
 				"PakFileUtilities",
 				"TimeManagement",
-                "LandscapeEditorUtilities",
 				"ScriptDisassembler",
 				"ToolMenus",
 				"FreeImage",
@@ -185,7 +166,6 @@ public class UnrealEd : ModuleRules
 				"TraceAnalysis",
 				"TraceServices",
 				"DeveloperSettings",
-				"AnimationModifiers",
 				"AnimationBlueprintLibrary",
 				"MaterialBaking",
 				"CookOnTheFly",
@@ -193,6 +173,7 @@ public class UnrealEd : ModuleRules
 				"Zen",
 				"BuildSettings",
 				"SubobjectEditor",
+				"HeadMountedDisplay",
 			}
 		);
 
@@ -299,19 +280,11 @@ public class UnrealEd : ModuleRules
 				"AssetTagsEditor",
 				"CollectionManager",
 				"ContentBrowser",
-				"BlueprintGraph",
 				"MeshUtilities",
 				"KismetCompiler",
-				"NavigationSystem",
 				"GameplayTasks",
-				"AIModule",
-				"Engine",
-				"SourceControl",
-				"UncontrolledChangelists",
 				"PropertyEditor",
 				"ClassViewer",
-				"EditorSubsystem",
-				"WidgetRegistration"
 			}
 		);
 
diff --git a/Engine/Source/Editor/VREditor/VREditor.Build.cs b/Engine/Source/Editor/VREditor/VREditor.Build.cs
--- a/Engine/Source/Editor/VREditor/VREditor.Build.cs
+++ b/Engine/Source/Editor/VREditor/VREditor.Build.cs
@@ -38,7 +38,8 @@ namespace UnrealBuildTool.Rules
 
 			PublicDependencyModuleNames.AddRange(
 				new string[] {
-					"ViewportInteraction"
+					"ViewportInteraction",
+					"HeadMountedDisplay"
 				}
 			);
 
diff --git a/Engine/Source/Runtime/NavigationSystem/NavigationSystem.Build.cs b/Engine/Source/Runtime/NavigationSystem/NavigationSystem.Build.cs
--- a/Engine/Source/Runtime/NavigationSystem/NavigationSystem.Build.cs
+++ b/Engine/Source/Runtime/NavigationSystem/NavigationSystem.Build.cs
@@ -10,17 +10,12 @@ namespace UnrealBuildTool.Rules
         {
             PublicDependencyModuleNames.AddRange(
                 new string[] {
-                    "Core",
+					"Chaos",
+					"Core",
                     "CoreUObject",
                     "Engine",
-                }
-                );
-
-            PrivateDependencyModuleNames.AddRange(
-                new string[] {
-                    "RHI",
-                    "RenderCore",
-                }
+					"GeometryCollectionEngine",
+				}
                 );
 
 			PrivateIncludePathModuleNames.AddRange(
@@ -31,14 +26,14 @@ namespace UnrealBuildTool.Rules
 				}
 				);
 
-            SetupModulePhysicsSupport(Target);
+			PrivateDependencyModuleNames.AddRange(
+				new string[] {
+					"RHI",
+					"RenderCore",
+				}
+				);
 
-            PublicDependencyModuleNames.AddRange(
-                new string[] {
-					"Chaos",
-                    "GeometryCollectionEngine",
-                }
-			);
+			SetupModulePhysicsSupport(Target);
 
             if (Target.bCompileRecast)
             {
