Fixed WorldOutliner force refreshing for nothing during PIE.
- Avoid broadcasting OnActorEditorContextClientChanged for game worlds
- Prevent UActorEditorContextSubsystem from broacasting ActorEditorContextSubsystemChanged when running Game/PIE.

diff --git a/Engine/Source/Editor/LevelEditor/Private/LevelEditorSubsystem.cpp b/Engine/Source/Editor/LevelEditor/Private/LevelEditorSubsystem.cpp
--- a/Engine/Source/Editor/LevelEditor/Private/LevelEditorSubsystem.cpp
+++ b/Engine/Source/Editor/LevelEditor/Private/LevelEditorSubsystem.cpp
@@ -925,12 +925,18 @@ bool ULevelEditorSubsystem::GetActorEditorContextDisplayInfo(UWorld* InWorld, FA
 
 void ULevelEditorSubsystem::OnLevelAddedOrRemoved(ULevel* InLevel, UWorld* InWorld)
 {
-	ActorEditorContextClientChanged.Broadcast(this);
+	if (!InWorld->IsGameWorld())
+	{
+		ActorEditorContextClientChanged.Broadcast(this);
+	}
 }
 
 void ULevelEditorSubsystem::OnCurrentLevelChanged(ULevel* InNewLevel, ULevel* InOldLevel, UWorld* InWorld)
 {
-	ActorEditorContextClientChanged.Broadcast(this);
+	if (!InWorld->IsGameWorld())
+	{
+		ActorEditorContextClientChanged.Broadcast(this);
+	}
 }
 
 void ULevelEditorSubsystem::HandleOnPreSaveWorldWithContext(class UWorld* World, FObjectPreSaveContext ObjectSaveContext)
diff --git a/Engine/Source/Editor/UnrealEd/Private/EditorActorFolders.cpp b/Engine/Source/Editor/UnrealEd/Private/EditorActorFolders.cpp
--- a/Engine/Source/Editor/UnrealEd/Private/EditorActorFolders.cpp
+++ b/Engine/Source/Editor/UnrealEd/Private/EditorActorFolders.cpp
@@ -107,7 +107,7 @@ void FActorFolders::BroadcastOnActorFolderCreated(UWorld& InWorld, const FFolder
 		OnFolderCreate.Broadcast(InWorld, InFolder.GetPath());
 	}
 	PRAGMA_ENABLE_DEPRECATION_WARNINGS
-	BroadcastOnActorEditorContextClientChanged();
+	BroadcastOnActorEditorContextClientChanged(InWorld);
 }
 
 void FActorFolders::BroadcastOnActorFolderDeleted(UWorld& InWorld, const FFolder& InFolder)
@@ -119,7 +119,7 @@ void FActorFolders::BroadcastOnActorFolderDeleted(UWorld& InWorld, const FFolder
 		OnFolderDelete.Broadcast(InWorld, InFolder.GetPath());
 	}
 	PRAGMA_ENABLE_DEPRECATION_WARNINGS
-	BroadcastOnActorEditorContextClientChanged();
+	BroadcastOnActorEditorContextClientChanged(InWorld);
 }
 
 void FActorFolders::BroadcastOnActorFolderMoved(UWorld& InWorld, const FFolder& InSrcFolder, const FFolder& InDstFolder)
@@ -131,7 +131,7 @@ void FActorFolders::BroadcastOnActorFolderMoved(UWorld& InWorld, const FFolder&
 		OnFolderMove.Broadcast(InWorld, InSrcFolder.GetPath(), InDstFolder.GetPath());
 	}
 	PRAGMA_ENABLE_DEPRECATION_WARNINGS
-	BroadcastOnActorEditorContextClientChanged();
+	BroadcastOnActorEditorContextClientChanged(InWorld);
 }
 
 void FActorFolders::OnLevelActorListChanged()
@@ -568,9 +568,12 @@ FFolder FActorFolders::GetActorEditorContextFolder(UWorld& InWorld, bool bMustMa
 
 void FActorFolders::SetActorEditorContextFolder(UWorld& InWorld, const FFolder& InFolder)
 {
-	if (GetOrCreateWorldFolders(InWorld).SetActorEditorContextFolder(InFolder))
+	if (!InWorld.IsGameWorld())
 	{
-		BroadcastOnActorEditorContextClientChanged();
+		if (GetOrCreateWorldFolders(InWorld).SetActorEditorContextFolder(InFolder))
+		{
+			BroadcastOnActorEditorContextClientChanged(InWorld);
+		}
 	}
 }
 
@@ -630,9 +633,12 @@ TSharedRef<SWidget> FActorFolders::GetActorEditorContextWidget(UWorld* InWorld)
 	return SNew(STextBlock).Text(Text);
 }
 
-void FActorFolders::BroadcastOnActorEditorContextClientChanged()
+void FActorFolders::BroadcastOnActorEditorContextClientChanged(UWorld& InWorld)
 {
-	ActorEditorContextClientChanged.Broadcast(this);
+	if (!InWorld.IsGameWorld())
+	{
+		ActorEditorContextClientChanged.Broadcast(this);
+	}
 }
 
 void FActorFolders::ForEachFolder(UWorld& InWorld, TFunctionRef<bool(const FFolder&)> Operation)
diff --git a/Engine/Source/Editor/UnrealEd/Private/Subsystems/ActorEditorContextSubsystem.cpp b/Engine/Source/Editor/UnrealEd/Private/Subsystems/ActorEditorContextSubsystem.cpp
--- a/Engine/Source/Editor/UnrealEd/Private/Subsystems/ActorEditorContextSubsystem.cpp
+++ b/Engine/Source/Editor/UnrealEd/Private/Subsystems/ActorEditorContextSubsystem.cpp
@@ -149,6 +149,11 @@ UWorld* UActorEditorContextSubsystem::GetWorld() const
 
 void UActorEditorContextSubsystem::OnActorEditorContextClientChanged(IActorEditorContextClient* Client)
 {
+	if (IsRunningGame() || (GIsEditor && GEditor->GetPIEWorldContext()))
+	{
+		return;
+	}
+
 	ActorEditorContextSubsystemChanged.Broadcast();
 }
 
diff --git a/Engine/Source/Editor/UnrealEd/Private/WorldFolders.cpp b/Engine/Source/Editor/UnrealEd/Private/WorldFolders.cpp
--- a/Engine/Source/Editor/UnrealEd/Private/WorldFolders.cpp
+++ b/Engine/Source/Editor/UnrealEd/Private/WorldFolders.cpp
@@ -145,7 +145,7 @@ bool UWorldFolders::RenameFolder(const FFolder& InOldFolder, const FFolder& InNe
 		}
 		if (bChanged)
 		{
-			FActorFolders::Get().BroadcastOnActorEditorContextClientChanged();
+			FActorFolders::Get().BroadcastOnActorEditorContextClientChanged(*World.Get());
 		}
 	}
 	return bSuccess;
diff --git a/Engine/Source/Editor/UnrealEd/Public/EditorActorFolders.h b/Engine/Source/Editor/UnrealEd/Public/EditorActorFolders.h
--- a/Engine/Source/Editor/UnrealEd/Public/EditorActorFolders.h
+++ b/Engine/Source/Editor/UnrealEd/Public/EditorActorFolders.h
@@ -221,7 +221,7 @@ private:
 	TMap<TWeakObjectPtr<UWorld>, TObjectPtr<UWorldFolders>> WorldFolders;
 
 	/** Called when ActorEditorContextClient changed. */
-	void BroadcastOnActorEditorContextClientChanged();
+	void BroadcastOnActorEditorContextClientChanged(UWorld& InWorld);
 
 	/** Delegate used to notify changes to ActorEditorContextSubsystem */
 	FOnActorEditorContextClientChanged ActorEditorContextClientChanged;
