Cooker and GeneratedPackages: Set the chunks of the generated package equal to the chunks of its generator.
#rn Bugfix, Cooking

diff --git a/Engine/Source/Editor/UnrealEd/Private/Commandlets/AssetRegistryGenerator.cpp b/Engine/Source/Editor/UnrealEd/Private/Commandlets/AssetRegistryGenerator.cpp
--- a/Engine/Source/Editor/UnrealEd/Private/Commandlets/AssetRegistryGenerator.cpp
+++ b/Engine/Source/Editor/UnrealEd/Private/Commandlets/AssetRegistryGenerator.cpp
@@ -867,10 +867,19 @@ void FAssetRegistryGenerator::CalculateChunkIdsAndAssignToManifest(const FName&
 	}
 	else
 	{
-		TArray<int32> PackageChunkIDs = GetExplicitChunkIDs(PackageFName);
-		ExistingChunkIDs = GetExistingPackageChunkAssignments(PackageFName);
+		FName PackageNameThatDefinesChunks = PackageFName;
+
+		// Generated packages use the chunks defined by their Generator
+		FName GeneratorName = GetGeneratorPackage(PackageFName, this->State);
+		if (!GeneratorName.IsNone())
+		{
+			PackageNameThatDefinesChunks = GeneratorName;
+		}
+
+		TArray<int32> PackageChunkIDs = GetExplicitChunkIDs(PackageNameThatDefinesChunks);
+		ExistingChunkIDs = GetExistingPackageChunkAssignments(PackageNameThatDefinesChunks);
 		PackageChunkIDs.Append(ExistingChunkIDs);
-		UAssetManager::Get().GetPackageChunkIds(PackageFName, TargetPlatform, PackageChunkIDs, TargetChunks);
+		UAssetManager::Get().GetPackageChunkIds(PackageNameThatDefinesChunks, TargetPlatform, PackageChunkIDs, TargetChunks);
 	}
 
 	// Add the package to the manifest for every chunk the AssetManager found it should belong to
@@ -1326,15 +1335,11 @@ void FAssetRegistryGenerator::ComputePackageDifferences(const FComputeDifference
 			// If it's a generated package, exclude it from the results list and do not remove it.
 			// It will be evaluated for identical/modified/removed only if the generator package is cooked,
 			// during the generator's process step.
-			TConstArrayView<const FAssetData*> PreviousAssets = PreviousState.GetAssetsByPackageName(PackageName);
-			bool bGenerated = !PreviousAssets.IsEmpty() && (PreviousAssets[0]->PackageFlags & PKG_CookGenerated) != 0;
-			if (bGenerated)
+			FName GeneratorName = GetGeneratorPackage(PackageName, PreviousState);
+			if (!GeneratorName.IsNone())
 			{
 				// Keep track of all generators and their list of generated
-				TArray<FAssetIdentifier> Referencers;
-				PreviousState.GetReferencers(FAssetIdentifier(PackageName), Referencers, UE::AssetRegistry::EDependencyCategory::Package);
-				FName GeneratorName = Referencers.Num() == 1 ? Referencers[0].PackageName : NAME_None;
-				const FAssetPackageData* GeneratorData = !GeneratorName.IsNone() ? State.GetAssetPackageData(GeneratorName) : nullptr;
+				const FAssetPackageData* GeneratorData = State.GetAssetPackageData(GeneratorName);
 				if (!GeneratorData)
 				{
 					// Mark it as removed; it will be regenerated when the Generator cooks
@@ -1479,6 +1484,21 @@ void FAssetRegistryGenerator::ComputePackageDifferences(const FComputeDifference
 	}
 }
 
+FName FAssetRegistryGenerator::GetGeneratorPackage(FName PackageName, const FAssetRegistryState& InState)
+{
+	TConstArrayView<const FAssetData*> Assets = InState.GetAssetsByPackageName(PackageName);
+	bool bGenerated = !Assets.IsEmpty() && (Assets[0]->PackageFlags & PKG_CookGenerated) != 0;
+	if (!bGenerated)
+	{
+		return NAME_None;
+	}
+
+	TArray<FAssetIdentifier> Referencers;
+	InState.GetReferencers(FAssetIdentifier(PackageName), Referencers, UE::AssetRegistry::EDependencyCategory::Package);
+	FName GeneratorName = Referencers.Num() == 1 ? Referencers[0].PackageName : NAME_None;
+	return GeneratorName;
+}
+
 void FAssetRegistryGenerator::ComputePackageRemovals(const FAssetRegistryState& PreviousState, TArray<FName>& OutRemovedPackages,
 	TMap<FName, FGeneratorPackageInfo>& OutGeneratorPackages, int32& OutNumNeverCookPlaceHolderPackages)
 {
diff --git a/Engine/Source/Editor/UnrealEd/Public/Commandlets/AssetRegistryGenerator.h b/Engine/Source/Editor/UnrealEd/Public/Commandlets/AssetRegistryGenerator.h
--- a/Engine/Source/Editor/UnrealEd/Public/Commandlets/AssetRegistryGenerator.h
+++ b/Engine/Source/Editor/UnrealEd/Public/Commandlets/AssetRegistryGenerator.h
@@ -474,6 +474,9 @@ private:
 	 * Helper function to find or create asset data for the input object. If the asset is not in the registry it will be added.
 	 */
 	const FAssetData* CreateOrFindAssetData(UObject& Object);
+
+	/** If InState records PackageName is generated, return the name of the Generator, otherwise return NAME_None. */
+	static FName GetGeneratorPackage(FName PackageName, const FAssetRegistryState& InState);
 };
 
 namespace UE::Cook
