Use new Blueprint association change for SNewClassDialog
Support new blueprint type for entity prefab

diff --git a/Engine/Source/Editor/GameProjectGeneration/Private/SNewClassDialog.cpp b/Engine/Source/Editor/GameProjectGeneration/Private/SNewClassDialog.cpp
--- a/Engine/Source/Editor/GameProjectGeneration/Private/SNewClassDialog.cpp
+++ b/Engine/Source/Editor/GameProjectGeneration/Private/SNewClassDialog.cpp
@@ -1128,7 +1128,7 @@ void SNewClassDialog::FinishClicked()
 			if (Package)
 			{
 				// Create and init a new Blueprint
-				UBlueprint* NewBP = FKismetEditorUtilities::CreateBlueprint(const_cast<UClass*>(ParentClassInfo.BaseClass), Package, FName(*NewClassName), BPTYPE_Normal, UBlueprint::StaticClass(), UBlueprintGeneratedClass::StaticClass());
+				UBlueprint* NewBP = FKismetEditorUtilities::CreateBlueprint(const_cast<UClass*>(ParentClassInfo.BaseClass), Package, FName(*NewClassName), BPTYPE_Normal);
 				if (NewBP)
 				{
 					// Notify the asset registry
diff --git a/Engine/Source/Editor/UnrealEd/Private/Kismet2/Kismet2.cpp b/Engine/Source/Editor/UnrealEd/Private/Kismet2/Kismet2.cpp
--- a/Engine/Source/Editor/UnrealEd/Private/Kismet2/Kismet2.cpp
+++ b/Engine/Source/Editor/UnrealEd/Private/Kismet2/Kismet2.cpp
@@ -414,6 +414,17 @@ void FKismetEditorUtilities::CreateDefaultEventGraphs(UBlueprint* Blueprint)
 
 extern UNREALED_API FSecondsCounterData BlueprintCompileAndLoadTimerData;
 
+/** Create a new Blueprint and initialize it to a valid state but uses the associated blueprint types. */
+UBlueprint* FKismetEditorUtilities::CreateBlueprint(UClass* ParentClass, UObject* Outer, const FName NewBPName,	EBlueprintType BlueprintType, FName CallingContext)
+{
+	UClass* BlueprintClassType = UBlueprint::StaticClass();
+	UClass* BlueprintGeneratedClassType = UBlueprintGeneratedClass::StaticClass();
+	const IKismetCompilerInterface& KismetCompilerModule = FModuleManager::LoadModuleChecked<IKismetCompilerInterface>("KismetCompiler");
+	KismetCompilerModule.GetBlueprintTypesForClass(ParentClass, BlueprintClassType, BlueprintGeneratedClassType);
+
+	return CreateBlueprint(ParentClass, Outer, NewBPName, BlueprintType, BlueprintClassType, BlueprintGeneratedClassType, CallingContext);
+}
+
 /** Create a new Blueprint and initialize it to a valid state. */
 UBlueprint* FKismetEditorUtilities::CreateBlueprint(UClass* ParentClass, UObject* Outer, const FName NewBPName, EBlueprintType BlueprintType, TSubclassOf<UBlueprint> BlueprintClassType, TSubclassOf<UBlueprintGeneratedClass> BlueprintGeneratedClassType, FName CallingContext)
 {
diff --git a/Engine/Source/Editor/UnrealEd/Public/Kismet2/KismetEditorUtilities.h b/Engine/Source/Editor/UnrealEd/Public/Kismet2/KismetEditorUtilities.h
--- a/Engine/Source/Editor/UnrealEd/Public/Kismet2/KismetEditorUtilities.h
+++ b/Engine/Source/Editor/UnrealEd/Public/Kismet2/KismetEditorUtilities.h
@@ -97,6 +97,18 @@ public:
 		FOnBlueprintCreated OnBlueprintCreated;
 	};
 
+	/**
+	 * Create a new Blueprint and initialize it to a valid state.
+	 *
+	 * @param ParentClass					the parent class of the new blueprint
+	 * @param Outer							the outer object of the new blueprint
+	 * @param NewBPName						the name of the new blueprint
+	 * @param BlueprintType					the type of the new blueprint (normal, const, etc)
+	 * @param CallingContext				the name of the calling method or module used to identify creation methods to engine analytics/usage stats (default None will be ignored)
+	 * @return								the new blueprint
+	 */
+	static UBlueprint* CreateBlueprint(UClass* ParentClass, UObject* Outer, const FName NewBPName, enum EBlueprintType BlueprintType, FName CallingContext = NAME_None);
+
 	/**
 	 * Create a new Blueprint and initialize it to a valid state.
 	 *
