Run cache modifications in a pipe to move them off the main thread.
Savesd 6s off asset registry scan in a large project.

diff --git a/Engine/Source/Editor/UnrealEd/Private/AutoReimport/AssetSourceFilenameCache.cpp b/Engine/Source/Editor/UnrealEd/Private/AutoReimport/AssetSourceFilenameCache.cpp
--- a/Engine/Source/Editor/UnrealEd/Private/AutoReimport/AssetSourceFilenameCache.cpp
+++ b/Engine/Source/Editor/UnrealEd/Private/AutoReimport/AssetSourceFilenameCache.cpp
@@ -19,6 +19,7 @@
 #include "UObject/Object.h"
 
 FAssetSourceFilenameCache::FAssetSourceFilenameCache()
+	: CachePipe(UE_SOURCE_LOCATION)
 {
 	if (IsEngineExitRequested())
 	{
@@ -28,8 +29,8 @@ FAssetSourceFilenameCache::FAssetSourceFilenameCache()
 
 	IAssetRegistry& AssetRegistry = FModuleManager::LoadModuleChecked<FAssetRegistryModule>("AssetRegistry").Get();
 
-	AssetRegistry.OnAssetAdded().AddRaw(this, &FAssetSourceFilenameCache::HandleOnAssetAdded);
-	AssetRegistry.OnAssetRemoved().AddRaw(this, &FAssetSourceFilenameCache::HandleOnAssetRemoved);
+	AssetRegistry.OnAssetsAdded().AddRaw(this, &FAssetSourceFilenameCache::HandleOnAssetsAdded);
+	AssetRegistry.OnAssetsRemoved().AddRaw(this, &FAssetSourceFilenameCache::HandleOnAssetsRemoved);
 	AssetRegistry.OnAssetRenamed().AddRaw(this, &FAssetSourceFilenameCache::HandleOnAssetRenamed);
 
 	UAssetImportData::OnImportDataChanged.AddRaw(this, &FAssetSourceFilenameCache::HandleOnAssetUpdated);
@@ -55,8 +56,8 @@ void FAssetSourceFilenameCache::Shutdown()
 		IAssetRegistry* AssetRegistry = AssetRegistryModule->TryGet();
 		if (AssetRegistry)
 		{
-			AssetRegistry->OnAssetAdded().RemoveAll(this);
-			AssetRegistry->OnAssetRemoved().RemoveAll(this);
+			AssetRegistry->OnAssetsAdded().RemoveAll(this);
+			AssetRegistry->OnAssetsRemoved().RemoveAll(this);
 			AssetRegistry->OnAssetRenamed().RemoveAll(this);
 		}
 	}
@@ -91,143 +92,149 @@ TOptional<FAssetImportInfo> FAssetSourceFilenameCache::ExtractAssetImportInfo(co
 	}
 }
 
-void FAssetSourceFilenameCache::HandleOnAssetAdded(const FAssetData& AssetData)
+void FAssetSourceFilenameCache::HandleOnAssetsAdded(TConstArrayView<FAssetData> Assets)
 {
-	TOptional<FAssetImportInfo> ImportData = ExtractAssetImportInfo(AssetData);
-	if (ImportData.IsSet())
-	{
-		for (const auto& SourceFile : ImportData->SourceFiles)
-		{
-			SourceFileToObjectPathCache.FindOrAdd(FPaths::GetCleanFilename(SourceFile.RelativeFilename)).Add(AssetData.GetSoftObjectPath());
-		}
-	}
-}
-
-void FAssetSourceFilenameCache::HandleOnAssetsAdded(const TArray<FAssetData>& Assets)
-{
-	const uint32 AssetCount = Assets.Num();
-	TArray<TOptional<FAssetImportInfo>> ImportDataList;
-	ImportDataList.SetNum(AssetCount);
-	ParallelFor(Assets.Num(), [&Assets, &ImportDataList](int32 Index)
-		{
-			ImportDataList[Index] = ExtractAssetImportInfo(Assets[Index]);
-		});
+	CachePipe.Launch(UE_SOURCE_LOCATION, [this, Assets=TArray<FAssetData>(Assets)](){
+		TRACE_CPUPROFILER_EVENT_SCOPE(FAssetSourceFilenameCache::HandleOnAssetsAdded);
+		const uint32 AssetCount = Assets.Num();
+		TArray<TOptional<FAssetImportInfo>> ImportDataList;
+		ImportDataList.SetNum(AssetCount);
+		ParallelFor(Assets.Num(), [&Assets, &ImportDataList](int32 Index)
+			{
+				ImportDataList[Index] = ExtractAssetImportInfo(Assets[Index]);
+			});
 
-	for (uint32 Index = 0; Index < AssetCount; Index++)
-	{
-		const TOptional<FAssetImportInfo>& ImportData = ImportDataList[Index];
-		if (ImportData.IsSet())
+		SourceFileToObjectPathCache.Reserve(SourceFileToObjectPathCache.Num() + Assets.Num());
+		for (uint32 Index = 0; Index < AssetCount; Index++)
 		{
-			for (const auto& SourceFile : ImportData->SourceFiles)
+			const TOptional<FAssetImportInfo>& ImportData = ImportDataList[Index];
+			if (ImportData.IsSet())
 			{
-				SourceFileToObjectPathCache.FindOrAdd(FPaths::GetCleanFilename(SourceFile.RelativeFilename)).Add(Assets[Index].GetSoftObjectPath());
+				for (const auto& SourceFile : ImportData->SourceFiles)
+				{
+					SourceFileToObjectPathCache.FindOrAdd(FPaths::GetCleanFilename(SourceFile.RelativeFilename)).Add(Assets[Index].GetSoftObjectPath());
+				}
 			}
 		}
-	}
+	}, LowLevelTasks::ETaskPriority::BackgroundNormal);
 }
 
-void FAssetSourceFilenameCache::HandleOnAssetRemoved(const FAssetData& AssetData)
+void FAssetSourceFilenameCache::HandleOnAssetsRemoved(TConstArrayView<FAssetData> Assets)
 {
-	TOptional<FAssetImportInfo> ImportData = ExtractAssetImportInfo(AssetData);
-	if (ImportData.IsSet())
-	{
-		for (auto& SourceFile : ImportData->SourceFiles)
+	CachePipe.Launch(UE_SOURCE_LOCATION, [this, Assets=TArray<FAssetData>(Assets)](){
+		TRACE_CPUPROFILER_EVENT_SCOPE(FAssetSourceFilenameCache::HandleOnAssetsRemoved);
+		for (const FAssetData& AssetData : Assets)
 		{
-			FString CleanFilename = FPaths::GetCleanFilename(SourceFile.RelativeFilename);
-			if (auto* Objects = SourceFileToObjectPathCache.Find(CleanFilename))
+			TOptional<FAssetImportInfo> ImportData = ExtractAssetImportInfo(AssetData);
+			if (ImportData.IsSet())
 			{
-				Objects->Remove(AssetData.GetSoftObjectPath());
-				if (Objects->Num() == 0)
+				for (auto& SourceFile : ImportData->SourceFiles)
 				{
-					SourceFileToObjectPathCache.Remove(CleanFilename);
+					FString CleanFilename = FPaths::GetCleanFilename(SourceFile.RelativeFilename);
+					if (auto* Objects = SourceFileToObjectPathCache.Find(CleanFilename))
+					{
+						Objects->Remove(AssetData.GetSoftObjectPath());
+						if (Objects->Num() == 0)
+						{
+							SourceFileToObjectPathCache.Remove(CleanFilename);
+						}
+					}
 				}
 			}
 		}
-	}
+	}, LowLevelTasks::ETaskPriority::BackgroundNormal);
 }
 
 void FAssetSourceFilenameCache::HandleOnAssetRenamed(const FAssetData& AssetData, const FString& OldPath)
 {
-	TOptional<FAssetImportInfo> ImportData = ExtractAssetImportInfo(AssetData);
-	if (ImportData.IsSet())
-	{
-		for (auto& SourceFile : ImportData->SourceFiles)
+	// Capture by ref as we wait inline
+	CachePipe.Launch(UE_SOURCE_LOCATION, [this, &AssetData, &OldPath]() {
+		TOptional<FAssetImportInfo> ImportData = ExtractAssetImportInfo(AssetData);
+		if (ImportData.IsSet())
 		{
-			FString CleanFilename = FPaths::GetCleanFilename(SourceFile.RelativeFilename);
-
-			if (auto* Objects = SourceFileToObjectPathCache.Find(CleanFilename))
+			for (auto& SourceFile : ImportData->SourceFiles)
 			{
-				Objects->Remove(FSoftObjectPath(OldPath));
-				if (Objects->Num() == 0)
+				FString CleanFilename = FPaths::GetCleanFilename(SourceFile.RelativeFilename);
+
+				if (auto* Objects = SourceFileToObjectPathCache.Find(CleanFilename))
 				{
-					SourceFileToObjectPathCache.Remove(CleanFilename);
+					Objects->Remove(FSoftObjectPath(OldPath));
+					if (Objects->Num() == 0)
+					{
+						SourceFileToObjectPathCache.Remove(CleanFilename);
+					}
 				}
-			}
 
-			SourceFileToObjectPathCache.FindOrAdd(CleanFilename).Add(AssetData.GetSoftObjectPath());
+				SourceFileToObjectPathCache.FindOrAdd(CleanFilename).Add(AssetData.GetSoftObjectPath());
+			}
 		}
-	}
+	}, LowLevelTasks::ETaskPriority::Normal, UE::Tasks::EExtendedTaskPriority::Inline).Wait();
 
 	AssetRenamedEvent.Broadcast(AssetData, OldPath);
 }
 
 void FAssetSourceFilenameCache::HandleOnAssetUpdated(const FAssetImportInfo& OldData, const UAssetImportData* ImportData)
 {
-	FSoftObjectPath ObjectPath = FSoftObjectPath(ImportData->GetOuter());
+	// Run inline on pipe to safely use ImportData
+	CachePipe.Launch(UE_SOURCE_LOCATION, [this, &OldData, ImportData]() {
+		FSoftObjectPath ObjectPath = FSoftObjectPath(ImportData->GetOuter());
 
-	for (const auto& SourceFile : OldData.SourceFiles)
-	{
-		FString CleanFilename = FPaths::GetCleanFilename(SourceFile.RelativeFilename);
-		if (auto* Objects = SourceFileToObjectPathCache.Find(CleanFilename))
+		for (const auto& SourceFile : OldData.SourceFiles)
 		{
-			Objects->Remove(ObjectPath);
-			if (Objects->Num() == 0)
+			FString CleanFilename = FPaths::GetCleanFilename(SourceFile.RelativeFilename);
+			if (auto* Objects = SourceFileToObjectPathCache.Find(CleanFilename))
 			{
-				SourceFileToObjectPathCache.Remove(CleanFilename);
+				Objects->Remove(ObjectPath);
+				if (Objects->Num() == 0)
+				{
+					SourceFileToObjectPathCache.Remove(CleanFilename);
+				}
 			}
 		}
-	}
 
-	for (const auto& SourceFile : ImportData->SourceData.SourceFiles)
-	{
-		SourceFileToObjectPathCache.FindOrAdd(FPaths::GetCleanFilename(SourceFile.RelativeFilename)).Add(ObjectPath);
-	}
+		for (const auto& SourceFile : ImportData->SourceData.SourceFiles)
+		{
+			SourceFileToObjectPathCache.FindOrAdd(FPaths::GetCleanFilename(SourceFile.RelativeFilename)).Add(ObjectPath);
+		}
+	}, LowLevelTasks::ETaskPriority::Normal, UE::Tasks::EExtendedTaskPriority::Inline).Wait();
 }
 
 TArray<FAssetData> FAssetSourceFilenameCache::GetAssetsPertainingToFile(const IAssetRegistry& Registry, const FString& AbsoluteFilename) const
 {
-	TArray<FAssetData> Assets;
-	
-	if (const auto* ObjectPaths = SourceFileToObjectPathCache.Find(FPaths::GetCleanFilename(AbsoluteFilename)))
+	return CachePipe.Launch(UE_SOURCE_LOCATION, [this, &Registry, &AbsoluteFilename]()
 	{
-		for (const FSoftObjectPath& Path : *ObjectPaths)
+		TArray<FAssetData> Assets;
+		if (const auto* ObjectPaths = SourceFileToObjectPathCache.Find(FPaths::GetCleanFilename(AbsoluteFilename)))
 		{
-			FAssetData Asset = Registry.GetAssetByObjectPath(Path);
-			TOptional<FAssetImportInfo> ImportInfo = ExtractAssetImportInfo(Asset);
-			if (ImportInfo.IsSet())
+			for (const FSoftObjectPath& Path : *ObjectPaths)
 			{
-				auto AssetPackagePath = FPackageName::LongPackageNameToFilename(Asset.PackagePath.ToString() / TEXT(""));
+				FAssetData Asset = Registry.GetAssetByObjectPath(Path);
+				TOptional<FAssetImportInfo> ImportInfo = ExtractAssetImportInfo(Asset);
+				if (ImportInfo.IsSet())
+				{
+					auto AssetPackagePath = FPackageName::LongPackageNameToFilename(Asset.PackagePath.ToString() / TEXT(""));
 
-				// Attempt to find the matching source filename in the list of imported sorce files (generally there are only one of these)
-				const bool bWasImportedFromFile = ImportInfo->SourceFiles.ContainsByPredicate([&](const FAssetImportInfo::FSourceFile& File){
+					// Attempt to find the matching source filename in the list of imported sorce files (generally there are only one of these)
+					const bool bWasImportedFromFile = ImportInfo->SourceFiles.ContainsByPredicate([&](const FAssetImportInfo::FSourceFile& File){
 
-					if (AbsoluteFilename == FPaths::ConvertRelativePathToFull(AssetPackagePath / File.RelativeFilename) ||
-						AbsoluteFilename == FPaths::ConvertRelativePathToFull(File.RelativeFilename))
-					{
-						return true;
-					}
+						if (AbsoluteFilename == FPaths::ConvertRelativePathToFull(AssetPackagePath / File.RelativeFilename) ||
+							AbsoluteFilename == FPaths::ConvertRelativePathToFull(File.RelativeFilename))
+						{
+							return true;
+						}
 
-					return false;
-				});
+						return false;
+					});
 
-				if (bWasImportedFromFile)
-				{
-					Assets.Emplace();
-					Swap(Asset, Assets[Assets.Num() - 1]);
+					if (bWasImportedFromFile)
+					{
+						Assets.Emplace();
+						Swap(Asset, Assets[Assets.Num() - 1]);
+					}
 				}
 			}
 		}
-	}
 
-	return Assets;
+		return Assets;
+	}, LowLevelTasks::ETaskPriority::Normal, UE::Tasks::EExtendedTaskPriority::Inline).GetResult();
 }
diff --git a/Engine/Source/Editor/UnrealEd/Public/AutoReimport/AssetSourceFilenameCache.h b/Engine/Source/Editor/UnrealEd/Public/AutoReimport/AssetSourceFilenameCache.h
--- a/Engine/Source/Editor/UnrealEd/Public/AutoReimport/AssetSourceFilenameCache.h
+++ b/Engine/Source/Editor/UnrealEd/Public/AutoReimport/AssetSourceFilenameCache.h
@@ -11,6 +11,7 @@
 #include "Delegates/Delegate.h"
 #include "EditorFramework/AssetImportData.h"
 #include "Misc/Optional.h"
+#include "Tasks/Pipe.h"
 #include "UObject/Object.h"
 
 class FName;
@@ -44,15 +45,17 @@ public:
 
 private:
 	/** Delegate bindings that keep the cache up-to-date */
-	void HandleOnAssetAdded(const FAssetData& AssetData);
-	void HandleOnAssetsAdded(const TArray<FAssetData>& Assets);
-	void HandleOnAssetRemoved(const FAssetData& AssetData);
+	void HandleOnAssetsAdded(TConstArrayView<FAssetData> Assets);
+	void HandleOnAssetsRemoved(TConstArrayView<FAssetData> AssetData);
 	void HandleOnAssetRenamed(const FAssetData& AssetData, const FString& OldPath);
 	void HandleOnAssetUpdated(const FAssetImportInfo& OldData, const UAssetImportData* ImportData);
 	
 	/** Event that is triggered when an asset has been renamed, and we've updated our cache */
 	FAssetRenamedEvent AssetRenamedEvent;
 
+	// Pipe to serialize access to cache 
+	mutable UE::Tasks::FPipe CachePipe;	
+
 	/** Map of clean filenames (no leading path information) to object paths that were imported with that file */
 	TMap<FString, TSet<FSoftObjectPath>> SourceFileToObjectPathCache;
 };
