Cooker: Tweak comments and whitespace to match upcoming change. Narrow the scope of !IsCookingInEditor in ReleaseCookedPlatformData because the upcoming change will add calls that need to be completed even in editor.
#rnx

diff --git a/Engine/Source/Editor/UnrealEd/Private/CookOnTheFlyServer.cpp b/Engine/Source/Editor/UnrealEd/Private/CookOnTheFlyServer.cpp
--- a/Engine/Source/Editor/UnrealEd/Private/CookOnTheFlyServer.cpp
+++ b/Engine/Source/Editor/UnrealEd/Private/CookOnTheFlyServer.cpp
@@ -3735,9 +3735,13 @@ UE::Cook::EPollStatus UCookOnTheFlyServer::CallBeginCacheOnObjects(UE::Cook::FPa
 		UObject* Obj = ObjectsData[ObjectIndex].Get();
 		if (!Obj)
 		{
-			// Objects can be marked as pending kill even without a garbage collect, and our weakptr.get will return null for them, so we have to always check the WeakPtr before using it
-			// Treat objects that have been marked as pending kill or deleted as no-longer-required for BeginCacheForCookedPlatformData and ClearAllCachedCookedPlatformData
-			ObjectsData[ObjectIndex] = nullptr; // If the weakptr is merely pendingkill, set it to null explicitly so we don't think that we've called BeginCacheForCookedPlatformData on it if it gets unmarked pendingkill later
+			// Objects can be marked as pending kill even without a garbage collect, and our weakptr.get will return
+			// null for them, so we have to always check the WeakPtr before using it.
+			// Treat objects that have been marked as pending kill or deleted as no-longer-required for
+			// BeginCacheForCookedPlatformData and ClearAllCachedCookedPlatformData
+			// In case the weakptr is merely pendingkill, set it to null explicitly so we don't think that we've called
+			// BeginCacheForCookedPlatformData on it if it gets unmarked pendingkill later
+			ObjectsData[ObjectIndex] = nullptr; 
 			++NextIndex;
 			continue;
 		}
@@ -3777,7 +3781,8 @@ UE::Cook::EPollStatus UCookOnTheFlyServer::CallBeginCacheOnObjects(UE::Cook::FPa
 		else
 		{
 			bool bNeedsResourceRelease = CurrentAsyncCache != nullptr;
-			PackageDatas->AddPendingCookedPlatformData(FPendingCookedPlatformData(Obj, TargetPlatform, PackageData, bNeedsResourceRelease, *this));
+			PackageDatas->AddPendingCookedPlatformData(FPendingCookedPlatformData(Obj, TargetPlatform,
+				PackageData, bNeedsResourceRelease, *this));
 		}
 
 		if (Timer.IsActionTimeUp())
@@ -3813,37 +3818,40 @@ void UCookOnTheFlyServer::ReleaseCookedPlatformData(UE::Cook::FPackageData& Pack
 	if (ReleaseSaveReason == EReleaseSaveReason::Completed)
 	{
 		// Since we have completed CookedPlatformData, we know we called BeginCacheForCookedPlatformData on all objects in the package, and none are pending
-		if (!IsCookingInEditor()) // ClearAllCachedCookedPlatformData and WillNeverCacheCookedPlatformDataAgain calls are only used when not in editor
+		UE_SCOPED_HIERARCHICAL_COOKTIMER(ClearAllCachedCookedPlatformData);
+		if (GenerationInfo)
 		{
-			UE_SCOPED_HIERARCHICAL_COOKTIMER(ClearAllCachedCookedPlatformData);
-			if (GenerationInfo)
-			{
-				check(GenerationInfo->HasTakenOverCachedCookedPlatformData());
-				check(PackageData.GetCachedObjectsInOuter().Num() == 0);
+			check(GenerationInfo->HasTakenOverCachedCookedPlatformData());
+			check(PackageData.GetCachedObjectsInOuter().Num() == 0);
 
-				for (FBeginCacheObject& BeginCacheObject : GenerationInfo->BeginCacheObjects.Objects)
+			for (FBeginCacheObject& BeginCacheObject : GenerationInfo->BeginCacheObjects.Objects)
+			{
+				UObject* Object = BeginCacheObject.Object.Get();
+				if (Object)
 				{
-					UObject* Object = BeginCacheObject.Object.Get();
-					if (Object)
+					check(BeginCacheObject.bHasFinishedRound);
+					if (!IsCookingInEditor()) // ClearAllCachedCookedPlatformData and WillNeverCacheCookedPlatformDataAgain calls are only used when not in editor
 					{
-						check(BeginCacheObject.bHasFinishedRound);
 						Object->ClearAllCachedCookedPlatformData();
-						if (IsDirectorCookByTheBook() && !IsCookingInEditor())
+						if (IsDirectorCookByTheBook())
 						{
 							Object->WillNeverCacheCookedPlatformDataAgain();
 						}
 					}
 				}
 			}
-			else
+		}
+		else
+		{
+			for (FWeakObjectPtr& WeakPtr : PackageData.GetCachedObjectsInOuter())
 			{
-				for (FWeakObjectPtr& WeakPtr : PackageData.GetCachedObjectsInOuter())
+				UObject* Object = WeakPtr.Get();
+				if (Object)
 				{
-					UObject* Object = WeakPtr.Get();
-					if (Object)
+					if (!IsCookingInEditor()) // ClearAllCachedCookedPlatformData and WillNeverCacheCookedPlatformDataAgain calls are only used when not in editor
 					{
 						Object->ClearAllCachedCookedPlatformData();
-						if (IsDirectorCookByTheBook() && !IsCookingInEditor())
+						if (IsDirectorCookByTheBook())
 						{
 							Object->WillNeverCacheCookedPlatformDataAgain();
 						}
@@ -3858,91 +3866,91 @@ void UCookOnTheFlyServer::ReleaseCookedPlatformData(UE::Cook::FPackageData& Pack
 		// Note that even after we return from this function, some objects with pending IsCachedCookedPlatformDataLoaded
 		// calls may still exist for this Package in PendingCookedPlatformDatas
 		// and this PackageData may therefore still have GetNumPendingCookedPlatformData > 0
-		if (!IsCookingInEditor()) // ClearAllCachedCookedPlatformData calls are only used when not in editor.
+		int32 NumPlatforms = PackageData.GetCachedObjectsInOuterNumPlatforms();
+		if (NumPlatforms > 0) // Shouldn't happen because PumpSaves checks for this, but avoid a divide by 0 if it does.
 		{
-			int32 NumPlatforms = PackageData.GetCachedObjectsInOuterNumPlatforms();
-			if (NumPlatforms > 0) // Shouldn't happen because PumpSaves checks for this, but avoid a divide by 0 if it does.
-			{
-				// We have only called BeginCacheForCookedPlatformData on Object,Platform pairs up to GetCookedPlatformDataNextIndex.
-				// Further, some of those calls might still be pending.
+			// We have only called BeginCacheForCookedPlatformData on Object,Platform pairs up to GetCookedPlatformDataNextIndex.
+			// Further, some of those calls might still be pending.
 
-				// Find all pending BeginCacheForCookedPlatformData for this FPackageData
-				TMap<UObject*, TArray<FPendingCookedPlatformData*>> PendingObjects;
-				PackageDatas->ForEachPendingCookedPlatformData(
-				[&PendingObjects, &PackageData](FPendingCookedPlatformData& PendingCookedPlatformData)
+			// Find all pending BeginCacheForCookedPlatformData for this FPackageData
+			TMap<UObject*, TArray<FPendingCookedPlatformData*>> PendingObjects;
+			PackageDatas->ForEachPendingCookedPlatformData(
+			[&PendingObjects, &PackageData](FPendingCookedPlatformData& PendingCookedPlatformData)
+			{
+				if (&PendingCookedPlatformData.PackageData == &PackageData && !PendingCookedPlatformData.PollIsComplete())
 				{
-					if (&PendingCookedPlatformData.PackageData == &PackageData && !PendingCookedPlatformData.PollIsComplete())
-					{
-						UObject* Object = PendingCookedPlatformData.Object.Get();
-						check(Object); // Otherwise PollIsComplete would have returned true
-						check(!PendingCookedPlatformData.bHasReleased); // bHasReleased should be false since PollIsComplete returned false
-						PendingObjects.FindOrAdd(Object).Add(&PendingCookedPlatformData);
-					}
-				});
+					UObject* Object = PendingCookedPlatformData.Object.Get();
+					check(Object); // Otherwise PollIsComplete would have returned true
+					check(!PendingCookedPlatformData.bHasReleased); // bHasReleased should be false since PollIsComplete returned false
+					PendingObjects.FindOrAdd(Object).Add(&PendingCookedPlatformData);
+				}
+			});
 
-				TArray<UObject*> ObjectsToClear;
-				TArray<FWeakObjectPtr>* CachedObjects = &PackageData.GetCachedObjectsInOuter();
-				int32 NumIndexes = PackageData.GetCookedPlatformDataNextIndex();
-				if (GenerationInfo)
+			TArray<UObject*> ObjectsToClear;
+			TArray<FWeakObjectPtr>* CachedObjects = &PackageData.GetCachedObjectsInOuter();
+			int32 NumIndexes = PackageData.GetCookedPlatformDataNextIndex();
+			if (GenerationInfo)
+			{
+				if (GenerationInfo->HasTakenOverCachedCookedPlatformData())
 				{
-					if (GenerationInfo->HasTakenOverCachedCookedPlatformData())
-					{
-						check(PackageData.GetCachedObjectsInOuter().Num() == 0);
+					check(PackageData.GetCachedObjectsInOuter().Num() == 0);
 
-						// Add all objects that were cached in previous rounds
-						for (FBeginCacheObject& BeginCacheObject : GenerationInfo->BeginCacheObjects.Objects)
+					// Add all objects that were cached in previous rounds
+					for (FBeginCacheObject& BeginCacheObject : GenerationInfo->BeginCacheObjects.Objects)
+					{
+						UObject* Object = BeginCacheObject.Object.Get();
+						if (Object && BeginCacheObject.bHasFinishedRound)
 						{
-							UObject* Object = BeginCacheObject.Object.Get();
-							if (Object && BeginCacheObject.bHasFinishedRound)
-							{
-								ObjectsToClear.Add(Object);
-							}
+							ObjectsToClear.Add(Object);
 						}
-						// Add all objects cached in the latest round
-						CachedObjects = &GenerationInfo->BeginCacheObjects.ObjectsInRound;
-						NumIndexes = GenerationInfo->BeginCacheObjects.NextIndexInRound;
-					}
-					else
-					{
-						check(GenerationInfo->BeginCacheObjects.Objects.Num() == 0);
 					}
+					// Add all objects cached in the latest round
+					CachedObjects = &GenerationInfo->BeginCacheObjects.ObjectsInRound;
+					NumIndexes = GenerationInfo->BeginCacheObjects.NextIndexInRound;
+				}
+				else
+				{
+					check(GenerationInfo->BeginCacheObjects.Objects.Num() == 0);
 				}
-				check(NumIndexes <= NumPlatforms * CachedObjects->Num());
+			}
+			check(NumIndexes <= NumPlatforms * CachedObjects->Num());
 
-				// Iterate over all objects in the FPackageData up to GetCookedPlatformDataNextIndex
-				// GetCookedPlatformDataNextIndex is a value in an inline iteration over the two-dimensional array of Objects x Platforms, in Object-major order.
-				// We take the ceiling of NextIndex/NumPlatforms to get the number of objects.
-				int32 NumObjects = (NumIndexes + NumPlatforms - 1) / NumPlatforms;
-				for (int32 ObjectIndex = 0; ObjectIndex < NumObjects; ++ObjectIndex)
+			// Iterate over all objects in the FPackageData up to GetCookedPlatformDataNextIndex
+			// GetCookedPlatformDataNextIndex is a value in an inline iteration over the two-dimensional array of Objects x Platforms, in Object-major order.
+			// We take the ceiling of NextIndex/NumPlatforms to get the number of objects.
+			int32 NumObjects = (NumIndexes + NumPlatforms - 1) / NumPlatforms;
+			for (int32 ObjectIndex = 0; ObjectIndex < NumObjects; ++ObjectIndex)
+			{
+				UObject* Object = (*CachedObjects)[ObjectIndex].Get();
+				if (!Object)
 				{
-					UObject* Object = (*CachedObjects)[ObjectIndex].Get();
-					if (!Object)
-					{
-						continue;
-					}
-					ObjectsToClear.Add(Object);
+					continue;
 				}
-				for (UObject* Object : ObjectsToClear)
+				ObjectsToClear.Add(Object);
+			}
+			for (UObject* Object : ObjectsToClear)
+			{
+				TArray<FPendingCookedPlatformData*>* PendingDatas = PendingObjects.Find(Object);
+				if (!PendingDatas || PendingDatas->Num() == 0)
 				{
-					TArray<FPendingCookedPlatformData*>* PendingDatas = PendingObjects.Find(Object);
-					if (!PendingDatas || PendingDatas->Num() == 0)
+					// No pending BeginCacheForCookedPlatformData calls for this object; clear it now.
+					if (!IsCookingInEditor()) // ClearAllCachedCookedPlatformData calls are only used when not in editor.
 					{
-						// No pending BeginCacheForCookedPlatformData calls for this object; clear it now.
 						Object->ClearAllCachedCookedPlatformData();
 					}
-					else
+				}
+				else
+				{
+					// For any pending Objects, we add a CancelManager to the FPendingCookedPlatformData to call ClearAllCachedCookedPlatformData when the pending Object,Platform pairs for that object complete.
+					FPendingCookedPlatformDataCancelManager* CancelManager = new FPendingCookedPlatformDataCancelManager();
+					CancelManager->NumPendingPlatforms = PendingDatas->Num();
+					for (FPendingCookedPlatformData* PendingCookedPlatformData : *PendingDatas)
 					{
-						// For any pending Objects, we add a CancelManager to the FPendingCookedPlatformData to call ClearAllCachedCookedPlatformData when the pending Object,Platform pairs for that object complete.
-						FPendingCookedPlatformDataCancelManager* CancelManager = new FPendingCookedPlatformDataCancelManager();
-						CancelManager->NumPendingPlatforms = PendingDatas->Num();
-						for (FPendingCookedPlatformData* PendingCookedPlatformData : *PendingDatas)
-						{
-							// We never start a new package until after the previous cancel finished, so all of the FPendingCookedPlatformData for the PlatformData we are cancelling can not have been cancelled before.  We would leak the CancelManager if we overwrote it here.
-							check(PendingCookedPlatformData->CancelManager == nullptr);
-							// If bHasReleaased on the PendingCookedPlatformData were already true, we would leak the CancelManager because the PendingCookedPlatformData would never call Release on it.
-							check(!PendingCookedPlatformData->bHasReleased);
-							PendingCookedPlatformData->CancelManager = CancelManager;
-						}
+						// We never start a new package until after the previous cancel finished, so all of the FPendingCookedPlatformData for the PlatformData we are cancelling can not have been cancelled before.  We would leak the CancelManager if we overwrote it here.
+						check(PendingCookedPlatformData->CancelManager == nullptr);
+						// If bHasReleaased on the PendingCookedPlatformData were already true, we would leak the CancelManager because the PendingCookedPlatformData would never call Release on it.
+						check(!PendingCookedPlatformData->bHasReleased);
+						PendingCookedPlatformData->CancelManager = CancelManager;
 					}
 				}
 			}
@@ -4533,10 +4541,11 @@ void UCookOnTheFlyServer::PostLoadPackageFixup(UE::Cook::FPackageData& PackageDa
 
 void UCookOnTheFlyServer::TickPrecacheObjectsForPlatforms(const float TimeSlice, const TArray<const ITargetPlatform*>& TargetPlatforms) 
 {
+	using namespace UE::Cook;
 	SCOPE_CYCLE_COUNTER(STAT_TickPrecacheCooking);
 
 
-	UE::Cook::FCookerTimer Timer(TimeSlice);
+	FCookerTimer Timer(TimeSlice);
 
 	if (LastUpdateTick > 50 ||
 		((CachedMaterialsToCacheArray.Num() == 0) && (CachedTexturesToCacheArray.Num() == 0)))
@@ -11595,7 +11604,8 @@ void UCookOnTheFlyServer::RegisterLocalizationChunkDataGenerator()
 	}
 }
 
-void UCookOnTheFlyServer::RouteBeginCacheForCookedPlatformData(FName PackageName, UObject* Obj, const ITargetPlatform* TargetPlatform)
+void UCookOnTheFlyServer::RouteBeginCacheForCookedPlatformData(FName PackageName, UObject* Obj,
+	const ITargetPlatform* TargetPlatform)
 {
 	LLM_SCOPE_BYTAG(Cooker_CachedPlatformData);
 	UE_SCOPED_TEXT_COOKTIMER(*WriteToString<128>(GetClassTraceScope(Obj), TEXT("_BeginCacheForCookedPlatformData")));
@@ -11604,7 +11614,8 @@ void UCookOnTheFlyServer::RouteBeginCacheForCookedPlatformData(FName PackageName
 	Obj->BeginCacheForCookedPlatformData(TargetPlatform);
 }
 
-bool UCookOnTheFlyServer::RouteIsCachedCookedPlatformDataLoaded(FName PackageName, UObject* Obj, const ITargetPlatform* TargetPlatform)
+bool UCookOnTheFlyServer::RouteIsCachedCookedPlatformDataLoaded(FName PackageName, UObject* Obj,
+	const ITargetPlatform* TargetPlatform)
 {
 	LLM_SCOPE_BYTAG(Cooker_CachedPlatformData);
 	UE_SCOPED_TEXT_COOKTIMER(*WriteToString<128>(GetClassTraceScope(Obj), TEXT("_IsCachedCookedPlatformDataLoaded")));
diff --git a/Engine/Source/Editor/UnrealEd/Private/Cooker/CookPackageData.cpp b/Engine/Source/Editor/UnrealEd/Private/Cooker/CookPackageData.cpp
--- a/Engine/Source/Editor/UnrealEd/Private/Cooker/CookPackageData.cpp
+++ b/Engine/Source/Editor/UnrealEd/Private/Cooker/CookPackageData.cpp
@@ -3039,7 +3039,7 @@ void FPackageDatas::GetCookedPackagesForPlatform(const ITargetPlatform* Platform
 void FPackageDatas::Clear()
 {
 	FWriteScopeLock ExistenceWriteLock(ExistenceLock);
-	PendingCookedPlatformDataLists.Empty(); // These destructors will dereference PackageDatas
+	PendingCookedPlatformDataLists.Empty(); // These destructors will read/write PackageDatas
 	RequestQueue.Empty();
 	SaveQueue.Empty();
 	PackageNameToPackageData.Empty();
