Cooker: Track BeginCacheForCookedPlatformData state for all objects being cooked. Refactor storage of which Objects have had BeginCacheForCookedPlatformData called to use a TMap from UObject* to state data for all packages. Refactor storage of GeneratedPackage's objects to use the same GetCachedObjectsInOuter that normal packages use.
#rnx

diff --git a/Engine/Source/Editor/UnrealEd/Classes/CookOnTheSide/CookOnTheFlyServer.h b/Engine/Source/Editor/UnrealEd/Classes/CookOnTheSide/CookOnTheFlyServer.h
--- a/Engine/Source/Editor/UnrealEd/Classes/CookOnTheSide/CookOnTheFlyServer.h
+++ b/Engine/Source/Editor/UnrealEd/Classes/CookOnTheSide/CookOnTheFlyServer.h
@@ -168,6 +168,7 @@ namespace UE::Cook
 	class ICookOnTheFlyRequestManager;
 	class ICookOnTheFlyNetworkServer;
 	class IWorkerRequests;
+	enum class ECachedCookedPlatformDataEvent : uint8;
 	enum class EPollStatus : uint8;
 	enum class EReleaseSaveReason : uint8;
 	enum class ESuppressCookReason : uint8;
@@ -1271,10 +1272,10 @@ private:
 		return (CookFlags & InCookFlags) != ECookInitializationFlags::None;
 	}
 
-	void RouteBeginCacheForCookedPlatformData(FName PackageName, UObject* Obj,
-		const ITargetPlatform* TargetPlatform);
-	bool RouteIsCachedCookedPlatformDataLoaded(FName PackageName, UObject* Obj,
-		const ITargetPlatform* TargetPlatform);
+	void RouteBeginCacheForCookedPlatformData(UE::Cook::FPackageData& PackageData, UObject* Obj,
+		const ITargetPlatform* TargetPlatform, UE::Cook::ECachedCookedPlatformDataEvent* ExistingEvent);
+	bool RouteIsCachedCookedPlatformDataLoaded(UE::Cook::FPackageData& PackageData, UObject* Obj,
+		const ITargetPlatform* TargetPlatform, UE::Cook::ECachedCookedPlatformDataEvent* ExistingEvent);
 
 	/** Cook (save) a package and process the results */
 	void SaveCookedPackage(UE::Cook::FSaveCookedPackageContext& Context);
diff --git a/Engine/Source/Editor/UnrealEd/Private/CookOnTheFlyServer.cpp b/Engine/Source/Editor/UnrealEd/Private/CookOnTheFlyServer.cpp
--- a/Engine/Source/Editor/UnrealEd/Private/CookOnTheFlyServer.cpp
+++ b/Engine/Source/Editor/UnrealEd/Private/CookOnTheFlyServer.cpp
@@ -3326,13 +3326,12 @@ UE::Cook::EPollStatus UCookOnTheFlyServer::BeginCacheObjectsToMove(UE::Cook::FGe
 		Info.SetSaveStateComplete(FCookGenerationInfo::ESaveState::CallObjectsToMove);
 	}
 
-	EPollStatus Result = CallBeginCacheOnObjects(PackageData, Package, Info.BeginCacheObjects.ObjectsInRound,
-		Info.BeginCacheObjects.NextIndexInRound, Timer);
+	EPollStatus Result = CallBeginCacheOnObjects(PackageData, Package, PackageData.GetCachedObjectsInOuter(),
+		PackageData.GetCookedPlatformDataNextIndex(), Timer);
 	if (Result != EPollStatus::Success)
 	{
 		return Result;
 	}
-	Info.BeginCacheObjects.EndRound(PackageData.GetCachedObjectsInOuterNumPlatforms());
 	return EPollStatus::Success;
 }
 
@@ -3434,8 +3433,8 @@ UE::Cook::EPollStatus UCookOnTheFlyServer::BeginCachePostMove(UE::Cook::FGenerat
 		Info.SetSaveStateComplete(FCookGenerationInfo::ESaveState::CallGetPostMoveObjects);
 	}
 
-	EPollStatus Result = CallBeginCacheOnObjects(PackageData, Package, Info.BeginCacheObjects.ObjectsInRound,
-		Info.BeginCacheObjects.NextIndexInRound, Timer);
+	EPollStatus Result = CallBeginCacheOnObjects(PackageData, Package, PackageData.GetCachedObjectsInOuter(),
+		PackageData.GetCookedPlatformDataNextIndex(), Timer);
 	if (PackageData.GetNumPendingCookedPlatformData() > 0 &&
 		!Generator.GetCookPackageSplitterInstance()->UseInternalReferenceToAvoidGarbageCollect() &&
 		!Info.HasIssuedUndeclaredMovedObjectsWarning())
@@ -3471,7 +3470,6 @@ UE::Cook::EPollStatus UCookOnTheFlyServer::BeginCachePostMove(UE::Cook::FGenerat
 	{
 		return Result;
 	}
-	Info.BeginCacheObjects.EndRound(PackageData.GetCachedObjectsInOuterNumPlatforms());
 
 	return EPollStatus::Success;
 }
@@ -3622,13 +3620,14 @@ UE::Cook::EPollStatus UCookOnTheFlyServer::PrepareSaveInternal(UE::Cook::FPackag
 		// Currently this does not cause significant cost since saving new platforms with some platforms already saved is a rare operation.
 
 		int32& CookedPlatformDataNextIndex = PackageData.GetCookedPlatformDataNextIndex();
-		if (CookedPlatformDataNextIndex == 0)
+		if (CookedPlatformDataNextIndex < 0)
 		{
 			if (!BuildDefinitions->TryRemovePendingBuilds(PackageData.GetPackageName()))
 			{
 				// Builds are in progress; wait for them to complete
 				return EPollStatus::Incomplete;
 			}
+			CookedPlatformDataNextIndex = 0;
 		}
 
 		TArray<FWeakObjectPtr>& CachedObjectsInOuter = PackageData.GetCachedObjectsInOuter();
@@ -3726,13 +3725,10 @@ UE::Cook::EPollStatus UCookOnTheFlyServer::CallBeginCacheOnObjects(UE::Cook::FPa
 	PackageData.GetCachedObjectsInOuterPlatforms(TargetPlatforms);
 
 	FWeakObjectPtr* ObjectsData = Objects.GetData();
-	int NumPlatforms = TargetPlatforms.Num();
-	int NumIndexes = Objects.Num() * NumPlatforms;
-	while (NextIndex < NumIndexes)
+	int NumObjects = Objects.Num();
+	for (; NextIndex < NumObjects; ++NextIndex)
 	{
-		int ObjectIndex = NextIndex / NumPlatforms;
-		int PlatformIndex = NextIndex - ObjectIndex * NumPlatforms;
-		UObject* Obj = ObjectsData[ObjectIndex].Get();
+		UObject* Obj = ObjectsData[NextIndex].Get();
 		if (!Obj)
 		{
 			// Objects can be marked as pending kill even without a garbage collect, and our weakptr.get will return
@@ -3741,56 +3737,77 @@ UE::Cook::EPollStatus UCookOnTheFlyServer::CallBeginCacheOnObjects(UE::Cook::FPa
 			// BeginCacheForCookedPlatformData and ClearAllCachedCookedPlatformData
 			// In case the weakptr is merely pendingkill, set it to null explicitly so we don't think that we've called
 			// BeginCacheForCookedPlatformData on it if it gets unmarked pendingkill later
-			ObjectsData[ObjectIndex] = nullptr; 
-			++NextIndex;
+			ObjectsData[NextIndex] = nullptr;
 			continue;
 		}
-		const ITargetPlatform* TargetPlatform = TargetPlatforms[PlatformIndex];
+		FCachedCookedPlatformDataState& CCPDState = PackageDatas->GetCachedCookedPlatformDataObjects().FindOrAdd(Obj);
+		if (CCPDState.PackageData != &PackageData)
+		{
+			if (CCPDState.PackageData != nullptr)
+			{
+				UE_LOG(LogCook, Error, TEXT("CachedCookedPlatformDatas: %s is unexpectedly in two packages at once.")
+					TEXT("\n\tPreviousPackage: %s")
+					TEXT("\n\tCurrentPackage: %s"),
+					*Obj->GetFullName(), *CCPDState.PackageData->GetPackageName().ToString(),
+					*PackageData.GetPackageName().ToString());
+				continue;
+			}
+			CCPDState.PackageData = &PackageData;
+		}
 
-		if (Obj->IsA(UMaterialInterface::StaticClass()))
+		for (const ITargetPlatform* TargetPlatform : TargetPlatforms)
 		{
-			if (GShaderCompilingManager->GetNumRemainingJobs() + 1 > MaxConcurrentShaderJobs)
+			ECachedCookedPlatformDataEvent& ExistingEvent =
+				CCPDState.PlatformStates.FindOrAdd(TargetPlatform, ECachedCookedPlatformDataEvent::None);
+			if (ExistingEvent != ECachedCookedPlatformDataEvent::None)
+			{
+				continue;
+			}
+
+			if (Obj->IsA(UMaterialInterface::StaticClass()))
 			{
+				if (GShaderCompilingManager->GetNumRemainingJobs() + 1 > MaxConcurrentShaderJobs)
+				{
 #if DEBUG_COOKONTHEFLY
-				UE_LOG(LogCook, Display, TEXT("Delaying shader compilation of material %s"), *Obj->GetFullName());
+					UE_LOG(LogCook, Display, TEXT("Delaying shader compilation of material %s"), *Obj->GetFullName());
 #endif
-				return EPollStatus::Incomplete;
+					return EPollStatus::Incomplete;
+				}
 			}
-		}
 
-		const FName ClassFName = Obj->GetClass()->GetFName();
-		int32* CurrentAsyncCache = CurrentAsyncCacheForType.Find(ClassFName);
-		if (CurrentAsyncCache != nullptr)
-		{
-			if (*CurrentAsyncCache < 1)
+			const FName ClassFName = Obj->GetClass()->GetFName();
+			int32* CurrentAsyncCache = CurrentAsyncCacheForType.Find(ClassFName);
+			if (CurrentAsyncCache != nullptr)
 			{
-				return EPollStatus::Incomplete;
+				if (*CurrentAsyncCache < 1)
+				{
+					return EPollStatus::Incomplete;
+				}
+				*CurrentAsyncCache -= 1;
 			}
-			*CurrentAsyncCache -= 1;
-		}
 
-		RouteBeginCacheForCookedPlatformData(PackageData.GetPackageName(), Obj, TargetPlatform);
-		++NextIndex;
-		if (RouteIsCachedCookedPlatformDataLoaded(PackageData.GetPackageName(), Obj, TargetPlatform))
-		{
-			if (CurrentAsyncCache)
+			RouteBeginCacheForCookedPlatformData(PackageData, Obj, TargetPlatform, &ExistingEvent);
+			if (RouteIsCachedCookedPlatformDataLoaded(PackageData, Obj, TargetPlatform, &ExistingEvent))
+			{
+				if (CurrentAsyncCache)
+				{
+					*CurrentAsyncCache += 1;
+				}
+			}
+			else
 			{
-				*CurrentAsyncCache += 1;
+				bool bNeedsResourceRelease = CurrentAsyncCache != nullptr;
+				PackageDatas->AddPendingCookedPlatformData(FPendingCookedPlatformData(Obj, TargetPlatform,
+					PackageData, bNeedsResourceRelease, *this));
 			}
-		}
-		else
-		{
-			bool bNeedsResourceRelease = CurrentAsyncCache != nullptr;
-			PackageDatas->AddPendingCookedPlatformData(FPendingCookedPlatformData(Obj, TargetPlatform,
-				PackageData, bNeedsResourceRelease, *this));
-		}
 
-		if (Timer.IsActionTimeUp())
-		{
+			if (Timer.IsActionTimeUp())
+			{
 #if DEBUG_COOKONTHEFLY
-			UE_LOG(LogCook, Display, TEXT("Object %s took too long to cache"), *Obj->GetFullName());
+				UE_LOG(LogCook, Display, TEXT("Object %s took too long to cache"), *Obj->GetFullName());
 #endif
-			return EPollStatus::Incomplete;
+				return EPollStatus::Incomplete;
+			}
 		}
 	}
 
@@ -3814,49 +3831,19 @@ void UCookOnTheFlyServer::ReleaseCookedPlatformData(UE::Cook::FPackageData& Pack
 	}
 	FCookGenerationInfo* GenerationInfo = Generator ? Generator->FindInfo(PackageData) : nullptr;
 
-	// For every Object on which we called BeginCacheForCookedPlatformData, we need to call ClearAllCachedCookedPlatformData
+	// For every BeginCacheForCookedPlatformData call we made we need to call ClearAllCachedCookedPlatformData
 	if (ReleaseSaveReason == EReleaseSaveReason::Completed)
 	{
-		// Since we have completed CookedPlatformData, we know we called BeginCacheForCookedPlatformData on all objects in the package, and none are pending
+		// Since we have completed CookedPlatformData, we know we called BeginCacheForCookedPlatformData on all
+		// objects in the package, and none are pending
 		UE_SCOPED_HIERARCHICAL_COOKTIMER(ClearAllCachedCookedPlatformData);
-		if (GenerationInfo)
+		for (FWeakObjectPtr& WeakPtr : PackageData.GetCachedObjectsInOuter())
 		{
-			check(GenerationInfo->HasTakenOverCachedCookedPlatformData());
-			check(PackageData.GetCachedObjectsInOuter().Num() == 0);
-
-			for (FBeginCacheObject& BeginCacheObject : GenerationInfo->BeginCacheObjects.Objects)
+			UObject* Object = WeakPtr.Get();
+			if (Object)
 			{
-				UObject* Object = BeginCacheObject.Object.Get();
-				if (Object)
-				{
-					check(BeginCacheObject.bHasFinishedRound);
-					if (!IsCookingInEditor()) // ClearAllCachedCookedPlatformData and WillNeverCacheCookedPlatformDataAgain calls are only used when not in editor
-					{
-						Object->ClearAllCachedCookedPlatformData();
-						if (IsDirectorCookByTheBook())
-						{
-							Object->WillNeverCacheCookedPlatformDataAgain();
-						}
-					}
-				}
-			}
-		}
-		else
-		{
-			for (FWeakObjectPtr& WeakPtr : PackageData.GetCachedObjectsInOuter())
-			{
-				UObject* Object = WeakPtr.Get();
-				if (Object)
-				{
-					if (!IsCookingInEditor()) // ClearAllCachedCookedPlatformData and WillNeverCacheCookedPlatformDataAgain calls are only used when not in editor
-					{
-						Object->ClearAllCachedCookedPlatformData();
-						if (IsDirectorCookByTheBook())
-						{
-							Object->WillNeverCacheCookedPlatformDataAgain();
-						}
-					}
-				}
+				FPendingCookedPlatformData::ClearCachedCookedPlatformData(Object, PackageData,
+					true /* bCompletedSuccesfully */);
 			}
 		}
 	}
@@ -3866,92 +3853,56 @@ void UCookOnTheFlyServer::ReleaseCookedPlatformData(UE::Cook::FPackageData& Pack
 		// Note that even after we return from this function, some objects with pending IsCachedCookedPlatformDataLoaded
 		// calls may still exist for this Package in PendingCookedPlatformDatas
 		// and this PackageData may therefore still have GetNumPendingCookedPlatformData > 0
-		int32 NumPlatforms = PackageData.GetCachedObjectsInOuterNumPlatforms();
-		if (NumPlatforms > 0) // Shouldn't happen because PumpSaves checks for this, but avoid a divide by 0 if it does.
-		{
-			// We have only called BeginCacheForCookedPlatformData on Object,Platform pairs up to GetCookedPlatformDataNextIndex.
-			// Further, some of those calls might still be pending.
+		// We have only called BeginCacheForCookedPlatformData on Object,Platform pairs up to GetCookedPlatformDataNextIndex.
+		// Further, some of those calls might still be pending.
 
-			// Find all pending BeginCacheForCookedPlatformData for this FPackageData
-			TMap<UObject*, TArray<FPendingCookedPlatformData*>> PendingObjects;
-			PackageDatas->ForEachPendingCookedPlatformData(
-			[&PendingObjects, &PackageData](FPendingCookedPlatformData& PendingCookedPlatformData)
+		// Find all pending BeginCacheForCookedPlatformData for this FPackageData
+		TMap<UObject*, TArray<FPendingCookedPlatformData*>> PendingObjects;
+		PackageDatas->ForEachPendingCookedPlatformData(
+		[&PendingObjects, &PackageData](FPendingCookedPlatformData& PendingCookedPlatformData)
+		{
+			if (&PendingCookedPlatformData.PackageData == &PackageData && !PendingCookedPlatformData.PollIsComplete())
 			{
-				if (&PendingCookedPlatformData.PackageData == &PackageData && !PendingCookedPlatformData.PollIsComplete())
-				{
-					UObject* Object = PendingCookedPlatformData.Object.Get();
-					check(Object); // Otherwise PollIsComplete would have returned true
-					check(!PendingCookedPlatformData.bHasReleased); // bHasReleased should be false since PollIsComplete returned false
-					PendingObjects.FindOrAdd(Object).Add(&PendingCookedPlatformData);
-				}
-			});
+				UObject* Object = PendingCookedPlatformData.Object.Get();
+				check(Object); // Otherwise PollIsComplete would have returned true
+				check(!PendingCookedPlatformData.bHasReleased); // bHasReleased should be false since PollIsComplete returned false
+				PendingObjects.FindOrAdd(Object).Add(&PendingCookedPlatformData);
+			}
+		});
 
-			TArray<UObject*> ObjectsToClear;
-			TArray<FWeakObjectPtr>* CachedObjects = &PackageData.GetCachedObjectsInOuter();
-			int32 NumIndexes = PackageData.GetCookedPlatformDataNextIndex();
-			if (GenerationInfo)
-			{
-				if (GenerationInfo->HasTakenOverCachedCookedPlatformData())
-				{
-					check(PackageData.GetCachedObjectsInOuter().Num() == 0);
 
-					// Add all objects that were cached in previous rounds
-					for (FBeginCacheObject& BeginCacheObject : GenerationInfo->BeginCacheObjects.Objects)
-					{
-						UObject* Object = BeginCacheObject.Object.Get();
-						if (Object && BeginCacheObject.bHasFinishedRound)
-						{
-							ObjectsToClear.Add(Object);
-						}
-					}
-					// Add all objects cached in the latest round
-					CachedObjects = &GenerationInfo->BeginCacheObjects.ObjectsInRound;
-					NumIndexes = GenerationInfo->BeginCacheObjects.NextIndexInRound;
-				}
-				else
-				{
-					check(GenerationInfo->BeginCacheObjects.Objects.Num() == 0);
-				}
+		// Iterate over all objects in the FPackageData up to GetCookedPlatformDataNextIndex
+		TArray<FWeakObjectPtr>& CachedObjects = PackageData.GetCachedObjectsInOuter();
+		for (int32 ObjectIndex = 0; ObjectIndex < PackageData.GetCookedPlatformDataNextIndex(); ++ObjectIndex)
+		{
+			UObject* Object = CachedObjects[ObjectIndex].Get();
+			if (!Object)
+			{
+				continue;
 			}
-			check(NumIndexes <= NumPlatforms * CachedObjects->Num());
-
-			// Iterate over all objects in the FPackageData up to GetCookedPlatformDataNextIndex
-			// GetCookedPlatformDataNextIndex is a value in an inline iteration over the two-dimensional array of Objects x Platforms, in Object-major order.
-			// We take the ceiling of NextIndex/NumPlatforms to get the number of objects.
-			int32 NumObjects = (NumIndexes + NumPlatforms - 1) / NumPlatforms;
-			for (int32 ObjectIndex = 0; ObjectIndex < NumObjects; ++ObjectIndex)
+			TArray<FPendingCookedPlatformData*>* PendingDatas = PendingObjects.Find(Object);
+			if (!PendingDatas || PendingDatas->Num() == 0)
 			{
-				UObject* Object = (*CachedObjects)[ObjectIndex].Get();
-				if (!Object)
-				{
-					continue;
-				}
-				ObjectsToClear.Add(Object);
+				// No pending BeginCacheForCookedPlatformData calls for this object; clear it now.
+				FPendingCookedPlatformData::ClearCachedCookedPlatformData(Object, PackageData,
+					false /* bCompletedSuccesfully */);
 			}
-			for (UObject* Object : ObjectsToClear)
+			else
 			{
-				TArray<FPendingCookedPlatformData*>* PendingDatas = PendingObjects.Find(Object);
-				if (!PendingDatas || PendingDatas->Num() == 0)
-				{
-					// No pending BeginCacheForCookedPlatformData calls for this object; clear it now.
-					if (!IsCookingInEditor()) // ClearAllCachedCookedPlatformData calls are only used when not in editor.
-					{
-						Object->ClearAllCachedCookedPlatformData();
-					}
-				}
-				else
+				// For any pending Objects, we add a CancelManager to the FPendingCookedPlatformData to call
+				// ClearAllCachedCookedPlatformData when the pending Object,Platform pairs for that object completes.
+				FPendingCookedPlatformDataCancelManager* CancelManager = new FPendingCookedPlatformDataCancelManager();
+				CancelManager->NumPendingPlatforms = PendingDatas->Num();
+				for (FPendingCookedPlatformData* PendingCookedPlatformData : *PendingDatas)
 				{
-					// For any pending Objects, we add a CancelManager to the FPendingCookedPlatformData to call ClearAllCachedCookedPlatformData when the pending Object,Platform pairs for that object complete.
-					FPendingCookedPlatformDataCancelManager* CancelManager = new FPendingCookedPlatformDataCancelManager();
-					CancelManager->NumPendingPlatforms = PendingDatas->Num();
-					for (FPendingCookedPlatformData* PendingCookedPlatformData : *PendingDatas)
-					{
-						// We never start a new package until after the previous cancel finished, so all of the FPendingCookedPlatformData for the PlatformData we are cancelling can not have been cancelled before.  We would leak the CancelManager if we overwrote it here.
-						check(PendingCookedPlatformData->CancelManager == nullptr);
-						// If bHasReleaased on the PendingCookedPlatformData were already true, we would leak the CancelManager because the PendingCookedPlatformData would never call Release on it.
-						check(!PendingCookedPlatformData->bHasReleased);
-						PendingCookedPlatformData->CancelManager = CancelManager;
-					}
+					// We never start a new package until after the previous cancel finished, so all of the
+					// FPendingCookedPlatformData for the PlatformData we are cancelling can not have been cancelled before.
+					// We would leak the CancelManager if we overwrote it here.
+					check(PendingCookedPlatformData->CancelManager == nullptr);
+					// If bHasReleaased on the PendingCookedPlatformData were already true, we would leak the CancelManager
+					// because the PendingCookedPlatformData would never call Release on it.
+					check(!PendingCookedPlatformData->bHasReleased);
+					PendingCookedPlatformData->CancelManager = CancelManager;
 				}
 			}
 		}
@@ -4599,9 +4550,9 @@ void UCookOnTheFlyServer::TickPrecacheObjectsForPlatforms(const float TimeSlice,
 			{
 				continue;
 			}
-			if (!RouteIsCachedCookedPlatformDataLoaded(PackageName, Material, TargetPlatform))
+			if (!Material->IsCachedCookedPlatformDataLoaded(TargetPlatform))
 			{
-				RouteBeginCacheForCookedPlatformData(PackageName, Material, TargetPlatform);
+				Material->BeginCacheForCookedPlatformData(TargetPlatform);
 				AllMaterialsCompiled = false;
 			}
 		}
@@ -4640,9 +4591,9 @@ void UCookOnTheFlyServer::TickPrecacheObjectsForPlatforms(const float TimeSlice,
 			{
 				continue;
 			}
-			if (!RouteIsCachedCookedPlatformDataLoaded(PackageName, Texture, TargetPlatform))
+			if (!Texture->IsCachedCookedPlatformDataLoaded(TargetPlatform))
 			{
-				RouteBeginCacheForCookedPlatformData(PackageName, Texture, TargetPlatform);
+				Texture->BeginCacheForCookedPlatformData(TargetPlatform);
 			}
 		}
 		if (Timer.IsActionTimeUp())
@@ -5190,6 +5141,8 @@ void UCookOnTheFlyServer::PostGarbageCollect()
 	NumObjectsHistory.AddInstance(GUObjectArray.GetObjectArrayNumMinusAvailable());
 	VirtualMemoryHistory.AddInstance(FPlatformMemory::GetStats().UsedVirtual);
 
+	TSet<UObject*> SaveQueueObjectsThatStillExist;
+
 	// If any PackageDatas with ObjectPointers had any of their object pointers deleted out from under them, demote them back to request
 	TArray<FPackageData*> Demotes;
 	for (FPackageData* PackageData : PackageDatas->GetSaveQueue())
@@ -5200,6 +5153,17 @@ void UCookOnTheFlyServer::PostGarbageCollect()
 		{
 			Demotes.Add(PackageData);
 		}
+		else
+		{
+			for (FWeakObjectPtr& WeakObjectPtr : PackageData->GetCachedObjectsInOuter())
+			{
+				UObject* Object = WeakObjectPtr.Get();
+				if (Object)
+				{
+					SaveQueueObjectsThatStillExist.Add(Object);
+				}
+			}
+		}
 	}
 	for (FPackageData* PackageData : Demotes)
 	{
@@ -5207,6 +5171,9 @@ void UCookOnTheFlyServer::PostGarbageCollect()
 		PackageDatas->GetRequestQueue().AddRequest(PackageData, /* bForceUrgent */ true);
 	}
 
+	// Remove objects that were deleted by garbage collection from our containers that track raw object pointers
+	PackageDatas->CachedCookedPlatformDataObjectsPostGarbageCollect(SaveQueueObjectsThatStillExist);
+
 	// If there was a GarbageCollect while we are saving a package, some of the WeakObjectPtr in SavingPackageData->CachedObjectPointers may have been deleted and set to null
 	// We need to handle nulls in that array at any point after calling SavePackage. We do not want to declare them as references and prevent their GC, in case there is 
 	// the expectation by some licensee code that removing references to an object will cause it to not be saved
@@ -11604,24 +11571,80 @@ void UCookOnTheFlyServer::RegisterLocalizationChunkDataGenerator()
 	}
 }
 
-void UCookOnTheFlyServer::RouteBeginCacheForCookedPlatformData(FName PackageName, UObject* Obj,
-	const ITargetPlatform* TargetPlatform)
+void UCookOnTheFlyServer::RouteBeginCacheForCookedPlatformData(UE::Cook::FPackageData& PackageData, UObject* Obj,
+	const ITargetPlatform* TargetPlatform, UE::Cook::ECachedCookedPlatformDataEvent* ExistingEvent)
 {
+	using namespace UE::Cook;
+
 	LLM_SCOPE_BYTAG(Cooker_CachedPlatformData);
 	UE_SCOPED_TEXT_COOKTIMER(*WriteToString<128>(GetClassTraceScope(Obj), TEXT("_BeginCacheForCookedPlatformData")));
+	FName PackageName = PackageData.GetPackageName();
 	UE_SCOPED_COOK_STAT(PackageName, EPackageEventStatType::BeginCacheForCookedPlatformData);
+	
+	if (!ExistingEvent)
+	{
+		FCachedCookedPlatformDataState& CCPDState = PackageDatas->GetCachedCookedPlatformDataObjects().FindOrAdd(Obj);
+		if (CCPDState.PackageData != &PackageData)
+		{
+			if (CCPDState.PackageData != nullptr)
+			{
+				ExistingEvent = CCPDState.PlatformStates.Find(TargetPlatform);
+				UE_LOG(LogCook, Error, TEXT("BeginCacheForCookedPlatformData called unexpectedly on %s. Skipping it.")
+					TEXT("\n\tPreviousEvent: %s for package %s.")
+					TEXT("\n\tCurrentEvent: %s for package %s."),
+					*Obj->GetFullName(), LexToString(ExistingEvent ? *ExistingEvent : ECachedCookedPlatformDataEvent::None),
+					*CCPDState.PackageData->GetPackageName().ToString(),
+					LexToString(ECachedCookedPlatformDataEvent::BeginCacheForCookedPlatformDataCalled),
+					*PackageName.ToString());
+				return;
+			}
+			CCPDState.PackageData = &PackageData;
+		}
+		ExistingEvent = &CCPDState.PlatformStates.FindOrAdd(TargetPlatform, ECachedCookedPlatformDataEvent::None);
+	}
 	FScopedActivePackage ScopedActivePackage(*this, PackageName, PackageAccessTrackingOps::NAME_CookerBuildObject);
 	Obj->BeginCacheForCookedPlatformData(TargetPlatform);
+	*ExistingEvent = ECachedCookedPlatformDataEvent::BeginCacheForCookedPlatformDataCalled;
 }
 
-bool UCookOnTheFlyServer::RouteIsCachedCookedPlatformDataLoaded(FName PackageName, UObject* Obj,
-	const ITargetPlatform* TargetPlatform)
+bool UCookOnTheFlyServer::RouteIsCachedCookedPlatformDataLoaded(UE::Cook::FPackageData& PackageData, UObject* Obj,
+	const ITargetPlatform* TargetPlatform, UE::Cook::ECachedCookedPlatformDataEvent* ExistingEvent)
 {
+	using namespace UE::Cook;
+
 	LLM_SCOPE_BYTAG(Cooker_CachedPlatformData);
 	UE_SCOPED_TEXT_COOKTIMER(*WriteToString<128>(GetClassTraceScope(Obj), TEXT("_IsCachedCookedPlatformDataLoaded")));
+	FName PackageName = PackageData.GetPackageName();
 	UE_SCOPED_COOK_STAT(Obj->GetPackage()->GetFName(), EPackageEventStatType::IsCachedCookedPlatformDataLoaded);
+
+	if (!ExistingEvent)
+	{
+		FCachedCookedPlatformDataState& CCPDState = PackageDatas->GetCachedCookedPlatformDataObjects().FindOrAdd(Obj);
+		if (CCPDState.PackageData != &PackageData)
+		{
+			if (CCPDState.PackageData != nullptr)
+			{
+				UE_LOG(LogCook, Error, TEXT("IsCachedCookedPlatformDataLoaded called unexpectedly on %s. Skipping it and returning true.")
+					TEXT("\n\tPreviousEvent: %s for package %s.")
+					TEXT("\n\tCurrentEvent: %s for package %s."),
+					*Obj->GetFullName(), LexToString(ExistingEvent ? *ExistingEvent : ECachedCookedPlatformDataEvent::None),
+					*CCPDState.PackageData->GetPackageName().ToString(),
+					LexToString(ECachedCookedPlatformDataEvent::IsCachedCookedPlatformDataLoadedCalled),
+					*PackageName.ToString());
+				return true;
+			}
+			CCPDState.PackageData = &PackageData;
+		}
+		ExistingEvent = &CCPDState.PlatformStates.FindOrAdd(TargetPlatform, ECachedCookedPlatformDataEvent::None);
+	}
+
 	FScopedActivePackage ScopedActivePackage(*this, PackageName, PackageAccessTrackingOps::NAME_CookerBuildObject);
-	return Obj->IsCachedCookedPlatformDataLoaded(TargetPlatform);
+	bool bResult = Obj->IsCachedCookedPlatformDataLoaded(TargetPlatform);
+	if (bResult)
+	{
+		*ExistingEvent = ECachedCookedPlatformDataEvent::IsCachedCookedPlatformDataLoadedReturnedTrue;
+	}
+	return bResult;
 }
 
 namespace UE::Cook
diff --git a/Engine/Source/Editor/UnrealEd/Private/Cooker/CookPackageData.cpp b/Engine/Source/Editor/UnrealEd/Private/Cooker/CookPackageData.cpp
--- a/Engine/Source/Editor/UnrealEd/Private/Cooker/CookPackageData.cpp
+++ b/Engine/Source/Editor/UnrealEd/Private/Cooker/CookPackageData.cpp
@@ -21,6 +21,7 @@
 #include "Engine/Console.h"
 #include "HAL/FileManager.h"
 #include "HAL/PlatformTime.h"
+#include "Interfaces/ITargetPlatform.h"
 #include "Misc/CommandLine.h"
 #include "Misc/CoreMiscDefines.h"
 #include "Misc/PackageAccessTrackingOps.h"
@@ -1122,19 +1123,6 @@ TArray<FWeakObjectPtr>& FPackageData::GetCachedObjectsInOuter()
 	return CachedObjectsInOuter;
 }
 
-int32 FPackageData::GetCachedObjectsInOuterNumPlatforms() const
-{
-	int32 Result = 0;
-	for (const TPair<const ITargetPlatform*, FPackagePlatformData>& Pair : PlatformDatas)
-	{
-		if (Pair.Value.IsRegisteredForCachedObjectsInOuter())
-		{
-			++Result;
-		}
-	}
-	return Result;
-}
-
 void FPackageData::CheckObjectCacheEmpty() const
 {
 	check(CachedObjectsInOuter.Num() == 0);
@@ -1224,10 +1212,6 @@ static TArray<UObject*> SetDifference(TArray<UObject*>& A, TArray<UObject*>& B)
 
 EPollStatus FPackageData::RefreshObjectCache(bool& bOutFoundNewObjects)
 {
-	// Caller will only call this function after CallBeginCacheOnObjects finished successfully
-	int32 NumPlatforms = GetCachedObjectsInOuterNumPlatforms();
-	check(NumPlatforms > 0);
-	check(GetCookedPlatformDataNextIndex()/NumPlatforms == GetCachedObjectsInOuter().Num());
 	check(Package.Get() != nullptr);
 
 	TArray<UObject*> OldObjects;
@@ -1272,6 +1256,29 @@ EPollStatus FPackageData::RefreshObjectCache(bool& bOutFoundNewObjects)
 
 void FPackageData::ClearObjectCache()
 {
+	TMap<UObject*, FCachedCookedPlatformDataState>& CCPDs = PackageDatas.GetCachedCookedPlatformDataObjects();
+	for (FWeakObjectPtr& WeakObjectPtr: CachedObjectsInOuter)
+	{
+		UObject* Object = WeakObjectPtr.Get();
+		if (!Object)
+		{
+			continue;
+		}
+		FCachedCookedPlatformDataState CCPDState;
+		CCPDs.RemoveAndCopyValue(Object, CCPDState);
+		if (CCPDState.PackageData != this)
+		{
+			// This is an error condition that is logged elsewhere; one of the objects in this package is also
+			// present in another package, which owns it. We should silently handle this case (silently because
+			// the error about it is logged elsewhere), which means we should not modify the object on account
+			// of this package.
+			// We removed it because we thought would own it if it existed, but since we were wrong, put it back.
+			if (CCPDState.PackageData != nullptr)
+			{
+				CCPDs.Add(Object, MoveTemp(CCPDState));
+			}
+		}
+	}
 	CachedObjectsInOuter.Empty();
 	for (TPair<const ITargetPlatform*, FPackagePlatformData>& Pair : PlatformDatas)
 	{
@@ -1312,7 +1319,7 @@ int32 FPackageData::GetMaxNumRetriesBeginCacheOnObjects()
 
 void FPackageData::CheckCookedPlatformDataEmpty() const
 {
-	check(GetCookedPlatformDataNextIndex() == 0);
+	check(GetCookedPlatformDataNextIndex() <= 0);
 	check(!GetCookedPlatformDataStarted());
 	check(!GetCookedPlatformDataCalled());
 	check(!GetCookedPlatformDataComplete());
@@ -1327,7 +1334,7 @@ void FPackageData::CheckCookedPlatformDataEmpty() const
 
 void FPackageData::ClearCookedPlatformData()
 {
-	CookedPlatformDataNextIndex = 0;
+	CookedPlatformDataNextIndex = -1;
 	NumRetriesBeginCacheOnObject = 0;
 	// Note that GetNumPendingCookedPlatformData is not cleared; it persists across Saves and CookSessions
 	SetCookedPlatformDataStarted(false);
@@ -1356,9 +1363,8 @@ void FPackageData::RemapTargetPlatforms(const TMap<ITargetPlatform*, ITargetPlat
 		NewPlatformDatas.FindOrAdd(NewKey) = MoveTemp(ExistingPair.Value);
 	}
 
-	// The save state (and maybe more in the future) depend on the order of the request platforms remaining
-	// unchanged, due to CookedPlatformDataNextIndex. If we change that order due to the remap, we need to
-	// demote back to request.
+	// The save state (and maybe more in the future) by contract can depend on the order of the request platforms remaining
+	// unchanged. If we change that order due to the remap, we need to demote back to request.
 	if (IsInProgress() && GetState() != EPackageState::Request)
 	{
 		bool bDemote = true;
@@ -1822,9 +1828,9 @@ void FGeneratorPackage::PreGarbageCollect(FCookGenerationInfo& Info, TArray<TObj
 				GCKeepPackageDatas.Add(Info.PackageData);
 			}
 			GCKeepPackages.Append(Info.KeepReferencedPackages);
-			for (FBeginCacheObject& BeginCacheObject : Info.BeginCacheObjects.Objects)
+			for (FWeakObjectPtr& WeakObjectPtr : Info.PackageData->GetCachedObjectsInOuter())
 			{
-				UObject* Object = BeginCacheObject.Object.Get();
+				UObject* Object = WeakObjectPtr.Get();
 				if (Object)
 				{
 					GCKeepObjects.Add(Object);
@@ -2007,19 +2013,21 @@ void FGeneratorPackage::ResetSaveState(FCookGenerationInfo& Info, UPackage* Pack
 	{
 		Info.SetSaveState(FCookGenerationInfo::ESaveState::StartPopulate);
 	}
-	if (Info.BeginCacheObjects.Objects.Num() != 0 &&
-		GetCookPackageSplitterInstance()->UseInternalReferenceToAvoidGarbageCollect() &&
-		(ReleaseSaveReason == EReleaseSaveReason::Demoted || ReleaseSaveReason == EReleaseSaveReason::RecreateObjectCache))
+	if (Info.HasTakenOverCachedCookedPlatformData())
 	{
-		UE_LOG(LogCook, Error, TEXT("CookPackageSplitter failure: We are demoting a %s package from save and removing our references that keep its objects loaded.\n")
-			TEXT("This will allow the objects to be garbage collected and cause failures in the splitter which expects them to remain loaded.\n")
-			TEXT("Package=%s, Splitter=%s, ReleaseSaveReason=%s"),
-			Info.IsGenerator() ? TEXT("generator") : TEXT("generated"),
-			Info.PackageData ? *Info.PackageData->GetPackageName().ToString() : *Info.RelativePath,
-			*GetSplitDataObjectName().ToString(), LexToString(ReleaseSaveReason));
+		if (Info.PackageData && Info.PackageData->GetCachedObjectsInOuter().Num() != 0 &&
+			GetCookPackageSplitterInstance()->UseInternalReferenceToAvoidGarbageCollect() &&
+			(ReleaseSaveReason == EReleaseSaveReason::Demoted || ReleaseSaveReason == EReleaseSaveReason::RecreateObjectCache))
+		{
+			UE_LOG(LogCook, Error, TEXT("CookPackageSplitter failure: We are demoting a %s package from save and removing our references that keep its objects loaded.\n")
+				TEXT("This will allow the objects to be garbage collected and cause failures in the splitter which expects them to remain loaded.\n")
+				TEXT("Package=%s, Splitter=%s, ReleaseSaveReason=%s"),
+				Info.IsGenerator() ? TEXT("generator") : TEXT("generated"),
+				Info.PackageData ? *Info.PackageData->GetPackageName().ToString() : *Info.RelativePath,
+				*GetSplitDataObjectName().ToString(), LexToString(ReleaseSaveReason));
+		}
+		Info.SetHasTakenOverCachedCookedPlatformData(false);
 	}
-	Info.BeginCacheObjects.Reset();
-	Info.SetHasTakenOverCachedCookedPlatformData(false);
 	Info.SetHasIssuedUndeclaredMovedObjectsWarning(false);
 	Info.KeepReferencedPackages.Reset();
 }
@@ -2047,10 +2055,14 @@ void FGeneratorPackage::UpdateSaveAfterGarbageCollect(const FPackageData& Packag
 		}
 	}
 
-	for (TArray<FBeginCacheObject>::TIterator Iter(Info->BeginCacheObjects.Objects); Iter; ++Iter)
+	TSet<UObject*> CachedObjectsInOuterSet;
+	TArray<FWeakObjectPtr>& CachedObjectsInOuter = Info->PackageData->GetCachedObjectsInOuter();
+	int32& NextIndexToBeginCache = Info->PackageData->GetCookedPlatformDataNextIndex();
+	for (int32 Index = 0; Index < CachedObjectsInOuter.Num(); ++Index)
 	{
-		const FBeginCacheObject& Object = *Iter;
-		if (!Object.Object.IsValid())
+		FWeakObjectPtr& WeakObjectPtr = CachedObjectsInOuter[Index];
+		UObject* Object = WeakObjectPtr.Get();
+		if (!Object)
 		{
 			if (GetCookPackageSplitterInstance()->UseInternalReferenceToAvoidGarbageCollect())
 			{
@@ -2062,13 +2074,30 @@ void FGeneratorPackage::UpdateSaveAfterGarbageCollect(const FPackageData& Packag
 					Info->IsGenerator() ? TEXT("PopulateGeneratorPackage") : TEXT("PopulateGeneratedPackage"),
 					*GetSplitDataObjectName().ToString(),
 					Info->IsGenerator() ? TEXT("") : *FString::Printf(TEXT(", Generated=%s."), *Info->PackageData->GetPackageName().ToString()));
-				Iter.RemoveCurrent();
+
+				CachedObjectsInOuter.RemoveAt(Index);
+				if (NextIndexToBeginCache > Index)
+				{
+					--NextIndexToBeginCache;
+				}
+				--Index;
 			}
 			else
 			{
 				bInOutDemote = true;
 			}
 		}
+		else
+		{
+			CachedObjectsInOuterSet.Add(Object);
+		}
+	}
+	for (TSet<UObject*>::TIterator Iter(Info->RootMovedObjects); Iter; ++Iter)
+	{
+		if (!CachedObjectsInOuterSet.Contains(*Iter))
+		{
+			Iter.RemoveCurrent();
+		}
 	}
 }
 
@@ -2092,20 +2121,15 @@ void FCookGenerationInfo::SetSaveStateComplete(ESaveState CompletedState)
 void FCookGenerationInfo::TakeOverCachedObjectsAndAddMoved(FGeneratorPackage& Generator,
 	TArray<FWeakObjectPtr>& CachedObjectsInOuter, TArray<UObject*>& MovedObjects)
 {
-	BeginCacheObjects.Objects.Reset();
+	RootMovedObjects.Reset();
+
 	TSet<UObject*> ObjectSet;
 	for (FWeakObjectPtr& ObjectInOuter : CachedObjectsInOuter)
 	{
 		UObject* Object = ObjectInOuter.Get();
 		if (Object)
 		{
-			bool bAlreadyExists;
-			ObjectSet.Add(Object, &bAlreadyExists);
-			if (!bAlreadyExists)
-			{
-				FBeginCacheObject& Added = BeginCacheObjects.Objects.Add_GetRef(FBeginCacheObject{ Object });
-				Added.bHasFinishedRound = true;
-			}
+			ObjectSet.Add(Object);
 		}
 	}
 
@@ -2121,14 +2145,12 @@ void FCookGenerationInfo::TakeOverCachedObjectsAndAddMoved(FGeneratorPackage& Ge
 				IsGenerator() ? TEXT("") : *FString::Printf(TEXT(", Package %s"), *PackageData->GetPackageName().ToString()));
 			continue;
 		}
-
 		bool bAlreadyExists;
 		ObjectSet.Add(Object, &bAlreadyExists);
 		if (!bAlreadyExists)
 		{
-			FBeginCacheObject& Added = BeginCacheObjects.Objects.Add_GetRef(FBeginCacheObject{ Object });
-			Added.bHasFinishedRound = false;
-			Added.bIsRootMovedObject = true;
+			RootMovedObjects.Add(Object);
+			CachedObjectsInOuter.Add(Object);
 			GetObjectsWithOuter(Object, ChildrenOfMovedObjects, true /* bIncludeNestedObjects */, RF_NoFlags, EInternalObjectFlags::Garbage);
 		}
 	}
@@ -2140,14 +2162,10 @@ void FCookGenerationInfo::TakeOverCachedObjectsAndAddMoved(FGeneratorPackage& Ge
 		ObjectSet.Add(Object, &bAlreadyExists);
 		if (!bAlreadyExists)
 		{
-			FBeginCacheObject& Added = BeginCacheObjects.Objects.Add_GetRef(FBeginCacheObject{ Object });
-			Added.bHasFinishedRound = false;
+			CachedObjectsInOuter.Add(Object);
 		}
 	}
 
-	BeginCacheObjects.StartRound();
-
-	CachedObjectsInOuter.Reset();
 	SetHasTakenOverCachedCookedPlatformData(true);
 }
 
@@ -2159,16 +2177,18 @@ EPollStatus FCookGenerationInfo::RefreshPackageObjects(FGeneratorPackage& Genera
 	GetObjectsWithOuter(Package, CurrentObjectsInOuter, true /* bIncludeNestedObjects */, RF_NoFlags, EInternalObjectFlags::Garbage);
 
 	TSet<UObject*> ObjectSet;
-	ObjectSet.Reserve(BeginCacheObjects.Objects.Num());
-	for (FBeginCacheObject& ExistingObject : BeginCacheObjects.Objects)
+	check(PackageData); // RefreshPackageObjects is only called when there is a PackageData
+	TArray<FWeakObjectPtr>& CachedObjectsInOuter = PackageData->GetCachedObjectsInOuter();
+	ObjectSet.Reserve(CachedObjectsInOuter.Num());
+	for (FWeakObjectPtr& ExistingObject : CachedObjectsInOuter)
 	{
-		UObject* Object = ExistingObject.Object.Get();
+		UObject* Object = ExistingObject.Get();
 		if (Object)
 		{
 			bool bAlreadyExists;
 			ObjectSet.Add(Object, &bAlreadyExists);
-			check(!bAlreadyExists); // Objects in BeginCacheObjects.Objects are guaranteed unique and we haven't added any others yet
-			if (ExistingObject.bIsRootMovedObject)
+			check(!bAlreadyExists); // Objects in GetCachedObjectsInOuter are guaranteed unique and we haven't added any others yet
+			if (RootMovedObjects.Contains(Object))
 			{
 				GetObjectsWithOuter(Object, CurrentObjectsInOuter, true /* bIncludeNestedObjects */, RF_NoFlags, EInternalObjectFlags::Garbage);
 			}
@@ -2181,8 +2201,7 @@ EPollStatus FCookGenerationInfo::RefreshPackageObjects(FGeneratorPackage& Genera
 		ObjectSet.Add(Object, &bAlreadyExists);
 		if (!bAlreadyExists)
 		{
-			FBeginCacheObject& Added = BeginCacheObjects.Objects.Add_GetRef(FBeginCacheObject{ Object });
-			Added.bHasFinishedRound = false;
+			CachedObjectsInOuter.Add(Object);
 			if (!FirstNewObject )
 			{
 				FirstNewObject  = Object;
@@ -2191,8 +2210,6 @@ EPollStatus FCookGenerationInfo::RefreshPackageObjects(FGeneratorPackage& Genera
 	}
 	bOutFoundNewObjects = FirstNewObject != nullptr;
 
-	BeginCacheObjects.StartRound();
-
 	if (FirstNewObject != nullptr && DemotionState != ESaveState::Last)
 	{
 		SetSaveState(DemotionState);
@@ -2210,37 +2227,6 @@ EPollStatus FCookGenerationInfo::RefreshPackageObjects(FGeneratorPackage& Genera
 	return EPollStatus::Success;
 }
 
-void FBeginCacheObjects::Reset()
-{
-	Objects.Reset();
-	ObjectsInRound.Reset();
-	NextIndexInRound = 0;
-}
-
-void FBeginCacheObjects::StartRound()
-{
-	ObjectsInRound.Reset(Objects.Num());
-	for (FBeginCacheObject& Object : Objects)
-	{
-		if (!Object.bHasFinishedRound)
-		{
-			ObjectsInRound.Add(Object.Object);
-		}
-	}
-	NextIndexInRound = 0;
-}
-
-void FBeginCacheObjects::EndRound(int32 NumPlatforms)
-{
-	check(NextIndexInRound == ObjectsInRound.Num()*NumPlatforms);
-	ObjectsInRound.Reset();
-	NextIndexInRound = 0;
-	for (FBeginCacheObject& Object : Objects)
-	{
-		Object.bHasFinishedRound = true;
-	}
-}
-
 void FCookGenerationInfo::CreateGuid()
 {
 	FBlake3 Blake3;
@@ -2354,7 +2340,7 @@ bool FPendingCookedPlatformData::PollIsComplete()
 		return true;
 	}
 	UCookOnTheFlyServer& COTFS = PackageData.GetPackageDatas().GetCookOnTheFlyServer();
-	if (COTFS.RouteIsCachedCookedPlatformDataLoaded(PackageData.GetPackageName(), LocalObject, TargetPlatform))
+	if (COTFS.RouteIsCachedCookedPlatformDataLoaded(PackageData, LocalObject, TargetPlatform, nullptr /* ExistingEvent */))
 	{
 		Release();
 		return true;
@@ -2408,6 +2394,59 @@ void FPendingCookedPlatformData::RemapTargetPlatforms(const TMap<ITargetPlatform
 	TargetPlatform = Remap[TargetPlatform];
 }
 
+void FPendingCookedPlatformData::ClearCachedCookedPlatformData(UObject* Object, FPackageData& PackageData,
+	bool bCompletedSuccesfully)
+{
+	FPackageDatas& PackageDatas = PackageData.GetPackageDatas();
+	UCookOnTheFlyServer& COTFS = PackageDatas.GetCookOnTheFlyServer();
+	FCachedCookedPlatformDataState CCPDState;
+	PackageDatas.GetCachedCookedPlatformDataObjects().RemoveAndCopyValue(Object, CCPDState);
+	if (CCPDState.PackageData != &PackageData)
+	{
+		if (CCPDState.PackageData)
+		{
+			ECachedCookedPlatformDataEvent ExistingEvent = ECachedCookedPlatformDataEvent::None;
+			for (TPair<const ITargetPlatform*, ECachedCookedPlatformDataEvent>& PlatformPair : CCPDState.PlatformStates)
+			{
+				if (PlatformPair.Value != ECachedCookedPlatformDataEvent::None)
+				{
+					ExistingEvent = PlatformPair.Value;
+					break;
+				}
+			}
+			UE_LOG(LogCook, Error, TEXT("ClearCachedCookedPlatformData called unexpectedly on %s. It will be skipped.")
+				TEXT("\n\tPreviousEvent: %s for package %s.")
+				TEXT("\n\tCurrentEvent: %s for package %s."),
+				*Object->GetFullName(), LexToString(ExistingEvent), *CCPDState.PackageData->GetPackageName().ToString(),
+				LexToString(ECachedCookedPlatformDataEvent::ClearCachedCookedPlatformDataCalled),
+				*PackageData.GetPackageName().ToString());
+			PackageDatas.GetCachedCookedPlatformDataObjects().Add(Object, MoveTemp(CCPDState));
+		}
+	}
+	else
+	{
+		// TODO: Add this call to ClearCachedCookedPlatformData once we take it over from SavePackage
+		constexpr bool bCallClearCachedCookedPlatformData = false;
+		if (bCallClearCachedCookedPlatformData)
+		{
+			for (const TPair<const ITargetPlatform*, ECachedCookedPlatformDataEvent>& PlatformPair : CCPDState.PlatformStates)
+			{
+				Object->ClearCachedCookedPlatformData(PlatformPair.Key);
+			}
+		}
+
+		// ClearAllCachedCookedPlatformData and WillNeverCacheCookedPlatformDataAgain calls are only used when not in editor
+		if (!COTFS.IsCookingInEditor())
+		{
+			Object->ClearAllCachedCookedPlatformData();
+			if (bCompletedSuccesfully && COTFS.IsDirectorCookByTheBook())
+			{
+				Object->WillNeverCacheCookedPlatformDataAgain();
+			}
+		}
+	}
+};
+
 
 //////////////////////////////////////////////////////////////////////////
 // FPendingCookedPlatformDataCancelManager
@@ -2422,7 +2461,8 @@ void FPendingCookedPlatformDataCancelManager::Release(FPendingCookedPlatformData
 		UObject* LocalObject = Data.Object.Get();
 		if (LocalObject)
 		{
-			LocalObject->ClearAllCachedCookedPlatformData();
+			FPendingCookedPlatformData::ClearCachedCookedPlatformData(LocalObject, Data.PackageData,
+				false /* bCompletedSuccesfully */);
 		}
 		delete this;
 	}
@@ -3044,6 +3084,7 @@ void FPackageDatas::Clear()
 	SaveQueue.Empty();
 	PackageNameToPackageData.Empty();
 	FileNameToPackageData.Empty();
+	CachedCookedPlatformDataObjects.Empty();
 	{
 		// All references must be cleared before any PackageDatas are destroyed
 		EnumeratePackageDatasWithinLock([](FPackageData* PackageData)
@@ -3417,4 +3458,30 @@ FPoppedPackageDataScope::~FPoppedPackageDataScope()
 }
 #endif
 
+const TCHAR* LexToString(ECachedCookedPlatformDataEvent Value)
+{
+	switch (Value)
+	{
+	case ECachedCookedPlatformDataEvent::None: return TEXT("None");
+	case ECachedCookedPlatformDataEvent::BeginCacheForCookedPlatformDataCalled: return TEXT("BeginCacheForCookedPlatformDataCalled");
+	case ECachedCookedPlatformDataEvent::IsCachedCookedPlatformDataLoadedCalled: return TEXT("IsCachedCookedPlatformDataLoadedCalled");
+	case ECachedCookedPlatformDataEvent::IsCachedCookedPlatformDataLoadedReturnedTrue: return TEXT("IsCachedCookedPlatformDataLoadedReturnedTrue");
+	case ECachedCookedPlatformDataEvent::ClearCachedCookedPlatformDataCalled: return TEXT("ClearCachedCookedPlatformDataCalled");
+	case ECachedCookedPlatformDataEvent::ClearAllCachedCookedPlatformDataCalled: return TEXT("ClearAllCachedCookedPlatformDataCalled");
+	default: return TEXT("Invalid");
+	}
+}
+
+void FPackageDatas::CachedCookedPlatformDataObjectsPostGarbageCollect(const TSet<UObject*>& SaveQueueObjectsThatStillExist)
+{
+	for (TMap<UObject*, FCachedCookedPlatformDataState>::TIterator Iter(this->CachedCookedPlatformDataObjects);
+		Iter; ++Iter)
+	{
+		if (!SaveQueueObjectsThatStillExist.Contains(Iter->Key))
+		{
+			Iter.RemoveCurrent();
+		}
+	}
+}
+
 } // namespace UE::Cook
diff --git a/Engine/Source/Editor/UnrealEd/Private/Cooker/CookPackageData.h b/Engine/Source/Editor/UnrealEd/Private/Cooker/CookPackageData.h
--- a/Engine/Source/Editor/UnrealEd/Private/Cooker/CookPackageData.h
+++ b/Engine/Source/Editor/UnrealEd/Private/Cooker/CookPackageData.h
@@ -50,11 +50,36 @@ namespace UE::Cook
 class FPackageDataQueue;
 class FRequestCluster;
 struct FGeneratorPackage;
+struct FPackageData;
 struct FPackageDataMonitor;
 struct FPendingCookedPlatformDataCancelManager;
 
 extern const TCHAR* GeneratedPackageSubPath;
 
+/**
+ * Events in the lifetime of an object related to BeginCacheForCookedPlatformData. Used by the cooker
+ * to track which calls have been made and still need to be made.
+ */
+enum class ECachedCookedPlatformDataEvent : uint8
+{
+	None,
+	BeginCacheForCookedPlatformDataCalled,
+	IsCachedCookedPlatformDataLoadedCalled,
+	IsCachedCookedPlatformDataLoadedReturnedTrue,
+	ClearCachedCookedPlatformDataCalled,
+	ClearAllCachedCookedPlatformDataCalled,
+};
+const TCHAR* LexToString(ECachedCookedPlatformDataEvent);
+/**
+ * BeginCachedForCookedPlatformData state about an object - which package owned it (not always the same
+ * one when PackageGenerators are involved) and the per-platform state for ECachedCookedPlatformDataEvent.
+ */
+struct FCachedCookedPlatformDataState
+{
+	FPackageData* PackageData = nullptr;
+	TMap<const ITargetPlatform*, ECachedCookedPlatformDataEvent> PlatformStates;
+};
+
 /** Flags specifying the behavior of FPackageData::SendToState */
 enum class ESendFlags : uint8
 {
@@ -189,13 +214,9 @@ private:
  *   The lifetime of this counter extends for the lifetime of the PackageData; it is shared across
  *   CookOnTheFlyServer sessions.
  *
- * CookedPlatformDataNextIndex - Index for the next (Object, ReachablePlatform) pair that needs to have
- *   BeginCacheForCookedPlatformData called on it for the current PackageSave. This field is only non-zero during
- *   the save state; it is cleared when successfully or unsucessfully leaving the save state. The field is an
- *   object-major index into the two-dimensional array given by
- *       {Each Object in GetCachedObjectsInOuter} x {each Platform that is reachable}
- *   e.g. the index into GetCachedObjectsInOuter is GetCookedPlatformDataNextIndex() / GetCachedObjectsInOuterNumPlatforms()
- *   and the index into GetCachedObjectsInOuterPlatforms() is GetCookedPlatformDataNextIndex() % GetCachedObjectsInOuterNumPlatforms()
+ * CookedPlatformDataNextIndex - Index for the next Object in CachedObjectsInOuter that needs to have
+ *   BeginCacheForCookedPlatformData called on it for the current PackageSave. This field is only >= 0 during
+ *   the save state; it is cleared to -1 when successfully or unsucessfully leaving the save state. 
  *
  * Other fields with explanation inline
 */
@@ -414,8 +435,6 @@ public:
 	template <typename ArrayType>
 	/** The list of platforms that were recorded as needscooking when CachedObjeObjectsInOuter was recorded. */
 	void GetCachedObjectsInOuterPlatforms(ArrayType& OutPlatforms) const;
-	/** The number of platforms that would be returned by GetCachedObjectsInOuterPlatforms. */
-	int32 GetCachedObjectsInOuterNumPlatforms() const;
 
 	/** Validate that the CachedObjectsInOuter-dependent variables are empty, when entering save. */
 	void CheckObjectCacheEmpty() const;
@@ -667,7 +686,7 @@ private:
 	};
 	FTrackedPreloadableFilePtr PreloadableFile;
 	int32 NumPendingCookedPlatformData = 0;
-	int32 CookedPlatformDataNextIndex = 0;
+	int32 CookedPlatformDataNextIndex = -1;
 	int32 NumRetriesBeginCacheOnObject = 0;
 	FOpenPackageResult PreloadableFileOpenResult;
 	FInstigator Instigator;
@@ -693,30 +712,6 @@ private:
 	uint32 bWasCookedThisSession : 1;
 };
 
-/** A single object in athe save of a package that might have had BeginCacheForCookedPlatformData called already */
-struct FBeginCacheObject
-{
-	FWeakObjectPtr Object;
-	bool bHasFinishedRound = false;
-	bool bIsRootMovedObject = false;
-};
-
-/**
- * Collection of the known objects in the save of a package that might have had BeginCacheForCookedPlatformData called already
- * and will have it called on them in the next round of calls if not.
- */
-struct FBeginCacheObjects
-{
-	TArray<FBeginCacheObject> Objects;
-	TArray<FWeakObjectPtr> ObjectsInRound;
-	int32 NextIndexInRound = 0;
-
-	void Reset();
-
-	void StartRound();
-	void EndRound(int32 NumPlatforms);
-};
-
 /**
  * Struct uses with ICookPackageSplitter, which contains the state information for the save of either the
  * Generator package or one of its Generated packages.
@@ -800,7 +795,6 @@ public:
 	TConstArrayView<FAssetDependency> GetDependencies() const { return PackageDependencies; }
 
 public:
-	FBeginCacheObjects BeginCacheObjects;
 	FGuid Guid;
 	FString RelativePath;
 	FString GeneratedRootPath;
@@ -808,6 +802,7 @@ public:
 	TArray<FAssetDependency> PackageDependencies;
 	FPackageData* PackageData = nullptr;
 	TArray<UPackage*> KeepReferencedPackages;
+	TSet<UObject*> RootMovedObjects;
 private:
 	ESaveState GeneratorSaveState = ESaveState::StartGenerate;
 	bool bCreateAsMap : 1;
@@ -924,6 +919,9 @@ struct FPendingCookedPlatformData
 	FPendingCookedPlatformData& operator=(const FPendingCookedPlatformData& Other) = delete;
 	FPendingCookedPlatformData& operator=(const FPendingCookedPlatformData&& Other) = delete;
 
+	/** Helper for both pending and synchronous; call ClearCachedCookedPlatformData and related teardowns. */
+	static void ClearCachedCookedPlatformData(UObject* Object, FPackageData& PackageData, bool bCompletedSuccesfully);
+
 	/**
 	 * Call IsCachedCookedPlatformDataLoaded on the object if it has not already returned true.
 	 * If IsCachedCookedPlatformDataLoaded returns true, this function releases all held resources related to the
@@ -1350,6 +1348,12 @@ public:
 	template <typename CallbackType>
 	void LockAndEnumeratePackageDatas(CallbackType&& Callback);
 
+	TMap<UObject*, FCachedCookedPlatformDataState>& GetCachedCookedPlatformDataObjects()
+	{
+		return CachedCookedPlatformDataObjects;
+	}
+	void CachedCookedPlatformDataObjectsPostGarbageCollect(const TSet<UObject*>& SaveQueueObjectsThatStillExist);
+
 private:
 	/**
 	 * Construct a new FPackageData with the given PackageName and FileName and store references to it in the maps.
@@ -1374,6 +1378,7 @@ private:
 	/* Guarded by ExistenceLock. Duplicates information on FPackageData, but can be read/write from any thread. */
 	TMap<FName, FThreadsafePackageData> ThreadsafePackageDatas;
 	TRingBuffer<FPendingCookedPlatformDataContainer> PendingCookedPlatformDataLists;
+	TMap<UObject*, FCachedCookedPlatformDataState> CachedCookedPlatformDataObjects;
 	int32 PendingCookedPlatformDataNum = 0;
 	FRequestQueue RequestQueue;
 	TFastPointerSet<FPackageData*> AssignedToWorkerSet;
