Cooker: Downgrade CachedCookedPlatformDatas errors to display while we diagnose them. They are occurring for WorldPartitionRuntimeCellDataSpatialHash objects in WorldPartition maps.
#rnx

diff --git a/Engine/Source/Editor/UnrealEd/Private/CookOnTheFlyServer.cpp b/Engine/Source/Editor/UnrealEd/Private/CookOnTheFlyServer.cpp
--- a/Engine/Source/Editor/UnrealEd/Private/CookOnTheFlyServer.cpp
+++ b/Engine/Source/Editor/UnrealEd/Private/CookOnTheFlyServer.cpp
@@ -3745,7 +3745,7 @@ UE::Cook::EPollStatus UCookOnTheFlyServer::CallBeginCacheOnObjects(UE::Cook::FPa
 		{
 			if (CCPDState.PackageData != nullptr)
 			{
-				UE_LOG(LogCook, Error, TEXT("CachedCookedPlatformDatas: %s is unexpectedly in two packages at once.")
+				UE_LOG(LogCook, Display, TEXT("CachedCookedPlatformDatas: %s is unexpectedly in two packages at once.")
 					TEXT("\n\tPreviousPackage: %s")
 					TEXT("\n\tCurrentPackage: %s"),
 					*Obj->GetFullName(), *CCPDState.PackageData->GetPackageName().ToString(),
@@ -11589,7 +11589,7 @@ void UCookOnTheFlyServer::RouteBeginCacheForCookedPlatformData(UE::Cook::FPackag
 			if (CCPDState.PackageData != nullptr)
 			{
 				ExistingEvent = CCPDState.PlatformStates.Find(TargetPlatform);
-				UE_LOG(LogCook, Error, TEXT("BeginCacheForCookedPlatformData called unexpectedly on %s. Skipping it.")
+				UE_LOG(LogCook, Display, TEXT("BeginCacheForCookedPlatformData called unexpectedly on %s. Skipping it.")
 					TEXT("\n\tPreviousEvent: %s for package %s.")
 					TEXT("\n\tCurrentEvent: %s for package %s."),
 					*Obj->GetFullName(), LexToString(ExistingEvent ? *ExistingEvent : ECachedCookedPlatformDataEvent::None),
@@ -11624,7 +11624,7 @@ bool UCookOnTheFlyServer::RouteIsCachedCookedPlatformDataLoaded(UE::Cook::FPacka
 		{
 			if (CCPDState.PackageData != nullptr)
 			{
-				UE_LOG(LogCook, Error, TEXT("IsCachedCookedPlatformDataLoaded called unexpectedly on %s. Skipping it and returning true.")
+				UE_LOG(LogCook, Display, TEXT("IsCachedCookedPlatformDataLoaded called unexpectedly on %s. Skipping it and returning true.")
 					TEXT("\n\tPreviousEvent: %s for package %s.")
 					TEXT("\n\tCurrentEvent: %s for package %s."),
 					*Obj->GetFullName(), LexToString(ExistingEvent ? *ExistingEvent : ECachedCookedPlatformDataEvent::None),
diff --git a/Engine/Source/Editor/UnrealEd/Private/Cooker/CookPackageData.cpp b/Engine/Source/Editor/UnrealEd/Private/Cooker/CookPackageData.cpp
--- a/Engine/Source/Editor/UnrealEd/Private/Cooker/CookPackageData.cpp
+++ b/Engine/Source/Editor/UnrealEd/Private/Cooker/CookPackageData.cpp
@@ -2414,7 +2414,7 @@ void FPendingCookedPlatformData::ClearCachedCookedPlatformData(UObject* Object,
 					break;
 				}
 			}
-			UE_LOG(LogCook, Error, TEXT("ClearCachedCookedPlatformData called unexpectedly on %s. It will be skipped.")
+			UE_LOG(LogCook, Display, TEXT("ClearCachedCookedPlatformData called unexpectedly on %s. It will be skipped.")
 				TEXT("\n\tPreviousEvent: %s for package %s.")
 				TEXT("\n\tCurrentEvent: %s for package %s."),
 				*Object->GetFullName(), LexToString(ExistingEvent), *CCPDState.PackageData->GetPackageName().ToString(),
