Add the VirtualizationEditor module back to the editor, but this time add it to UnrealEd.Build.cs (where the module is loaded)

#rnx

- Was removed in CL 25782459 from LevelEditor.Build.cs
- Cleaned up some whitespace (spaces to tabs) for consistency.

diff --git a/Engine/Source/Editor/UnrealEd/UnrealEd.Build.cs b/Engine/Source/Editor/UnrealEd/UnrealEd.Build.cs
--- a/Engine/Source/Editor/UnrealEd/UnrealEd.Build.cs
+++ b/Engine/Source/Editor/UnrealEd/UnrealEd.Build.cs
@@ -153,8 +153,8 @@ public class UnrealEd : ModuleRules
 				"ViewportInteraction",
 				"VREditor",
 				"ClothingSystemEditor",
-                "ClothingSystemRuntimeInterface",
-                "ClothingSystemRuntimeCommon",
+				"ClothingSystemRuntimeInterface",
+				"ClothingSystemRuntimeCommon",
 				"PIEPreviewDeviceProfileSelector",
 				"PakFileUtilities",
 				"TimeManagement",
@@ -177,6 +177,7 @@ public class UnrealEd : ModuleRules
 				"SubobjectEditor",
 				"HeadMountedDisplay",
 				"FieldNotification",
+				"VirtualizationEditor",
 			}
 		);
 
@@ -348,5 +349,5 @@ public class UnrealEd : ModuleRules
 		{
 			PrivateIncludePathModuleNames.Add("LiveCoding");
 		}
-    }
+	}
 }
