EditorTransaction: Catch objects being invalidated during applying a transaction

diff --git a/Engine/Source/Editor/UnrealEd/Private/EditorTransaction.cpp b/Engine/Source/Editor/UnrealEd/Private/EditorTransaction.cpp
--- a/Engine/Source/Editor/UnrealEd/Private/EditorTransaction.cpp
+++ b/Engine/Source/Editor/UnrealEd/Private/EditorTransaction.cpp
@@ -907,10 +907,32 @@ void FTransaction::Apply()
 		return Cast<UActorComponent>(&A) != nullptr;
 	});
 
+	auto ObjectWasInvalidatedDuringTransaction = [](const UObject* InObject) -> bool
+	{
+		// skip records which point to a SKEL_ or REINST_ class
+		// (any class that we're keeping around only to GC old instances of the class).
+		// a previous record may have caused a blueprint compilation
+		// which may cause the consecutive changed object to turn invalid.
+		if(InObject != nullptr && !IsValid(InObject) &&
+			InObject->GetOutermost() == GetTransientPackage())
+		{
+			if(InObject->GetClass() && InObject->GetClass()->HasAnyClassFlags(CLASS_NewerVersionExists))
+			{
+				return true;
+			}
+		}
+		return false;
+	};
+
 	TArray<ULevel*> LevelsToCommitModelSurface;
 	for (const TTuple<UObject*, FChangedObjectValue>& ChangedObjectIt : ChangedObjects)
 	{
 		UObject* ChangedObject = ChangedObjectIt.Key;
+		if(ObjectWasInvalidatedDuringTransaction(ChangedObject))
+		{
+			continue;
+		}
+		
 		UModel* Model = Cast<UModel>(ChangedObject);
 		if (Model && Model->Nodes.Num())
 		{
@@ -987,7 +1009,11 @@ void FTransaction::Apply()
 	for (auto ChangedObjectIt : ChangedObjects)
 	{
 		UObject* ChangedObject = ChangedObjectIt.Key;
-		ChangedObject->CheckDefaultSubobjects();
+		if(ObjectWasInvalidatedDuringTransaction(ChangedObject))
+        {
+        	continue;
+        }
+		(void)ChangedObject->CheckDefaultSubobjects();
 	}
 
 	ChangedObjects.Reset();
