Cooker and SavePackage: Move calls to BeginCacheForCookedPlatformData and ClearCachedCookedPlatformData from SavePackage into the cooker. Moving the Clear calls allows us to clear only after repeated SavePackage attempts from the diffonly cooker have been completed, preventing a break in contract by having objects being saved after Clear is called. Moving the Begin calls is necessary to track which objects need Clear called on them.
#rn Minor, Cooking

diff --git a/Engine/Source/Developer/IoStoreUtilities/Private/ZenStoreWriter.cpp b/Engine/Source/Developer/IoStoreUtilities/Private/ZenStoreWriter.cpp
--- a/Engine/Source/Developer/IoStoreUtilities/Private/ZenStoreWriter.cpp
+++ b/Engine/Source/Developer/IoStoreUtilities/Private/ZenStoreWriter.cpp
@@ -1259,6 +1259,12 @@ void FZenStoreWriter::CompleteExportsArchiveForDiff(FPackageInfo& Info, FLargeMe
 	ExistingState.PreOptimizedPackage = MoveTemp(PackageStorePackage);
 }
 
+EPackageWriterResult FZenStoreWriter::BeginCacheForCookedPlatformData(
+	FBeginCacheForCookedPlatformDataInfo& Info)
+{
+	return BeginCacheCallback(Info);
+}
+
 TFuture<FCbObject> FZenStoreWriter::WriteMPCookMessageForPackage(FName PackageName)
 {
 	TArray<FString> AdditionalFiles;
@@ -1333,7 +1339,7 @@ bool FZenStoreWriter::TryReadMPCookMessageForPackage(FName PackageName, FCbObjec
 			ExistingPackageHashes = ThisPackageHashes;
 			if (bAlreadyExisted)
 			{
-				UE_LOG(LogSavePackage, Error, TEXT("FLooseCookedPackageWriter encountered the same package twice in a cook! (%s)"),
+				UE_LOG(LogSavePackage, Error, TEXT("FZenStoreWriter encountered the same package twice in a cook! (%s)"),
 					*PackageName.ToString());
 			}
 		}
diff --git a/Engine/Source/Developer/IoStoreUtilities/Public/ZenStoreWriter.h b/Engine/Source/Developer/IoStoreUtilities/Public/ZenStoreWriter.h
--- a/Engine/Source/Developer/IoStoreUtilities/Public/ZenStoreWriter.h
+++ b/Engine/Source/Developer/IoStoreUtilities/Public/ZenStoreWriter.h
@@ -67,6 +67,12 @@ public:
 		return this;
 	};
 
+	// Delegates to forward PackageWriter events onto UCookOnTheFlyServer when cooking
+	void SetBeginCacheCallback(FBeginCacheCallback&& InBeginCacheCallback)
+	{
+		BeginCacheCallback = MoveTemp(InBeginCacheCallback);
+	}
+
 	struct ZenHostInfo
 	{
 		FString ProjectId;
@@ -126,6 +132,7 @@ public:
 	IOSTOREUTILITIES_API virtual void MarkPackagesUpToDate(TArrayView<const FName> UpToDatePackages) override;
 	IOSTOREUTILITIES_API bool GetPreviousCookedBytes(const FPackageInfo& Info, FPreviousCookedBytesData& OutData) override;
 	IOSTOREUTILITIES_API void CompleteExportsArchiveForDiff(FPackageInfo& Info, FLargeMemoryWriter& ExportsArchive) override;
+	IOSTOREUTILITIES_API virtual EPackageWriterResult BeginCacheForCookedPlatformData(FBeginCacheForCookedPlatformDataInfo& Info) override;
 	IOSTOREUTILITIES_API virtual TFuture<FCbObject> WriteMPCookMessageForPackage(FName PackageName) override;
 	IOSTOREUTILITIES_API virtual bool TryReadMPCookMessageForPackage(FName PackageName, FCbObjectView Message) override;
 
@@ -228,6 +235,8 @@ private:
 	TUniquePtr<FZenFileSystemManifest>	ZenFileSystemManifest;
 	TMap<FName, TArray<FString>>		PackageAdditionalFiles;
 
+	FBeginCacheCallback					BeginCacheCallback;
+
 	FEntryCreatedEvent					EntryCreatedEvent;
 	FCriticalSection					CommitEventCriticalSection;
 	FCommitEvent						CommitEvent;
diff --git a/Engine/Source/Editor/UnrealEd/Classes/CookOnTheSide/CookOnTheFlyServer.h b/Engine/Source/Editor/UnrealEd/Classes/CookOnTheSide/CookOnTheFlyServer.h
--- a/Engine/Source/Editor/UnrealEd/Classes/CookOnTheSide/CookOnTheFlyServer.h
+++ b/Engine/Source/Editor/UnrealEd/Classes/CookOnTheSide/CookOnTheFlyServer.h
@@ -35,6 +35,7 @@ class IAssetRegistry;
 class ICookedPackageWriter;
 class IPlugin;
 class ITargetPlatform;
+enum class EPackageWriterResult : uint8;
 enum class ODSCRecompileCommand;
 struct FBeginCookContext;
 struct FCrashContextExtendedWriter;
@@ -1276,6 +1277,8 @@ private:
 		const ITargetPlatform* TargetPlatform, UE::Cook::ECachedCookedPlatformDataEvent* ExistingEvent);
 	bool RouteIsCachedCookedPlatformDataLoaded(UE::Cook::FPackageData& PackageData, UObject* Obj,
 		const ITargetPlatform* TargetPlatform, UE::Cook::ECachedCookedPlatformDataEvent* ExistingEvent);
+	EPackageWriterResult SavePackageBeginCacheForCookedPlatformData(FName PackageName,
+		const ITargetPlatform* TargetPlatform, TConstArrayView<UObject*> SaveableObjects, uint32 SaveFlags);
 
 	/** Cook (save) a package and process the results */
 	void SaveCookedPackage(UE::Cook::FSaveCookedPackageContext& Context);
@@ -1460,6 +1463,11 @@ private:
 	bool bRandomizeCookOrder = false;
 	/** True if commandline arguments specified that we suppress the cook of packages based on filter criteria. */
 	bool bCookFilter = false;
+	/**
+	 * Experimental feature to correctly invoke the BeginCacheForCookedPlatformData contracts for
+	 * objects created by PreSave
+	 */
+	bool bCallIsCachedOnSaveCreatedObjects = false;
 
 	/** Timers for tracking how long we have been busy, to manage retries and warnings of deadlock */
 	double SaveBusyStartTimeSeconds = MAX_flt;
diff --git a/Engine/Source/Editor/UnrealEd/Private/CookOnTheFlyServer.cpp b/Engine/Source/Editor/UnrealEd/Private/CookOnTheFlyServer.cpp
--- a/Engine/Source/Editor/UnrealEd/Private/CookOnTheFlyServer.cpp
+++ b/Engine/Source/Editor/UnrealEd/Private/CookOnTheFlyServer.cpp
@@ -6615,6 +6615,8 @@ void UCookOnTheFlyServer::SetInitializeConfigSettings(UE::Cook::FInitializeConfi
 		(FParse::Param(FCommandLine::Get(), TEXT("DIFFONLY")) && !FParse::Param(FCommandLine::Get(), TEXT("DIFFNORANDCOOK")));
 
 	ParseCookFilters();
+
+	bCallIsCachedOnSaveCreatedObjects = FParse::Param(FCommandLine::Get(), TEXT("CallIsCachedOnSaveCreatedObjects"));
 }
 
 void UCookOnTheFlyServer::ParseCookFilters()
@@ -10059,15 +10061,23 @@ UE::Cook::FCookSavePackageContext* UCookOnTheFlyServer::CreateSaveContext(const
 
 	ICookedPackageWriter* PackageWriter = nullptr;
 	FString WriterDebugName;
+	ICookedPackageWriter::FBeginCacheCallback BeginCacheCallback(
+		[this](ICookedPackageWriter::FBeginCacheForCookedPlatformDataInfo& Info)
+		{
+			return this->SavePackageBeginCacheForCookedPlatformData(Info.PackageName,
+				Info.TargetPlatform, Info.SaveableObjects, Info.SaveFlags);
+		});
 	if (IsUsingZenStore())
 	{
-		PackageWriter = new FZenStoreWriter(ResolvedRootPath, ResolvedMetadataPath, TargetPlatform);
+		FZenStoreWriter* ZenWriter = new FZenStoreWriter(ResolvedRootPath, ResolvedMetadataPath, TargetPlatform);
+		ZenWriter->SetBeginCacheCallback(MoveTemp(BeginCacheCallback));
+		PackageWriter = ZenWriter;
 		WriterDebugName = TEXT("ZenStore");
 	}
 	else
 	{
 		PackageWriter = new FLooseCookedPackageWriter(ResolvedRootPath, ResolvedMetadataPath, TargetPlatform,
-			GetAsyncIODelete(), *PackageDatas, PluginsToRemap);
+			GetAsyncIODelete(), PluginsToRemap, MoveTemp(BeginCacheCallback));
 		WriterDebugName = TEXT("LooseCookedPackageWriter");
 	}
 
@@ -11602,7 +11612,14 @@ void UCookOnTheFlyServer::RouteBeginCacheForCookedPlatformData(UE::Cook::FPackag
 		}
 		ExistingEvent = &CCPDState.PlatformStates.FindOrAdd(TargetPlatform, ECachedCookedPlatformDataEvent::None);
 	}
-	FScopedActivePackage ScopedActivePackage(*this, PackageName, PackageAccessTrackingOps::NAME_CookerBuildObject);
+
+	// We need to set our scopes for e.g. TObjectPtr reads around the call to BeginCacheForCookedPlatformData,
+	// but in some cases we have already set the scope (e.g. when calling BeginCache from inside SavePackage)
+	TOptional<FScopedActivePackage> ScopedActivePackage;
+	if (!ActivePackageData.bActive)
+	{
+		ScopedActivePackage.Emplace(*this, PackageName, PackageAccessTrackingOps::NAME_CookerBuildObject);
+	}
 	Obj->BeginCacheForCookedPlatformData(TargetPlatform);
 	*ExistingEvent = ECachedCookedPlatformDataEvent::BeginCacheForCookedPlatformDataCalled;
 }
@@ -11638,7 +11655,14 @@ bool UCookOnTheFlyServer::RouteIsCachedCookedPlatformDataLoaded(UE::Cook::FPacka
 		ExistingEvent = &CCPDState.PlatformStates.FindOrAdd(TargetPlatform, ECachedCookedPlatformDataEvent::None);
 	}
 
-	FScopedActivePackage ScopedActivePackage(*this, PackageName, PackageAccessTrackingOps::NAME_CookerBuildObject);
+	// We need to set our scopes for e.g. TObjectPtr reads around the call to BeginCacheForCookedPlatformData,
+	// but in some cases we have already set the scope (e.g. when calling IsCached from inside SavePackage)
+	TOptional<FScopedActivePackage> ScopedActivePackage;
+	if (!ActivePackageData.bActive)
+	{
+		ScopedActivePackage.Emplace(*this, PackageName, PackageAccessTrackingOps::NAME_CookerBuildObject);
+	}
+
 	bool bResult = Obj->IsCachedCookedPlatformDataLoaded(TargetPlatform);
 	if (bResult)
 	{
@@ -11647,6 +11671,117 @@ bool UCookOnTheFlyServer::RouteIsCachedCookedPlatformDataLoaded(UE::Cook::FPacka
 	return bResult;
 }
 
+EPackageWriterResult UCookOnTheFlyServer::SavePackageBeginCacheForCookedPlatformData(FName PackageName,
+	const ITargetPlatform* TargetPlatform, TConstArrayView<UObject*> SaveableObjects, uint32 SaveFlags)
+{
+	using namespace UE::Cook;
+
+	FPackageData* PackageData = PackageDatas->FindPackageDataByPackageName(PackageName);
+	check(PackageData); // This callback is called from a packagesave we initiated, so it should exist
+
+	TArray<FWeakObjectPtr>& CachedObjectsInOuter = PackageData->GetCachedObjectsInOuter();
+	int32& NextIndex = PackageData->GetCookedPlatformDataNextIndex();
+	TArray<UObject*> PendingObjects;
+	for (UObject* Object : SaveableObjects)
+	{
+		FCachedCookedPlatformDataState& CCPDState = PackageDatas->GetCachedCookedPlatformDataObjects().FindOrAdd(Object);
+		if (CCPDState.PackageData != PackageData)
+		{
+			if (CCPDState.PackageData != nullptr)
+			{
+				UE_LOG(LogCook, Display, TEXT("CachedCookedPlatformDatas: %s is unexpectedly in two packages at once.")
+					TEXT("\n\tPreviousPackage: %s")
+					TEXT("\n\tCurrentPackage: %s"),
+					*Object->GetFullName(), *CCPDState.PackageData->GetPackageName().ToString(),
+					*PackageData->GetPackageName().ToString());
+				continue;
+			}
+			CCPDState.PackageData = PackageData;
+
+			// NextIndex is usually at the end of CachedObjectsInOuter, but in case it is not, insert the new Object 
+			// at NextIndex so that we still record that we have not called BeginCace on objects after it. Then increment
+			// NextIndex to indicate we have already called (down below) BeginCache on the added object, so that
+			// ReleaseCookedPlatformData knows that it needs to call Clear on it.
+			check(NextIndex >= 0);
+			CachedObjectsInOuter.Insert(Object, NextIndex);
+			++NextIndex;
+		}
+
+		ECachedCookedPlatformDataEvent& ExistingEvent =
+			CCPDState.PlatformStates.FindOrAdd(TargetPlatform, ECachedCookedPlatformDataEvent::None);
+		if (ExistingEvent == ECachedCookedPlatformDataEvent::IsCachedCookedPlatformDataLoadedReturnedTrue)
+		{
+			// Already called
+			continue;
+		}
+
+		RouteBeginCacheForCookedPlatformData(*PackageData, Object, TargetPlatform, &ExistingEvent);
+		if (bCallIsCachedOnSaveCreatedObjects)
+		{
+			// TODO: Enable bCallIsCachedOnSaveCreatedObjects so that we call IsCachedCookedPlatformDataLoaded on all the objects until it
+			// returns true. This is required for the BeginCacheForCookedPlatformData contract.
+			// Doing so will cause us to return Timeout and retry the save later after the pending objects have completed.
+			// We tried enabling this once, but it created knockon bugs: Textures created by landscape were not handling it correctly
+			// (which we have subsequently fixed) and MaterialInstanceConstants created by landscape were not handling it correctly (which 
+			// we have not yet diagnosed).
+			if (!RouteIsCachedCookedPlatformDataLoaded(*PackageData, Object, TargetPlatform, &ExistingEvent))
+			{
+				PackageDatas->AddPendingCookedPlatformData(FPendingCookedPlatformData(Object, TargetPlatform,
+					*PackageData, false /* bNeedsResourceRelease */, *this));
+				PendingObjects.Add(Object);
+			}
+		}
+	}
+
+	if (!PendingObjects.IsEmpty())
+	{
+		constexpr float MaxWaitSeconds = 30.f;
+		constexpr float SleepTimeSeconds = 0.010f;
+		const double EndTimeSeconds = FPlatformTime::Seconds() + MaxWaitSeconds;
+		for (;;)
+		{
+			UE_SCOPED_HIERARCHICAL_COOKTIMER(PollPendingCookedPlatformDatas);
+			PackageDatas->PollPendingCookedPlatformDatas(true /* bForce */, LastCookableObjectTickTime);
+			if (PackageData->GetNumPendingCookedPlatformData() == 0)
+			{
+				break;
+			}
+			if (SaveFlags & SAVE_AllowTimeout)
+			{
+				return EPackageWriterResult::Timeout;
+			}
+			if (FPlatformTime::Seconds() > EndTimeSeconds)
+			{
+				UObject* Culprit = nullptr;
+				for (UObject* Object : PendingObjects)
+				{
+					FCachedCookedPlatformDataState& CCPDState = PackageDatas->GetCachedCookedPlatformDataObjects().FindOrAdd(Object);
+					ECachedCookedPlatformDataEvent& ExistingEvent = CCPDState.PlatformStates.FindOrAdd(TargetPlatform,
+						ECachedCookedPlatformDataEvent::None);
+					if (ExistingEvent != ECachedCookedPlatformDataEvent::IsCachedCookedPlatformDataLoadedReturnedTrue)
+					{
+						Culprit = Object;
+						break;
+					}
+				}
+				if (!Culprit)
+				{
+					UE_LOG(LogCook, Warning,
+						TEXT("SavePackageBeginCacheForCookedPlatformData Error for package %s: GetNumPendingCookedPlatformData() != 0 but no Culprit found. Ignoring it and continuing."),
+						*PackageName.ToString());
+					break;
+				}
+				UE_LOG(LogSavePackage, Error, TEXT("Save of %s failed: timed out waiting for IsCachedCookedPlatformDataLoaded on %s."),
+					*PackageName.ToString(), *Culprit->GetFullName());
+				return EPackageWriterResult::Error;
+			}
+
+			FPlatformProcess::Sleep(SleepTimeSeconds);
+		}
+	}
+	return EPackageWriterResult::Success;
+}
+
 namespace UE::Cook
 {
 
diff --git a/Engine/Source/Editor/UnrealEd/Private/Cooker/CookPackageData.cpp b/Engine/Source/Editor/UnrealEd/Private/Cooker/CookPackageData.cpp
--- a/Engine/Source/Editor/UnrealEd/Private/Cooker/CookPackageData.cpp
+++ b/Engine/Source/Editor/UnrealEd/Private/Cooker/CookPackageData.cpp
@@ -2425,14 +2425,9 @@ void FPendingCookedPlatformData::ClearCachedCookedPlatformData(UObject* Object,
 	}
 	else
 	{
-		// TODO: Add this call to ClearCachedCookedPlatformData once we take it over from SavePackage
-		constexpr bool bCallClearCachedCookedPlatformData = false;
-		if (bCallClearCachedCookedPlatformData)
+		for (const TPair<const ITargetPlatform*, ECachedCookedPlatformDataEvent>& PlatformPair : CCPDState.PlatformStates)
 		{
-			for (const TPair<const ITargetPlatform*, ECachedCookedPlatformDataEvent>& PlatformPair : CCPDState.PlatformStates)
-			{
-				Object->ClearCachedCookedPlatformData(PlatformPair.Key);
-			}
+			Object->ClearCachedCookedPlatformData(PlatformPair.Key);
 		}
 
 		// ClearAllCachedCookedPlatformData and WillNeverCacheCookedPlatformDataAgain calls are only used when not in editor
diff --git a/Engine/Source/Editor/UnrealEd/Private/Cooker/LooseCookedPackageWriter.cpp b/Engine/Source/Editor/UnrealEd/Private/Cooker/LooseCookedPackageWriter.cpp
--- a/Engine/Source/Editor/UnrealEd/Private/Cooker/LooseCookedPackageWriter.cpp
+++ b/Engine/Source/Editor/UnrealEd/Private/Cooker/LooseCookedPackageWriter.cpp
@@ -59,13 +59,13 @@ LLM_DEFINE_TAG(Cooker_PackageStoreManifest);
 
 FLooseCookedPackageWriter::FLooseCookedPackageWriter(const FString& InOutputPath,
 	const FString& InMetadataDirectoryPath, const ITargetPlatform* InTargetPlatform, FAsyncIODelete& InAsyncIODelete,
-	UE::Cook::FPackageDatas& InPackageDatas, const TArray<TSharedRef<IPlugin>>& InPluginsToRemap)
+	const TArray<TSharedRef<IPlugin>>& InPluginsToRemap, FBeginCacheCallback&& InBeginCacheCallback)
 	: OutputPath(InOutputPath)
 	, MetadataDirectoryPath(InMetadataDirectoryPath)
 	, TargetPlatform(*InTargetPlatform)
-	, PackageDatas(InPackageDatas)
 	, PluginsToRemap(InPluginsToRemap)
 	, AsyncIODelete(InAsyncIODelete)
+	, BeginCacheCallback(MoveTemp(InBeginCacheCallback))
 {
 }
 
@@ -555,6 +555,12 @@ void FLooseCookedPackageWriter::Initialize(const FCookInfo& Info)
 	}
 }
 
+EPackageWriterResult FLooseCookedPackageWriter::BeginCacheForCookedPlatformData(
+	FBeginCacheForCookedPlatformDataInfo& Info)
+{
+	return BeginCacheCallback(Info);
+}
+
 void FLooseCookedPackageWriter::WriteOplogEntry(FCbWriter& Writer, const FOplogPackageInfo& PackageInfo)
 {
 	Writer.BeginObject();
diff --git a/Engine/Source/Editor/UnrealEd/Private/Cooker/LooseCookedPackageWriter.h b/Engine/Source/Editor/UnrealEd/Private/Cooker/LooseCookedPackageWriter.h
--- a/Engine/Source/Editor/UnrealEd/Private/Cooker/LooseCookedPackageWriter.h
+++ b/Engine/Source/Editor/UnrealEd/Private/Cooker/LooseCookedPackageWriter.h
@@ -29,7 +29,6 @@ class FMD5;
 class IPlugin;
 class ITargetPlatform;
 template <typename ReferencedType> class TRefCountPtr;
-namespace UE::Cook { struct FPackageDatas; }
 
 /** A CookedPackageWriter that saves cooked packages in separate .uasset,.uexp,.ubulk files in the Saved\Cooked\[Platform] directory. */
 class FLooseCookedPackageWriter : public TPackageWriterToSharedBuffer<ICookedPackageWriter>
@@ -39,7 +38,7 @@ public:
 
 	FLooseCookedPackageWriter(const FString& OutputPath, const FString& MetadataDirectoryPath,
 		const ITargetPlatform* TargetPlatform, FAsyncIODelete& InAsyncIODelete,
-		UE::Cook::FPackageDatas& InPackageDatas, const TArray<TSharedRef<IPlugin>>& InPluginsToRemap);
+		const TArray<TSharedRef<IPlugin>>& InPluginsToRemap, FBeginCacheCallback&& InBeginCacheCallback);
 	~FLooseCookedPackageWriter();
 
 	virtual FCookCapabilities GetCookCapabilities() const override
@@ -65,6 +64,8 @@ public:
 	virtual bool TryReadMPCookMessageForPackage(FName PackageName, FCbObjectView Message) override;
 	virtual bool GetPreviousCookedBytes(const FPackageInfo& Info, FPreviousCookedBytesData& OutData) override;
 	virtual void CompleteExportsArchiveForDiff(FPackageInfo& Info, FLargeMemoryWriter& ExportsArchive) override;
+	virtual EPackageWriterResult BeginCacheForCookedPlatformData(FBeginCacheForCookedPlatformDataInfo& Info) override;
+
 	virtual void CommitPackageInternal(FPackageWriterRecords::FPackage&& BaseRecord,
 		const FCommitPackageInfo& Info) override;
 	virtual FPackageWriterRecords::FPackage* ConstructRecord() override;
@@ -171,10 +172,10 @@ private:
 	FString OutputPath;
 	FString MetadataDirectoryPath;
 	const ITargetPlatform& TargetPlatform;
-	UE::Cook::FPackageDatas& PackageDatas;
 	TMap<FName, FOplogPackageInfo> Oplog;
 	const TArray<TSharedRef<IPlugin>>& PluginsToRemap;
 	FAsyncIODelete& AsyncIODelete;
+	FBeginCacheCallback BeginCacheCallback;
 	bool bIterateSharedBuild = false;
 	bool bProvidePerPackageResults = false;
 };
diff --git a/Engine/Source/Runtime/Core/Public/Serialization/PackageWriter.h b/Engine/Source/Runtime/Core/Public/Serialization/PackageWriter.h
--- a/Engine/Source/Runtime/Core/Public/Serialization/PackageWriter.h
+++ b/Engine/Source/Runtime/Core/Public/Serialization/PackageWriter.h
@@ -25,6 +25,13 @@ struct FPackageStoreEntryResource;
 struct FSavePackageArgs;
 struct FSavePackageResultStruct;
 
+enum class EPackageWriterResult : uint8
+{
+	Success,
+	Error,
+	Timeout,
+};
+
 /** Interface for SavePackage to write packages to storage. */
 class IPackageWriter
 {
@@ -326,6 +333,23 @@ public:
 		// The subclass must override this method if it returns bDiffModeSupported=true in GetCookCapabilities
 		unimplemented();
 	}
+
+	struct FBeginCacheForCookedPlatformDataInfo
+	{
+		FName PackageName;
+		const ITargetPlatform* TargetPlatform;
+		TConstArrayView<UObject*> SaveableObjects;
+		uint32 SaveFlags;
+	};
+	// Helper callback type for Writers that need to send the message on to UCookOnTheFlyServer
+	using FBeginCacheCallback = TUniqueFunction<EPackageWriterResult(FBeginCacheForCookedPlatformDataInfo& Info)>;
+
+	virtual EPackageWriterResult BeginCacheForCookedPlatformData(FBeginCacheForCookedPlatformDataInfo& Info)
+	{
+		unimplemented();
+		return EPackageWriterResult::Error;
+	}
+
 	/** Modify the SaveArgs if required before the first Save. Used for diffing. */
 	virtual void UpdateSaveArguments(FSavePackageArgs& SaveArgs)
 	{
diff --git a/Engine/Source/Runtime/CoreUObject/Private/UObject/SavePackage2.cpp b/Engine/Source/Runtime/CoreUObject/Private/UObject/SavePackage2.cpp
--- a/Engine/Source/Runtime/CoreUObject/Private/UObject/SavePackage2.cpp
+++ b/Engine/Source/Runtime/CoreUObject/Private/UObject/SavePackage2.cpp
@@ -2521,84 +2521,36 @@ ESavePackageResult FinalizeFile(FStructuredArchive::FRecord& StructuredArchiveRo
 ESavePackageResult BeginCachePlatformCookedData(FSaveContext& SaveContext)
 {
 #if WITH_EDITOR
-	// TODO: Call IsCachedCookedPlatformDataLoaded on all the objects until it returns true. This is required for the contract.
-	// We tried enabling this once, but it created knockon bugs: Textures created by landscape were not handling it correctly
-	// (which we have subsequently fixed) and MaterialInstanceConstants created by landscape were not handling it correctly (which 
-	// we have not yet diagnosed)
-#define SAVEPACKAGE_CALL_ISCACHEDCOOKEDPLATFORMDATALOADED 0
-
-	// TODO: Remove BeginCacheForCookedPlatformData from SavePackage; it is already called by the cooker.
-	// Cache platform cooked data
-	if (SaveContext.IsCooking() && !SaveContext.IsConcurrent())
+	if (!SaveContext.IsCooking() || SaveContext.IsConcurrent())
 	{
-		const ITargetPlatform* TargetPlatform = SaveContext.GetTargetPlatform();
-		TArray<UObject*> ObjectsInPackage;
-		GetObjectsWithPackage(SaveContext.GetPackage(), ObjectsInPackage);
-		for (TArray<UObject*>::TIterator Iter(ObjectsInPackage); Iter; ++Iter)
-		{
-			UObject* Object = *Iter;
-			if (SaveContext.IsUnsaveable(Object))
-			{
-				Iter.RemoveCurrentSwap();
-				continue;
-			}
-			Object->BeginCacheForCookedPlatformData(TargetPlatform);
-#if SAVEPACKAGE_CALL_ISCACHEDCOOKEDPLATFORMDATALOADED
-			if (Object->IsCachedCookedPlatformDataLoaded(TargetPlatform))
-			{
-				Iter.RemoveCurrentSwap();
-				continue;
-			}
-#endif
-		}
-#if SAVEPACKAGE_CALL_ISCACHEDCOOKEDPLATFORMDATALOADED
-		if (ObjectsInPackage.Num())
-		{
-			if (SaveContext.GetSaveArgs().SaveFlags & SAVE_AllowTimeout)
-			{
-				return ESavePackageResult::Timeout;
-			}
-			const float MaxWaitSeconds = 30.f;
-			const double EndTimeSeconds = FPlatformTime::Seconds() + MaxWaitSeconds;
-			while (ObjectsInPackage.Num())
-			{
-				if (FPlatformTime::Seconds() > EndTimeSeconds)
-				{
-					UE_LOG(LogSavePackage, Error, TEXT("Save of %s failed: timed out waiting for IsCachedCookedPlatformDataLoaded on %s."),
-						SaveContext.GetFilename(), *ObjectsInPackage[0]->GetFullName());
-					return ESavePackageResult::Error;
-				}
-				for (TArray<UObject*>::TIterator Iter(ObjectsInPackage); Iter; ++Iter)
-				{
-					if ((*Iter)->IsCachedCookedPlatformDataLoaded(TargetPlatform))
-					{
-						Iter.RemoveCurrentSwap();
-						continue;
-					}
-				}
-			}
-		}
-#endif
+		// BeginCacheForCookedPlatformData is not called if not cooking
+		// When saving concurrently the cooker has called it ahead of time (because it is not threadsafe)
+		return ESavePackageResult::Success;
 	}
-#endif
-	return ESavePackageResult::Success;
-}
+	ICookedPackageWriter* CookedWriter = SaveContext.GetPackageWriter()->AsCookedPackageWriter();
+	check(CookedWriter); // Cooking requires a CookedPackageWriter
 
-void ClearCachedPlatformCookedData(FSaveContext& SaveContext)
-{
-#if WITH_EDITOR
-	if (SaveContext.IsCooking() && !SaveContext.IsConcurrent())
+	// Find the saveable objects
+	TArray<UObject*> ObjectsInPackage;
+	UPackage* Package = SaveContext.GetPackage();
+	GetObjectsWithPackage(Package, ObjectsInPackage);
+	ObjectsInPackage.RemoveAllSwap([&SaveContext](UObject* Object) { return SaveContext.IsUnsaveable(Object); });
+
+	// Call the PackageWriter to dispatch the BeginCache calls
+	ICookedPackageWriter::FBeginCacheForCookedPlatformDataInfo Info{
+		Package->GetFName(), SaveContext.GetTargetPlatform(), TConstArrayView<UObject*>(ObjectsInPackage),
+		SaveContext.GetSaveArgs().SaveFlags
+	};
+	EPackageWriterResult Result = CookedWriter->BeginCacheForCookedPlatformData(Info);
+	switch (Result)
 	{
-		TArray<UObject*> ObjectsInPackage;
-		GetObjectsWithPackage(SaveContext.GetPackage(), ObjectsInPackage);
-		for (UObject* Object : ObjectsInPackage)
-		{
-			if (!SaveContext.IsUnsaveable(Object))
-			{
-				Object->ClearCachedCookedPlatformData(SaveContext.GetTargetPlatform());
-			}
-		}
+	case EPackageWriterResult::Success: return ESavePackageResult::Success;
+	case EPackageWriterResult::Error: return ESavePackageResult::Error;
+	case EPackageWriterResult::Timeout: return ESavePackageResult::Timeout;
+	default: return ESavePackageResult::Error;
 	}
+#else
+	return ESavePackageResult::Success;
 #endif
 }
 
@@ -3096,8 +3048,10 @@ FSavePackageResultStruct UPackage::Save2(UPackage* InPackage, UObject* InAsset,
 		}
 	}
 
-	// Trigger platform cooked data caching before package harvesting 
-	// because it might modify some property and hence affect harvested property name for tagged property for example
+	// Trigger platform cooked data caching after PreSave but before package harvesting 
+	// After PreSave because objects can be created during PreSave and we need to cache them
+	// Before package harvesting because it might modify some property and hence affect the harvested
+	// property name of a tagged property for example
 	SaveContext.Result = BeginCachePlatformCookedData(SaveContext);
 	if (SaveContext.Result != ESavePackageResult::Success)
 	{
@@ -3110,7 +3064,7 @@ FSavePackageResultStruct UPackage::Save2(UPackage* InPackage, UObject* InAsset,
 		//FScopedSavingFlag IsSavingFlag(SaveContext.IsConcurrent());
 		SaveContext.Result = InnerSave(SaveContext);
 
-		// in case of failure or cancellation, do not exit here, still run cleanup (PostSaveRoot/ClearCachedPlatformCookedData)
+		// in case of failure or cancellation, do not exit here, still run cleanup (e.g. PostSaveRoot)
 	}
 
 	// PostSave Asset
@@ -3121,8 +3075,6 @@ FSavePackageResultStruct UPackage::Save2(UPackage* InPackage, UObject* InAsset,
 		SaveContext.SetPostSaveRootRequired(false);
 	}
 
-	ClearCachedPlatformCookedData(SaveContext);
-
 	// PostSave Package - edit in memory package and send events if save was successful
 	SlowTask.EnterProgressFrame();
 	PostSavePackage(SaveContext);
@@ -3209,8 +3161,6 @@ ESavePackageResult UPackage::SaveConcurrent(TArrayView<FPackageSaveInfo> InPacka
 				SaveContext.SetPreSaveCleanup(false);
 			}
 
-			ClearCachedPlatformCookedData(SaveContext);
-
 			// PostSave Package - edit in memory package and send events
 			if (SaveContext.Result == ESavePackageResult::Success)
 			{
