Cooker: BeginCacheForCookedPlatformData and generated packages: Support objects being passed through their BeginCacheForCookedPlatformData functions from two packages at the same time, because it can occur when a generator package moves an object into a generated package. Track references from the packages to prevent calling the ClearCacheCookedPlatformData functions on the object until all packages are done with it.
#rnx

diff --git a/Engine/Source/Editor/UnrealEd/Private/CookOnTheFlyServer.cpp b/Engine/Source/Editor/UnrealEd/Private/CookOnTheFlyServer.cpp
--- a/Engine/Source/Editor/UnrealEd/Private/CookOnTheFlyServer.cpp
+++ b/Engine/Source/Editor/UnrealEd/Private/CookOnTheFlyServer.cpp
@@ -3741,19 +3741,7 @@ UE::Cook::EPollStatus UCookOnTheFlyServer::CallBeginCacheOnObjects(UE::Cook::FPa
 			continue;
 		}
 		FCachedCookedPlatformDataState& CCPDState = PackageDatas->GetCachedCookedPlatformDataObjects().FindOrAdd(Obj);
-		if (CCPDState.PackageData != &PackageData)
-		{
-			if (CCPDState.PackageData != nullptr)
-			{
-				UE_LOG(LogCook, Display, TEXT("CachedCookedPlatformDatas: %s is unexpectedly in two packages at once.")
-					TEXT("\n\tPreviousPackage: %s")
-					TEXT("\n\tCurrentPackage: %s"),
-					*Obj->GetFullName(), *CCPDState.PackageData->GetPackageName().ToString(),
-					*PackageData.GetPackageName().ToString());
-				continue;
-			}
-			CCPDState.PackageData = &PackageData;
-		}
+		CCPDState.AddRefFrom(&PackageData);
 
 		for (const ITargetPlatform* TargetPlatform : TargetPlatforms)
 		{
@@ -11594,24 +11582,14 @@ void UCookOnTheFlyServer::RouteBeginCacheForCookedPlatformData(UE::Cook::FPackag
 	if (!ExistingEvent)
 	{
 		FCachedCookedPlatformDataState& CCPDState = PackageDatas->GetCachedCookedPlatformDataObjects().FindOrAdd(Obj);
-		if (CCPDState.PackageData != &PackageData)
-		{
-			if (CCPDState.PackageData != nullptr)
-			{
-				ExistingEvent = CCPDState.PlatformStates.Find(TargetPlatform);
-				UE_LOG(LogCook, Display, TEXT("BeginCacheForCookedPlatformData called unexpectedly on %s. Skipping it.")
-					TEXT("\n\tPreviousEvent: %s for package %s.")
-					TEXT("\n\tCurrentEvent: %s for package %s."),
-					*Obj->GetFullName(), LexToString(ExistingEvent ? *ExistingEvent : ECachedCookedPlatformDataEvent::None),
-					*CCPDState.PackageData->GetPackageName().ToString(),
-					LexToString(ECachedCookedPlatformDataEvent::BeginCacheForCookedPlatformDataCalled),
-					*PackageName.ToString());
-				return;
-			}
-			CCPDState.PackageData = &PackageData;
-		}
+		CCPDState.AddRefFrom(&PackageData);
 		ExistingEvent = &CCPDState.PlatformStates.FindOrAdd(TargetPlatform, ECachedCookedPlatformDataEvent::None);
 	}
+	if (*ExistingEvent != ECachedCookedPlatformDataEvent::None)
+	{
+		// BeginCacheForCookedPlatformData was already called; do not call it again
+		return;
+	}
 
 	// We need to set our scopes for e.g. TObjectPtr reads around the call to BeginCacheForCookedPlatformData,
 	// but in some cases we have already set the scope (e.g. when calling BeginCache from inside SavePackage)
@@ -11637,21 +11615,7 @@ bool UCookOnTheFlyServer::RouteIsCachedCookedPlatformDataLoaded(UE::Cook::FPacka
 	if (!ExistingEvent)
 	{
 		FCachedCookedPlatformDataState& CCPDState = PackageDatas->GetCachedCookedPlatformDataObjects().FindOrAdd(Obj);
-		if (CCPDState.PackageData != &PackageData)
-		{
-			if (CCPDState.PackageData != nullptr)
-			{
-				UE_LOG(LogCook, Display, TEXT("IsCachedCookedPlatformDataLoaded called unexpectedly on %s. Skipping it and returning true.")
-					TEXT("\n\tPreviousEvent: %s for package %s.")
-					TEXT("\n\tCurrentEvent: %s for package %s."),
-					*Obj->GetFullName(), LexToString(ExistingEvent ? *ExistingEvent : ECachedCookedPlatformDataEvent::None),
-					*CCPDState.PackageData->GetPackageName().ToString(),
-					LexToString(ECachedCookedPlatformDataEvent::IsCachedCookedPlatformDataLoadedCalled),
-					*PackageName.ToString());
-				return true;
-			}
-			CCPDState.PackageData = &PackageData;
-		}
+		CCPDState.AddRefFrom(&PackageData);
 		ExistingEvent = &CCPDState.PlatformStates.FindOrAdd(TargetPlatform, ECachedCookedPlatformDataEvent::None);
 	}
 
@@ -11685,18 +11649,9 @@ EPackageWriterResult UCookOnTheFlyServer::SavePackageBeginCacheForCookedPlatform
 	for (UObject* Object : SaveableObjects)
 	{
 		FCachedCookedPlatformDataState& CCPDState = PackageDatas->GetCachedCookedPlatformDataObjects().FindOrAdd(Object);
-		if (CCPDState.PackageData != PackageData)
+		if (!CCPDState.PackageDatas.Contains(PackageData))
 		{
-			if (CCPDState.PackageData != nullptr)
-			{
-				UE_LOG(LogCook, Display, TEXT("CachedCookedPlatformDatas: %s is unexpectedly in two packages at once.")
-					TEXT("\n\tPreviousPackage: %s")
-					TEXT("\n\tCurrentPackage: %s"),
-					*Object->GetFullName(), *CCPDState.PackageData->GetPackageName().ToString(),
-					*PackageData->GetPackageName().ToString());
-				continue;
-			}
-			CCPDState.PackageData = PackageData;
+			CCPDState.AddRefFrom(PackageData);
 
 			// NextIndex is usually at the end of CachedObjectsInOuter, but in case it is not, insert the new Object 
 			// at NextIndex so that we still record that we have not called BeginCace on objects after it. Then increment
@@ -11709,26 +11664,26 @@ EPackageWriterResult UCookOnTheFlyServer::SavePackageBeginCacheForCookedPlatform
 
 		ECachedCookedPlatformDataEvent& ExistingEvent =
 			CCPDState.PlatformStates.FindOrAdd(TargetPlatform, ECachedCookedPlatformDataEvent::None);
-		if (ExistingEvent == ECachedCookedPlatformDataEvent::IsCachedCookedPlatformDataLoadedReturnedTrue)
+		if (ExistingEvent != ECachedCookedPlatformDataEvent::IsCachedCookedPlatformDataLoadedReturnedTrue)
 		{
-			// Already called
-			continue;
-		}
-
-		RouteBeginCacheForCookedPlatformData(*PackageData, Object, TargetPlatform, &ExistingEvent);
-		if (bCallIsCachedOnSaveCreatedObjects)
-		{
-			// TODO: Enable bCallIsCachedOnSaveCreatedObjects so that we call IsCachedCookedPlatformDataLoaded on all the objects until it
-			// returns true. This is required for the BeginCacheForCookedPlatformData contract.
-			// Doing so will cause us to return Timeout and retry the save later after the pending objects have completed.
-			// We tried enabling this once, but it created knockon bugs: Textures created by landscape were not handling it correctly
-			// (which we have subsequently fixed) and MaterialInstanceConstants created by landscape were not handling it correctly (which 
-			// we have not yet diagnosed).
-			if (!RouteIsCachedCookedPlatformDataLoaded(*PackageData, Object, TargetPlatform, &ExistingEvent))
+			if (ExistingEvent == ECachedCookedPlatformDataEvent::None)
+			{
+				RouteBeginCacheForCookedPlatformData(*PackageData, Object, TargetPlatform, &ExistingEvent);
+			}
+			if (bCallIsCachedOnSaveCreatedObjects)
 			{
-				PackageDatas->AddPendingCookedPlatformData(FPendingCookedPlatformData(Object, TargetPlatform,
-					*PackageData, false /* bNeedsResourceRelease */, *this));
-				PendingObjects.Add(Object);
+				// TODO: Enable bCallIsCachedOnSaveCreatedObjects so that we call IsCachedCookedPlatformDataLoaded on all the objects until it
+				// returns true. This is required for the BeginCacheForCookedPlatformData contract.
+				// Doing so will cause us to return Timeout and retry the save later after the pending objects have completed.
+				// We tried enabling this once, but it created knockon bugs: Textures created by landscape were not handling it correctly
+				// (which we have subsequently fixed) and MaterialInstanceConstants created by landscape were not handling it correctly (which 
+				// we have not yet diagnosed).
+				if (!RouteIsCachedCookedPlatformDataLoaded(*PackageData, Object, TargetPlatform, &ExistingEvent))
+				{
+					PackageDatas->AddPendingCookedPlatformData(FPendingCookedPlatformData(Object, TargetPlatform,
+						*PackageData, false /* bNeedsResourceRelease */, *this));
+					PendingObjects.Add(Object);
+				}
 			}
 		}
 	}
diff --git a/Engine/Source/Editor/UnrealEd/Private/Cooker/CookPackageData.cpp b/Engine/Source/Editor/UnrealEd/Private/Cooker/CookPackageData.cpp
--- a/Engine/Source/Editor/UnrealEd/Private/Cooker/CookPackageData.cpp
+++ b/Engine/Source/Editor/UnrealEd/Private/Cooker/CookPackageData.cpp
@@ -1256,29 +1256,9 @@ EPollStatus FPackageData::RefreshObjectCache(bool& bOutFoundNewObjects)
 
 void FPackageData::ClearObjectCache()
 {
-	TMap<UObject*, FCachedCookedPlatformDataState>& CCPDs = PackageDatas.GetCachedCookedPlatformDataObjects();
-	for (FWeakObjectPtr& WeakObjectPtr: CachedObjectsInOuter)
-	{
-		UObject* Object = WeakObjectPtr.Get();
-		if (!Object)
-		{
-			continue;
-		}
-		FCachedCookedPlatformDataState CCPDState;
-		CCPDs.RemoveAndCopyValue(Object, CCPDState);
-		if (CCPDState.PackageData != this)
-		{
-			// This is an error condition that is logged elsewhere; one of the objects in this package is also
-			// present in another package, which owns it. We should silently handle this case (silently because
-			// the error about it is logged elsewhere), which means we should not modify the object on account
-			// of this package.
-			// We removed it because we thought would own it if it existed, but since we were wrong, put it back.
-			if (CCPDState.PackageData != nullptr)
-			{
-				CCPDs.Add(Object, MoveTemp(CCPDState));
-			}
-		}
-	}
+	// Note we do not need to remove objects in CachedObjectsInOuter from CachedCookedPlatformDataObjects
+	// That removal is handled by ReleaseCookedPlatformData, and the caller is responsible for calling
+	// ReleaseCookedPlatformData before calling ClearObjectCache
 	CachedObjectsInOuter.Empty();
 	for (TPair<const ITargetPlatform*, FPackagePlatformData>& Pair : PlatformDatas)
 	{
@@ -2399,38 +2379,26 @@ void FPendingCookedPlatformData::ClearCachedCookedPlatformData(UObject* Object,
 {
 	FPackageDatas& PackageDatas = PackageData.GetPackageDatas();
 	UCookOnTheFlyServer& COTFS = PackageDatas.GetCookOnTheFlyServer();
-	FCachedCookedPlatformDataState CCPDState;
-	PackageDatas.GetCachedCookedPlatformDataObjects().RemoveAndCopyValue(Object, CCPDState);
-	if (CCPDState.PackageData != &PackageData)
+	using FCCPDMapType = TMap<UObject*, FCachedCookedPlatformDataState>;
+	FCCPDMapType& CCPDs = PackageDatas.GetCachedCookedPlatformDataObjects();
+
+	uint32 ObjectKeyHash = FCCPDMapType::KeyFuncsType::GetKeyHash(Object);
+	FCachedCookedPlatformDataState* CCPDState = CCPDs.FindByHash(ObjectKeyHash, Object);
+	if (!CCPDState)
 	{
-		if (CCPDState.PackageData)
-		{
-			ECachedCookedPlatformDataEvent ExistingEvent = ECachedCookedPlatformDataEvent::None;
-			for (TPair<const ITargetPlatform*, ECachedCookedPlatformDataEvent>& PlatformPair : CCPDState.PlatformStates)
-			{
-				if (PlatformPair.Value != ECachedCookedPlatformDataEvent::None)
-				{
-					ExistingEvent = PlatformPair.Value;
-					break;
-				}
-			}
-			UE_LOG(LogCook, Display, TEXT("ClearCachedCookedPlatformData called unexpectedly on %s. It will be skipped.")
-				TEXT("\n\tPreviousEvent: %s for package %s.")
-				TEXT("\n\tCurrentEvent: %s for package %s."),
-				*Object->GetFullName(), LexToString(ExistingEvent), *CCPDState.PackageData->GetPackageName().ToString(),
-				LexToString(ECachedCookedPlatformDataEvent::ClearCachedCookedPlatformDataCalled),
-				*PackageData.GetPackageName().ToString());
-			PackageDatas.GetCachedCookedPlatformDataObjects().Add(Object, MoveTemp(CCPDState));
-		}
+		return;
 	}
-	else
+
+	CCPDState->ReleaseFrom(&PackageData);
+	if (!CCPDState->IsReferenced())
 	{
-		for (const TPair<const ITargetPlatform*, ECachedCookedPlatformDataEvent>& PlatformPair : CCPDState.PlatformStates)
+		for (const TPair<const ITargetPlatform*, ECachedCookedPlatformDataEvent>&
+			PlatformPair : CCPDState->PlatformStates)
 		{
 			Object->ClearCachedCookedPlatformData(PlatformPair.Key);
 		}
 
-		// ClearAllCachedCookedPlatformData and WillNeverCacheCookedPlatformDataAgain calls are only used when not in editor
+		// ClearAllCachedCookedPlatformData and WillNeverCacheCookedPlatformDataAgain are not used in editor
 		if (!COTFS.IsCookingInEditor())
 		{
 			Object->ClearAllCachedCookedPlatformData();
@@ -2439,6 +2407,8 @@ void FPendingCookedPlatformData::ClearCachedCookedPlatformData(UObject* Object,
 				Object->WillNeverCacheCookedPlatformDataAgain();
 			}
 		}
+
+		CCPDs.RemoveByHash(ObjectKeyHash, Object);
 	}
 };
 
@@ -3479,4 +3449,26 @@ void FPackageDatas::CachedCookedPlatformDataObjectsPostGarbageCollect(const TSet
 	}
 }
 
+void FCachedCookedPlatformDataState::AddRefFrom(FPackageData* PackageData)
+{
+	// Most objects will only be referenced by a single package.
+	// The exceptions:
+	//   1) Generator Packages that move the object from the generator into a generated
+	//   2) Bugs
+	// Even in case (1), the number of referencers will be 2.
+	// We therefore for now just use a flat array and AddUnique, to minimize memory and performance in the 
+	// usual case on only a single referencer.
+	PackageDatas.AddUnique(PackageData);
+}
+
+void FCachedCookedPlatformDataState::ReleaseFrom(FPackageData* PackageData)
+{
+	PackageDatas.Remove(PackageData);
+}
+
+bool FCachedCookedPlatformDataState::IsReferenced() const
+{
+	return !PackageDatas.IsEmpty();
+}
+
 } // namespace UE::Cook
diff --git a/Engine/Source/Editor/UnrealEd/Private/Cooker/CookPackageData.h b/Engine/Source/Editor/UnrealEd/Private/Cooker/CookPackageData.h
--- a/Engine/Source/Editor/UnrealEd/Private/Cooker/CookPackageData.h
+++ b/Engine/Source/Editor/UnrealEd/Private/Cooker/CookPackageData.h
@@ -71,12 +71,21 @@ enum class ECachedCookedPlatformDataEvent : uint8
 };
 const TCHAR* LexToString(ECachedCookedPlatformDataEvent);
 /**
- * BeginCachedForCookedPlatformData state about an object - which package owned it (not always the same
+ * BeginCachedForCookedPlatformData state about an object - which packages owned it (not always the same
  * one when PackageGenerators are involved) and the per-platform state for ECachedCookedPlatformDataEvent.
  */
 struct FCachedCookedPlatformDataState
 {
-	FPackageData* PackageData = nullptr;
+	void AddRefFrom(FPackageData* PackageData);
+	void ReleaseFrom(FPackageData* PackageData);
+	bool IsReferenced() const;
+
+	/**
+	 * The packages that have called any of the BeginCacheForCookedPlatformData family of function
+	 * on this object. Usually only a single 1, sometimes 2, see comment in AddPackageData.
+	 */
+	TArray<FPackageData*, TInlineAllocator<2>> PackageDatas;
+	/** The per-platform state of which BeginCacheForCookedPlatformData events have been passed. */
 	TMap<const ITargetPlatform*, ECachedCookedPlatformDataEvent> PlatformStates;
 };
 
