Make static mesh instance picking a persistent project setting

diff --git a/Engine/Source/Editor/UnrealEd/Private/Settings/EditorProjectSettings.cpp b/Engine/Source/Editor/UnrealEd/Private/Settings/EditorProjectSettings.cpp
--- a/Engine/Source/Editor/UnrealEd/Private/Settings/EditorProjectSettings.cpp
+++ b/Engine/Source/Editor/UnrealEd/Private/Settings/EditorProjectSettings.cpp
@@ -181,6 +181,34 @@ void ULevelEditor2DSettings::PostEditChangeProperty(struct FPropertyChangedEvent
 	Super::PostEditChangeProperty(PropertyChangedEvent);
 }
 
+/* ULevelEditorProjectSettings
+*****************************************************************************/
+
+ULevelEditorProjectSettings::ULevelEditorProjectSettings( const FObjectInitializer& ObjectInitializer )
+	: Super(ObjectInitializer)
+	, bEnableViewportSMInstanceSelection(true)
+{
+}
+
+void ULevelEditorProjectSettings::PostInitProperties()
+{
+	Super::PostInitProperties();
+	if (FProperty* EnableSMInstanceSelectionProperty = GetClass()->FindPropertyByName(GET_MEMBER_NAME_CHECKED(ULevelEditorProjectSettings, bEnableViewportSMInstanceSelection)))
+	{
+		ExportValuesToConsoleVariables(EnableSMInstanceSelectionProperty);
+	}
+}
+
+void ULevelEditorProjectSettings::PostEditChangeProperty(FPropertyChangedEvent& PropertyChangedEvent)
+{
+	Super::PostEditChangeProperty(PropertyChangedEvent);
+	if (!PropertyChangedEvent.Property)
+	{
+		return;
+	}
+	ExportValuesToConsoleVariables(PropertyChangedEvent.Property);
+}
+
 
 /* UEditorPerformanceProjectSettings
 *****************************************************************************/
diff --git a/Engine/Source/Editor/UnrealEd/Public/Settings/EditorProjectSettings.h b/Engine/Source/Editor/UnrealEd/Public/Settings/EditorProjectSettings.h
--- a/Engine/Source/Editor/UnrealEd/Public/Settings/EditorProjectSettings.h
+++ b/Engine/Source/Editor/UnrealEd/Public/Settings/EditorProjectSettings.h
@@ -182,6 +182,27 @@ public:
 
 };
 
+/**
+ * Configure per-project settings for the Level Editor
+ */
+UCLASS(config=Editor, meta=(DisplayName="Level Editor"), defaultconfig)
+class UNREALED_API ULevelEditorProjectSettings : public UDeveloperSettings
+{
+	GENERATED_BODY()
+
+public:
+	UPROPERTY(EditAnywhere, config, Category=Editing, meta=(
+		DisplayName="Enable viewport static mesh instance selection",
+		ConsoleVariable="TypedElements.EnableViewportSMInstanceSelection"))
+	bool bEnableViewportSMInstanceSelection;
+
+public:
+	// UObject interface
+	virtual void PostInitProperties() override;
+	virtual void PostEditChangeProperty(struct FPropertyChangedEvent& PropertyChangedEvent) override;
+	// End of UObject interface
+};
+
 /**
  * Configure per-project performance settings for the Editor
  */
