Fix CIS compilation

diff --git a/Engine/Source/Editor/UnrealEd/Public/Settings/EditorProjectSettings.h b/Engine/Source/Editor/UnrealEd/Public/Settings/EditorProjectSettings.h
--- a/Engine/Source/Editor/UnrealEd/Public/Settings/EditorProjectSettings.h
+++ b/Engine/Source/Editor/UnrealEd/Public/Settings/EditorProjectSettings.h
@@ -197,6 +197,7 @@ public:
 	bool bEnableViewportSMInstanceSelection;
 
 public:
+	ULevelEditorProjectSettings(const class FObjectInitializer& ObjectInitializer);
 	// UObject interface
 	virtual void PostInitProperties() override;
 	virtual void PostEditChangeProperty(struct FPropertyChangedEvent& PropertyChangedEvent) override;
