Cooker and ZenStore: Fix CookWorkers to call EndCook on the PackageWriter; it is guaranteed to happen by the CookedPackageWriter contract and ZenStore requires it.
#rnx

diff --git a/Engine/Source/Editor/UnrealEd/Private/CookOnTheFlyServer.cpp b/Engine/Source/Editor/UnrealEd/Private/CookOnTheFlyServer.cpp
--- a/Engine/Source/Editor/UnrealEd/Private/CookOnTheFlyServer.cpp
+++ b/Engine/Source/Editor/UnrealEd/Private/CookOnTheFlyServer.cpp
@@ -10479,6 +10479,7 @@ void UCookOnTheFlyServer::CookAsCookWorkerFinished()
 	FShaderLibraryCooker::EndCookingLibrary(ActualLibraryName);
 	FShaderLibraryCooker::Shutdown();
 	ShutdownShaderCompilers(PlatformManager->GetSessionPlatforms());
+	FinalizePackageStore();
 
 	if (IsDirectorCookOnTheFly())
 	{
