Cooker BeginCache contracts: Do not allow the cancelmanager from a package's previous save attempt block a new save attempt; this situation can occur after a garbage collect and it is possible the old objects in the cancelmanager will never return true from IsCachedCookedPlatformDataLoaded until BeginCacheCookedPlatformData is called again.
#rnx

diff --git a/Engine/Source/Editor/UnrealEd/Private/CookOnTheFlyServer.cpp b/Engine/Source/Editor/UnrealEd/Private/CookOnTheFlyServer.cpp
--- a/Engine/Source/Editor/UnrealEd/Private/CookOnTheFlyServer.cpp
+++ b/Engine/Source/Editor/UnrealEd/Private/CookOnTheFlyServer.cpp
@@ -3605,9 +3605,25 @@ UE::Cook::EPollStatus UCookOnTheFlyServer::PrepareSaveInternal(UE::Cook::FPackag
 		{
 			if (PackageData.GetNumPendingCookedPlatformData() > 0)
 			{
-				// A previous Save was started and deleted after some calls to BeginCacheForCookedPlatformData occurred, and some of those objects have still not returned true for IsCachedCookedPlatformDataLoaded
-				// We need to wait for all of pending async calls from the cancelled save to finish before we start the new ones
-				return EPollStatus::Incomplete;
+				// A previous Save was started and demoted after some calls to BeginCacheForCookedPlatformData
+				// occurred, and some of those objects have still not returned true for
+				// IsCachedCookedPlatformDataLoaded. We were keeping them around to call Clear on them after they
+				// return true from IsCachedCooked before we call Begin on them again. But depending on their state
+				// after a garbage collect, they might now never return true. So rather than blocking the BeginCache
+				// calls, clear the cancel manager and call BeginCache on them again even though they never returned
+				// true from IsCached. The contract for BeginCacheCookedPlatformData and
+				// IsCachedCookedPlatformDataLoaded includes the provision that is valid for the cooker to call
+				// BeginCacheCookedPlatformData multiple times before IsCachedCookedPlatformData is returns true, and
+				// the object should remain in a valid state (possibly reset to the beginning of its async work)
+				// afterwards and still eventually return true from a future IsCachedCookedPlatformData call.
+				PackageDatas->ClearCancelManager(PackageData);
+				if (PackageData.GetNumPendingCookedPlatformData() > 0)
+				{
+					UE_LOG(LogCook, Error,
+						TEXT("CookerBug: Package %s is blocked from entering save due to GetNumPendingCookedCookedPlatformData() == %d."),
+						*PackageData.GetPackageName().ToString(), PackageData.GetNumPendingCookedPlatformData());
+					return EPollStatus::Error;
+				}
 			}
 			PackageData.SetCookedPlatformDataStarted(true);
 		}
@@ -3883,11 +3899,11 @@ void UCookOnTheFlyServer::ReleaseCookedPlatformData(UE::Cook::FPackageData& Pack
 				CancelManager->NumPendingPlatforms = PendingDatas->Num();
 				for (FPendingCookedPlatformData* PendingCookedPlatformData : *PendingDatas)
 				{
-					// We never start a new package until after the previous cancel finished, so all of the
+					// We never start a new package until after clearing the previous cancels, so all of the
 					// FPendingCookedPlatformData for the PlatformData we are cancelling can not have been cancelled before.
 					// We would leak the CancelManager if we overwrote it here.
 					check(PendingCookedPlatformData->CancelManager == nullptr);
-					// If bHasReleaased on the PendingCookedPlatformData were already true, we would leak the CancelManager
+					// If bHasReleased on the PendingCookedPlatformData were already true, we would leak the CancelManager
 					// because the PendingCookedPlatformData would never call Release on it.
 					check(!PendingCookedPlatformData->bHasReleased);
 					PendingCookedPlatformData->CancelManager = CancelManager;
@@ -5143,6 +5159,7 @@ void UCookOnTheFlyServer::PostGarbageCollect()
 		}
 		else
 		{
+			// Mark that the objects for this package should be kept in CachedCookedPlatformData records
 			for (FWeakObjectPtr& WeakObjectPtr : PackageData->GetCachedObjectsInOuter())
 			{
 				UObject* Object = WeakObjectPtr.Get();
@@ -5159,6 +5176,21 @@ void UCookOnTheFlyServer::PostGarbageCollect()
 		PackageDatas->GetRequestQueue().AddRequest(PackageData, /* bForceUrgent */ true);
 	}
 
+	// Mark that any objects in PendingCookedPlatformDatas should be kept in CachedCookedPlatformData records
+	PackageDatas->ForEachPendingCookedPlatformData(
+		[&SaveQueueObjectsThatStillExist](FPendingCookedPlatformData& CookedPlatformData)
+		{
+			UObject* Object = CookedPlatformData.Object.Get();
+			if (Object)
+			{
+				SaveQueueObjectsThatStillExist.Add(Object);
+			}
+			else
+			{
+				CookedPlatformData.Release();
+			}
+		});
+
 	// Remove objects that were deleted by garbage collection from our containers that track raw object pointers
 	PackageDatas->CachedCookedPlatformDataObjectsPostGarbageCollect(SaveQueueObjectsThatStillExist);
 
diff --git a/Engine/Source/Editor/UnrealEd/Private/Cooker/CookPackageData.cpp b/Engine/Source/Editor/UnrealEd/Private/Cooker/CookPackageData.cpp
--- a/Engine/Source/Editor/UnrealEd/Private/Cooker/CookPackageData.cpp
+++ b/Engine/Source/Editor/UnrealEd/Private/Cooker/CookPackageData.cpp
@@ -3210,6 +3210,22 @@ void FPackageDatas::PollPendingCookedPlatformDatas(bool bForce, double& LastCook
 	}
 }
 
+void FPackageDatas::ClearCancelManager(FPackageData& PackageData)
+{
+	ForEachPendingCookedPlatformData(
+		[&PackageData](FPendingCookedPlatformData& PendingCookedPlatformData)
+		{
+			if (&PendingCookedPlatformData.PackageData == &PackageData)
+			{
+				if (!PendingCookedPlatformData.PollIsComplete())
+				{
+					// Abandon it
+					PendingCookedPlatformData.Release();
+				}
+			}
+		});
+}
+
 void FPackageDatas::RemapTargetPlatforms(const TMap<ITargetPlatform*, ITargetPlatform*>& Remap)
 {
 	LockAndEnumeratePackageDatas([&Remap](FPackageData* PackageData)
diff --git a/Engine/Source/Editor/UnrealEd/Private/Cooker/CookPackageData.h b/Engine/Source/Editor/UnrealEd/Private/Cooker/CookPackageData.h
--- a/Engine/Source/Editor/UnrealEd/Private/Cooker/CookPackageData.h
+++ b/Engine/Source/Editor/UnrealEd/Private/Cooker/CookPackageData.h
@@ -1346,6 +1346,7 @@ public:
 	 * releasing their resources and pending count if so.
 	 */
 	void PollPendingCookedPlatformDatas(bool bForce, double& LastCookableObjectTickTime);
+	void ClearCancelManager(FPackageData& PackageData);
 
 	/** Swap all ITargetPlatform* stored on this instance according to the mapping in @param Remap. */
 	void RemapTargetPlatforms(const TMap<ITargetPlatform*, ITargetPlatform*>& Remap);
