Added public dependencies that were missing.

diff --git a/Engine/Source/Editor/UnrealEd/UnrealEd.Build.cs b/Engine/Source/Editor/UnrealEd/UnrealEd.Build.cs
--- a/Engine/Source/Editor/UnrealEd/UnrealEd.Build.cs
+++ b/Engine/Source/Editor/UnrealEd/UnrealEd.Build.cs
@@ -89,7 +89,9 @@ public class UnrealEd : ModuleRules
 				"PhysicsUtilities",
 				"ToolWidgets",
 				"AssetTools", 
-				"WidgetRegistration"
+				"WidgetRegistration",
+				"GameplayTasks",
+				"TargetPlatform"
 			}
 		);
 
