Avoid delta serialization against an archetype of a different type

diff --git a/Engine/Source/Editor/UnrealEd/Private/Kismet2/KismetReinstanceUtilities.cpp b/Engine/Source/Editor/UnrealEd/Private/Kismet2/KismetReinstanceUtilities.cpp
--- a/Engine/Source/Editor/UnrealEd/Private/Kismet2/KismetReinstanceUtilities.cpp
+++ b/Engine/Source/Editor/UnrealEd/Private/Kismet2/KismetReinstanceUtilities.cpp
@@ -2018,6 +2018,13 @@ static void ReplaceObjectHelper(UObject*& OldObject, UClass* OldClass, UObject*&
 	Options.bNotifyObjectReplacement = true;
 	Options.bSkipCompilerGeneratedDefaults = true;
 	Options.OptionalReplacementMappings = &OldToNewInstanceMap;
+	// this currently happens because of some misguided logic in UBlueprintGeneratedClass::FindArchetype that
+	// points us to a mismatched archetype, in which case delta serialization becomes unsafe.. without
+	// that logic we could lose data, so for now i'm disabling delta serialization when we detect that situation
+	if(Options.SourceObjectArchetype && !OldObject->IsA(Options.SourceObjectArchetype->GetClass()))
+	{
+		Options.bDoDelta = false;
+	}
 	UEditorEngine::CopyPropertiesForUnrelatedObjects(OldObject, NewUObject, Options);
 	// Generate new subobjects
 	InstancedPropertyUtils::FArchiveInsertInstancedSubObjects InstancedSubObjSpawner(NewUObject, InstancedPropertyMap);
@@ -2214,6 +2221,12 @@ static void ReplaceActorHelper(AActor* OldActor, UClass* OldClass, UObject*& New
 	Params.bAggressiveDefaultSubobjectReplacement = true;
 	PRAGMA_ENABLE_DEPRECATION_WARNINGS
 	Params.bNotifyObjectReplacement = true;
+	// This shouldn't be possible, but if GetArchetype has a bug we could crash in delta serialization
+	// attempting to use it:
+	if (OldArchetype &&!ensure(OldActor->IsA(OldArchetype->GetClass())))
+	{
+		Params.bDoDelta = false;
+	}
 	UEngine::CopyPropertiesForUnrelatedObjects(OldActor, NewActor, Params);
 
 	// reset properties/streams
