Add a console command to manually trigger package reload.  Used for testing issues around hot reload.

#rnx

diff --git a/Engine/Source/Editor/UnrealEd/Private/PackageTools.cpp b/Engine/Source/Editor/UnrealEd/Private/PackageTools.cpp
--- a/Engine/Source/Editor/UnrealEd/Private/PackageTools.cpp
+++ b/Engine/Source/Editor/UnrealEd/Private/PackageTools.cpp
@@ -97,6 +97,37 @@ void UPackageTools::FlushAsyncCompilation(TArrayView<UPackage* const> InPackages
 	}
 }
 
+struct FPackageToolsCommands
+{
+	FPackageToolsCommands()
+		: ReloadConsoleCommand(TEXT("PackageTools.ReloadPackage"),
+							   TEXT("Force a reload of the named package, e.g. PackageTools.ReloadPackage /Game/MyAsset"),
+							   FConsoleCommandWithArgsDelegate::CreateRaw(this, &FPackageToolsCommands::ReloadPackageCommand))
+	{
+	}
+
+	void ReloadPackageCommand(const TArray<FString>& Params)
+	{
+		TArray<UPackage*> Packages;
+		for (const FString& Param : Params)
+		{
+			UPackage* ExistingPackage = FindPackage(nullptr, *Param);
+			if (ExistingPackage)
+			{
+				Packages.Add(ExistingPackage);
+			}
+		}
+		if (Packages.Num() > 0)
+		{
+			UPackageTools::ReloadPackages(Packages);
+		}
+	}
+
+	FAutoConsoleCommand ReloadConsoleCommand;
+};
+
+static FPackageToolsCommands PackageCommands;
+
 UPackageTools::UPackageTools(const FObjectInitializer& ObjectInitializer)
 	: Super(ObjectInitializer)
 {
