Guard against reentry in FEditorViewportClient::StopTracking. This seems possible if during end transaction, a window is shown and takes over the focus from the level viewport.

diff --git a/Engine/Source/Editor/UnrealEd/Private/EditorViewportClient.cpp b/Engine/Source/Editor/UnrealEd/Private/EditorViewportClient.cpp
--- a/Engine/Source/Editor/UnrealEd/Private/EditorViewportClient.cpp
+++ b/Engine/Source/Editor/UnrealEd/Private/EditorViewportClient.cpp
@@ -470,6 +470,7 @@ FEditorViewportClient::FEditorViewportClient(FEditorModeTools* InModeTools, FPre
 	, bShouldInvalidateViewportWidget(false)
 	, DragStartView(nullptr)
 	, DragStartViewFamily(nullptr)
+	, bIsTrackingBeingStopped(false)
 {
 	InitViewOptionsArray();
 	if (!ModeTools)
@@ -3066,8 +3067,10 @@ bool FEditorViewportClient::Internal_InputKey(const FInputKeyEventArgs& EventArg
 
 void FEditorViewportClient::StopTracking()
 {
-	if( bIsTracking )
+	if( bIsTracking && !bIsTrackingBeingStopped )
 	{
+		bIsTrackingBeingStopped = true;
+
 		DragStartView = nullptr;
 		if (DragStartViewFamily != nullptr)
 		{
@@ -3094,6 +3097,7 @@ void FEditorViewportClient::StopTracking()
   		CheckHoveredHitProxy(HitProxy);
 
 		bIsTracking = false;
+		bIsTrackingBeingStopped = false;
 	}
 
 	bHasMouseMovedSinceClick = false;
diff --git a/Engine/Source/Editor/UnrealEd/Public/EditorViewportClient.h b/Engine/Source/Editor/UnrealEd/Public/EditorViewportClient.h
--- a/Engine/Source/Editor/UnrealEd/Public/EditorViewportClient.h
+++ b/Engine/Source/Editor/UnrealEd/Public/EditorViewportClient.h
@@ -2050,6 +2050,11 @@ private:
 
 	/** True if this is the current viewport client in editing */
 	bool bIsCurrentLevelEditingFocus = false;
+
+	/**
+	 * true during FEditorViewportClient::StopTracking. Used to guard against reentry.
+	 */
+	bool bIsTrackingBeingStopped;
 };
 
 
