[UnrealEd]
* Moved dll export from type to methods/statics

diff --git a/Engine/Source/Editor/UnrealEd/Classes/ActorFactories/ActorFactory.h b/Engine/Source/Editor/UnrealEd/Classes/ActorFactories/ActorFactory.h
--- a/Engine/Source/Editor/UnrealEd/Classes/ActorFactories/ActorFactory.h
+++ b/Engine/Source/Editor/UnrealEd/Classes/ActorFactories/ActorFactory.h
@@ -22,8 +22,8 @@ class ULevel;
 class UInstancedPlacemenClientSettings;
 class AVolume;
 class UBrushBuilder;
-UCLASS(collapsecategories, hidecategories=Object, editinlinenew, config=Editor, abstract, transient)
-class UNREALED_API UActorFactory : public UObject, public IAssetFactoryInterface
+UCLASS(collapsecategories, hidecategories=Object, editinlinenew, config=Editor, abstract, transient, MinimalAPI)
+class UActorFactory : public UObject, public IAssetFactoryInterface
 {
 	GENERATED_UCLASS_BODY()
 
@@ -59,72 +59,72 @@ class UNREALED_API UActorFactory : public UObject, public IAssetFactoryInterface
 
 	/** Called to actual create an actor with the supplied transform (scale is ignored), using the properties in the ActorFactory */
 	UE_DEPRECATED(5.0, "This function has been deprecated in favor of the other version that takes a FActorSpawnParameters in parameter")
-	AActor* CreateActor(UObject* Asset, ULevel* InLevel, const FTransform& Transform, EObjectFlags InObjectFlags, const FName InName = NAME_None);
-	AActor* CreateActor(UObject* InAsset, ULevel* InLevel, const FTransform& InTransform, const FActorSpawnParameters& InSpawnParams = FActorSpawnParameters());
+	UNREALED_API AActor* CreateActor(UObject* Asset, ULevel* InLevel, const FTransform& Transform, EObjectFlags InObjectFlags, const FName InName = NAME_None);
+	UNREALED_API AActor* CreateActor(UObject* InAsset, ULevel* InLevel, const FTransform& InTransform, const FActorSpawnParameters& InSpawnParams = FActorSpawnParameters());
 
 	/** Called to create a blueprint class that can be used to spawn an actor from this factory */
 	UE_DEPRECATED(5.0, "This function is no longer used. See FKismetEditorUtilities::CreateBlueprint.")
-	UBlueprint* CreateBlueprint( UObject* Instance, UObject* Outer, const FName Name, const FName CallingContext = NAME_None );
+	UNREALED_API UBlueprint* CreateBlueprint( UObject* Instance, UObject* Outer, const FName Name, const FName CallingContext = NAME_None );
 
-	virtual bool CanCreateActorFrom( const FAssetData& AssetData, FText& OutErrorMsg );
+	UNREALED_API virtual bool CanCreateActorFrom( const FAssetData& AssetData, FText& OutErrorMsg );
 
 	/** Name to put on context menu. */
 	FText GetDisplayName() const { return DisplayName; }
 
 	/** Initialize NewActorClass if necessary, and return default actor for that class. */
-	virtual AActor* GetDefaultActor( const FAssetData& AssetData );
+	UNREALED_API virtual AActor* GetDefaultActor( const FAssetData& AssetData );
 
 	/** Initialize NewActorClass if necessary, and return that class. */
-	virtual UClass* GetDefaultActorClass( const FAssetData& AssetData );
+	UNREALED_API virtual UClass* GetDefaultActorClass( const FAssetData& AssetData );
 
 	/** Given an instance of an actor, find the wrapped asset object which can be used to create a valid FAssetData.
 	 *  Returns nullptr if the given ActorInstance is not valid for this factory.
 	 *  Override this function if the factory actor is a different class than the asset data's class which this factory operates on.
 	 *  For example, if this is the static mesh actor factory, the class of the asset data is UStaticMesh, but the actor factory's class is AStaticMeshActor
 	 */
-	virtual UObject* GetAssetFromActorInstance(AActor* ActorInstance);
+	UNREALED_API virtual UObject* GetAssetFromActorInstance(AActor* ActorInstance);
 
 	/** Return a quaternion which aligns this actor type to the specified surface normal */
-	virtual FQuat AlignObjectToSurfaceNormal(const FVector& InSurfaceNormal, const FQuat& ActorRotation = FQuat::Identity) const;
+	UNREALED_API virtual FQuat AlignObjectToSurfaceNormal(const FVector& InSurfaceNormal, const FQuat& ActorRotation = FQuat::Identity) const;
 
 	// Begin IAssetFactoryInterface Interface
-	virtual bool CanPlaceElementsFromAssetData(const FAssetData& InAssetData) override;
-	virtual bool PrePlaceAsset(FAssetPlacementInfo& InPlacementInfo, const FPlacementOptions& InPlacementOptions) override;
-	virtual TArray<FTypedElementHandle> PlaceAsset(const FAssetPlacementInfo& InPlacementInfo, const FPlacementOptions& InPlacementOptions) override;
-	virtual void PostPlaceAsset(TArrayView<const FTypedElementHandle> InHandle, const FAssetPlacementInfo& InPlacementInfo, const FPlacementOptions& InPlacementOptions) override;
-	virtual FAssetData GetAssetDataFromElementHandle(const FTypedElementHandle& InHandle) override;
-	virtual void BeginPlacement(const FPlacementOptions& InPlacementOptions) override;
-	virtual void EndPlacement(TArrayView<const FTypedElementHandle> InPlacedElements, const FPlacementOptions& InPlacementOptions) override;
-	virtual UInstancedPlacemenClientSettings* FactorySettingsObjectForPlacement(const FAssetData& InAssetData, const FPlacementOptions& InPlacementOptions) override;
+	UNREALED_API virtual bool CanPlaceElementsFromAssetData(const FAssetData& InAssetData) override;
+	UNREALED_API virtual bool PrePlaceAsset(FAssetPlacementInfo& InPlacementInfo, const FPlacementOptions& InPlacementOptions) override;
+	UNREALED_API virtual TArray<FTypedElementHandle> PlaceAsset(const FAssetPlacementInfo& InPlacementInfo, const FPlacementOptions& InPlacementOptions) override;
+	UNREALED_API virtual void PostPlaceAsset(TArrayView<const FTypedElementHandle> InHandle, const FAssetPlacementInfo& InPlacementInfo, const FPlacementOptions& InPlacementOptions) override;
+	UNREALED_API virtual FAssetData GetAssetDataFromElementHandle(const FTypedElementHandle& InHandle) override;
+	UNREALED_API virtual void BeginPlacement(const FPlacementOptions& InPlacementOptions) override;
+	UNREALED_API virtual void EndPlacement(TArrayView<const FTypedElementHandle> InPlacedElements, const FPlacementOptions& InPlacementOptions) override;
+	UNREALED_API virtual UInstancedPlacemenClientSettings* FactorySettingsObjectForPlacement(const FAssetData& InAssetData, const FPlacementOptions& InPlacementOptions) override;
 	// End IAssetFactoryInterface Interface
 
-	static void CreateBrushForVolumeActor(AVolume* NewActor, UBrushBuilder* BrushBuilder);
+	static UNREALED_API void CreateBrushForVolumeActor(AVolume* NewActor, UBrushBuilder* BrushBuilder);
 
 protected:
 
 	/** Validates the input params for SpawnActor and returns the appropriate level to use depending on whether InLevel and/or InSpawnParams.OverrideLevel is passed : */
-	ULevel* ValidateSpawnActorLevel(ULevel* InLevel, const FActorSpawnParameters& InSpawnParams) const;
+	UNREALED_API ULevel* ValidateSpawnActorLevel(ULevel* InLevel, const FActorSpawnParameters& InSpawnParams) const;
 
-	virtual bool PreSpawnActor( UObject* Asset, FTransform& InOutLocation);
+	UNREALED_API virtual bool PreSpawnActor( UObject* Asset, FTransform& InOutLocation);
 
 	UE_DEPRECATED(5.0, "This function has been deprecated in favor of the other version that takes a FActorSpawnParameters in parameter")
-	virtual AActor* SpawnActor(UObject* Asset, ULevel* InLevel, const FTransform& Transform, EObjectFlags ObjectFlags, const FName Name ) final;
-	virtual AActor* SpawnActor(UObject* InAsset, ULevel* InLevel, const FTransform& InTransform, const FActorSpawnParameters& InSpawnParams);
+	UNREALED_API virtual AActor* SpawnActor(UObject* Asset, ULevel* InLevel, const FTransform& Transform, EObjectFlags ObjectFlags, const FName Name ) final;
+	UNREALED_API virtual AActor* SpawnActor(UObject* InAsset, ULevel* InLevel, const FTransform& InTransform, const FActorSpawnParameters& InSpawnParams);
 
 	/** Subclasses may implement this to modify the actor after it has been spawned 
 	    IMPORTANT: If you override this, you should usually also override PostCreateBlueprint()! */
-	virtual void PostSpawnActor( UObject* Asset, AActor* NewActor );
+	UNREALED_API virtual void PostSpawnActor( UObject* Asset, AActor* NewActor );
 
 	/** Override this in derived factory classes if needed.  This is called after a blueprint is created by this factory to
 	    update the blueprint's CDO properties with state from the asset for this factory.
 		IMPORTANT: If you override this, you should usually also override PostSpawnActor()! */
-	virtual void PostCreateBlueprint( UObject* Asset, AActor* CDO );
+	UNREALED_API virtual void PostCreateBlueprint( UObject* Asset, AActor* CDO );
 
 	/**
 	 * Get the default label that should be used for the actor spawned by the given asset (does not include any numeric suffix).
 	 * For classes or BPs that inherit from AActor this will defer to AActor::GetDefaultActorLabel, and for everything else it will use the asset name.
 	 */
-	virtual FString GetDefaultActorLabel(UObject* Asset) const;
+	UNREALED_API virtual FString GetDefaultActorLabel(UObject* Asset) const;
 };
 
-extern UNREALED_API FQuat FindActorAlignmentRotation(const FQuat& InActorRotation, const FVector& InModelAxis, const FVector& InWorldNormal, FQuat* OutDeltaRotation = nullptr);
\ No newline at end of file
+extern UNREALED_API FQuat FindActorAlignmentRotation(const FQuat& InActorRotation, const FVector& InModelAxis, const FVector& InWorldNormal, FQuat* OutDeltaRotation = nullptr);
diff --git a/Engine/Source/Editor/UnrealEd/Classes/ActorFactories/ActorFactoryBlueprint.h b/Engine/Source/Editor/UnrealEd/Classes/ActorFactories/ActorFactoryBlueprint.h
--- a/Engine/Source/Editor/UnrealEd/Classes/ActorFactories/ActorFactoryBlueprint.h
+++ b/Engine/Source/Editor/UnrealEd/Classes/ActorFactories/ActorFactoryBlueprint.h
@@ -11,16 +11,16 @@
 class AActor;
 struct FAssetData;
 
-UCLASS(config=Editor, collapsecategories, hidecategories=Object)
-class UNREALED_API UActorFactoryBlueprint : public UActorFactory
+UCLASS(config=Editor, collapsecategories, hidecategories=Object, MinimalAPI)
+class UActorFactoryBlueprint : public UActorFactory
 {
 	GENERATED_UCLASS_BODY()
 
 	//~ Begin UActorFactory Interface
-	virtual bool PreSpawnActor( UObject* Asset, FTransform& InOutLocation ) override;
-	virtual bool CanCreateActorFrom( const FAssetData& AssetData, FText& OutErrorMsg ) override;
-	virtual AActor* GetDefaultActor( const FAssetData& AssetData ) override;
-	virtual UClass* GetDefaultActorClass( const FAssetData& AssetData ) override;
+	UNREALED_API virtual bool PreSpawnActor( UObject* Asset, FTransform& InOutLocation ) override;
+	UNREALED_API virtual bool CanCreateActorFrom( const FAssetData& AssetData, FText& OutErrorMsg ) override;
+	UNREALED_API virtual AActor* GetDefaultActor( const FAssetData& AssetData ) override;
+	UNREALED_API virtual UClass* GetDefaultActorClass( const FAssetData& AssetData ) override;
 	//~ End UActorFactory Interface	
 };
 
diff --git a/Engine/Source/Editor/UnrealEd/Classes/Animation/DebugSkelMeshComponent.h b/Engine/Source/Editor/UnrealEd/Classes/Animation/DebugSkelMeshComponent.h
--- a/Engine/Source/Editor/UnrealEd/Classes/Animation/DebugSkelMeshComponent.h
+++ b/Engine/Source/Editor/UnrealEd/Classes/Animation/DebugSkelMeshComponent.h
@@ -161,8 +161,8 @@ enum class ESkeletonDrawMode : uint8
 	GreyedOut
 };
 
-UCLASS(transient)
-class UNREALED_API UDebugSkelMeshComponent : public USkeletalMeshComponent
+UCLASS(transient, MinimalAPI)
+class UDebugSkelMeshComponent : public USkeletalMeshComponent
 {
 	GENERATED_UCLASS_BODY()
 	
@@ -339,102 +339,102 @@ class UNREALED_API UDebugSkelMeshComponent : public USkeletalMeshComponent
 	bool bTrackAttachedInstanceLOD;
 
 	// Helper method that sets the forced lod
-	void SetDebugForcedLOD(int32 InNewForcedLOD);
+	UNREALED_API void SetDebugForcedLOD(int32 InNewForcedLOD);
 
 	//~ Begin USceneComponent Interface.
-	virtual FBoxSphereBounds CalcBounds(const FTransform& LocalToWorld) const override;
+	UNREALED_API virtual FBoxSphereBounds CalcBounds(const FTransform& LocalToWorld) const override;
 	//~ End USceneComponent Interface.
 
 	//~ Begin UPrimitiveComponent Interface.
-	virtual FPrimitiveSceneProxy* CreateSceneProxy() override;
+	UNREALED_API virtual FPrimitiveSceneProxy* CreateSceneProxy() override;
 	
 	// engine only draw bounds IF selected
 	// @todo fix this properly
 	// this isn't really the best way to do this, but for now
 	// we'll just mark as selected
-	virtual bool ShouldRenderSelected() const override;
+	UNREALED_API virtual bool ShouldRenderSelected() const override;
 	//~ End UPrimitiveComponent Interface.
 
 	//~ Begin SkinnedMeshComponent Interface
-	virtual bool ShouldCPUSkin() override;
-	virtual void PostInitMeshObject(class FSkeletalMeshObject* MeshObject) override;
-	virtual void RefreshBoneTransforms(FActorComponentTickFunction* TickFunction = NULL) override;
+	UNREALED_API virtual bool ShouldCPUSkin() override;
+	UNREALED_API virtual void PostInitMeshObject(class FSkeletalMeshObject* MeshObject) override;
+	UNREALED_API virtual void RefreshBoneTransforms(FActorComponentTickFunction* TickFunction = NULL) override;
 	virtual int32 GetLODBias() const override { return 0; }
 	//~ End SkinnedMeshComponent Interface
 
 	//~ Begin SkeletalMeshComponent Interface
-	virtual void InitAnim(bool bForceReinit) override;
+	UNREALED_API virtual void InitAnim(bool bForceReinit) override;
 	virtual bool IsWindEnabled() const override { return true; }
-	virtual void SetAnimClass(class UClass* NewClass) override;
-	virtual void OnClearAnimScriptInstance() override;
-	virtual void SetSkeletalMesh(USkeletalMesh* InSkelMesh, bool bReinitPose = true) override;
+	UNREALED_API virtual void SetAnimClass(class UClass* NewClass) override;
+	UNREALED_API virtual void OnClearAnimScriptInstance() override;
+	UNREALED_API virtual void SetSkeletalMesh(USkeletalMesh* InSkelMesh, bool bReinitPose = true) override;
 	//~ End SkeletalMeshComponent Interface
 
 	// return true if currently preview animation asset is on
-	virtual bool IsPreviewOn() const;
+	UNREALED_API virtual bool IsPreviewOn() const;
 
 	// @todo document
-	FString GetPreviewText() const;
+	UNREALED_API FString GetPreviewText() const;
 
 	// @todo anim : you still need to give asset, so that we know which one to disable
 	// we can disable per asset, so that if some other window disabled before me, I don't accidently turn it off
-	virtual void EnablePreview(bool bEnable, class UAnimationAsset * PreviewAsset);
+	UNREALED_API virtual void EnablePreview(bool bEnable, class UAnimationAsset * PreviewAsset);
 
 	// Create the preview instance to use (default UAnimPreviewInstance)
-	virtual TObjectPtr<UAnimPreviewInstance> CreatePreviewInstance();
+	UNREALED_API virtual TObjectPtr<UAnimPreviewInstance> CreatePreviewInstance();
 
 	// reference pose for this component
 	// we don't want to use default refpose because you still want to move joint when this mode is on
-	virtual void ShowReferencePose(bool bRefPose);
-	virtual bool IsReferencePoseShown() const;
+	UNREALED_API virtual void ShowReferencePose(bool bRefPose);
+	UNREALED_API virtual bool IsReferencePoseShown() const;
 
 	/**
 	 * Update material information depending on color render mode 
 	 * Refresh/replace materials 
 	 */
-	void SetShowBoneWeight(bool bNewShowBoneWeight);
+	UNREALED_API void SetShowBoneWeight(bool bNewShowBoneWeight);
 
 	/**
 	* Update material information depending on color render mode
 	* Refresh/replace materials
 	*/
-	void SetShowMorphTargetVerts(bool bNewShowMorphTargetVerts);
+	UNREALED_API void SetShowMorphTargetVerts(bool bNewShowMorphTargetVerts);
 
 	/**
 	 * Does it use in-game bounds or bounds calculated from bones
 	 */
-	bool IsUsingInGameBounds() const;
+	UNREALED_API bool IsUsingInGameBounds() const;
 
 	/**
 	 * Set to use in-game bounds or bounds calculated from bones
 	 */
-	void UseInGameBounds(bool bUseInGameBounds);
+	UNREALED_API void UseInGameBounds(bool bUseInGameBounds);
 
 	/**
 	 * Does it use pre-skinned bounds
 	 */
-	bool IsUsingPreSkinnedBounds() const;
+	UNREALED_API bool IsUsingPreSkinnedBounds() const;
 
 	/**
 	 * Set to use pre-skinned bounds
 	 */
-	void UsePreSkinnedBounds(bool bUsePreSkinnedBounds);
+	UNREALED_API void UsePreSkinnedBounds(bool bUsePreSkinnedBounds);
 
 	/**
 	 * Test if in-game bounds are as big as preview bounds
 	 */
-	bool CheckIfBoundsAreCorrrect();
+	UNREALED_API bool CheckIfBoundsAreCorrrect();
 
 	/** Get the in-game bounds of the skeleton mesh */
-	FBoxSphereBounds CalcGameBounds(const FTransform& LocalToWorld) const;
+	UNREALED_API FBoxSphereBounds CalcGameBounds(const FTransform& LocalToWorld) const;
 
 	/** 
 	 * Update components position based on animation root motion
 	 */
-	void ConsumeRootMotion(const FVector& FloorMin, const FVector& FloorMax);
+	UNREALED_API void ConsumeRootMotion(const FVector& FloorMin, const FVector& FloorMax);
 
 	/** Sets the flag used to determine whether or not the current active cloth sim mesh should be rendered */
-	void SetShowClothProperty(bool bState);
+	UNREALED_API void SetShowClothProperty(bool bState);
 
 	/** Get whether we should be previewing root motion */
 	UE_DEPRECATED(5.0, "Please use IsProcessingRootMotion or GetProcessRootMotionMode")
@@ -445,22 +445,22 @@ class UNREALED_API UDebugSkelMeshComponent : public USkeletalMeshComponent
 	void SetPreviewRootMotion(bool bInPreviewRootMotion) { SetProcessRootMotionMode(bInPreviewRootMotion ? EProcessRootMotionMode::Loop : EProcessRootMotionMode::Ignore); }
 
 	/** Whether we are processing root motion or not */
-	bool IsProcessingRootMotion() const;
+	UNREALED_API bool IsProcessingRootMotion() const;
 
 	/** Gets process root motion mode */
-	EProcessRootMotionMode GetProcessRootMotionMode() const;
+	UNREALED_API EProcessRootMotionMode GetProcessRootMotionMode() const;
 
 	/** Sets process root motion mode. Note: disabling root motion preview resets transform. */
-	void SetProcessRootMotionMode(EProcessRootMotionMode Mode);
+	UNREALED_API void SetProcessRootMotionMode(EProcessRootMotionMode Mode);
 
 	/** Whether the supplied root motion mode can be used for the current asset */
-	bool CanUseProcessRootMotionMode(EProcessRootMotionMode Mode) const;
+	UNREALED_API bool CanUseProcessRootMotionMode(EProcessRootMotionMode Mode) const;
 
 	/** Whether the current asset or animation blueprint is using root motion */
-	bool DoesCurrentAssetHaveRootMotion() const;
+	UNREALED_API bool DoesCurrentAssetHaveRootMotion() const;
 
 	/** Whether the current LOD of the debug mesh is being synced with the attached (preview) mesh instance. */
-	bool IsTrackingAttachedLOD() const;
+	UNREALED_API bool IsTrackingAttachedLOD() const;
 
 	/** Set the wireframe mesh overlay color, which basically controls the color of the wireframe. */
 	void SetWireframeMeshOverlayColor(FLinearColor Color) { WireframeMeshOverlayColor = Color; }
@@ -480,19 +480,19 @@ class UNREALED_API UDebugSkelMeshComponent : public USkeletalMeshComponent
 		TArray<FString> Errors;
 	};
 	TArray<FAnimNotifyErrors> AnimNotifyErrors;
-	virtual void ReportAnimNotifyError(const FText& Error, UObject* InSourceNotify) override;
-	virtual void ClearAnimNotifyErrors(UObject* InSourceNotify) override;
+	UNREALED_API virtual void ReportAnimNotifyError(const FText& Error, UObject* InSourceNotify) override;
+	UNREALED_API virtual void ClearAnimNotifyErrors(UObject* InSourceNotify) override;
 
 	/** 
 	 * Extended viewport text delegate handling. Registering a delegate allows external
 	 * objects to place custom text in the anim tools viewports.
 	 */
-	FDelegateHandle RegisterExtendedViewportTextDelegate(const FGetExtendedViewportText& InDelegate);
-	void UnregisterExtendedViewportTextDelegate(const FDelegateHandle& InDelegateHandle);
+	UNREALED_API FDelegateHandle RegisterExtendedViewportTextDelegate(const FGetExtendedViewportText& InDelegate);
+	UNREALED_API void UnregisterExtendedViewportTextDelegate(const FDelegateHandle& InDelegateHandle);
 	const TArray<FGetExtendedViewportText>& GetExtendedViewportTextDelegates() const { return ExtendedViewportTextDelegates; }
 
-	FDelegateHandle RegisterOnDebugForceLODChangedDelegate(const FOnDebugForceLODChanged& InDelegate);
-	void UnregisterOnDebugForceLODChangedDelegate();
+	UNREALED_API FDelegateHandle RegisterOnDebugForceLODChangedDelegate(const FOnDebugForceLODChanged& InDelegate);
+	UNREALED_API void UnregisterOnDebugForceLODChangedDelegate();
 
 private:
 	TArray<FGetExtendedViewportText> ExtendedViewportTextDelegates;
@@ -521,18 +521,18 @@ public:
 	 * if bShowOnlyClothSections is true, shows only cloth sections. On the other hand, 
 	 * if bShowOnlyClothSections is false, hides only cloth sections.
 	 */
-	void ToggleClothSectionsVisibility(bool bShowOnlyClothSections);
+	UNREALED_API void ToggleClothSectionsVisibility(bool bShowOnlyClothSections);
 	/** Restore all section visibilities to original states for all LODs */
-	void RestoreClothSectionsVisibility();
+	UNREALED_API void RestoreClothSectionsVisibility();
 
 	/** 
 	 * To normal game/runtime code we don't want to expose a non-const pointer to the simulation, so we can only get
 	 * one from this editor-only component. Intended for debug options/visualisations/editor-only code to poke the sim
 	 */
-	IClothingSimulation* GetMutableClothingSimulation();
+	UNREALED_API IClothingSimulation* GetMutableClothingSimulation();
 
 	/** to avoid clothing reset while modifying properties in Persona */
-	virtual void CheckClothTeleport() override;
+	UNREALED_API virtual void CheckClothTeleport() override;
 
 	/** The currently selected asset guid if we're painting, used to build dynamic mesh to paint sim parameters */
 	FGuid SelectedClothingGuidForPainting;
@@ -544,17 +544,17 @@ public:
 	int32 SelectedClothingLodMaskForPainting;
 
 	/** Find a section using a clothing asset with the given GUID and set its visiblity */
-	void SetMeshSectionVisibilityForCloth(FGuid InClothGuid, bool bVisibility);
+	UNREALED_API void SetMeshSectionVisibilityForCloth(FGuid InClothGuid, bool bVisibility);
 
 	// fixes up the disabled flags so clothing is enabled and originals are disabled as
 	// ToggleMeshSectionForCloth will make these get out of sync
-	void ResetMeshSectionVisibility();
+	UNREALED_API void ResetMeshSectionVisibility();
 
 	// Rebuilds the fixed parameter on the mesh to mesh data, to be used if the editor has
 	// changed a vert to be fixed or unfixed otherwise the simulation will not work
 	// bInvalidateDerivedDataCache can only be false during previewing as otherwise the changes won't be correctly saved
 	UE_DEPRECATED(5.0, "This function is redundant, since it is always called after ApplyParameterMasks and therefore will be removed.")
-	void RebuildClothingSectionsFixedVerts(bool bInvalidateDerivedDataCache = true);
+	UNREALED_API void RebuildClothingSectionsFixedVerts(bool bInvalidateDerivedDataCache = true);
 
 	TArray<FVector3f> SkinnedSelectedClothingPositions;
 	TArray<FVector3f> SkinnedSelectedClothingNormals;
@@ -563,22 +563,22 @@ private:
 	// Rebuilds the fixed vertex attribute on any cloth deformer mappings,
 	// including LOD bias mappings, that reference the specified LOD section.
 	UE_DEPRECATED(5.0, "This function is redundant, since it is always called after ApplyParameterMasks and therefore will be removed.")
-	void RebuildClothingSectionFixedVerts(int32 LODIndex, int32 SectionIndex);
+	UNREALED_API void RebuildClothingSectionFixedVerts(int32 LODIndex, int32 SectionIndex);
 
 	// Helper function to generate space bases for current frame
-	void GenSpaceBases(TArray<FTransform>& OutSpaceBases);
+	UNREALED_API void GenSpaceBases(TArray<FTransform>& OutSpaceBases);
 
 	// Helper function to enable overlay material
-	void EnableOverlayMaterial(bool bEnable);
+	UNREALED_API void EnableOverlayMaterial(bool bEnable);
 
 	// Rebuilds the cloth bounds for the asset.
-	void RebuildCachedClothBounds();
+	UNREALED_API void RebuildCachedClothBounds();
 protected:
 
 	// Overridden to support single clothing ticks
-	virtual bool ShouldRunClothTick() const override;
+	UNREALED_API virtual bool ShouldRunClothTick() const override;
 
-	virtual void SendRenderDynamicData_Concurrent() override;
+	UNREALED_API virtual void SendRenderDynamicData_Concurrent() override;
 
 public:
 	/** Current turn table mode */
@@ -586,13 +586,13 @@ public:
 	/** Current turn table speed scaling */
 	float TurnTableSpeedScaling;
 
-	virtual void TickComponent(float DeltaTime, enum ELevelTick TickType, FActorComponentTickFunction *ThisTickFunction) override;
+	UNREALED_API virtual void TickComponent(float DeltaTime, enum ELevelTick TickType, FActorComponentTickFunction *ThisTickFunction) override;
 
-	void RefreshSelectedClothingSkinnedPositions();
+	UNREALED_API void RefreshSelectedClothingSkinnedPositions();
 
 	virtual bool CanOverrideCollisionProfile() const { return true; }
 
-	virtual void GetUsedMaterials(TArray<UMaterialInterface *>& OutMaterials, bool bGetDebugMaterials = false) const override;
+	UNREALED_API virtual void GetUsedMaterials(TArray<UMaterialInterface *>& OutMaterials, bool bGetDebugMaterials = false) const override;
 
 	/**
 	 * Define Custom Default pose for this component for preview
@@ -647,7 +647,7 @@ public:
  * Important it should be destroy after the PostEditChange of the skeletalmesh is done and the renderdata have been recreate
  * i.e. FScopedSkeletalMeshPostEditChange should be create after FScopedSuspendAlternateSkinWeightPreview and delete before FScopedSuspendAlternateSkinWeightPreview
  */
-class UNREALED_API FScopedSuspendAlternateSkinWeightPreview
+class FScopedSuspendAlternateSkinWeightPreview
 {
 public:
 	/*
@@ -655,13 +655,13 @@ public:
 	 * Parameters:
 	 * @param InSkeletalMesh - SkeletalMesh use to know which preview component we have to suspend the alternate skinning preview.
 	 */
-	FScopedSuspendAlternateSkinWeightPreview(class USkeletalMesh* InSkeletalMesh);
+	UNREALED_API FScopedSuspendAlternateSkinWeightPreview(class USkeletalMesh* InSkeletalMesh);
 
 	/*
 	 * This destructor put back the preview alternate skinning
 	 */
-	~FScopedSuspendAlternateSkinWeightPreview();
+	UNREALED_API ~FScopedSuspendAlternateSkinWeightPreview();
 
 private:
 	TArray< TTuple<UDebugSkelMeshComponent*, FName> > SuspendedComponentArray;
-};
\ No newline at end of file
+};
diff --git a/Engine/Source/Editor/UnrealEd/Classes/Animation/EditorParentPlayerListObj.h b/Engine/Source/Editor/UnrealEd/Classes/Animation/EditorParentPlayerListObj.h
--- a/Engine/Source/Editor/UnrealEd/Classes/Animation/EditorParentPlayerListObj.h
+++ b/Engine/Source/Editor/UnrealEd/Classes/Animation/EditorParentPlayerListObj.h
@@ -11,25 +11,25 @@
 
 class UAnimGraphNode_Base;
 
-UCLASS()
-class UNREALED_API UEditorParentPlayerListObj : public UObject
+UCLASS(MinimalAPI)
+class UEditorParentPlayerListObj : public UObject
 {
 	GENERATED_UCLASS_BODY()
 public:
 	// Build a possible list of overrides from a blueprint and pull in any existing overrides.
-	void InitialiseFromBlueprint(UAnimBlueprint* Blueprint);
+	UNREALED_API void InitialiseFromBlueprint(UAnimBlueprint* Blueprint);
 
 	// Adds an overridable node to the possible override list.
-	FAnimParentNodeAssetOverride& AddOverridableNode(UAnimGraphNode_Base* Node);
+	UNREALED_API FAnimParentNodeAssetOverride& AddOverridableNode(UAnimGraphNode_Base* Node);
 
 	// Gets a pointer to a U-node from the stored GUID for that node.
-	UAnimGraphNode_Base* GetVisualNodeFromGuid(FGuid InGuid) const;
+	UNREALED_API UAnimGraphNode_Base* GetVisualNodeFromGuid(FGuid InGuid) const;
 
 	// Pushes an override into the blueprint we're editing.
-	void ApplyOverrideToBlueprint(FAnimParentNodeAssetOverride& Override);
+	UNREALED_API void ApplyOverrideToBlueprint(FAnimParentNodeAssetOverride& Override);
 
 	// Get the current blueprint
-	const UAnimBlueprint* GetBlueprint() const;
+	UNREALED_API const UAnimBlueprint* GetBlueprint() const;
 
 	// List of possible overrides to display.
 	UPROPERTY(Category = AnimationGraphOverrides, EditAnywhere)
diff --git a/Engine/Source/Editor/UnrealEd/Classes/Commandlets/BaseIteratePackagesCommandlet.h b/Engine/Source/Editor/UnrealEd/Classes/Commandlets/BaseIteratePackagesCommandlet.h
--- a/Engine/Source/Editor/UnrealEd/Classes/Commandlets/BaseIteratePackagesCommandlet.h
+++ b/Engine/Source/Editor/UnrealEd/Classes/Commandlets/BaseIteratePackagesCommandlet.h
@@ -23,11 +23,11 @@ DECLARE_DELEGATE_TwoParams(FBaseIteratePackagesCommandletActorCallback, AActor*
 DECLARE_DELEGATE_TwoParams(FBaseIteratePackagesCommandletObjectCallback, UObject* /*Actor*/, bool& /*bSavePackage*/);
 
 
-UCLASS()
+UCLASS(MinimalAPI)
 // Added UNREALED_API to expose this to the save packages test
 // Base commandlet used to iterate packages provided on a commandline
 // has callbacks for processing each package
-class UNREALED_API UBaseIteratePackagesCommandlet : public UCommandlet
+class UBaseIteratePackagesCommandlet : public UCommandlet
 {
     GENERATED_UCLASS_BODY()
 
@@ -187,22 +187,22 @@ protected:
 	 *
 	 * @return	0 to indicate that the commandlet should continue; otherwise, the error code that should be returned by Main()
 	 */
-	virtual int32 InitializeParameters( const TArray<FString>& Tokens, TArray<FString>& MapPathNames );
+	UNREALED_API virtual int32 InitializeParameters( const TArray<FString>& Tokens, TArray<FString>& MapPathNames );
 
 	virtual void InitializePackageNames(const TArray<FString>& Tokens, TArray<FString>& MapPathNames, bool& bExplicitPackages) { }
 
-	void ParseSourceControlOptions(const TArray<FString>& Tokens);
+	UNREALED_API void ParseSourceControlOptions(const TArray<FString>& Tokens);
 
-	void OnAddResaveOnDemandPackage(FName SystemName, FName PackageName);
+	UNREALED_API void OnAddResaveOnDemandPackage(FName SystemName, FName PackageName);
 
 	/** Loads and saves a single package */
-	virtual void LoadAndSaveOnePackage(const FString& Filename);
+	UNREALED_API virtual void LoadAndSaveOnePackage(const FString& Filename);
 
 	/** Checks to see if a package should be skipped */
-	virtual bool ShouldSkipPackage(const FString& Filename);
+	UNREALED_API virtual bool ShouldSkipPackage(const FString& Filename);
 
 	/** Deletes a single package */
-	virtual void DeleteOnePackage(const FString& Filename);
+	UNREALED_API virtual void DeleteOnePackage(const FString& Filename);
 
 	/**
 	 * Allow the commandlet to perform any operations on the export/import table of the package before all objects in the package are loaded.
@@ -211,7 +211,7 @@ protected:
 	 * @param	bSavePackage	[in]	indicates whether the package is currently going to be saved
 	 *							[out]	set to true to resave the package
 	 */
-	virtual void PerformPreloadOperations( FLinkerLoad* PackageLinker, bool& bSavePackage );
+	UNREALED_API virtual void PerformPreloadOperations( FLinkerLoad* PackageLinker, bool& bSavePackage );
 
 	/**
 	 * Allows the commandlet to perform any additional operations on the object before it is resaved.
@@ -261,23 +261,23 @@ protected:
 	 *
 	 * @return	true to resave the package
 	 */
-	bool CleanClassesFromContentPackages( class UPackage* Package );
+	UNREALED_API bool CleanClassesFromContentPackages( class UPackage* Package );
 
 	// Get the changelist description to use if automatically checking packages out
-	virtual FText GetChangelistDescription() const;
+	UNREALED_API virtual FText GetChangelistDescription() const;
 
-	bool CheckoutFile(const FString& Filename, bool bAddFile = false, bool bIgnoreAlreadyCheckedOut = false);
-	bool RevertFile(const FString& Filename);
-	bool CanCheckoutFile(const FString& Filename, FString& CheckedOutUser);
-	void CheckoutAndSavePackage(UPackage* Package, TArray<FString>& SublevelFilenames, bool bIgnoreAlreadyCheckedOut = false);
-	void CheckInFiles(const TArray<FString>& InFilesToSubmit, const FText& InDescription) const;
+	UNREALED_API bool CheckoutFile(const FString& Filename, bool bAddFile = false, bool bIgnoreAlreadyCheckedOut = false);
+	UNREALED_API bool RevertFile(const FString& Filename);
+	UNREALED_API bool CanCheckoutFile(const FString& Filename, FString& CheckedOutUser);
+	UNREALED_API void CheckoutAndSavePackage(UPackage* Package, TArray<FString>& SublevelFilenames, bool bIgnoreAlreadyCheckedOut = false);
+	UNREALED_API void CheckInFiles(const TArray<FString>& InFilesToSubmit, const FText& InDescription) const;
 
 	/**
 	 * Creates and returns a unique filename in the temporary file directory.
 	 *
 	 * @return	A unique temporary filename
 	 */
-	FString CreateTempFilename();
+	UNREALED_API FString CreateTempFilename();
 
 	/**
 	 * Returns the path of the directory we use to store all temporary files for
@@ -285,25 +285,25 @@ protected:
 	 *
 	 * @return Path to the temporary file directory
 	 */
-	FString GetTempFilesDirectory();
+	UNREALED_API FString GetTempFilesDirectory();
 
 	/**
 	 * Delete the remaining files in the temporary files directory for the commandlet.
 	 */
-	void CleanTempFiles();	
+	UNREALED_API void CleanTempFiles();	
 
 	// Print out a message only if running in very verbose mode
-	void VerboseMessage(const FString& Message);
+	UNREALED_API void VerboseMessage(const FString& Message);
 
 	/** Parse commandline to decide whether resaveondemand is activated for the given system. */
-	TSet<FName> ParseResaveOnDemandSystems();
+	UNREALED_API TSet<FName> ParseResaveOnDemandSystems();
 
 private:
 
-	void SavePackages(const TArray<UPackage*>& PackagesToSave);
+	UNREALED_API void SavePackages(const TArray<UPackage*>& PackagesToSave);
 
 public:		
 	//~ Begin UCommandlet Interface
-	virtual int32 Main(const FString& Params) override;
+	UNREALED_API virtual int32 Main(const FString& Params) override;
 	//~ End UCommandlet Interface
 };
diff --git a/Engine/Source/Editor/UnrealEd/Classes/Commandlets/ConvertLevelsToExternalActorsCommandlet.h b/Engine/Source/Editor/UnrealEd/Classes/Commandlets/ConvertLevelsToExternalActorsCommandlet.h
--- a/Engine/Source/Editor/UnrealEd/Classes/Commandlets/ConvertLevelsToExternalActorsCommandlet.h
+++ b/Engine/Source/Editor/UnrealEd/Classes/Commandlets/ConvertLevelsToExternalActorsCommandlet.h
@@ -8,28 +8,28 @@
 #include "ISourceControlProvider.h"
 #include "ConvertLevelsToExternalActorsCommandlet.generated.h"
 
-UCLASS()
-class UNREALED_API UConvertLevelsToExternalActorsCommandlet : public UCommandlet
+UCLASS(MinimalAPI)
+class UConvertLevelsToExternalActorsCommandlet : public UCommandlet
 {
 	GENERATED_UCLASS_BODY()
 
 public:
 	//~ Begin UCommandlet Interface
-	virtual int32 Main(const FString& Params) override;
+	UNREALED_API virtual int32 Main(const FString& Params) override;
 	//~ End UCommandlet Interface
 
 protected:
-	ULevel* LoadLevel(const FString& LevelToLoad) const;
-	void GetSubLevelsToConvert(ULevel* MainLevel, TSet<ULevel*>& SubLevels, bool bRecursive);
+	UNREALED_API ULevel* LoadLevel(const FString& LevelToLoad) const;
+	UNREALED_API void GetSubLevelsToConvert(ULevel* MainLevel, TSet<ULevel*>& SubLevels, bool bRecursive);
 
-	bool CheckExternalActors(const FString& Level, bool bRepair);
+	UNREALED_API bool CheckExternalActors(const FString& Level, bool bRepair);
 
 	bool UseSourceControl() const { return SourceControlProvider != nullptr; }
 	ISourceControlProvider& GetSourceControlProvider() { check(UseSourceControl()); return *SourceControlProvider; }
-	bool AddPackageToSourceControl(UPackage* Package);
-	bool CheckoutPackage(UPackage* Package);
-	bool SavePackage(UPackage* Package);
-	bool DeleteFile(const FString& Filename);
+	UNREALED_API bool AddPackageToSourceControl(UPackage* Package);
+	UNREALED_API bool CheckoutPackage(UPackage* Package);
+	UNREALED_API bool SavePackage(UPackage* Package);
+	UNREALED_API bool DeleteFile(const FString& Filename);
 
 protected:
 	ISourceControlProvider* SourceControlProvider;
diff --git a/Engine/Source/Editor/UnrealEd/Classes/Commandlets/ExternalActorsCommandlet.h b/Engine/Source/Editor/UnrealEd/Classes/Commandlets/ExternalActorsCommandlet.h
--- a/Engine/Source/Editor/UnrealEd/Classes/Commandlets/ExternalActorsCommandlet.h
+++ b/Engine/Source/Editor/UnrealEd/Classes/Commandlets/ExternalActorsCommandlet.h
@@ -8,16 +8,16 @@
 #include "PackageSourceControlHelper.h"
 #include "ExternalActorsCommandlet.generated.h"
 
-UCLASS()
-class UNREALED_API UExternalActorsCommandlet : public UCommandlet
+UCLASS(MinimalAPI)
+class UExternalActorsCommandlet : public UCommandlet
 {
 	GENERATED_UCLASS_BODY()
 
 public:
 	//~ Begin UCommandlet Interface
-	virtual int32 Main(const FString& Params) override;
+	UNREALED_API virtual int32 Main(const FString& Params) override;
 	//~ End UCommandlet Interface
 
 protected:
-	UWorld* LoadWorld(const FString& LevelToLoad);
+	UNREALED_API UWorld* LoadWorld(const FString& LevelToLoad);
 };
diff --git a/Engine/Source/Editor/UnrealEd/Classes/Commandlets/GatherTextCommandletBase.h b/Engine/Source/Editor/UnrealEd/Classes/Commandlets/GatherTextCommandletBase.h
--- a/Engine/Source/Editor/UnrealEd/Classes/Commandlets/GatherTextCommandletBase.h
+++ b/Engine/Source/Editor/UnrealEd/Classes/Commandlets/GatherTextCommandletBase.h
@@ -10,15 +10,15 @@
 #include "LocalizedAssetUtil.h"
 #include "GatherTextCommandletBase.generated.h"
 
-struct UNREALED_API FGatherTextDelegates
+struct FGatherTextDelegates
 {
 	/** Delegate called during a localization gather to allow code to inject new gather and exclude paths for the given localization target */
 	DECLARE_MULTICAST_DELEGATE_ThreeParams(FGetAdditionalGatherPaths, const FString& /*InLocalizationTargetName*/, TArray<FString>& /*InOutIncludePathFilters*/, TArray<FString>& /*InOutExcludePathFilters*/);
-	static FGetAdditionalGatherPaths GetAdditionalGatherPaths;
+	static UNREALED_API FGetAdditionalGatherPaths GetAdditionalGatherPaths;
 };
 
 /** Performs fuzzy path matching against a set of include and exclude paths */
-class UNREALED_API FFuzzyPathMatcher
+class FFuzzyPathMatcher
 {
 public:
 	enum class EPathMatch
@@ -29,9 +29,9 @@ public:
 	};
 
 public:
-	FFuzzyPathMatcher(const TArray<FString>& InIncludePathFilters, const TArray<FString>& InExcludePathFilters);
+	UNREALED_API FFuzzyPathMatcher(const TArray<FString>& InIncludePathFilters, const TArray<FString>& InExcludePathFilters);
 
-	EPathMatch TestPath(const FString& InPathToTest) const;
+	UNREALED_API EPathMatch TestPath(const FString& InPathToTest) const;
 
 private:
 	enum class EPathType : uint8
@@ -58,29 +58,29 @@ private:
 /**
  *	UGatherTextCommandletBase: Base class for localization commandlets. Just to force certain behaviors and provide helper functionality. 
  */
-UCLASS()
-class UNREALED_API UGatherTextCommandletBase : public UCommandlet
+UCLASS(MinimalAPI)
+class UGatherTextCommandletBase : public UCommandlet
 {
 	GENERATED_UCLASS_BODY()
 
 public:
-	virtual void Initialize( const TSharedRef< FLocTextHelper >& InGatherManifestHelper, const TSharedPtr< FLocalizationSCC >& InSourceControlInfo );
-	virtual void BeginDestroy() override;
+	UNREALED_API virtual void Initialize( const TSharedRef< FLocTextHelper >& InGatherManifestHelper, const TSharedPtr< FLocalizationSCC >& InSourceControlInfo );
+	UNREALED_API virtual void BeginDestroy() override;
 
 	// Wrappers for extracting config values
-	bool GetBoolFromConfig( const TCHAR* Section, const TCHAR* Key, bool& OutValue, const FString& Filename );
-	bool GetStringFromConfig( const TCHAR* Section, const TCHAR* Key, FString& OutValue, const FString& Filename );
-	bool GetPathFromConfig( const TCHAR* Section, const TCHAR* Key, FString& OutValue, const FString& Filename );
-	int32 GetStringArrayFromConfig( const TCHAR* Section, const TCHAR* Key, TArray<FString>& OutArr, const FString& Filename );
-	int32 GetPathArrayFromConfig( const TCHAR* Section, const TCHAR* Key, TArray<FString>& OutArr, const FString& Filename );
+	UNREALED_API bool GetBoolFromConfig( const TCHAR* Section, const TCHAR* Key, bool& OutValue, const FString& Filename );
+	UNREALED_API bool GetStringFromConfig( const TCHAR* Section, const TCHAR* Key, FString& OutValue, const FString& Filename );
+	UNREALED_API bool GetPathFromConfig( const TCHAR* Section, const TCHAR* Key, FString& OutValue, const FString& Filename );
+	UNREALED_API int32 GetStringArrayFromConfig( const TCHAR* Section, const TCHAR* Key, TArray<FString>& OutArr, const FString& Filename );
+	UNREALED_API int32 GetPathArrayFromConfig( const TCHAR* Section, const TCHAR* Key, TArray<FString>& OutArr, const FString& Filename );
 
 	// Utilities for split platform detection
-	bool IsSplitPlatformName(const FName InPlatformName) const;
-	bool ShouldSplitPlatformForPath(const FString& InPath, FName* OutPlatformName = nullptr) const;
-	FName GetSplitPlatformNameFromPath(const FString& InPath) const;
+	UNREALED_API bool IsSplitPlatformName(const FName InPlatformName) const;
+	UNREALED_API bool ShouldSplitPlatformForPath(const FString& InPath, FName* OutPlatformName = nullptr) const;
+	UNREALED_API FName GetSplitPlatformNameFromPath(const FString& InPath) const;
 
 	// Utility to get the correct base path (engine or project) for the current environment
-	static const FString& GetProjectBasePath();
+	static UNREALED_API const FString& GetProjectBasePath();
 
 	/**
 * Returns true if this commandlet should run during a preview run.
@@ -101,12 +101,12 @@ protected:
 	TMap<FName, FString> SplitPlatforms;
 
 	// Common params and switches among all text gathering commadnlets 
-	static const TCHAR* ConfigParam;
-	static const TCHAR* EnableSourceControlSwitch;
-	static const TCHAR* DisableSubmitSwitch;
-	static const TCHAR* PreviewSwitch;
-	static const TCHAR* GatherTypeParam;
+	static UNREALED_API const TCHAR* ConfigParam;
+	static UNREALED_API const TCHAR* EnableSourceControlSwitch;
+	static UNREALED_API const TCHAR* DisableSubmitSwitch;
+	static UNREALED_API const TCHAR* PreviewSwitch;
+	static UNREALED_API const TCHAR* GatherTypeParam;
 
 private:
-	virtual void CreateCustomEngine(const FString& Params) override ; //Disallow other text commandlets to make their own engine.	
+	UNREALED_API virtual void CreateCustomEngine(const FString& Params) override ; //Disallow other text commandlets to make their own engine.	
 };
diff --git a/Engine/Source/Editor/UnrealEd/Classes/Commandlets/ResavePackagesCommandlet.h b/Engine/Source/Editor/UnrealEd/Classes/Commandlets/ResavePackagesCommandlet.h
--- a/Engine/Source/Editor/UnrealEd/Classes/Commandlets/ResavePackagesCommandlet.h
+++ b/Engine/Source/Editor/UnrealEd/Classes/Commandlets/ResavePackagesCommandlet.h
@@ -13,9 +13,9 @@
 // Log category should be accessible by derived classes
 UNREALED_API DECLARE_LOG_CATEGORY_EXTERN(LogContentCommandlet, Log, All);
 
-UCLASS()
+UCLASS(MinimalAPI)
 // Added UNREALED_API to expose this to the save packages test
-class UNREALED_API UResavePackagesCommandlet : public UCommandlet
+class UResavePackagesCommandlet : public UCommandlet
 {
     GENERATED_UCLASS_BODY()
 
@@ -178,20 +178,20 @@ protected:
 	 *
 	 * @return	0 to indicate that the commandlet should continue; otherwise, the error code that should be returned by Main()
 	 */
-	virtual int32 InitializeResaveParameters( const TArray<FString>& Tokens, TArray<FString>& MapPathNames );
+	UNREALED_API virtual int32 InitializeResaveParameters( const TArray<FString>& Tokens, TArray<FString>& MapPathNames );
 
-	void ParseSourceControlOptions(const TArray<FString>& Tokens);
+	UNREALED_API void ParseSourceControlOptions(const TArray<FString>& Tokens);
 
-	void OnAddResaveOnDemandPackage(FName SystemName, FName PackageName);
+	UNREALED_API void OnAddResaveOnDemandPackage(FName SystemName, FName PackageName);
 
 	/** Loads and saves a single package */
-	virtual void LoadAndSaveOnePackage(const FString& Filename);
+	UNREALED_API virtual void LoadAndSaveOnePackage(const FString& Filename);
 
 	/** Checks to see if a package should be skipped */
-	virtual bool ShouldSkipPackage(const FString& Filename);
+	UNREALED_API virtual bool ShouldSkipPackage(const FString& Filename);
 
 	/** Deletes a single package */
-	virtual void DeleteOnePackage(const FString& Filename);
+	UNREALED_API virtual void DeleteOnePackage(const FString& Filename);
 
 	/**
 	 * Allow the commandlet to perform any operations on the export/import table of the package before all objects in the package are loaded.
@@ -200,7 +200,7 @@ protected:
 	 * @param	bSavePackage	[in]	indicates whether the package is currently going to be saved
 	 *							[out]	set to true to resave the package
 	 */
-	virtual void PerformPreloadOperations( FLinkerLoad* PackageLinker, bool& bSavePackage );
+	UNREALED_API virtual void PerformPreloadOperations( FLinkerLoad* PackageLinker, bool& bSavePackage );
 
 	/**
 	 * Allows the commandlet to perform any additional operations on the object before it is resaved.
@@ -209,7 +209,7 @@ protected:
 	 * @param	bSavePackage	[in]	indicates whether the package is currently going to be saved
 	 *							[out]	set to true to resave the package
 	 */
-	virtual void PerformAdditionalOperations( class UObject* Object, bool& bSavePackage );
+	UNREALED_API virtual void PerformAdditionalOperations( class UObject* Object, bool& bSavePackage );
 
 	/**
 	 * Allows the commandlet to perform any additional operations on the package before it is resaved.
@@ -218,7 +218,7 @@ protected:
 	 * @param	bSavePackage	[in]	indicates whether the package is currently going to be saved
 	 *							[out]	set to true to resave the package
 	 */
-	virtual void PerformAdditionalOperations( class UPackage* Package, bool& bSavePackage );
+	UNREALED_API virtual void PerformAdditionalOperations( class UPackage* Package, bool& bSavePackage );
 
 	/**
 	* Allows the commandlet to perform any additional operations on the world before it is resaved.
@@ -227,7 +227,7 @@ protected:
 	* @param	bSavePackage	[in]	indicates whether the package is currently going to be saved
 	*							[out]	set to true to resave the package
 	*/
-	virtual void PerformAdditionalOperations(class UWorld* World, bool& bSavePackage);
+	UNREALED_API virtual void PerformAdditionalOperations(class UWorld* World, bool& bSavePackage);
 
 	/**
 	 * Removes any UClass exports from packages which aren't script packages.
@@ -236,24 +236,24 @@ protected:
 	 *
 	 * @return	true to resave the package
 	 */
-	bool CleanClassesFromContentPackages( class UPackage* Package );
+	UNREALED_API bool CleanClassesFromContentPackages( class UPackage* Package );
 
 	// Get the changelist description to use if automatically checking packages out
-	virtual FText GetChangelistDescription() const;
+	UNREALED_API virtual FText GetChangelistDescription() const;
 
-	bool CheckoutFile(const FString& Filename, bool bAddFile = false, bool bIgnoreAlreadyCheckedOut = false);
-	bool RevertFile(const FString& Filename);
-	bool CanCheckoutFile(const FString& Filename, FString& CheckedOutUser);
-	void CheckoutAndSavePackage(UPackage* Package, TArray<FString>& SublevelFilenames, bool bIgnoreAlreadyCheckedOut = false);
-	void CheckInFiles(const TArray<FString>& InFilesToSubmit, const FText& InDescription) const;
-	static bool TryVirtualization(const TArray<FString>& FilesToSubmit, FText& InOutDescription);
+	UNREALED_API bool CheckoutFile(const FString& Filename, bool bAddFile = false, bool bIgnoreAlreadyCheckedOut = false);
+	UNREALED_API bool RevertFile(const FString& Filename);
+	UNREALED_API bool CanCheckoutFile(const FString& Filename, FString& CheckedOutUser);
+	UNREALED_API void CheckoutAndSavePackage(UPackage* Package, TArray<FString>& SublevelFilenames, bool bIgnoreAlreadyCheckedOut = false);
+	UNREALED_API void CheckInFiles(const TArray<FString>& InFilesToSubmit, const FText& InDescription) const;
+	static UNREALED_API bool TryVirtualization(const TArray<FString>& FilesToSubmit, FText& InOutDescription);
 
 	/**
 	 * Creates and returns a unique filename in the temporary file directory.
 	 *
 	 * @return	A unique temporary filename
 	 */
-	FString CreateTempFilename();
+	UNREALED_API FString CreateTempFilename();
 
 	/**
 	 * Returns the path of the directory we use to store all temporary files for
@@ -261,21 +261,21 @@ protected:
 	 *
 	 * @return Path to the temporary file directory
 	 */
-	FString GetTempFilesDirectory();
+	UNREALED_API FString GetTempFilesDirectory();
 
 	/**
 	 * Delete the remaining files in the temporary files directory for the commandlet.
 	 */
-	void CleanTempFiles();	
+	UNREALED_API void CleanTempFiles();	
 
 	// Print out a message only if running in very verbose mode
-	void VerboseMessage(const FString& Message);
+	UNREALED_API void VerboseMessage(const FString& Message);
 
 	/** Parse commandline to decide whether resaveondemand is activated for the given system. */
-	TSet<FName> ParseResaveOnDemandSystems();
+	UNREALED_API TSet<FName> ParseResaveOnDemandSystems();
 
 public:		
 	//~ Begin UCommandlet Interface
-	virtual int32 Main(const FString& Params) override;
+	UNREALED_API virtual int32 Main(const FString& Params) override;
 	//~ End UCommandlet Interface
 };
diff --git a/Engine/Source/Editor/UnrealEd/Classes/Commandlets/WorldPartitionConvertCommandlet.h b/Engine/Source/Editor/UnrealEd/Classes/Commandlets/WorldPartitionConvertCommandlet.h
--- a/Engine/Source/Editor/UnrealEd/Classes/Commandlets/WorldPartitionConvertCommandlet.h
+++ b/Engine/Source/Editor/UnrealEd/Classes/Commandlets/WorldPartitionConvertCommandlet.h
@@ -31,48 +31,48 @@ struct FHLODLayerActorMapping
 	FString HLODLayer;
 };
 
-UCLASS(Config = Engine)
-class UNREALED_API UWorldPartitionConvertCommandlet : public UCommandlet
+UCLASS(Config = Engine, MinimalAPI)
+class UWorldPartitionConvertCommandlet : public UCommandlet
 {
 	GENERATED_UCLASS_BODY()
 
 public:
 	//~ Begin UCommandlet Interface
-	virtual int32 Main(const FString& Params) override;
+	UNREALED_API virtual int32 Main(const FString& Params) override;
 	//~ End UCommandlet Interface
 
-	static const FString GetConversionSuffix(const bool bInOnlyMergeSubLevels);
+	static UNREALED_API const FString GetConversionSuffix(const bool bInOnlyMergeSubLevels);
 
 private:
-	void GatherAndPrepareSubLevelsToConvert(ULevel* Level, TArray<ULevel*>& SubLevels);
+	UNREALED_API void GatherAndPrepareSubLevelsToConvert(ULevel* Level, TArray<ULevel*>& SubLevels);
 
 protected:
-	virtual bool GetAdditionalLevelsToConvert(ULevel* Level, TArray<ULevel*>& SubLevels);
-	virtual bool PrepareStreamingLevelForConversion(ULevelStreaming* StreamingLevel);
-	virtual bool ShouldConvertStreamingLevel(ULevelStreaming* StreamingLevel);
-	virtual bool ShouldDeleteActor(AActor* Actor, bool bMainLevel) const;
-	virtual void PerformAdditionalWorldCleanup(UWorld* World) const;
+	UNREALED_API virtual bool GetAdditionalLevelsToConvert(ULevel* Level, TArray<ULevel*>& SubLevels);
+	UNREALED_API virtual bool PrepareStreamingLevelForConversion(ULevelStreaming* StreamingLevel);
+	UNREALED_API virtual bool ShouldConvertStreamingLevel(ULevelStreaming* StreamingLevel);
+	UNREALED_API virtual bool ShouldDeleteActor(AActor* Actor, bool bMainLevel) const;
+	UNREALED_API virtual void PerformAdditionalWorldCleanup(UWorld* World) const;
 	virtual void PerformAdditionalActorChanges(AActor* Actor) const {}
-	virtual void OutputConversionReport() const;
-	virtual void OnWorldLoaded(UWorld* World);
+	UNREALED_API virtual void OutputConversionReport() const;
+	UNREALED_API virtual void OnWorldLoaded(UWorld* World);
 	virtual void ReadAdditionalTokensAndSwitches(const TArray<FString>& Tokens, const TArray<FString>& Switches) {}
 
-	UWorldPartition* CreateWorldPartition(class AWorldSettings* MainWorldSettings);
-	UWorld* LoadWorld(const FString& LevelToLoad);
-	ULevel* InitWorld(UWorld* World);
+	UNREALED_API UWorldPartition* CreateWorldPartition(class AWorldSettings* MainWorldSettings);
+	UNREALED_API UWorld* LoadWorld(const FString& LevelToLoad);
+	UNREALED_API ULevel* InitWorld(UWorld* World);
 
-	void ChangeObjectOuter(UObject* Object, UObject* NewOuter);
-	void FixupSoftObjectPaths(UPackage* OuterPackage);
+	UNREALED_API void ChangeObjectOuter(UObject* Object, UObject* NewOuter);
+	UNREALED_API void FixupSoftObjectPaths(UPackage* OuterPackage);
 
-	bool DetachDependantLevelPackages(ULevel* Level);
-	bool RenameWorldPackageWithSuffix(UWorld* World);
+	UNREALED_API bool DetachDependantLevelPackages(ULevel* Level);
+	UNREALED_API bool RenameWorldPackageWithSuffix(UWorld* World);
 
-	void SetupHLOD();
-	void SetupHLODLayerAssets();
-	UHLODLayer* CreateHLODLayerFromINI(const FString& InHLODLayerName);
+	UNREALED_API void SetupHLOD();
+	UNREALED_API void SetupHLODLayerAssets();
+	UNREALED_API UHLODLayer* CreateHLODLayerFromINI(const FString& InHLODLayerName);
 
-	void SetActorGuid(AActor* Actor, const FGuid& NewGuid);
-	void CreateWorldMiniMapTexture(UWorld* World);
+	UNREALED_API void SetActorGuid(AActor* Actor, const FGuid& NewGuid);
+	UNREALED_API void CreateWorldMiniMapTexture(UWorld* World);
 
 	// Conversion report
 	TSet<FString> MapsWithLevelScriptsBPs;
diff --git a/Engine/Source/Editor/UnrealEd/Classes/Commandlets/WorldPartitionDataLayerToAssetCommandLet.h b/Engine/Source/Editor/UnrealEd/Classes/Commandlets/WorldPartitionDataLayerToAssetCommandLet.h
--- a/Engine/Source/Editor/UnrealEd/Classes/Commandlets/WorldPartitionDataLayerToAssetCommandLet.h
+++ b/Engine/Source/Editor/UnrealEd/Classes/Commandlets/WorldPartitionDataLayerToAssetCommandLet.h
@@ -123,8 +123,8 @@ private:
 	TArray<TWeakObjectPtr<UDataLayerConversionInfo>> ConvertingDataLayerInfo;
 };
 
-UCLASS(Config = Engine)
-class UNREALED_API UDataLayerToAssetCommandlet : public UCommandlet
+UCLASS(Config = Engine, MinimalAPI)
+class UDataLayerToAssetCommandlet : public UCommandlet
 {
 	GENERATED_UCLASS_BODY()
 
@@ -139,7 +139,7 @@ public:
 	};
 
 	//~ Begin UCommandlet Interface
-	virtual int32 Main(const FString& Params) override;
+	UNREALED_API virtual int32 Main(const FString& Params) override;
 	//~ End UCommandlet Interface
 
 protected:
@@ -147,27 +147,27 @@ protected:
 	virtual bool PerformAdditionalActorConversions(TStrongObjectPtr<UDataLayerToAssetCommandletContext>& CommandletContext, AActor* Actor) { return true; }
 
 private:
-	bool InitializeFromCommandLine(TArray<FString>& Tokens, TArray<FString> const& Switches, TMap<FString, FString> const& Params);
+	UNREALED_API bool InitializeFromCommandLine(TArray<FString>& Tokens, TArray<FString> const& Switches, TMap<FString, FString> const& Params);
 
-	bool BuildConversionInfos(TStrongObjectPtr<UDataLayerToAssetCommandletContext>& CommandletContext, FPackageSourceControlHelper& PackageHelper);
-	bool CreateConversionFromDataLayer(TStrongObjectPtr<UDataLayerToAssetCommandletContext>& CommandletContext, const UDeprecatedDataLayerInstance* DataLayer, FPackageSourceControlHelper& PackageHelper);
-	TObjectPtr<UDataLayerAsset> GetOrCreateDataLayerAssetForConversion(TStrongObjectPtr<UDataLayerToAssetCommandletContext>& CommandletContext, FName DataLayerName);
+	UNREALED_API bool BuildConversionInfos(TStrongObjectPtr<UDataLayerToAssetCommandletContext>& CommandletContext, FPackageSourceControlHelper& PackageHelper);
+	UNREALED_API bool CreateConversionFromDataLayer(TStrongObjectPtr<UDataLayerToAssetCommandletContext>& CommandletContext, const UDeprecatedDataLayerInstance* DataLayer, FPackageSourceControlHelper& PackageHelper);
+	UNREALED_API TObjectPtr<UDataLayerAsset> GetOrCreateDataLayerAssetForConversion(TStrongObjectPtr<UDataLayerToAssetCommandletContext>& CommandletContext, FName DataLayerName);
 
-	bool ResolvePreviousConversionsToCurrent(TStrongObjectPtr<UDataLayerToAssetCommandletContext>& CommandletContext);
+	UNREALED_API bool ResolvePreviousConversionsToCurrent(TStrongObjectPtr<UDataLayerToAssetCommandletContext>& CommandletContext);
 
-	bool CreateDataLayerInstances(TStrongObjectPtr<UDataLayerToAssetCommandletContext>& CommandletContext);
-	bool RebuildDataLayerHierarchies(TStrongObjectPtr<UDataLayerToAssetCommandletContext>& CommandletContext);
+	UNREALED_API bool CreateDataLayerInstances(TStrongObjectPtr<UDataLayerToAssetCommandletContext>& CommandletContext);
+	UNREALED_API bool RebuildDataLayerHierarchies(TStrongObjectPtr<UDataLayerToAssetCommandletContext>& CommandletContext);
 
-	bool RemapActorDataLayersToAssets(TStrongObjectPtr<UDataLayerToAssetCommandletContext>& CommandletContext, FPackageSourceControlHelper& PackageHelper);
-	uint32 RemapActorDataLayers(TStrongObjectPtr<UDataLayerToAssetCommandletContext>& CommandletContext, AActor* Actor);
-	uint32 RemapDataLayersAssetsFromPreviousConversions(TStrongObjectPtr<UDataLayerToAssetCommandletContext>& CommandletContext, AActor* Actor);
+	UNREALED_API bool RemapActorDataLayersToAssets(TStrongObjectPtr<UDataLayerToAssetCommandletContext>& CommandletContext, FPackageSourceControlHelper& PackageHelper);
+	UNREALED_API uint32 RemapActorDataLayers(TStrongObjectPtr<UDataLayerToAssetCommandletContext>& CommandletContext, AActor* Actor);
+	UNREALED_API uint32 RemapDataLayersAssetsFromPreviousConversions(TStrongObjectPtr<UDataLayerToAssetCommandletContext>& CommandletContext, AActor* Actor);
 	
-	bool DeletePreviousConversionsData(TStrongObjectPtr<UDataLayerToAssetCommandletContext>& CommandletContext, FPackageSourceControlHelper& PackageHelper);
+	UNREALED_API bool DeletePreviousConversionsData(TStrongObjectPtr<UDataLayerToAssetCommandletContext>& CommandletContext, FPackageSourceControlHelper& PackageHelper);
 
-	bool CommitConversion(TStrongObjectPtr<UDataLayerToAssetCommandletContext>& CommandletContext, FPackageSourceControlHelper& PackageHelper);
+	UNREALED_API bool CommitConversion(TStrongObjectPtr<UDataLayerToAssetCommandletContext>& CommandletContext, FPackageSourceControlHelper& PackageHelper);
 
-	FString GetConversionFolder() const;
-	bool IsAssetInConversionFolder(const TObjectPtr<UDataLayerAsset> DataLayerAsset);
+	UNREALED_API FString GetConversionFolder() const;
+	UNREALED_API bool IsAssetInConversionFolder(const TObjectPtr<UDataLayerAsset> DataLayerAsset);
 
 	UPROPERTY(Config)
 	FString DestinationFolder;
@@ -186,4 +186,4 @@ private:
 
 	UPROPERTY(Transient)
 	TObjectPtr<UWorld> MainWorld;
-};
\ No newline at end of file
+};
