[UnrealEd]

diff --git a/Engine/Source/Editor/UnrealEd/Classes/CookOnTheSide/CookOnTheFlyServer.h b/Engine/Source/Editor/UnrealEd/Classes/CookOnTheSide/CookOnTheFlyServer.h
--- a/Engine/Source/Editor/UnrealEd/Classes/CookOnTheSide/CookOnTheFlyServer.h
+++ b/Engine/Source/Editor/UnrealEd/Classes/CookOnTheSide/CookOnTheFlyServer.h
@@ -216,13 +216,13 @@ private:
 typedef TMap<FName, TMap<FName, TMap<FName, TArray<FString>>>> FIniSettingContainer;
 
 }
-UCLASS()
-class UNREALED_API UCookOnTheFlyServer : public UObject, public FTickableEditorObject, public FExec, public UE::Cook::ICookInfo
+UCLASS(MinimalAPI)
+class UCookOnTheFlyServer : public UObject, public FTickableEditorObject, public FExec, public UE::Cook::ICookInfo
 {
 	GENERATED_BODY()
 
-	UCookOnTheFlyServer(const FObjectInitializer& ObjectInitializer = FObjectInitializer::Get());
-	UCookOnTheFlyServer(FVTableHelper& Helper); // Declare the FVTableHelper constructor manually so that we can forward-declare-only TUniquePtrs in the header without getting compile error in generated cpp
+	UNREALED_API UCookOnTheFlyServer(const FObjectInitializer& ObjectInitializer = FObjectInitializer::Get());
+	UNREALED_API UCookOnTheFlyServer(FVTableHelper& Helper); // Declare the FVTableHelper constructor manually so that we can forward-declare-only TUniquePtrs in the header without getting compile error in generated cpp
 
 private:
 
@@ -400,7 +400,7 @@ private:
 	TArray<FWeakObjectPtr> CachedTexturesToCacheArray;
 	int32 LastUpdateTick = 0;
 	int32 MaxPrecacheShaderJobs = 0;
-	void TickPrecacheObjectsForPlatforms(const float TimeSlice, const TArray<const ITargetPlatform*>& TargetPlatform);
+	UNREALED_API void TickPrecacheObjectsForPlatforms(const float TimeSlice, const TArray<const ITargetPlatform*>& TargetPlatform);
 
 	//////////////////////////////////////////////////////////////////////////
 
@@ -426,46 +426,46 @@ private:
 	mutable UE::Cook::FIniSettingContainer AccessedIniStrings;
 	TArray<const FConfigFile*> OpenConfigFiles;
 	TArray<FString> ConfigSettingDenyList;
-	void OnFConfigDeleted(const FConfigFile* Config);
-	void OnFConfigCreated(const FConfigFile* Config);
+	UNREALED_API void OnFConfigDeleted(const FConfigFile* Config);
+	UNREALED_API void OnFConfigCreated(const FConfigFile* Config);
 
-	void ProcessAccessedIniSettings(const FConfigFile* Config, UE::Cook::FIniSettingContainer& AccessedIniStrings) const;
+	UNREALED_API void ProcessAccessedIniSettings(const FConfigFile* Config, UE::Cook::FIniSettingContainer& AccessedIniStrings) const;
 
-	void OnRequestClusterCompleted(const UE::Cook::FRequestCluster& RequestCluster);
+	UNREALED_API void OnRequestClusterCompleted(const UE::Cook::FRequestCluster& RequestCluster);
 
 	/**
 	* OnTargetPlatformChangedSupportedFormats
 	* called when target platform changes the return value of supports shader formats 
 	* used to reset the cached cooked shaders
 	*/
-	void OnTargetPlatformChangedSupportedFormats(const ITargetPlatform* TargetPlatform);
+	UNREALED_API void OnTargetPlatformChangedSupportedFormats(const ITargetPlatform* TargetPlatform);
 
 	/* Initializing Platforms must be done on the tickloop thread; Platform data is read only on other threads */
-	void AddCookOnTheFlyPlatformFromGameThread(ITargetPlatform* TargetPlatform);
+	UNREALED_API void AddCookOnTheFlyPlatformFromGameThread(ITargetPlatform* TargetPlatform);
 	/* Start the session for the given platform in response to the first client connection. */
-	void StartCookOnTheFlySessionFromGameThread(ITargetPlatform* TargetPlatform);
+	UNREALED_API void StartCookOnTheFlySessionFromGameThread(ITargetPlatform* TargetPlatform);
 
 	/* Callback to recalculate all ITargetPlatform* pointers when they change due to modules reloading */
-	void OnTargetPlatformsInvalidated();
+	UNREALED_API void OnTargetPlatformsInvalidated();
 
 	/* Update polled fields used by CookOnTheFly's network request handlers */
-	void TickNetwork();
+	UNREALED_API void TickNetwork();
 
 	/** Tick the cook for the current cook mode: request, load, save, finish. */
-	void TickMainCookLoop(UE::Cook::FTickStackData& StackData);
+	UNREALED_API void TickMainCookLoop(UE::Cook::FTickStackData& StackData);
 
 	/** Execute operations that need to be done after each Scheduler task, such as checking for new external requests. */
-	void TickCookStatus(UE::Cook::FTickStackData& StackData);
-	void SetSaveBusy(bool bInSaveBusy);
-	void SetLoadBusy(bool bInLoadBusy);
+	UNREALED_API void TickCookStatus(UE::Cook::FTickStackData& StackData);
+	UNREALED_API void SetSaveBusy(bool bInSaveBusy);
+	UNREALED_API void SetLoadBusy(bool bInLoadBusy);
 	enum class EIdleStatus
 	{
 		Active,
 		Idle,
 		Done
 	};
-	void SetIdleStatus(UE::Cook::FTickStackData& StackData, EIdleStatus InStatus);
-	void UpdateDisplay(UE::Cook::FTickStackData& StackData, bool bForceDisplay);
+	UNREALED_API void SetIdleStatus(UE::Cook::FTickStackData& StackData, EIdleStatus InStatus);
+	UNREALED_API void UpdateDisplay(UE::Cook::FTickStackData& StackData, bool bForceDisplay);
 	enum class ECookAction
 	{
 		Done,			// The cook is complete; no requests remain in any non-idle state
@@ -480,81 +480,81 @@ private:
 		YieldTick,		// Progress is blocked by an async result. Temporarily exit TickMainCookLoop.
 	};
 	/** Inspect all tasks the scheduler could do and return which one it should do. */
-	ECookAction DecideNextCookAction(UE::Cook::FTickStackData& StackData);
+	UNREALED_API ECookAction DecideNextCookAction(UE::Cook::FTickStackData& StackData);
 	/** How many packages are assigned to the CookOnTheFlyServer for saving/loading? Returns 0 if the current process is not a CookDirector. */
-	int32 NumMultiprocessLocalWorkerAssignments() const;
+	UNREALED_API int32 NumMultiprocessLocalWorkerAssignments() const;
 	/** Execute any existing external callbacks and push any existing external cook requests into new RequestClusters. */
-	void PumpExternalRequests(const UE::Cook::FCookerTimer& CookerTimer);
+	UNREALED_API void PumpExternalRequests(const UE::Cook::FCookerTimer& CookerTimer);
 	/** Send the PackageData back to request state to create a request cluster, if it has not yet been explored. */
-	bool TryCreateRequestCluster(UE::Cook::FPackageData& PackageData);
+	UNREALED_API bool TryCreateRequestCluster(UE::Cook::FPackageData& PackageData);
 	/** Inspect the next package in the RequestQueue and push it on to its next state. Report the number of PackageDatas that were pushed to load. */
-	void PumpRequests(UE::Cook::FTickStackData& StackData, int32& OutNumPushed);
+	UNREALED_API void PumpRequests(UE::Cook::FTickStackData& StackData, int32& OutNumPushed);
 	/**
 	 * Assign the requests found in PumpRequests; either pushing them to ReadyRequests if SingleProcess Cook or
 	 * to CookWorkers if MultiProcess. Input Requests have been sorted by leaf to root load order.
 	 */
-	void AssignRequests(TArrayView<UE::Cook::FPackageData*> Requests, UE::Cook::FRequestQueue& RequestQueue,
+	UNREALED_API void AssignRequests(TArrayView<UE::Cook::FPackageData*> Requests, UE::Cook::FRequestQueue& RequestQueue,
 		TMap<UE::Cook::FPackageData*, TArray<UE::Cook::FPackageData*>>&& RequestGraph);
 	/**
 	 * Multiprocess cook: Notify the CookDirector that a package it assigned to a worker was demoted
 	 * due to e.g. cancelled cook and should be removed from the CookWorker.
 	 */
-	void NotifyRemovedFromWorker(UE::Cook::FPackageData& PackageData);
+	UNREALED_API void NotifyRemovedFromWorker(UE::Cook::FPackageData& PackageData);
 
 	/** Load packages in the LoadQueue until it's time to break. Report the number of loads that were pushed to save. */
-	void PumpLoads(UE::Cook::FTickStackData& StackData, uint32 DesiredQueueLength, int32& OutNumPushed, bool& bOutBusy);
+	UNREALED_API void PumpLoads(UE::Cook::FTickStackData& StackData, uint32 DesiredQueueLength, int32& OutNumPushed, bool& bOutBusy);
 	/** Move packages from LoadPrepare's entry queue into the PreloadingQueue until we run out of Preload slots. */
-	void PumpPreloadStarts();
+	UNREALED_API void PumpPreloadStarts();
 	/** Move preload-completed packages from LoadPrepare->LoadReady until we find one that is not finished preloading. */
-	void PumpPreloadCompletes();
+	UNREALED_API void PumpPreloadCompletes();
 	/**
 	 * Load the given PackageData that was in the load queue and send it on to its next state.
 	 * Report the number of PackageDatas that were pushed to save (0 or 1)
 	 */
-	void LoadPackageInQueue(UE::Cook::FPackageData& PackageData, uint32& ResultFlags, int32& OutNumPushed);
+	UNREALED_API void LoadPackageInQueue(UE::Cook::FPackageData& PackageData, uint32& ResultFlags, int32& OutNumPushed);
 	/** Mark that the given PackageData failed to load and return it to idle. */
-	void RejectPackageToLoad(UE::Cook::FPackageData& PackageData, const TCHAR* ReasonText,
+	UNREALED_API void RejectPackageToLoad(UE::Cook::FPackageData& PackageData, const TCHAR* ReasonText,
 		UE::Cook::ESuppressCookReason Reason);
 
 	/**
 	 * Try to save all packages in the SaveQueue until it's time to break.
 	 * Report the number of requests that were completed (either skipped or successfully saved or failed to save)
 	 */
-	void PumpSaves(UE::Cook::FTickStackData& StackData, uint32 DesiredQueueLength, int32& OutNumPushed, bool& bOutBusy);
+	UNREALED_API void PumpSaves(UE::Cook::FTickStackData& StackData, uint32 DesiredQueueLength, int32& OutNumPushed, bool& bOutBusy);
 
 	/**
 	 * Inspect the given package and queue it for saving if necessary.
 	 *
 	 * @param PackageData			The PackageData to be considered for saving.
 	 */
-	void QueueDiscoveredPackage(UE::Cook::FPackageData& PackageData, UE::Cook::FInstigator&& Instigator, 
+	UNREALED_API void QueueDiscoveredPackage(UE::Cook::FPackageData& PackageData, UE::Cook::FInstigator&& Instigator, 
 		UE::Cook::FDiscoveredPlatformSet&& ReachablePlatforms, bool bUrgent=false);
-	void QueueDiscoveredPackageOnDirector(UE::Cook::FPackageData& PackageData, UE::Cook::FInstigator&& Instigator,
+	UNREALED_API void QueueDiscoveredPackageOnDirector(UE::Cook::FPackageData& PackageData, UE::Cook::FInstigator&& Instigator,
 		UE::Cook::FDiscoveredPlatformSet&& ReachablePlatforms, bool bUrgent);
 
 	/** Called when a package is cancelled and returned to idle. Notifies CookDirector when on a CookWorker. */
-	void DemoteToIdle(UE::Cook::FPackageData& PackageData, UE::Cook::ESendFlags SendFlags, UE::Cook::ESuppressCookReason Reason);
+	UNREALED_API void DemoteToIdle(UE::Cook::FPackageData& PackageData, UE::Cook::ESendFlags SendFlags, UE::Cook::ESuppressCookReason Reason);
 	/** Called when a package completes its save and returns to idle. Notifies the CookDirector when on a CookWorker. */
-	void PromoteToSaveComplete(UE::Cook::FPackageData& PackageData, UE::Cook::ESendFlags SendFlags);
+	UNREALED_API void PromoteToSaveComplete(UE::Cook::FPackageData& PackageData, UE::Cook::ESendFlags SendFlags);
 
 	/**
 	 * Remove all request data about the given platform from any data in the CookOnTheFlyServer.
 	 * Called when a platform is removed from the list of session platforms e.g. because it has not been recently
 	 * used by CookOnTheFly. Does not modify Cooked platforms.
 	 */
-	void OnRemoveSessionPlatform(const ITargetPlatform* TargetPlatform, int32 RemovedIndex);
+	UNREALED_API void OnRemoveSessionPlatform(const ITargetPlatform* TargetPlatform, int32 RemovedIndex);
 	/** Update structures that have data per platfrom when COTF adds a new SessionPlatform. */
-	void OnPlatformAddedToSession(const ITargetPlatform* TargetPlatform);
-
-	void InitializePollables();
-	void PumpPollables(UE::Cook::FTickStackData& StackData, bool bIsIdle);
-	void PollFlushRenderingCommands();
-	TRefCountPtr<FPollable> CreatePollableLLM();
-	TRefCountPtr<FPollable> CreatePollableTriggerGC();
-	void PollGarbageCollection(UE::Cook::FTickStackData& StackData);
-	void PollQueuedCancel(UE::Cook::FTickStackData& StackData);
-	void WaitForAsync(UE::Cook::FTickStackData& StackData);
-	void TickRecompileShaderRequestsPrivate();
+	UNREALED_API void OnPlatformAddedToSession(const ITargetPlatform* TargetPlatform);
+
+	UNREALED_API void InitializePollables();
+	UNREALED_API void PumpPollables(UE::Cook::FTickStackData& StackData, bool bIsIdle);
+	UNREALED_API void PollFlushRenderingCommands();
+	UNREALED_API TRefCountPtr<FPollable> CreatePollableLLM();
+	UNREALED_API TRefCountPtr<FPollable> CreatePollableTriggerGC();
+	UNREALED_API void PollGarbageCollection(UE::Cook::FTickStackData& StackData);
+	UNREALED_API void PollQueuedCancel(UE::Cook::FTickStackData& StackData);
+	UNREALED_API void WaitForAsync(UE::Cook::FTickStackData& StackData);
+	UNREALED_API void TickRecompileShaderRequestsPrivate();
 
 public:
 
@@ -610,24 +610,24 @@ public:
 		TArray<ITargetPlatform*> TargetPlatforms;
 	};
 
-	virtual ~UCookOnTheFlyServer();
+	UNREALED_API virtual ~UCookOnTheFlyServer();
 
 	// FTickableEditorObject interface used by cook on the side
-	TStatId GetStatId() const override;
-	void Tick(float DeltaTime) override;
-	bool IsTickable() const override;
+	UNREALED_API TStatId GetStatId() const override;
+	UNREALED_API void Tick(float DeltaTime) override;
+	UNREALED_API bool IsTickable() const override;
 	ECookMode::Type GetCookMode() const { return CurrentCookMode; }
 	ECookInitializationFlags GetCookFlags() const { return CookFlags; }
 
 	// ICookInfo interface
-	virtual UE::Cook::FInstigator GetInstigator(FName PackageName) override;
-	virtual TArray<UE::Cook::FInstigator> GetInstigatorChain(FName PackageName) override;
+	UNREALED_API virtual UE::Cook::FInstigator GetInstigator(FName PackageName) override;
+	UNREALED_API virtual TArray<UE::Cook::FInstigator> GetInstigatorChain(FName PackageName) override;
 
 	/** Dumps cooking stats to the log. Run from the exec command "Cook stats". */
-	void DumpStats();
+	UNREALED_API void DumpStats();
 
 	/** Initialize *this so that either CookOnTheFly or CookByTheBook can be started and ticked */
-	void Initialize( ECookMode::Type DesiredCookMode, ECookInitializationFlags InCookInitializationFlags,
+	UNREALED_API void Initialize( ECookMode::Type DesiredCookMode, ECookInitializationFlags InCookInitializationFlags,
 		const FString& OutputDirectoryOverride = FString() );
 
 	/**
@@ -637,74 +637,74 @@ public:
 	 *
 	 * @return true on success, false otherwise.
 	 */
-	bool StartCookOnTheFly(FCookOnTheFlyStartupOptions InCookOnTheFlyOptions); 
+	UNREALED_API bool StartCookOnTheFly(FCookOnTheFlyStartupOptions InCookOnTheFlyOptions); 
 
 	/** Broadcast the fileserver's presence on the network */
-	bool BroadcastFileserverPresence( const FGuid &InstanceId );
+	UNREALED_API bool BroadcastFileserverPresence( const FGuid &InstanceId );
 	
 	/** Shutdown *this from running in CookOnTheFly mode */
-	void ShutdownCookOnTheFly();
+	UNREALED_API void ShutdownCookOnTheFly();
 
 	/**
 	* Start a cook by the book session
 	* Cook on the fly can't run at the same time as cook by the book
 	*/
-	void StartCookByTheBook( const FCookByTheBookStartupOptions& CookByTheBookStartupOptions );
+	UNREALED_API void StartCookByTheBook( const FCookByTheBookStartupOptions& CookByTheBookStartupOptions );
 
 	/** Queue a cook by the book cancel (can be called from any thread) */
-	void QueueCancelCookByTheBook();
+	UNREALED_API void QueueCancelCookByTheBook();
 
 	/** Cancel the currently running cook by the book (needs to be called from the game thread) */
-	void CancelCookByTheBook();
+	UNREALED_API void CancelCookByTheBook();
 
 	/** Connect to the CookDirector host, Initialize this UCOTFS, and start the CookWorker session */
-	bool TryInitializeCookWorker();
+	UNREALED_API bool TryInitializeCookWorker();
 
 	/** Log stats for the CookWorker; this is called before the connection to the director is terminated. */
-	void LogCookWorkerStats();
+	UNREALED_API void LogCookWorkerStats();
 
 	/** Terminate the CookWorker session */
-	void ShutdownCookAsCookWorker();
+	UNREALED_API void ShutdownCookAsCookWorker();
 
 	/**
 	 * Is the local CookOnTheFlyServer in a cook session in any CookMode?
 	 * Used to restrict operations when cooking and reduce cputime when not cooking.
 	 */
-	bool IsInSession() const;
+	UNREALED_API bool IsInSession() const;
 	/** Is the local CookOnTheFlyServer in a cook session in CookByTheBook Mode? */
-	bool IsCookByTheBookRunning() const;
+	UNREALED_API bool IsCookByTheBookRunning() const;
 
 	/** Execute class-specific special case cook postloads and reference discovery on a given package. */
-	void PostLoadPackageFixup(UE::Cook::FPackageData& PackageData, UPackage* Package);
+	UNREALED_API void PostLoadPackageFixup(UE::Cook::FPackageData& PackageData, UPackage* Package);
 
 	/** Tick CBTB until it finishes or needs to yield. Should only be called when in CookByTheBook Mode. */
-	uint32 TickCookByTheBook(const float TimeSlice, ECookTickFlags TickFlags = ECookTickFlags::None);
+	UNREALED_API uint32 TickCookByTheBook(const float TimeSlice, ECookTickFlags TickFlags = ECookTickFlags::None);
 	/** Tick COTF until it finishes or needs to yield. Should only be called when in CookOnTheFly Mode. */
-	uint32 TickCookOnTheFly(const float TimeSlice, ECookTickFlags TickFlags = ECookTickFlags::None);
+	UNREALED_API uint32 TickCookOnTheFly(const float TimeSlice, ECookTickFlags TickFlags = ECookTickFlags::None);
 	/** Tick CookWorker until it finishes or needs to yield. Should only be called when in CookWorker Mode. */
-	uint32 TickCookWorker();
+	UNREALED_API uint32 TickCookWorker();
 	/** Log a list of all of the transitively requested packages. */
-	void RunCookList(ECookListOptions CookListOptions);
+	UNREALED_API void RunCookList(ECookListOptions CookListOptions);
 
 	/** Clear all the previously cooked data all cook requests from now on will be considered recook requests */
-	void ClearAllCookedData();
+	UNREALED_API void ClearAllCookedData();
 
 	/** Demote PackageDatas in any queue back to Idle, and eliminate pending requests. Used when canceling a cook. */
-	void CancelAllQueues();
+	UNREALED_API void CancelAllQueues();
 
 	/**
 	 * Clear any cached cooked platform data for a platform and call ClearCachedCookedPlatformData on all UObjects.
 	 * 
 	 * @param TargetPlatform platform to clear all the cached data for
 	 */
-	void ClearCachedCookedPlatformDataForPlatform(const ITargetPlatform* TargetPlatform);
+	UNREALED_API void ClearCachedCookedPlatformDataForPlatform(const ITargetPlatform* TargetPlatform);
 
 	/**
 	 * Clear all the previously cooked data for the platform passed in 
 	 * 
 	 * @param TargetPlatform the platform to clear the cooked packages for
 	 */
-	void ClearPlatformCookedData(const ITargetPlatform* TargetPlatform);
+	UNREALED_API void ClearPlatformCookedData(const ITargetPlatform* TargetPlatform);
 
 	/**
 	 * Clear platforms' explored flags for all PackageDatas and optionally clear the cookresult flags.
@@ -712,7 +712,7 @@ public:
 	 * @param TargetPlatforms List of pairs with targetplatforms to reset and bool bResetResults indicating whether
 	 *        the platform should clear CookResults in addition to clearing IsExplored.
 	 */
-	void ResetCook(TConstArrayView<TPair<const ITargetPlatform*, bool>> TargetPlatforms);
+	UNREALED_API void ResetCook(TConstArrayView<TPair<const ITargetPlatform*, bool>> TargetPlatforms);
 
 	/**
 	 * Recompile any global shader changes 
@@ -720,85 +720,85 @@ public:
 	 * 
 	 * @return return true if shaders were recompiled
 	 */
-	bool RecompileChangedShaders(const TArray<const ITargetPlatform*>& TargetPlatforms);
+	UNREALED_API bool RecompileChangedShaders(const TArray<const ITargetPlatform*>& TargetPlatforms);
 
 	/**
 	 * Force stop whatever pending cook requests are going on and clear all the cooked data
 	 * Note CookOnTheFly clients may not be able to recover from this if they are waiting on a cook request to complete
 	 */
-	void StopAndClearCookedData();
+	UNREALED_API void StopAndClearCookedData();
 
-	void TickRequestManager();
+	UNREALED_API void TickRequestManager();
 
 	/**
 	 * Return whether the tick needs to take any action for the current session. If not, the session is done.
 	 * Used for external managers of the cooker to know when to tick it.
 	 */
-	bool HasRemainingWork() const;
-	void WaitForRequests(int TimeoutMs);
+	UNREALED_API bool HasRemainingWork() const;
+	UNREALED_API void WaitForRequests(int TimeoutMs);
 
-	uint32 NumConnections() const;
+	UNREALED_API uint32 NumConnections() const;
 
 	/** Is the local CookOnTheFlyServer running in the editor? */
-	bool IsCookingInEditor() const;
+	UNREALED_API bool IsCookingInEditor() const;
 
 	/** Is the local CookOnTheFlyServer initialized to run in real time mode (respects the timeslice), (e.g. in Editor)? */
-	bool IsRealtimeMode() const;
+	UNREALED_API bool IsRealtimeMode() const;
 
 	/** Is the local CookOnTheFlyServer initialized to run CookByTheBook (find all required packages and cook them)? */
-	bool IsCookByTheBookMode() const;
+	UNREALED_API bool IsCookByTheBookMode() const;
 	/**
 	 * Is the CookDirector (local CookOnTheFlyServer for SPCook or MPDirector, or the remote MPDirector for a
 	 * CookWorker) initialized to run CookByTheBook?
 	 */
-	bool IsDirectorCookByTheBook() const;
+	UNREALED_API bool IsDirectorCookByTheBook() const;
 
 	/**
 	 * Is the local CookOnTheFlyServer in a mode that uses ShaderCodeLibraries rather than storing shaders in
 	 * the Packages that use them?
 	 */
-	bool IsUsingShaderCodeLibrary() const;
+	UNREALED_API bool IsUsingShaderCodeLibrary() const;
 
 	/**
 	 * Is the local CookOnTheFlyServer using ZenStore (cooked packages are stored in Zen's Cache storage rather
 	 * than as loose files on disk)?
 	 */
-	bool IsUsingZenStore() const;
+	UNREALED_API bool IsUsingZenStore() const;
 
 	/**
 	 * Is the local CookOnTheFlyServer initialized to run CookOnTheFly (accept connections from game executables,
 	 * cook what they ask for)?
 	 */
-	bool IsCookOnTheFlyMode() const;
+	UNREALED_API bool IsCookOnTheFlyMode() const;
 	/**
 	 * Is the CookDirector (local CookOnTheFlyServer for SPCook or MPDirector, or the remote MPDirector for a
 	 * CookWorker) initialized to run CookOnTheFly?
 	 */
-	bool IsDirectorCookOnTheFly() const;
+	UNREALED_API bool IsDirectorCookOnTheFly() const;
 
 	/**
 	 * Is the local CookOnTheFlyServer initialized to run as a CookWorker (connect to a Director cooker cooking
 	 * CBTB or COTF, cook what it assigns to us)?
 	 */
-	bool IsCookWorkerMode() const;
+	UNREALED_API bool IsCookWorkerMode() const;
 
-	virtual void BeginDestroy() override;
+	UNREALED_API virtual void BeginDestroy() override;
 
 	/** Returns the configured number of packages to process before GC */
-	uint32 GetPackagesPerGC() const;
+	UNREALED_API uint32 GetPackagesPerGC() const;
 
 	/** Returns the configured number of packages to process before partial GC */
-	uint32 GetPackagesPerPartialGC() const;
+	UNREALED_API uint32 GetPackagesPerPartialGC() const;
 
 	/** Returns the configured amount of idle time before forcing a GC */
-	double GetIdleTimeToGC() const;
+	UNREALED_API double GetIdleTimeToGC() const;
 
 	UE_DEPRECATED(5.2, "UCookOnTheFlyServer now uses a more complicated private GC scheme; HasExceededMaxMemory is no longer used and returns false")
 	bool HasExceededMaxMemory() { return false; }
-	void SetGarbageCollectType(uint32 ResultFlagsFromTick);
-	void ClearGarbageCollectType();
+	UNREALED_API void SetGarbageCollectType(uint32 ResultFlagsFromTick);
+	UNREALED_API void ClearGarbageCollectType();
 
-	void EvaluateGarbageCollectionResults(bool bWasDueToOOM, bool bWasPartialGC, uint32 ResultFlags,
+	UNREALED_API void EvaluateGarbageCollectionResults(bool bWasDueToOOM, bool bWasPartialGC, uint32 ResultFlags,
 		int32 NumObjectsBeforeGC, const FPlatformMemoryStats& MemStatsBeforeGC,
 		const FGenericMemoryStats& AllocatorStatsBeforeGC,
 		int32 NumObjectsAfterGC, const FPlatformMemoryStats& MemStatsAfterGC,
@@ -811,7 +811,7 @@ public:
 	 * @param TargetPlatforms The TargetPlatforms we want this package cooked for
 	 * @param bForceFrontOfQueue should package go to front of the cook queue (next to be processed) or the end
 	 */
-	bool RequestPackage(const FName& StandardFileName, const TArrayView<const ITargetPlatform* const>& TargetPlatforms,
+	UNREALED_API bool RequestPackage(const FName& StandardFileName, const TArrayView<const ITargetPlatform* const>& TargetPlatforms,
 		const bool bForceFrontOfQueue);
 
 	/**
@@ -821,17 +821,17 @@ public:
 	 * @param StandardPackageFName name of the package in standard format as returned by FPaths::MakeStandardFilename
 	 * @param bForceFrontOfQueue should package go to front of the cook queue (next to be processed) or the end
 	 */
-	bool RequestPackage(const FName& StandardPackageFName, const bool bForceFrontOfQueue);
+	UNREALED_API bool RequestPackage(const FName& StandardPackageFName, const bool bForceFrontOfQueue);
 
 
 	/**
 	* Callbacks from editor 
 	*/
 
-	void OnObjectModified( UObject *ObjectMoving );
-	void OnObjectPropertyChanged(UObject* ObjectBeingModified, FPropertyChangedEvent& PropertyChangedEvent);
-	void OnObjectUpdated( UObject *Object );
-	void OnObjectSaved( UObject *ObjectSaved, FObjectPreSaveContext SaveContext );
+	UNREALED_API void OnObjectModified( UObject *ObjectMoving );
+	UNREALED_API void OnObjectPropertyChanged(UObject* ObjectBeingModified, FPropertyChangedEvent& PropertyChangedEvent);
+	UNREALED_API void OnObjectUpdated( UObject *Object );
+	UNREALED_API void OnObjectSaved( UObject *ObjectSaved, FObjectPreSaveContext SaveContext );
 
 	DECLARE_MULTICAST_DELEGATE(FOnCookByTheBookStarted);
 	static FOnCookByTheBookStarted& OnCookByTheBookStarted() { return CookByTheBookStartedEvent; };
@@ -842,13 +842,13 @@ public:
 	* Marks a package as dirty for cook
 	* causes package to be recooked on next request (and all dependent packages which are currently cooked)
 	*/
-	void MarkPackageDirtyForCooker( UPackage *Package, bool bAllowInSession = false );
+	UNREALED_API void MarkPackageDirtyForCooker( UPackage *Package, bool bAllowInSession = false );
 
 	/**
 	 * Helper function for MarkPackageDirtyForCooker. Executes the MarkPackageDirtyForCooker operations that are
 	 * only safe to execute from the scheduler's designated point for handling external requests.
 	 */
-	void MarkPackageDirtyForCookerFromSchedulerThread(const FName& PackageName);
+	UNREALED_API void MarkPackageDirtyForCookerFromSchedulerThread(const FName& PackageName);
 
 	/**
 	 * MaybeMarkPackageAsAlreadyLoaded
@@ -857,74 +857,74 @@ public:
 	 * 
 	 * @param Package to mark as not requiring reload
 	 */
-	void MaybeMarkPackageAsAlreadyLoaded(UPackage* Package);
+	UNREALED_API void MaybeMarkPackageAsAlreadyLoaded(UPackage* Package);
 
 	// Callbacks from UObject globals
-	void PreGarbageCollect();
-	void CookerAddReferencedObjects(FReferenceCollector& Ar);
-	void PostGarbageCollect();
+	UNREALED_API void PreGarbageCollect();
+	UNREALED_API void CookerAddReferencedObjects(FReferenceCollector& Ar);
+	UNREALED_API void PostGarbageCollect();
 
 	/** Calculate the ShaderLibrary codedir and metadatadir */
-	void GetShaderLibraryPaths(const ITargetPlatform* TargetPlatform, FString& OutShaderCodeDir,
+	UNREALED_API void GetShaderLibraryPaths(const ITargetPlatform* TargetPlatform, FString& OutShaderCodeDir,
 		FString& OutMetaDataPath, bool bUseProjectDirForDLC=false);
 
 	/** Print detailed stats from the cook. */
-	void PrintDetailedCookStats();
+	UNREALED_API void PrintDetailedCookStats();
 
 	/** Is the local CookOnTheFlyServer cooking a DLC plugin rather than a Project+Engine+EmbeddedPlugins? */
-	bool IsCookingDLC() const;
+	UNREALED_API bool IsCookingDLC() const;
 
 	/** Return whether EDLCookInfo verification has NOT been rendered useless by settings such as CookFilter. */
-	bool ShouldVerifyEDLCookInfo() const;
+	UNREALED_API bool ShouldVerifyEDLCookInfo() const;
 
 protected:
 	// FExec interface used in the editor
-	virtual bool Exec_Editor(class UWorld* InWorld, const TCHAR* Cmd, FOutputDevice& Ar) override;
+	UNREALED_API virtual bool Exec_Editor(class UWorld* InWorld, const TCHAR* Cmd, FOutputDevice& Ar) override;
 
 private:
 
-	bool PumpHasExceededMaxMemory(uint32& OutResultFlags);
+	UNREALED_API bool PumpHasExceededMaxMemory(uint32& OutResultFlags);
 
 	/**
 	 * Is the local CookOnTheFlyServer initialized to run in CookOnTheFly AND using the legacy scheduling method
 	 * that predates COTF2?
 	 */
-	bool IsUsingLegacyCookOnTheFlyScheduling() const;
-	bool IsDebugRecordUnsolicited() const;
+	UNREALED_API bool IsUsingLegacyCookOnTheFlyScheduling() const;
+	UNREALED_API bool IsDebugRecordUnsolicited() const;
 
 	/** Update accumulators of editor data and consume editor changes after last cook when starting cooks in the editor. */
-	void BeginCookEditorSystems();
-	void BeginCookPackageWriters(FBeginCookContext& BeginContext);
-	void BeginCookDirector(FBeginCookContext& BeginContext);
-	void InitializeSession();
+	UNREALED_API void BeginCookEditorSystems();
+	UNREALED_API void BeginCookPackageWriters(FBeginCookContext& BeginContext);
+	UNREALED_API void BeginCookDirector(FBeginCookContext& BeginContext);
+	UNREALED_API void InitializeSession();
 
 	//////////////////////////////////////////////////////////////////////////
 	// cook by the book specific functions
-	const FCookByTheBookStartupOptions& BlockOnPrebootCookGate(bool& bOutAbortCook,
+	UNREALED_API const FCookByTheBookStartupOptions& BlockOnPrebootCookGate(bool& bOutAbortCook,
 		const FCookByTheBookStartupOptions& CookByTheBookStartupOptions,
 		TOptional<FCookByTheBookStartupOptions>& ModifiedStartupOptions);
 	/** Construct the on-stack data for StartCook functions, based on the arguments to StartCookByTheBook */
-	FBeginCookContext CreateBeginCookByTheBookContext(const FCookByTheBookStartupOptions& StartupOptions);
+	UNREALED_API FBeginCookContext CreateBeginCookByTheBookContext(const FCookByTheBookStartupOptions& StartupOptions);
 	/** Set the PlatformManager's session platforms and finish filling out the BeginContext for StartCookByTheBook */
-	void SelectSessionPlatforms(FBeginCookContext& BeginContext);
+	UNREALED_API void SelectSessionPlatforms(FBeginCookContext& BeginContext);
 	/** CollectFilesToCook and add them as external requests. */
-	void GenerateInitialRequests(FBeginCookContext& BeginContext);
+	UNREALED_API void GenerateInitialRequests(FBeginCookContext& BeginContext);
 	/** If cooking DLC, initialize cooked PackageDatas for all of the packages cooked by the base release. */
-	void RecordDLCPackagesFromBaseGame(FBeginCookContext& BeginContext);
-	void RegisterCookByTheBookDelegates();
-	void UnregisterCookByTheBookDelegates();
+	UNREALED_API void RecordDLCPackagesFromBaseGame(FBeginCookContext& BeginContext);
+	UNREALED_API void RegisterCookByTheBookDelegates();
+	UNREALED_API void UnregisterCookByTheBookDelegates();
 	/** Start the collection of EDL diagnostics for the cook, if applicable. */
-	void BeginCookEDLCookInfo(FBeginCookContext& BeginContext);
+	UNREALED_API void BeginCookEDLCookInfo(FBeginCookContext& BeginContext);
 
 	/** Collect all the files which need to be cooked for a cook by the book session */
-	void CollectFilesToCook(TArray<FName>& FilesInPath, TMap<FName, UE::Cook::FInstigator>& Instigators,
+	UNREALED_API void CollectFilesToCook(TArray<FName>& FilesInPath, TMap<FName, UE::Cook::FInstigator>& Instigators,
 		const TArray<FString>& CookMaps, const TArray<FString>& CookDirectories, 
 		const TArray<FString>& IniMapSections, ECookByTheBookOptions FilesToCookFlags,
 		const TArrayView<const ITargetPlatform* const>& TargetPlatforms,
 		const TMap<FName, TArray<FName>>& GameDefaultObjects);
 
 	/** Gets all game default objects for all platforms. */
-	static void GetGameDefaultObjects(const TArray<ITargetPlatform*>& TargetPlatforms, TMap<FName,
+	static UNREALED_API void GetGameDefaultObjects(const TArray<ITargetPlatform*>& TargetPlatforms, TMap<FName,
 		TArray<FName>>& GameDefaultObjectsOut);
 
 	/*
@@ -932,62 +932,62 @@ private:
 	 * Does not include checking UAssetManager, which has to be queried later
 	 * This function is const because it is not always called and should avoid side effects
 	 */
-	TArray<FName> GetNeverCookPackageFileNames(TArrayView<const FString> ExtraNeverCookDirectories
+	UNREALED_API TArray<FName> GetNeverCookPackageFileNames(TArrayView<const FString> ExtraNeverCookDirectories
 		= TArrayView<const FString>()) const;
 
 
 	/** AddFileToCook add file to cook list */
-	void AddFileToCook( TArray<FName>& InOutFilesToCook, TMap<FName, UE::Cook::FInstigator>& InOutInstigators,
+	UNREALED_API void AddFileToCook( TArray<FName>& InOutFilesToCook, TMap<FName, UE::Cook::FInstigator>& InOutInstigators,
 		const FString &InFilename, const UE::Cook::FInstigator& Instigator) const;
 
 	/** Return the name to use for the project's global shader library */
-	FString GetProjectShaderLibraryName() const;
+	UNREALED_API FString GetProjectShaderLibraryName() const;
 
 	/** Invokes the necessary FShaderCodeLibrary functions to start cooking the shader code library. */
-	void BeginCookStartShaderCodeLibrary(FBeginCookContext& BeginContext);
+	UNREALED_API void BeginCookStartShaderCodeLibrary(FBeginCookContext& BeginContext);
 	/**
 	 * Finishes async operations from ShaderCodeLibraryBeginCookStart, saves initial data to disk,
 	 * and opens the library for the rest of the cook
 	 */
-	void BeginCookFinishShaderCodeLibrary(FBeginCookContext& BeginContext);
+	UNREALED_API void BeginCookFinishShaderCodeLibrary(FBeginCookContext& BeginContext);
     
 	/** Opens Global shader library. Global shaderlib is a special case - no chunks and not connected to assets. */
-	void OpenGlobalShaderLibrary();
+	UNREALED_API void OpenGlobalShaderLibrary();
 
 	/** Saves Global shader library. Global shaderlib is a special case - no chunks and not connected to assets. */
-	void SaveAndCloseGlobalShaderLibrary();
+	UNREALED_API void SaveAndCloseGlobalShaderLibrary();
 
 	/** Invokes the necessary FShaderCodeLibrary functions to open a named code library. */
-	void OpenShaderLibrary(FString const& Name);
+	UNREALED_API void OpenShaderLibrary(FString const& Name);
     
 	/**
 	 * Finishes collection of data that should be in the named code library (e.g. load data from previous
 	 * iterative cook)
 	 */
-	void FinishPopulateShaderLibrary(const ITargetPlatform* TargetPlatform, FString const& Name);
+	UNREALED_API void FinishPopulateShaderLibrary(const ITargetPlatform* TargetPlatform, FString const& Name);
 
 	/** Invokes the necessary FShaderCodeLibrary functions to save and close a named code library. */
-	void SaveShaderLibrary(const ITargetPlatform* TargetPlatform, FString const& Name);
+	UNREALED_API void SaveShaderLibrary(const ITargetPlatform* TargetPlatform, FString const& Name);
 
 	/**
 	 * Calls the ShaderPipelineCacheToolsCommandlet to build a upipelinecache file from the stable pipeline cache
 	 * (.spc, in the past textual .stablepc.csv) file, if any
 	 */
-	void CreatePipelineCache(const ITargetPlatform* TargetPlatform, const FString& LibraryName);
+	UNREALED_API void CreatePipelineCache(const ITargetPlatform* TargetPlatform, const FString& LibraryName);
 
 	/** Invokes the necessary FShaderCodeLibrary functions to clean out all the temporary files. */
-	void CleanShaderCodeLibraries();
+	UNREALED_API void CleanShaderCodeLibraries();
 
-	void RegisterShaderChunkDataGenerator();
+	UNREALED_API void RegisterShaderChunkDataGenerator();
 
 	/** Called at the end of CookByTheBook to write aggregated data such as AssetRegistry and shaders. */
-	void CookByTheBookFinished();
+	UNREALED_API void CookByTheBookFinished();
 
 	/** Clears session-lifetime data from COTFS and CookPackageDatas. */
-	void ShutdownCookSession();
+	UNREALED_API void ShutdownCookSession();
 
 	/** Print some stats when finishing or cancelling CookByTheBook */
-	void PrintFinishStats();
+	UNREALED_API void PrintFinishStats();
 
 	/**
 	 * Get all the packages which are listed in asset registry passed in.  
@@ -998,49 +998,49 @@ private:
 	 * @param OutPackageDatas out list of packagename/filenames for packages contained in the asset registry file
 	 * @return true if successfully read false otherwise
 	 */
-	bool GetAllPackageFilenamesFromAssetRegistry( const FString& AssetRegistryPath, bool bVerifyPackagesExist,
+	UNREALED_API bool GetAllPackageFilenamesFromAssetRegistry( const FString& AssetRegistryPath, bool bVerifyPackagesExist,
 		 bool bSkipUncookedPackages, TArray<UE::Cook::FConstructPackageData>& OutPackageDatas) const;
 
 	/** Build a map of the package dependencies used by each Map Package. */
-	TMap<FName, TSet<FName>> BuildMapDependencyGraph(const ITargetPlatform* TargetPlatform);
+	UNREALED_API TMap<FName, TSet<FName>> BuildMapDependencyGraph(const ITargetPlatform* TargetPlatform);
 
 	/** Write a MapDependencyGraph to a metadata file in the sandbox for the given platform. */
-	void WriteMapDependencyGraph(const ITargetPlatform* TargetPlatform, TMap<FName, TSet<FName>>& MapDependencyGraph);
+	UNREALED_API void WriteMapDependencyGraph(const ITargetPlatform* TargetPlatform, TMap<FName, TSet<FName>>& MapDependencyGraph);
 
 	/** Generates the CachedEditorThumbnails.bin file */
-	void GenerateCachedEditorThumbnails();
+	UNREALED_API void GenerateCachedEditorThumbnails();
 
-	void InitializeAllCulturesToCook(TConstArrayView<FString> CookCultures);
-	void CompileDLCLocalization(FBeginCookContext& BeginContext);
+	UNREALED_API void InitializeAllCulturesToCook(TConstArrayView<FString> CookCultures);
+	UNREALED_API void CompileDLCLocalization(FBeginCookContext& BeginContext);
 	/** Find localization dependencies for all packages, used to add required localization files as soft references. */
-	void GenerateLocalizationReferences();
-	void RegisterLocalizationChunkDataGenerator();
+	UNREALED_API void GenerateLocalizationReferences();
+	UNREALED_API void RegisterLocalizationChunkDataGenerator();
 
 	//////////////////////////////////////////////////////////////////////////
 	// cook on the fly specific functions
 
 	/** Construct the on-stack data for StartCook functions, based on the arguments to StartCookOnTheFly */
-	FBeginCookContext CreateBeginCookOnTheFlyContext(const FCookOnTheFlyStartupOptions& Options);
+	UNREALED_API FBeginCookContext CreateBeginCookOnTheFlyContext(const FCookOnTheFlyStartupOptions& Options);
 	/** Construct the on-stack data for StartCook functions, based on the arguments when adding a COTF platform */
-	FBeginCookContext CreateAddPlatformContext(ITargetPlatform* TargetPlatform);
-	void GetCookOnTheFlyUnsolicitedFiles(const ITargetPlatform* TargetPlatform, const FString& PlatformName,
+	UNREALED_API FBeginCookContext CreateAddPlatformContext(ITargetPlatform* TargetPlatform);
+	UNREALED_API void GetCookOnTheFlyUnsolicitedFiles(const ITargetPlatform* TargetPlatform, const FString& PlatformName,
 		TArray<FString>& UnsolicitedFiles, const FString& Filename, bool bIsCookable);
 
 	//////////////////////////////////////////////////////////////////////////
 	// CookWorker specific functions
 	/** Start a CookWorker Session */
-	void StartCookAsCookWorker();
+	UNREALED_API void StartCookAsCookWorker();
 	/** Construct the on-stack data for StartCook functions, based on the CookWorker's configuration */
-	FBeginCookContext CreateCookWorkerContext();
-	void CookAsCookWorkerFinished();
+	UNREALED_API FBeginCookContext CreateCookWorkerContext();
+	UNREALED_API void CookAsCookWorkerFinished();
 	/** Return the best packages to retract and give to other workers -the ones with least effort spent so far. */
-	void GetPackagesToRetract(int32 NumToRetract, TArray<FName>& OutRetractionPackages);
+	UNREALED_API void GetPackagesToRetract(int32 NumToRetract, TArray<FName>& OutRetractionPackages);
 
 	//////////////////////////////////////////////////////////////////////////
 	// general functions
 
 	/** Perform any special processing for freshly loaded packages */
-	void ProcessUnsolicitedPackages(TArray<FName>* OutDiscoveredPackageNames = nullptr,
+	UNREALED_API void ProcessUnsolicitedPackages(TArray<FName>* OutDiscoveredPackageNames = nullptr,
 		TMap<FName, UE::Cook::FInstigator>* OutInstigators = nullptr);
 
 	/**
@@ -1053,44 +1053,44 @@ private:
 	 * @param ReportingPackageData PackageData which caused the load, for e.g. generated packages
 	 * @return Whether LoadPackage was completely successful and the package can be cooked
 	 */
-	bool LoadPackageForCooking(UE::Cook::FPackageData& PackageData, UPackage*& OutPackage,
+	UNREALED_API bool LoadPackageForCooking(UE::Cook::FPackageData& PackageData, UPackage*& OutPackage,
 		UE::Cook::FPackageData* ReportingPackageData = nullptr);
 
 	/**
 	 * Read information about the previous cook from disk and set whether the current cook is iterative based on
 	 * previous cook, config, and commandline.
 	 */
-	void LoadBeginCookIterativeFlags(FBeginCookContext& BeginContext);
+	UNREALED_API void LoadBeginCookIterativeFlags(FBeginCookContext& BeginContext);
 	/** const because it is not always calledand should avoid sideeffects */
-	void LoadBeginCookIterativeFlagsLocal(FBeginCookContext& BeginContext) const;
+	UNREALED_API void LoadBeginCookIterativeFlagsLocal(FBeginCookContext& BeginContext) const;
 	/** Initialize the sandbox for a new cook session */
-	void BeginCookSandbox(FBeginCookContext& BeginContext);
+	UNREALED_API void BeginCookSandbox(FBeginCookContext& BeginContext);
 
 	/** Set parameters that rely on config settings from startup and don't depend on StartCook* functions. */
-	void LoadInitializeConfigSettings(const FString& InOutputDirectoryOverride);
-	void SetInitializeConfigSettings(UE::Cook::FInitializeConfigSettings&& Settings);
-	void ParseCookFilters();
+	UNREALED_API void LoadInitializeConfigSettings(const FString& InOutputDirectoryOverride);
+	UNREALED_API void SetInitializeConfigSettings(UE::Cook::FInitializeConfigSettings&& Settings);
+	UNREALED_API void ParseCookFilters();
 
 	/** Set parameters that rely on config and CookByTheBook settings. */
-	void LoadBeginCookConfigSettings(FBeginCookContext& BeginContext);
-	void SetBeginCookConfigSettings(FBeginCookContext& BeginContext, UE::Cook::FBeginCookConfigSettings&& Settings);
-	void SetNeverCookPackageConfigSettings(FBeginCookContext& BeginContext, UE::Cook::FBeginCookConfigSettings& Settings);
+	UNREALED_API void LoadBeginCookConfigSettings(FBeginCookContext& BeginContext);
+	UNREALED_API void SetBeginCookConfigSettings(FBeginCookContext& BeginContext, UE::Cook::FBeginCookConfigSettings&& Settings);
+	UNREALED_API void SetNeverCookPackageConfigSettings(FBeginCookContext& BeginContext, UE::Cook::FBeginCookConfigSettings& Settings);
 
 	/** Finalize the package store */
-	void FinalizePackageStore();
+	UNREALED_API void FinalizePackageStore();
 
 	/** Empties SavePackageContexts and deletes the contents */
-	void ClearPackageStoreContexts();
+	UNREALED_API void ClearPackageStoreContexts();
 
 	/** Initialize shaders for the specified platforms when running cook on the fly. */
-	void InitializeShadersForCookOnTheFly(const TArrayView<ITargetPlatform* const>& NewTargetPlatforms);
+	UNREALED_API void InitializeShadersForCookOnTheFly(const TArrayView<ITargetPlatform* const>& NewTargetPlatforms);
 
 	/**
 	 * Some content plugins does not support all target platforms.
 	 * Build up a map of unsupported packages per platform that can be checked before saving.
 	 * const because it is not always called and should avoid side effects
 	 */
-	void DiscoverPlatformSpecificNeverCookPackages(
+	UNREALED_API void DiscoverPlatformSpecificNeverCookPackages(
 		const TArrayView<const ITargetPlatform* const>& TargetPlatforms, const TArray<FString>& UBTPlatformStrings,
 		UE::Cook::FBeginCookConfigSettings& Settings) const;
 
@@ -1101,7 +1101,7 @@ private:
 	 * @param Packages List of packages to use as the root set for dependency checking
 	 * @param Found return value, all objects which package is dependent on
 	 */
-	void GetDependentPackages( const TSet<UPackage*>& Packages, TSet<FName>& Found);
+	UNREALED_API void GetDependentPackages( const TSet<UPackage*>& Packages, TSet<FName>& Found);
 
 	/**
 	 * GetDependentPackages
@@ -1110,7 +1110,7 @@ private:
 	 * @param Root set of packages to use when looking for dependencies
 	 * @param FoundPackages list of packages which were found
 	 */
-	void GetDependentPackages(const TSet<FName>& RootPackages, TSet<FName>& FoundPackages);
+	UNREALED_API void GetDependentPackages(const TSet<FName>& RootPackages, TSet<FName>& FoundPackages);
 
 	/**
 	 * ContainsWorld
@@ -1119,7 +1119,7 @@ private:
 	 * @param PackageName to return if it contains the a UWorld object or a ULevel
 	 * @return true if the Package contains a UWorld or ULevel false otherwise
 	 */
-	bool ContainsMap(const FName& PackageName) const;
+	UNREALED_API bool ContainsMap(const FName& PackageName) const;
 
 
 	/** 
@@ -1129,7 +1129,7 @@ private:
 	 * @param RedirectedPaths map of original to redirected object paths
 	 * @return true if the Package contains a redirector false otherwise
 	 */
-	bool ContainsRedirector(const FName& PackageName, TMap<FSoftObjectPath, FSoftObjectPath>& RedirectedPaths) const;
+	UNREALED_API bool ContainsRedirector(const FName& PackageName, TMap<FSoftObjectPath, FSoftObjectPath>& RedirectedPaths) const;
 	
 	/**
 	 * Prepares save by calling BeginCacheForCookedPlatformData on all UObjects in the package. 
@@ -1141,12 +1141,12 @@ private:
 	 * @param bIsPrecaching true if called for precaching
 	 * @return false if time slice was reached, true if all objects have had BeginCacheForCookedPlatformData called
 	 */
-	UE::Cook::EPollStatus PrepareSave(UE::Cook::FPackageData& PackageData, UE::Cook::FCookerTimer& Timer,
+	UNREALED_API UE::Cook::EPollStatus PrepareSave(UE::Cook::FPackageData& PackageData, UE::Cook::FCookerTimer& Timer,
 		bool bPrecaching);
-	UE::Cook::EPollStatus PrepareSaveInternal(UE::Cook::FPackageData& PackageData, UE::Cook::FCookerTimer& Timer,
+	UNREALED_API UE::Cook::EPollStatus PrepareSaveInternal(UE::Cook::FPackageData& PackageData, UE::Cook::FCookerTimer& Timer,
 		bool bPrecaching);
 	/** Call BeginCacheForCookedPlatformData on all objects in a PackageData or a GeneratorPackage's current round. */
-	UE::Cook::EPollStatus CallBeginCacheOnObjects(UE::Cook::FPackageData& PackageData, UPackage* Package,
+	UNREALED_API UE::Cook::EPollStatus CallBeginCacheOnObjects(UE::Cook::FPackageData& PackageData, UPackage* Package,
 		TArray<FWeakObjectPtr>& Objects, int32& NextIndex, UE::Cook::FCookerTimer& Timer);
 
 	/**
@@ -1158,14 +1158,14 @@ private:
 	 *        the save again. If true, all data will be wiped.
 	 * @param ReleaseSaveReason Why the save data is being released, allows specifying how much to tear down
 	 */
-	void ReleaseCookedPlatformData(UE::Cook::FPackageData& PackageData, UE::Cook::EReleaseSaveReason ReleaseSaveReason);
+	UNREALED_API void ReleaseCookedPlatformData(UE::Cook::FPackageData& PackageData, UE::Cook::EReleaseSaveReason ReleaseSaveReason);
 
 	/**
 	 * Poll the PendingCookedPlatformDatas and release their resources when they are complete.
 	 * This is done inside of PumpSaveQueue, but is also required after a cancelled cook so that references to
 	 * the pending objects will be eventually dropped.
 	 */
-	void TickCancels();
+	UNREALED_API void TickCancels();
 
 	/**
 	* GetCurrentIniVersionStrings gets the current ini version strings for compare against previous cook
@@ -1173,7 +1173,7 @@ private:
 	* @param IniVersionStrings return list of the important current ini version strings
 	* @return false if function fails (should assume all platforms are out of date)
 	*/
-	bool GetCurrentIniVersionStrings( const ITargetPlatform* TargetPlatform,
+	UNREALED_API bool GetCurrentIniVersionStrings( const ITargetPlatform* TargetPlatform,
 		UE::Cook::FIniSettingContainer& IniVersionStrings ) const;
 
 	/**
@@ -1182,73 +1182,73 @@ private:
 	* @param IniVersionStrings return list of the previous cooks ini version strings
 	* @return false if function fails to find the ini version strings
 	*/
-	bool GetCookedIniVersionStrings( const ITargetPlatform* TargetPlatform,
+	UNREALED_API bool GetCookedIniVersionStrings( const ITargetPlatform* TargetPlatform,
 		UE::Cook::FIniSettingContainer& IniVersionStrings, TMap<FString, FString>& AdditionalStrings ) const;
 
 	/**
 	 * Test the CurrentCookSettings against the previous cooksettings to decide if we have to wipe the
 	 * previous cook even when running iteratively.
 	 */
-	bool ArePreviousCookSettingsCompatible(const TMap<FName, FString>& CurrentCookSettings,
+	UNREALED_API bool ArePreviousCookSettingsCompatible(const TMap<FName, FString>& CurrentCookSettings,
 		const ITargetPlatform* TargetPlatform) const;
 	/** Save the CurrentCookSettings into the output directory. */
-	void SaveCookSettings(const TMap<FName, FString>& CurrentCookSettings, const ITargetPlatform* TargetPlatform);
+	UNREALED_API void SaveCookSettings(const TMap<FName, FString>& CurrentCookSettings, const ITargetPlatform* TargetPlatform);
 	/**
 	 * Populate a map suitable for saving as an ini with the current value of the Cook settings that need to be
 	 * tested for compatibility.
 	 */
-	TMap<FName, FString> CalculateCookSettingStrings() const;
+	UNREALED_API TMap<FName, FString> CalculateCookSettingStrings() const;
 	/** Get the path to the CookSettings metadata file for the given platform and current output directory. */
-	FString GetCookSettingsFileName(const ITargetPlatform* TargetPlatform) const;
+	UNREALED_API FString GetCookSettingsFileName(const ITargetPlatform* TargetPlatform) const;
 
 	/**
 	 * Convert a path to a full sandbox path. Is effected by the cooking dlc settings.
 	 * This function should be used instead of calling the FSandbox Sandbox->ConvertToSandboxPath functions.
 	 */
-	FString ConvertToFullSandboxPath( const FString &FileName, bool bForWrite = false ) const;
-	FString ConvertToFullSandboxPath( const FString &FileName, bool bForWrite, const FString& PlatformName ) const;
+	UNREALED_API FString ConvertToFullSandboxPath( const FString &FileName, bool bForWrite = false ) const;
+	UNREALED_API FString ConvertToFullSandboxPath( const FString &FileName, bool bForWrite, const FString& PlatformName ) const;
 
 	/**
 	 * GetSandboxAssetRegistryFilename
 	 * 
 	 * @return full path of the asset registry in the sandbox
 	 */
-	FString GetSandboxAssetRegistryFilename();
+	UNREALED_API FString GetSandboxAssetRegistryFilename();
 
-	FString GetCookedAssetRegistryFilename(const FString& PlatformName);
+	UNREALED_API FString GetCookedAssetRegistryFilename(const FString& PlatformName);
 
 	/* @return Full path of the CachedEditorThumbnails.bin file in the sandbox */
-	FString GetSandboxCachedEditorThumbnailsFilename();
+	UNREALED_API FString GetSandboxCachedEditorThumbnailsFilename();
 
 	/**
 	 * Get the sandbox root directory for that platform. Is effected by the CookingDlc settings.
 	 * This should be used instead of calling the Sandbox function.
 	 */
-	FString GetSandboxDirectory( const FString& PlatformName ) const;
+	UNREALED_API FString GetSandboxDirectory( const FString& PlatformName ) const;
 
 	/* Create the delete-old-cooked-directory helper.*/
-	FAsyncIODelete& GetAsyncIODelete();
+	UNREALED_API FAsyncIODelete& GetAsyncIODelete();
 
 	/** Is the local CookOnTheFlyServer cooking a Project+Engine+EmbeddedPlugin Patch based on a previous Release? */
-	bool IsCookingAgainstFixedBase() const;
+	UNREALED_API bool IsCookingAgainstFixedBase() const;
 
 	/** Returns whether or not we should populate the Asset Registry using the main game content */
-	bool ShouldPopulateFullAssetRegistry() const;
+	UNREALED_API bool ShouldPopulateFullAssetRegistry() const;
 
 	/**
 	 * GetBaseDirectoryForDLC
 	 * 
 	 * @return return the path to the DLC
 	 */
-	FString GetBaseDirectoryForDLC() const;
+	UNREALED_API FString GetBaseDirectoryForDLC() const;
 
-	FString GetContentDirectoryForDLC() const;
+	UNREALED_API FString GetContentDirectoryForDLC() const;
 
 	/**
 	 * Is the local CookOnTheFlyServer cooking a Project+Engine+EmbeddedPlugin Release that can be used as a
 	 * base for future DLC or Patches?
 	 */
-	bool IsCreatingReleaseVersion();
+	UNREALED_API bool IsCreatingReleaseVersion();
 
 	/**
 	 * Checks if important ini settings have changed since last cook for each target platform 
@@ -1256,14 +1256,14 @@ private:
 	 * @param TargetPlatforms to check if out of date
 	 * @param OutOfDateTargetPlatforms return list of out of date target platforms which should be cleaned
 	 */
-	bool IniSettingsOutOfDate( const ITargetPlatform* TargetPlatform ) const;
+	UNREALED_API bool IniSettingsOutOfDate( const ITargetPlatform* TargetPlatform ) const;
 
 	/**
 	 * Saves ini settings which are in the memory cache to the hard drive in ini files
 	 *
 	 * @param TargetPlatforms to save
 	 */
-	bool SaveCurrentIniSettings( const ITargetPlatform* TargetPlatform ) const;
+	UNREALED_API bool SaveCurrentIniSettings( const ITargetPlatform* TargetPlatform ) const;
 
 
 
@@ -1273,17 +1273,17 @@ private:
 		return (CookFlags & InCookFlags) != ECookInitializationFlags::None;
 	}
 
-	void RouteBeginCacheForCookedPlatformData(UE::Cook::FPackageData& PackageData, UObject* Obj,
+	UNREALED_API void RouteBeginCacheForCookedPlatformData(UE::Cook::FPackageData& PackageData, UObject* Obj,
 		const ITargetPlatform* TargetPlatform, UE::Cook::ECachedCookedPlatformDataEvent* ExistingEvent);
-	bool RouteIsCachedCookedPlatformDataLoaded(UE::Cook::FPackageData& PackageData, UObject* Obj,
+	UNREALED_API bool RouteIsCachedCookedPlatformDataLoaded(UE::Cook::FPackageData& PackageData, UObject* Obj,
 		const ITargetPlatform* TargetPlatform, UE::Cook::ECachedCookedPlatformDataEvent* ExistingEvent);
-	EPackageWriterResult SavePackageBeginCacheForCookedPlatformData(FName PackageName,
+	UNREALED_API EPackageWriterResult SavePackageBeginCacheForCookedPlatformData(FName PackageName,
 		const ITargetPlatform* TargetPlatform, TConstArrayView<UObject*> SaveableObjects, uint32 SaveFlags);
 
 	/** Cook (save) a package and process the results */
-	void SaveCookedPackage(UE::Cook::FSaveCookedPackageContext& Context);
+	UNREALED_API void SaveCookedPackage(UE::Cook::FSaveCookedPackageContext& Context);
 	/** Helper for package saves using ExternalActors: record ExternalActors for iterative builds. */
-	void RecordExternalActorDependencies(TConstArrayView<FName> ExternalActorDependencies);
+	UNREALED_API void RecordExternalActorDependencies(TConstArrayView<FName> ExternalActorDependencies);
 	friend class UE::Cook::FSaveCookedPackageContext;
 
 	/**
@@ -1291,80 +1291,80 @@ private:
 	 *
 	 * @param Platforms List of platforms to make global shader maps for
 	 */
-	void SaveGlobalShaderMapFiles(const TArrayView<const ITargetPlatform* const>& Platforms,
+	UNREALED_API void SaveGlobalShaderMapFiles(const TArrayView<const ITargetPlatform* const>& Platforms,
 		ODSCRecompileCommand RecompileCommand);
 
 
 	/** Create sandbox file in directory using current settings supplied */
-	void CreateSandboxFile(FBeginCookContext& BeginContext);
+	UNREALED_API void CreateSandboxFile(FBeginCookContext& BeginContext);
 	/** Gets the output directory respecting any command line overrides */
-	FString GetOutputDirectoryOverride(FBeginCookContext& BeginContext) const;
+	UNREALED_API FString GetOutputDirectoryOverride(FBeginCookContext& BeginContext) const;
 
 	/**
 	 * Populate cooked packages from the PackageWriter's previous manifest and assetregistry of cooked output.
 	 * Delete any out of date packages from the PackageWriter's manifest.
 	 */
-	void PopulateCookedPackages(const TConstArrayView<const ITargetPlatform*> TargetPlatforms);
+	UNREALED_API void PopulateCookedPackages(const TConstArrayView<const ITargetPlatform*> TargetPlatforms);
 
 	/** Waits for the AssetRegistry to complete so that we know any missing assets are missing on disk */
-	void BlockOnAssetRegistry();
+	UNREALED_API void BlockOnAssetRegistry();
 
 	/** Construct or refresh-for-filechanges the platform-specific asset registry for the given platforms */
-	void RefreshPlatformAssetRegistries(const TArrayView<const ITargetPlatform* const>& TargetPlatforms);
+	UNREALED_API void RefreshPlatformAssetRegistries(const TArrayView<const ITargetPlatform* const>& TargetPlatforms);
 
 	/** Generates long package names for all files to be cooked */
-	void GenerateLongPackageNames(TArray<FName>& FilesInPath, TMap<FName, UE::Cook::FInstigator>& Instigators);
+	UNREALED_API void GenerateLongPackageNames(TArray<FName>& FilesInPath, TMap<FName, UE::Cook::FInstigator>& Instigators);
 
-	UE::Cook::EPollStatus ConditionalCreateGeneratorPackage(UE::Cook::FPackageData& PackageData, bool bPrecaching);
+	UNREALED_API UE::Cook::EPollStatus ConditionalCreateGeneratorPackage(UE::Cook::FPackageData& PackageData, bool bPrecaching);
 
 	/** Generate the list of cook-time-created packages created by the Generator. */
-	UE::Cook::EPollStatus QueueGeneratedPackages(UE::Cook::FGeneratorPackage& Generator,
+	UNREALED_API UE::Cook::EPollStatus QueueGeneratedPackages(UE::Cook::FGeneratorPackage& Generator,
 		UE::Cook::FPackageData& PackageData);
 	/** Run additional steps in PrepareSave that are required when the package has a Generator. */
-	UE::Cook::EPollStatus PrepareSaveGeneratedPackage(UE::Cook::FGeneratorPackage& Generator,
+	UNREALED_API UE::Cook::EPollStatus PrepareSaveGeneratedPackage(UE::Cook::FGeneratorPackage& Generator,
 		UE::Cook::FPackageData& PackageData, UE::Cook::FCookerTimer& Timer, bool bPrecaching);
 	/** Call BeginCacheForCookedPlatformData on the objects the Generator plans to move into its main UPackage. */
-	UE::Cook::EPollStatus BeginCacheObjectsToMove(UE::Cook::FGeneratorPackage& Generator,
+	UNREALED_API UE::Cook::EPollStatus BeginCacheObjectsToMove(UE::Cook::FGeneratorPackage& Generator,
 		UE::Cook::FCookGenerationInfo& Info, UE::Cook::FCookerTimer& Timer,
 		TArray<ICookPackageSplitter::FGeneratedPackageForPreSave>& GeneratedPackagesForPresave);
 	/** Call the Generator's PreSaveGeneratorPackage to create/move objects into its main UPackage. */
-	UE::Cook::EPollStatus PreSaveGeneratorPackage(UE::Cook::FPackageData& PackageData,
+	UNREALED_API UE::Cook::EPollStatus PreSaveGeneratorPackage(UE::Cook::FPackageData& PackageData,
 		UE::Cook::FGeneratorPackage& Generator, UE::Cook::FCookGenerationInfo& Info,
 		TArray<ICookPackageSplitter::FGeneratedPackageForPreSave>& GeneratedPackagesForPresave);
 	/** Construct the list of generated packages that is required for some of the CookPackageSplitter interface calls. */
-	void ConstructGeneratedPackagesForPresave(UE::Cook::FPackageData& PackageData, UE::Cook::FGeneratorPackage& Generator,
+	UNREALED_API void ConstructGeneratedPackagesForPresave(UE::Cook::FPackageData& PackageData, UE::Cook::FGeneratorPackage& Generator,
 		TArray<ICookPackageSplitter::FGeneratedPackageForPreSave>& GeneratedPackagesForPresave);
 	/** Call BeginCacheForCookedPlatformData on any undeclared objects in the Generator's main UPackage after the move. */
-	UE::Cook::EPollStatus BeginCachePostMove(UE::Cook::FGeneratorPackage& Generator,
+	UNREALED_API UE::Cook::EPollStatus BeginCachePostMove(UE::Cook::FGeneratorPackage& Generator,
 		UE::Cook::FCookGenerationInfo& Info, UE::Cook::FCookerTimer& Timer);
 
 	/** Try creating (or finding from earlier creation) the generated package for later population */
-	UPackage* TryCreateGeneratedPackage(UE::Cook::FGeneratorPackage& Generator, UE::Cook::FCookGenerationInfo& GeneratedInfo);
+	UNREALED_API UPackage* TryCreateGeneratedPackage(UE::Cook::FGeneratorPackage& Generator, UE::Cook::FCookGenerationInfo& GeneratedInfo);
 	/** Try calling the splitter's populate to create the package */
-	UE::Cook::EPollStatus TryPopulateGeneratedPackage(UE::Cook::FGeneratorPackage& Generator,
+	UNREALED_API UE::Cook::EPollStatus TryPopulateGeneratedPackage(UE::Cook::FGeneratorPackage& Generator,
 		UE::Cook::FCookGenerationInfo& GeneratedInfo);
 
-	ICookedPackageWriter& FindOrCreatePackageWriter(const ITargetPlatform* TargetPlatform);
-	const ICookedPackageWriter* FindPackageWriter(const ITargetPlatform* TargetPlatform) const;
-	void FindOrCreateSaveContexts(TConstArrayView<const ITargetPlatform*> TargetPlatforms);
-	UE::Cook::FCookSavePackageContext& FindOrCreateSaveContext(const ITargetPlatform* TargetPlatform);
-	const UE::Cook::FCookSavePackageContext* FindSaveContext(const ITargetPlatform* TargetPlatform) const;
+	UNREALED_API ICookedPackageWriter& FindOrCreatePackageWriter(const ITargetPlatform* TargetPlatform);
+	UNREALED_API const ICookedPackageWriter* FindPackageWriter(const ITargetPlatform* TargetPlatform) const;
+	UNREALED_API void FindOrCreateSaveContexts(TConstArrayView<const ITargetPlatform*> TargetPlatforms);
+	UNREALED_API UE::Cook::FCookSavePackageContext& FindOrCreateSaveContext(const ITargetPlatform* TargetPlatform);
+	UNREALED_API const UE::Cook::FCookSavePackageContext* FindSaveContext(const ITargetPlatform* TargetPlatform) const;
 	/** Allocate a new FCookSavePackageContext and ICookedPackageWriter for the given platform. */
-	UE::Cook::FCookSavePackageContext* CreateSaveContext(const ITargetPlatform* TargetPlatform);
+	UNREALED_API UE::Cook::FCookSavePackageContext* CreateSaveContext(const ITargetPlatform* TargetPlatform);
 	/**
 	 * Delete files that exist for the package in the PackageWriter from a previous cook. Used when packages are
 	 * detected as invalidated for iterative cook after cook startup (e.g. generated packages).
 	 */
-	void DeleteOutputForPackage(FName PackageName, const ITargetPlatform* TargetPlatform);
+	UNREALED_API void DeleteOutputForPackage(FName PackageName, const ITargetPlatform* TargetPlatform);
 
 	/**
 	 * Set the given package as the active package for diagnostics (e.g. hidden dependencies).
  	 * Must be balanced with call to ClearActivePackage. Consider using FScopedActivePackage instead.
 	 * @param PackageTrackingOpsName Optional op for TObjectPtr tracking. Set to NAME_None to not set TObjectPtr context.
 	 */
-	void SetActivePackage(FName PackageName, FName PackageTrackingOpsName);
+	UNREALED_API void SetActivePackage(FName PackageName, FName PackageTrackingOpsName);
 	/** Clear the package for diagnostics. Must be balanced with call to SetActivePackage */
-	void ClearActivePackage();
+	UNREALED_API void ClearActivePackage();
 	struct FActivePackageData
 	{
 		FName PackageName;
@@ -1389,15 +1389,15 @@ private:
 		TGuardValue<bool> SoftGCGuard;
 	};
 	/** Callback for FGenericCrashContext; provides the current ActivePackage as context. */
-	void DumpCrashContext(FCrashContextExtendedWriter& Writer);
+	UNREALED_API void DumpCrashContext(FCrashContextExtendedWriter& Writer);
 	/** Callback for analytics when a new UPackage is loaded. */
-	void OnDiscoveredPackageDebug(FName PackageName, const UE::Cook::FInstigator& Instigator);
+	UNREALED_API void OnDiscoveredPackageDebug(FName PackageName, const UE::Cook::FInstigator& Instigator);
 	/** Callback for analytics when a TObjectPtr is read. */
-	void OnObjectHandleReadDebug(TArrayView<const UObject*const> ReadObjects);
+	UNREALED_API void OnObjectHandleReadDebug(TArrayView<const UObject*const> ReadObjects);
 	/** Send warnings/telemetry when a discovered or read package is found to be a hidden dependency. */
-	void ReportHiddenDependency(FName Referencer, FName Dependency);
+	UNREALED_API void ReportHiddenDependency(FName Referencer, FName Dependency);
 
-	static UCookOnTheFlyServer* ActiveCOTFS;
+	static UNREALED_API UCookOnTheFlyServer* ActiveCOTFS;
 	uint32		StatLoadedPackageCount = 0;
 	uint32		StatSavedPackageCount = 0;
 	int32		PackageDataFromBaseGameNum = 0;
@@ -1480,8 +1480,8 @@ private:
 	double LastCookableObjectTickTime = 0.;
 
 	// Cook events that can be listenned to
-	static FOnCookByTheBookStarted CookByTheBookStartedEvent;
-	static FOnCookByTheBookFinished CookByTheBookFinishedEvent;
+	static UNREALED_API FOnCookByTheBookStarted CookByTheBookStartedEvent;
+	static UNREALED_API FOnCookByTheBookFinished CookByTheBookFinishedEvent;
 
 
 	// These structs are TUniquePtr rather than inline members so we can keep their headers private.
diff --git a/Engine/Source/Editor/UnrealEd/Classes/Editor/AssetGuideline.h b/Engine/Source/Editor/UnrealEd/Classes/Editor/AssetGuideline.h
--- a/Engine/Source/Editor/UnrealEd/Classes/Editor/AssetGuideline.h
+++ b/Engine/Source/Editor/UnrealEd/Classes/Editor/AssetGuideline.h
@@ -7,7 +7,7 @@
 
 // Helper struct for setting string console ini values.
 USTRUCT()
-struct UNREALED_API FIniStringValue
+struct FIniStringValue
 {
 	GENERATED_BODY()
 		
@@ -34,8 +34,8 @@ struct UNREALED_API FIniStringValue
 * This class intentionally does not accept FText arguments. The project using your bundled asset would need to have
 * your localization tables, and we currently do not support text table referencing.
 */
-UCLASS(Blueprintable)
-class UNREALED_API UAssetGuideline : public UAssetUserData
+UCLASS(Blueprintable, MinimalAPI)
+class UAssetGuideline : public UAssetUserData
 {
 	GENERATED_BODY()
 
@@ -58,23 +58,23 @@ private:
 
 	TWeakPtr<class SNotificationItem> NotificationPtr;
 
-	static bool bAssetGuidelinesEnabled;
+	static UNREALED_API bool bAssetGuidelinesEnabled;
 #endif // WITH_EDITORONLY_DATA
 
 #if WITH_EDITOR
 public:
 	/** Begin UObject interface */
-	virtual bool IsPostLoadThreadSafe() const override;
-	virtual void PostLoad() override;
-	virtual void BeginDestroy() override;
+	UNREALED_API virtual bool IsPostLoadThreadSafe() const override;
+	UNREALED_API virtual void PostLoad() override;
+	UNREALED_API virtual void BeginDestroy() override;
 	/** End UObject interface */
 
 	static void EnableAssetGuidelines(bool isEnabled) { bAssetGuidelinesEnabled = isEnabled; }
 	static bool AreAssetGuidelinesEnabled() { return bAssetGuidelinesEnabled; }
 
 private:
-	void EnableMissingGuidelines(TArray<FString> IncorrectPlugins, TArray<FIniStringValue> IncorrectProjectSettings);
-	void DismissNotifications();
-	void RemoveAssetGuideline();
+	UNREALED_API void EnableMissingGuidelines(TArray<FString> IncorrectPlugins, TArray<FIniStringValue> IncorrectProjectSettings);
+	UNREALED_API void DismissNotifications();
+	UNREALED_API void RemoveAssetGuideline();
 #endif // WITH_EDITOR
 };
