[UnrealEd]

diff --git a/Engine/Source/Editor/UnrealEd/Classes/Editor/EditorEngine.h b/Engine/Source/Editor/UnrealEd/Classes/Editor/EditorEngine.h
--- a/Engine/Source/Editor/UnrealEd/Classes/Editor/EditorEngine.h
+++ b/Engine/Source/Editor/UnrealEd/Classes/Editor/EditorEngine.h
@@ -287,14 +287,14 @@ struct FAssetReferenceFilterContext
  */
 PRAGMA_DISABLE_DEPRECATION_WARNINGS
 
-UCLASS(config=Engine, transient)
-class UNREALED_API UEditorEngine : public UEngine
+UCLASS(config=Engine, transient, MinimalAPI)
+class UEditorEngine : public UEngine
 {
 public:
 	GENERATED_BODY()
 
 public:
-	UEditorEngine(const FObjectInitializer& ObjectInitializer = FObjectInitializer::Get());
+	UNREALED_API UEditorEngine(const FObjectInitializer& ObjectInitializer = FObjectInitializer::Get());
 
 	// Objects.
 	UPROPERTY()
@@ -583,10 +583,10 @@ public:
 	 * Add a viewport client.
 	 * @return Index to the new item
 	 */
-	int32 AddViewportClients(FEditorViewportClient* ViewportClient);
+	UNREALED_API int32 AddViewportClients(FEditorViewportClient* ViewportClient);
 
 	/** Remove a viewport client */
-	void RemoveViewportClients(FEditorViewportClient* ViewportClient);
+	UNREALED_API void RemoveViewportClients(FEditorViewportClient* ViewportClient);
 
 	/** List of level editor viewport clients for level specific actions */
 	const TArray<class FLevelEditorViewportClient*>& GetLevelViewportClients() { return LevelViewportClients; }
@@ -596,10 +596,10 @@ public:
 	 * Add a viewport client.
 	 * @return Index to the new item
 	 */
-	int32 AddLevelViewportClients(FLevelEditorViewportClient* ViewportClient);
+	UNREALED_API int32 AddLevelViewportClients(FLevelEditorViewportClient* ViewportClient);
 
 	/** Remove a level editor viewport client */
-	void RemoveLevelViewportClients(FLevelEditorViewportClient* ViewportClient);
+	UNREALED_API void RemoveLevelViewportClients(FLevelEditorViewportClient* ViewportClient);
 
 	/** Called when the level editor viewport clients list changed */
 	FViewportClientListChangedEvent& OnLevelViewportClientListChanged() { return LevelViewportClientListChangedEvent; }
@@ -761,130 +761,130 @@ public:
 	void BroadcastEndCameraMovement(UObject& Object) const { OnEndCameraTransformEvent.Broadcast(Object); }
 
 	/**	Broadcasts that an object has been reimported. THIS SHOULD NOT BE PUBLIC */
-	void BroadcastObjectReimported(UObject* InObject);
+	UNREALED_API void BroadcastObjectReimported(UObject* InObject);
 
 	//~ Begin UObject Interface.
-	virtual void FinishDestroy() override;
-	virtual void PostEditChangeProperty(FPropertyChangedEvent& PropertyChangedEvent) override;
-	static void AddReferencedObjects(UObject* InThis, FReferenceCollector& Collector);
+	UNREALED_API virtual void FinishDestroy() override;
+	UNREALED_API virtual void PostEditChangeProperty(FPropertyChangedEvent& PropertyChangedEvent) override;
+	static UNREALED_API void AddReferencedObjects(UObject* InThis, FReferenceCollector& Collector);
 	//~ End UObject Interface.
 
 	//~ Begin UEngine Interface.
 public:
-	virtual void Init(IEngineLoop* InEngineLoop) override;
-	virtual void PreExit() override;
-	virtual float GetMaxTickRate(float DeltaTime, bool bAllowFrameRateSmoothing = true) const override;
-	virtual void Tick(float DeltaSeconds, bool bIdleMode) override;
-	virtual bool ShouldDrawBrushWireframe(AActor* InActor) override;
-	virtual void NotifyToolsOfObjectReplacement(const TMap<UObject*, UObject*>& OldToNewInstanceMap) override;
-	virtual bool ShouldThrottleCPUUsage() const override;
-	virtual bool IsPropertyColorationColorFeatureActivated() const override;
-	virtual bool GetPropertyColorationColor(class UObject* Object, FColor& OutColor) override;
-	virtual bool WorldIsPIEInNewViewport(UWorld* InWorld) override;
-	virtual void FocusNextPIEWorld(UWorld* CurrentPieWorld, bool previous = false) override;
-	virtual void ResetPIEAudioSetting(UWorld *CurrentPieWorld) override;
-	virtual class UGameViewportClient* GetNextPIEViewport(UGameViewportClient* CurrentViewport) override;
-	virtual UWorld* CreatePIEWorldByDuplication(FWorldContext &WorldContext, UWorld* InWorld, FString &PlayWorldMapName) override;
-	virtual void PostCreatePIEWorld(UWorld* InWorld) override;
+	UNREALED_API virtual void Init(IEngineLoop* InEngineLoop) override;
+	UNREALED_API virtual void PreExit() override;
+	UNREALED_API virtual float GetMaxTickRate(float DeltaTime, bool bAllowFrameRateSmoothing = true) const override;
+	UNREALED_API virtual void Tick(float DeltaSeconds, bool bIdleMode) override;
+	UNREALED_API virtual bool ShouldDrawBrushWireframe(AActor* InActor) override;
+	UNREALED_API virtual void NotifyToolsOfObjectReplacement(const TMap<UObject*, UObject*>& OldToNewInstanceMap) override;
+	UNREALED_API virtual bool ShouldThrottleCPUUsage() const override;
+	UNREALED_API virtual bool IsPropertyColorationColorFeatureActivated() const override;
+	UNREALED_API virtual bool GetPropertyColorationColor(class UObject* Object, FColor& OutColor) override;
+	UNREALED_API virtual bool WorldIsPIEInNewViewport(UWorld* InWorld) override;
+	UNREALED_API virtual void FocusNextPIEWorld(UWorld* CurrentPieWorld, bool previous = false) override;
+	UNREALED_API virtual void ResetPIEAudioSetting(UWorld *CurrentPieWorld) override;
+	UNREALED_API virtual class UGameViewportClient* GetNextPIEViewport(UGameViewportClient* CurrentViewport) override;
+	UNREALED_API virtual UWorld* CreatePIEWorldByDuplication(FWorldContext &WorldContext, UWorld* InWorld, FString &PlayWorldMapName) override;
+	UNREALED_API virtual void PostCreatePIEWorld(UWorld* InWorld) override;
 	virtual bool GetMapBuildCancelled() const override { return false; }
 	virtual void SetMapBuildCancelled(bool InCancelled) override { /* Intentionally empty. */ }
-	virtual void HandleNetworkFailure(UWorld *World, UNetDriver *NetDriver, ENetworkFailure::Type FailureType, const FString& ErrorString) override;
+	UNREALED_API virtual void HandleNetworkFailure(UWorld *World, UNetDriver *NetDriver, ENetworkFailure::Type FailureType, const FString& ErrorString) override;
 	virtual ERHIFeatureLevel::Type GetDefaultWorldFeatureLevel() const override { return DefaultWorldFeatureLevel; }
-	virtual bool GetPreviewPlatformName(FName& PlatformName) const override;
+	UNREALED_API virtual bool GetPreviewPlatformName(FName& PlatformName) const override;
 
 protected:
-	virtual void InitializeObjectReferences() override;
-	virtual void ProcessToggleFreezeCommand(UWorld* InWorld) override;
-	virtual void ProcessToggleFreezeStreamingCommand(UWorld* InWorld) override;
-	virtual void HandleBrowseToDefaultMapFailure(FWorldContext& Context, const FString& TextURL, const FString& Error) override;
+	UNREALED_API virtual void InitializeObjectReferences() override;
+	UNREALED_API virtual void ProcessToggleFreezeCommand(UWorld* InWorld) override;
+	UNREALED_API virtual void ProcessToggleFreezeStreamingCommand(UWorld* InWorld) override;
+	UNREALED_API virtual void HandleBrowseToDefaultMapFailure(FWorldContext& Context, const FString& TextURL, const FString& Error) override;
 private:
-	virtual void RemapGamepadControllerIdForPIE(class UGameViewportClient* GameViewport, int32 &ControllerId) override;
-	virtual TSharedPtr<SViewport> GetGameViewportWidget() const override;
-	virtual void TriggerStreamingDataRebuild() override;
+	UNREALED_API virtual void RemapGamepadControllerIdForPIE(class UGameViewportClient* GameViewport, int32 &ControllerId) override;
+	UNREALED_API virtual TSharedPtr<SViewport> GetGameViewportWidget() const override;
+	UNREALED_API virtual void TriggerStreamingDataRebuild() override;
 
-	virtual bool NetworkRemapPath(UNetConnection* Connection, FString& Str, bool bReading = true) override;
-	virtual bool NetworkRemapPath(UPendingNetGame* PendingNetGame, FString& Str, bool bReading = true) override;
+	UNREALED_API virtual bool NetworkRemapPath(UNetConnection* Connection, FString& Str, bool bReading = true) override;
+	UNREALED_API virtual bool NetworkRemapPath(UPendingNetGame* PendingNetGame, FString& Str, bool bReading = true) override;
 
-	virtual bool AreEditorAnalyticsEnabled() const override;
-	virtual void CreateStartupAnalyticsAttributes(TArray<FAnalyticsEventAttribute>& StartSessionAttributes) const override;
-	virtual void CheckAndHandleStaleWorldObjectReferences(FWorldContext* InWorldContext) override;
+	UNREALED_API virtual bool AreEditorAnalyticsEnabled() const override;
+	UNREALED_API virtual void CreateStartupAnalyticsAttributes(TArray<FAnalyticsEventAttribute>& StartSessionAttributes) const override;
+	UNREALED_API virtual void CheckAndHandleStaleWorldObjectReferences(FWorldContext* InWorldContext) override;
 
 	/** Called during editor init and whenever the vanilla status might have changed, to set the flag on the base class */
-	void UpdateIsVanillaProduct();
+	UNREALED_API void UpdateIsVanillaProduct();
 
 	/** Called when hotreload adds a new class to create volume factories */
-	void CreateVolumeFactoriesForNewClasses(const TArray<UClass*>& NewClasses);
+	UNREALED_API void CreateVolumeFactoriesForNewClasses(const TArray<UClass*>& NewClasses);
 	//~ End UEngine Interface.
 	
 	//~ Begin FExec Interface
 protected:
-	virtual bool Exec_Editor( UWorld* InWorld, const TCHAR* Cmd, FOutputDevice& Ar=*GLog ) override;
+	UNREALED_API virtual bool Exec_Editor( UWorld* InWorld, const TCHAR* Cmd, FOutputDevice& Ar=*GLog ) override;
 	//~ End FExec Interface
 
 public:
-	bool	CommandIsDeprecated( const TCHAR* CommandStr, FOutputDevice& Ar );
+	UNREALED_API bool	CommandIsDeprecated( const TCHAR* CommandStr, FOutputDevice& Ar );
 	
 	/**
 	 * Exec command handlers
 	 */
-	bool	HandleCallbackCommand( UWorld* InWorld, const TCHAR* Str, FOutputDevice& Ar );
-	bool	HandleTestPropsCommand( const TCHAR* Str, FOutputDevice& Ar );
-	bool	HandleMapCommand( const TCHAR* Str, FOutputDevice& Ar, UWorld* InWorld );
-	bool	HandleSelectCommand( const TCHAR* Str, FOutputDevice& Ar, UWorld* InWorld );
-	bool	HandleDeleteCommand( const TCHAR* Str, FOutputDevice& Ar, UWorld* InWorld );
-	bool	HandleLightmassDebugCommand( const TCHAR* Str, FOutputDevice& Ar );
-	bool	HandleLightmassStatsCommand( const TCHAR* Str, FOutputDevice& Ar );
-	bool	HandleSwarmDistributionCommand( const TCHAR* Str, FOutputDevice& Ar );
-	bool	HandleLightmassImmediateImportCommand( const TCHAR* Str, FOutputDevice& Ar );
-	bool	HandleLightmassImmediateProcessCommand( const TCHAR* Str, FOutputDevice& Ar );
-	bool	HandleLightmassSortCommand( const TCHAR* Str, FOutputDevice& Ar );
-	bool	HandleLightmassDebugMaterialCommand( const TCHAR* Str, FOutputDevice& Ar );
-	bool	HandleLightmassPaddingCommand( const TCHAR* Str, FOutputDevice& Ar );
-	bool	HandleLightmassDebugPaddingCommand( const TCHAR* Str, FOutputDevice& Ar );
-	bool	HandleLightmassProfileCommand( const TCHAR* Str, FOutputDevice& Ar );
-	bool	HandleSelectNameCommand( const TCHAR* Str, FOutputDevice& Ar, UWorld* InWorld  );
-	bool	HandleDumpPublicCommand( const TCHAR* Str, FOutputDevice& Ar );
-	bool	HandleJumpToCommand( const TCHAR* Str, FOutputDevice& Ar );
-	bool	HandleBugItGoCommand( const TCHAR* Str, FOutputDevice& Ar );
-	bool	HandleBugItCommand(const TCHAR* Str, FOutputDevice& Ar);
-	bool	HandleTagSoundsCommand( const TCHAR* Str, FOutputDevice& Ar );
-	bool	HandlecheckSoundsCommand( const TCHAR* Str, FOutputDevice& Ar );
-	bool	HandleFixupBadAnimNotifiersCommand( const TCHAR* Str, FOutputDevice& Ar );
-	bool	HandleSetDetailModeCommand( const TCHAR* Str, FOutputDevice& Ar );
-	bool	HandleSetDetailModeViewCommand( const TCHAR* Str, FOutputDevice& Ar, UWorld* InWorld );
-	bool	HandleCleanBSPMaterialCommand( const TCHAR* Str, FOutputDevice& Ar, UWorld* InWorld  );
-	bool	HandleAutoMergeStaticMeshCommand( const TCHAR* Str, FOutputDevice& Ar );
-	bool	HandleAddSelectedCommand( const TCHAR* Str, FOutputDevice& Ar );
-	bool	HandleToggleSocketGModeCommand( const TCHAR* Str, FOutputDevice& Ar );
-	bool	HandleListMapPackageDependenciesCommand( const TCHAR* Str, FOutputDevice& Ar );
-	bool	HandleRebuildVolumesCommand( const TCHAR* Str, FOutputDevice& Ar, UWorld* InWorld );
-	bool	HandleRemoveArchtypeFlagCommand( const TCHAR* Str, FOutputDevice& Ar );
-	bool	HandleStartMovieCaptureCommand( const TCHAR* Cmd, FOutputDevice& Ar );
-	bool	HandleBuildMaterialTextureStreamingData( const TCHAR* Cmd, FOutputDevice& Ar );
+	UNREALED_API bool	HandleCallbackCommand( UWorld* InWorld, const TCHAR* Str, FOutputDevice& Ar );
+	UNREALED_API bool	HandleTestPropsCommand( const TCHAR* Str, FOutputDevice& Ar );
+	UNREALED_API bool	HandleMapCommand( const TCHAR* Str, FOutputDevice& Ar, UWorld* InWorld );
+	UNREALED_API bool	HandleSelectCommand( const TCHAR* Str, FOutputDevice& Ar, UWorld* InWorld );
+	UNREALED_API bool	HandleDeleteCommand( const TCHAR* Str, FOutputDevice& Ar, UWorld* InWorld );
+	UNREALED_API bool	HandleLightmassDebugCommand( const TCHAR* Str, FOutputDevice& Ar );
+	UNREALED_API bool	HandleLightmassStatsCommand( const TCHAR* Str, FOutputDevice& Ar );
+	UNREALED_API bool	HandleSwarmDistributionCommand( const TCHAR* Str, FOutputDevice& Ar );
+	UNREALED_API bool	HandleLightmassImmediateImportCommand( const TCHAR* Str, FOutputDevice& Ar );
+	UNREALED_API bool	HandleLightmassImmediateProcessCommand( const TCHAR* Str, FOutputDevice& Ar );
+	UNREALED_API bool	HandleLightmassSortCommand( const TCHAR* Str, FOutputDevice& Ar );
+	UNREALED_API bool	HandleLightmassDebugMaterialCommand( const TCHAR* Str, FOutputDevice& Ar );
+	UNREALED_API bool	HandleLightmassPaddingCommand( const TCHAR* Str, FOutputDevice& Ar );
+	UNREALED_API bool	HandleLightmassDebugPaddingCommand( const TCHAR* Str, FOutputDevice& Ar );
+	UNREALED_API bool	HandleLightmassProfileCommand( const TCHAR* Str, FOutputDevice& Ar );
+	UNREALED_API bool	HandleSelectNameCommand( const TCHAR* Str, FOutputDevice& Ar, UWorld* InWorld  );
+	UNREALED_API bool	HandleDumpPublicCommand( const TCHAR* Str, FOutputDevice& Ar );
+	UNREALED_API bool	HandleJumpToCommand( const TCHAR* Str, FOutputDevice& Ar );
+	UNREALED_API bool	HandleBugItGoCommand( const TCHAR* Str, FOutputDevice& Ar );
+	UNREALED_API bool	HandleBugItCommand(const TCHAR* Str, FOutputDevice& Ar);
+	UNREALED_API bool	HandleTagSoundsCommand( const TCHAR* Str, FOutputDevice& Ar );
+	UNREALED_API bool	HandlecheckSoundsCommand( const TCHAR* Str, FOutputDevice& Ar );
+	UNREALED_API bool	HandleFixupBadAnimNotifiersCommand( const TCHAR* Str, FOutputDevice& Ar );
+	UNREALED_API bool	HandleSetDetailModeCommand( const TCHAR* Str, FOutputDevice& Ar );
+	UNREALED_API bool	HandleSetDetailModeViewCommand( const TCHAR* Str, FOutputDevice& Ar, UWorld* InWorld );
+	UNREALED_API bool	HandleCleanBSPMaterialCommand( const TCHAR* Str, FOutputDevice& Ar, UWorld* InWorld  );
+	UNREALED_API bool	HandleAutoMergeStaticMeshCommand( const TCHAR* Str, FOutputDevice& Ar );
+	UNREALED_API bool	HandleAddSelectedCommand( const TCHAR* Str, FOutputDevice& Ar );
+	UNREALED_API bool	HandleToggleSocketGModeCommand( const TCHAR* Str, FOutputDevice& Ar );
+	UNREALED_API bool	HandleListMapPackageDependenciesCommand( const TCHAR* Str, FOutputDevice& Ar );
+	UNREALED_API bool	HandleRebuildVolumesCommand( const TCHAR* Str, FOutputDevice& Ar, UWorld* InWorld );
+	UNREALED_API bool	HandleRemoveArchtypeFlagCommand( const TCHAR* Str, FOutputDevice& Ar );
+	UNREALED_API bool	HandleStartMovieCaptureCommand( const TCHAR* Cmd, FOutputDevice& Ar );
+	UNREALED_API bool	HandleBuildMaterialTextureStreamingData( const TCHAR* Cmd, FOutputDevice& Ar );
 
 	/**
 	 * Initializes the Editor.
 	 */
-	void InitEditor(IEngineLoop* InEngineLoop);
+	UNREALED_API void InitEditor(IEngineLoop* InEngineLoop);
 
 	/**
 	 * Constructs a default cube builder brush, this function MUST be called at the AFTER UEditorEngine::Init in order to guarantee builder brush and other required subsystems exist.
 	 *
 	 * @param		InWorld			World in which to create the builder brush.
 	 */
-	void InitBuilderBrush( UWorld* InWorld );
+	UNREALED_API void InitBuilderBrush( UWorld* InWorld );
 
 	/** Access user setting for audio mute. */
-	bool IsRealTimeAudioMuted() const;
+	UNREALED_API bool IsRealTimeAudioMuted() const;
 
 	/** Set user setting for audio mute. */
-	void MuteRealTimeAudio(bool bMute);
+	UNREALED_API void MuteRealTimeAudio(bool bMute);
 
 	/** Access user setting for audio level. Fractional volume 0.0->1.0 */
-	float GetRealTimeAudioVolume() const;
+	UNREALED_API float GetRealTimeAudioVolume() const;
 
 	/** Set user setting for audio mute. Fractional volume 0.0->1.0 */
-	void SetRealTimeAudioVolume(float VolumeLevel);
+	UNREALED_API void SetRealTimeAudioVolume(float VolumeLevel);
 
 	/**
 	 * Updates a single viewport
@@ -893,7 +893,7 @@ public:
 	 * @param bLinkedOrthoMovement	True if orthographic viewport movement is linked
 	 * @return - Whether a NON-realtime viewport has updated in this call.  Used to help time-slice canvas redraws
 	 */
-	bool UpdateSingleViewportClient(FEditorViewportClient* InViewportClient, const bool bInAllowNonRealtimeViewportToDraw, bool bLinkedOrthoMovement );
+	UNREALED_API bool UpdateSingleViewportClient(FEditorViewportClient* InViewportClient, const bool bInAllowNonRealtimeViewportToDraw, bool bLinkedOrthoMovement );
 
 	/** Used for generating status bar text */
 	enum EMousePositionType
@@ -909,15 +909,15 @@ public:
 
 
 	// Execute a command that is safe for rebuilds.
-	virtual bool SafeExec( UWorld* InWorld, const TCHAR* Cmd, FOutputDevice& Ar=*GLog );
+	UNREALED_API virtual bool SafeExec( UWorld* InWorld, const TCHAR* Cmd, FOutputDevice& Ar=*GLog );
 	// Process an incoming network message meant for the editor server
-	bool Exec_StaticMesh( UWorld* InWorld, const TCHAR* Str, FOutputDevice& Ar );
-	bool Exec_Brush( UWorld* InWorld, const TCHAR* Str, FOutputDevice& Ar );
-	bool Exec_Poly( UWorld* InWorld, const TCHAR* Str, FOutputDevice& Ar );
-	bool Exec_Obj( const TCHAR* Str, FOutputDevice& Ar );
-	bool Exec_Camera( const TCHAR* Str, FOutputDevice& Ar );
-	bool Exec_Transaction(const TCHAR* Str, FOutputDevice& Ar);
-	bool Exec_Particle(const TCHAR* Str, FOutputDevice& Ar);
+	UNREALED_API bool Exec_StaticMesh( UWorld* InWorld, const TCHAR* Str, FOutputDevice& Ar );
+	UNREALED_API bool Exec_Brush( UWorld* InWorld, const TCHAR* Str, FOutputDevice& Ar );
+	UNREALED_API bool Exec_Poly( UWorld* InWorld, const TCHAR* Str, FOutputDevice& Ar );
+	UNREALED_API bool Exec_Obj( const TCHAR* Str, FOutputDevice& Ar );
+	UNREALED_API bool Exec_Camera( const TCHAR* Str, FOutputDevice& Ar );
+	UNREALED_API bool Exec_Transaction(const TCHAR* Str, FOutputDevice& Ar);
+	UNREALED_API bool Exec_Particle(const TCHAR* Str, FOutputDevice& Ar);
 
 	/**
 	 * Executes each line of text in a file sequentially, as if each were a separate command
@@ -926,20 +926,20 @@ public:
 	 * @param InFilename The name of the file to load and execute
 	 * @param Ar Output device
 	 */
-	void ExecFile( UWorld* InWorld, const TCHAR* InFilename, FOutputDevice& Ar );
+	UNREALED_API void ExecFile( UWorld* InWorld, const TCHAR* InFilename, FOutputDevice& Ar );
 
 	//~ Begin Transaction Interfaces.
-	virtual int32 BeginTransaction(const TCHAR* TransactionContext, const FText& Description, UObject* PrimaryObject) override;
-	int32 BeginTransaction(const FText& Description);
-	virtual bool CanTransact() override;
-	virtual int32 EndTransaction() override;
-	virtual void CancelTransaction(int32 Index) override;
-	void ResetTransaction(const FText& Reason);
-	bool UndoTransaction(bool bCanRedo = true);
-	bool RedoTransaction();
-	bool IsTransactionActive() const;
-	FText GetTransactionName() const;
-	bool IsObjectInTransactionBuffer( const UObject* Object ) const;
+	UNREALED_API virtual int32 BeginTransaction(const TCHAR* TransactionContext, const FText& Description, UObject* PrimaryObject) override;
+	UNREALED_API int32 BeginTransaction(const FText& Description);
+	UNREALED_API virtual bool CanTransact() override;
+	UNREALED_API virtual int32 EndTransaction() override;
+	UNREALED_API virtual void CancelTransaction(int32 Index) override;
+	UNREALED_API void ResetTransaction(const FText& Reason);
+	UNREALED_API bool UndoTransaction(bool bCanRedo = true);
+	UNREALED_API bool RedoTransaction();
+	UNREALED_API bool IsTransactionActive() const;
+	UNREALED_API FText GetTransactionName() const;
+	UNREALED_API bool IsObjectInTransactionBuffer( const UObject* Object ) const;
 
 	enum EMapRebuildType
 	{
@@ -957,14 +957,14 @@ public:
 	 * @param RebuildMap	The map to be rebuilt
 	 * @param RebuildType	Allows to filter which of the map's level should be rebuilt
 	 */
-	void RebuildMap(UWorld* RebuildMap, EMapRebuildType RebuildType);
+	UNREALED_API void RebuildMap(UWorld* RebuildMap, EMapRebuildType RebuildType);
 	
 	/**
 	 * Quickly rebuilds a single level (no bounds build, visibility testing or Bsp tree optimization).
 	 *
 	 * @param Level	The level to be rebuilt.
 	 */
-	void RebuildLevel(ULevel& Level);
+	UNREALED_API void RebuildLevel(ULevel& Level);
 	
 	/**
 	 * Builds up a model from a set of brushes. Used by RebuildLevel.
@@ -973,7 +973,7 @@ public:
 	 * @param bSelectedBrushesOnly	Use all brushes in the current level or just the selected ones?.
 	 * @param bTreatMovableBrushesAsStatic	Treat moveable brushes as static?.
 	 */
-	void RebuildModelFromBrushes(UModel* Model, bool bSelectedBrushesOnly, bool bTreatMovableBrushesAsStatic = false);
+	UNREALED_API void RebuildModelFromBrushes(UModel* Model, bool bSelectedBrushesOnly, bool bTreatMovableBrushesAsStatic = false);
 
 	/**
 	 * Builds up a model from a given set of brushes. Used by BspConversionTool to build brushes before converting them
@@ -982,20 +982,20 @@ public:
 	 * @param BrushesToBuild	List of brushes to build.
 	 * @param Model				Model into which to put the output.
 	 */
-	void RebuildModelFromBrushes(TArray<ABrush*>& BrushesToBuild, UModel* Model);
+	UNREALED_API void RebuildModelFromBrushes(TArray<ABrush*>& BrushesToBuild, UModel* Model);
 
 	/**
 	 * Rebuilds levels containing currently selected brushes and should be invoked after a brush has been modified
 	 */
-	void RebuildAlteredBSP();
+	UNREALED_API void RebuildAlteredBSP();
 
 	/** Helper method for executing the de/intersect CSG operation */
-	void BSPIntersectionHelper(UWorld* InWorld, ECsgOper Operation);
+	UNREALED_API void BSPIntersectionHelper(UWorld* InWorld, ECsgOper Operation);
 
 	/**
 	 * @return	A pointer to the named actor or NULL if not found.
 	 */
-	AActor* SelectNamedActor(const TCHAR *TargetActorName);
+	UNREALED_API AActor* SelectNamedActor(const TCHAR *TargetActorName);
 
 
 	/**
@@ -1007,14 +1007,14 @@ public:
 	 * @param InCameraOrigin	The location of the camera in world space
 	 * @param InCameraDirection	The normalized direction vector of the camera
 	 */
-	void MoveActorInFrontOfCamera( AActor& InActor, const FVector& InCameraOrigin, const FVector& InCameraDirection );
+	UNREALED_API void MoveActorInFrontOfCamera( AActor& InActor, const FVector& InCameraOrigin, const FVector& InCameraDirection );
 
 	/**
 	 * Moves all viewport cameras to the target actor.
 	 * @param	Actor					Target actor.
 	 * @param	bActiveViewportOnly		If true, move/reorient only the active viewport.
 	 */
-	void MoveViewportCamerasToActor(AActor& Actor,  bool bActiveViewportOnly);
+	UNREALED_API void MoveViewportCamerasToActor(AActor& Actor,  bool bActiveViewportOnly);
 
 	/**
 	* Moves all viewport cameras to focus on the provided array of actors.
@@ -1022,7 +1022,7 @@ public:
 
 	* @param	bActiveViewportOnly		If true, move/reorient only the active viewport.
 	*/
-	void MoveViewportCamerasToActor(const TArray<AActor*> &Actors, bool bActiveViewportOnly);
+	UNREALED_API void MoveViewportCamerasToActor(const TArray<AActor*> &Actors, bool bActiveViewportOnly);
 
 	/**
 	* Moves all viewport cameras to focus on the provided array of actors.
@@ -1030,21 +1030,21 @@ public:
 	* @param	Components				Target components (used of actors array is empty)
 	* @param	bActiveViewportOnly		If true, move/reorient only the active viewport.
 	*/
-	void MoveViewportCamerasToActor(const TArray<AActor*> &Actors, const TArray<UPrimitiveComponent*>& Components, bool bActiveViewportOnly);
+	UNREALED_API void MoveViewportCamerasToActor(const TArray<AActor*> &Actors, const TArray<UPrimitiveComponent*>& Components, bool bActiveViewportOnly);
 
 	/**
 	* Moves all viewport cameras to focus on the provided component.
 	* @param	Component				Target component
 	* @param	bActiveViewportOnly		If true, move/reorient only the active viewport.
 	*/
-	void MoveViewportCamerasToComponent(const USceneComponent* Component, bool bActiveViewportOnly);
+	UNREALED_API void MoveViewportCamerasToComponent(const USceneComponent* Component, bool bActiveViewportOnly);
 
 	/**
 	* Moves all viewport cameras to focus on the provided set of elements.
 	* @param	SelectionSet			Target elements
 	* @param	bActiveViewportOnly		If true, move/reorient only the active viewport.
 	*/
-	void MoveViewportCamerasToElement(const UTypedElementSelectionSet* SelectionSet, bool bActiveViewportOnly) const;
+	UNREALED_API void MoveViewportCamerasToElement(const UTypedElementSelectionSet* SelectionSet, bool bActiveViewportOnly) const;
 
 	/**
 	 * Moves all viewport cameras to focus on the provided bounding box.
@@ -1052,7 +1052,7 @@ public:
 	 * @param	bActiveViewportOnly			If true, move/reorient only the active viewport.
 	 * @param	DrawDebubBoxTimeInSeconds	If greater than 0 a debug box is drawn representing the bounding box.It will be drawn for specified time.
 	 */
-	void MoveViewportCamerasToBox(const FBox& BoundingBox, bool bActiveViewportOnly, float DrawDebugBoxTimeInSeconds = 0.f) const;
+	UNREALED_API void MoveViewportCamerasToBox(const FBox& BoundingBox, bool bActiveViewportOnly, float DrawDebugBoxTimeInSeconds = 0.f) const;
 
 	/** 
 	 * Snaps an element in a direction.  Optionally will align with the trace normal.
@@ -1064,14 +1064,14 @@ public:
 	 * @param InDestination		The destination element we want to move this actor to, unset assumes we just want to go towards the floor
 	 * @return					Whether or not the actor was moved.
 	 */
-	bool SnapElementTo(const FTypedElementHandle& InElementHandle, const bool InAlign, const bool InUseLineTrace, const bool InUseBounds, const bool InUsePivot, const FTypedElementHandle& InDestination = FTypedElementHandle(), TArrayView<const FTypedElementHandle> InElementsToIgnore = TArrayView<const FTypedElementHandle>());
+	UNREALED_API bool SnapElementTo(const FTypedElementHandle& InElementHandle, const bool InAlign, const bool InUseLineTrace, const bool InUseBounds, const bool InUsePivot, const FTypedElementHandle& InDestination = FTypedElementHandle(), TArrayView<const FTypedElementHandle> InElementsToIgnore = TArrayView<const FTypedElementHandle>());
 
 	/**
 	 * Snaps the view of the camera to that of the provided element.
 	 *
 	 * @param InElementHandle	The element the camera is going to be snapped to.
 	 */
-	void SnapViewTo(const FTypedElementHandle& InElementHandle);
+	UNREALED_API void SnapViewTo(const FTypedElementHandle& InElementHandle);
 
 	/**
 	 * Remove the roll, pitch and/or yaw from the perspective viewports' cameras.
@@ -1080,7 +1080,7 @@ public:
 	 * @param	Pitch		If true, the camera pitch is reset to zero.
 	 * @param	Yaw			If true, the camera yaw is reset to zero.
 	 */
-	void RemovePerspectiveViewRotation(bool Roll, bool Pitch, bool Yaw);
+	UNREALED_API void RemovePerspectiveViewRotation(bool Roll, bool Pitch, bool Yaw);
 
 	//
 	// Pivot handling.
@@ -1108,7 +1108,7 @@ public:
 	 * @param	Redraw			Whether to redraw viewports
 	 * @param	TransReset		Human readable reason for resetting the transaction system
 	 */
-	virtual void Cleanse( bool ClearSelection, bool Redraw, const FText& TransReset, bool bTransReset = true );
+	UNREALED_API virtual void Cleanse( bool ClearSelection, bool Redraw, const FText& TransReset, bool bTransReset = true );
 	virtual void FinishAllSnaps() { }
 
 	/**
@@ -1137,7 +1137,7 @@ public:
 	 * @param	bSelectActor	Whether or not to select the spawned actor.
 	 * @return					A pointer to the newly added actor, or NULL if add failed.
 	 */
-	virtual AActor* AddActor(ULevel* InLevel, UClass* Class, const FTransform& Transform, bool bSilent = false, EObjectFlags ObjectFlags = RF_Transactional, bool bSelectActor = true);
+	UNREALED_API virtual AActor* AddActor(ULevel* InLevel, UClass* Class, const FTransform& Transform, bool bSilent = false, EObjectFlags ObjectFlags = RF_Transactional, bool bSelectActor = true);
 
 	/**
 	 * Adds actors to the world at the specified location using export text.
@@ -1147,27 +1147,27 @@ public:
 	 * @param	ObjectFlags		The object flags to place on the spawned actor.
 	 * @return					A pointer to the newly added actor, or NULL if add failed.
 	 */
-	virtual TArray<AActor*> AddExportTextActors(const FString& ExportText, bool bSilent, EObjectFlags ObjectFlags = RF_Transactional);
+	UNREALED_API virtual TArray<AActor*> AddExportTextActors(const FString& ExportText, bool bSilent, EObjectFlags ObjectFlags = RF_Transactional);
 
 	virtual void NoteActorMovement() { check(0); }
-	virtual UTransactor* CreateTrans();
+	UNREALED_API virtual UTransactor* CreateTrans();
 
 	/** Plays an editor sound, loading the sound on demand if necessary (if the user has sounds enabled.)  The reference to the sound asset is not retained. */
-	void PlayEditorSound( const FString& SoundAssetName );
+	UNREALED_API void PlayEditorSound( const FString& SoundAssetName );
 
 	/** Plays an editor sound (if the user has sounds enabled.) */
-	void PlayEditorSound( USoundBase* InSound );
+	UNREALED_API void PlayEditorSound( USoundBase* InSound );
 
 	/**
 	 * Returns true if currently able to play a sound and if the user has sounds enabled.
 	 */
-	bool CanPlayEditorSound() const;
+	UNREALED_API bool CanPlayEditorSound() const;
 
 
 	/**
 	 * Returns the preview audio component
 	 */
-	UAudioComponent* GetPreviewAudioComponent();
+	UNREALED_API UAudioComponent* GetPreviewAudioComponent();
 
 	/**
 	 * Returns an audio component linked to the current scene that it is safe to play a sound on
@@ -1175,7 +1175,7 @@ public:
 	 * @param	Sound		A sound to attach to the audio component
 	 * @param	SoundNode	A sound node that is attached to the audio component when the sound is NULL
 	 */
-	UAudioComponent* ResetPreviewAudioComponent(USoundBase* Sound = nullptr, USoundNode* SoundNode = nullptr);
+	UNREALED_API UAudioComponent* ResetPreviewAudioComponent(USoundBase* Sound = nullptr, USoundNode* SoundNode = nullptr);
 
 	/**
 	 * Plays a preview of a specified sound or node
@@ -1183,25 +1183,25 @@ public:
 	 * @param	Sound		A sound to attach to the audio component
 	 * @param	SoundNode	A sound node that is attached to the audio component when the sound is NULL
 	 */
-	UAudioComponent* PlayPreviewSound(USoundBase* Sound, USoundNode* SoundNode = nullptr);
+	UNREALED_API UAudioComponent* PlayPreviewSound(USoundBase* Sound, USoundNode* SoundNode = nullptr);
 
 	/**
 	 * Clean up any world specific editor components so they can be GC correctly
 	 */
-	void ClearPreviewComponents();
+	UNREALED_API void ClearPreviewComponents();
 
 	
 	/**
 	 * Close all the edit windows for assets that are owned by the world being closed
 	 */
-	void CloseEditedWorldAssets(UWorld* InWorld);
+	UNREALED_API void CloseEditedWorldAssets(UWorld* InWorld);
 
 	/**
 	 * Redraws all editor viewport clients.
 	 *
 	 * @param	bInvalidateHitProxies		[opt] If true (the default), invalidates cached hit proxies too.
 	 */
-	void RedrawAllViewports(bool bInvalidateHitProxies=true);
+	UNREALED_API void RedrawAllViewports(bool bInvalidateHitProxies=true);
 
 	UE_DEPRECATED(5.3, "InvalidateChildViewports is deprecated")
 	void InvalidateChildViewports(FSceneViewStateInterface* InParentView, bool bInvalidateHitProxies=true) {}
@@ -1212,7 +1212,7 @@ public:
 	 * @param	InClass		The class to find the factory for.
 	 * @return				A pointer to the factory to use.  NULL if no factories support this class.
 	 */
-	UActorFactory* FindActorFactoryForActorClass( const UClass* InClass );
+	UNREALED_API UActorFactory* FindActorFactoryForActorClass( const UClass* InClass );
 
 	/**
 	 * Looks for an actor factory spawned from the specified class.
@@ -1220,7 +1220,7 @@ public:
 	 * @param	InClass		The factories class to find
 	 * @return				A pointer to the factory to use.  NULL if no factories support this class.
 	 */
-	UActorFactory* FindActorFactoryByClass( const UClass* InClass ) const;
+	UNREALED_API UActorFactory* FindActorFactoryByClass( const UClass* InClass ) const;
 
 	/**
 	* Looks for an actor factory spawned from the specified class, for the specified UClass for an actor.
@@ -1229,7 +1229,7 @@ public:
 	* @param	InActorClass		The class to find the factory for.
 	* @return				A pointer to the factory to use.  NULL if no factories support this class.
 	*/
-	UActorFactory* FindActorFactoryByClassForActorClass( const UClass* InFactoryClass, const UClass* InActorClass );
+	UNREALED_API UActorFactory* FindActorFactoryByClassForActorClass( const UClass* InFactoryClass, const UClass* InActorClass );
 
 	/**
 	 * Uses the supplied factory to create an actor at the clicked location and adds to level.
@@ -1238,7 +1238,7 @@ public:
 	 * @param	ObjectFlags				[opt] The flags to apply to the actor when it is created
 	 * @return							A pointer to the new actor, or NULL on fail.
 	 */
-	AActor* UseActorFactoryOnCurrentSelection( UActorFactory* Factory, const FTransform* InActorTransform, EObjectFlags ObjectFlags = RF_Transactional );
+	UNREALED_API AActor* UseActorFactoryOnCurrentSelection( UActorFactory* Factory, const FTransform* InActorTransform, EObjectFlags ObjectFlags = RF_Transactional );
 
 	/**
 	 * Uses the supplied factory to create an actor at the clicked location and adds to level.
@@ -1250,7 +1250,7 @@ public:
 	 * @param	ObjectFlags				[opt] The flags to apply to the actor when it is created
 	 * @return							A pointer to the new actor, or NULL on fail.
 	 */
-	AActor* UseActorFactory( UActorFactory* Factory, const FAssetData& AssetData, const FTransform* ActorLocation, EObjectFlags ObjectFlags = RF_Transactional );
+	UNREALED_API AActor* UseActorFactory( UActorFactory* Factory, const FAssetData& AssetData, const FTransform* ActorLocation, EObjectFlags ObjectFlags = RF_Transactional );
 
 	/**
 	 * Replaces the selected Actors with the same number of a different kind of Actor using the specified factory to spawn the new Actors
@@ -1258,7 +1258,7 @@ public:
 	 * 
 	 * @param Factory - the Factory to use to create Actors
 	 */
-	void ReplaceSelectedActors(UActorFactory* Factory, const FAssetData& AssetData, bool bCopySourceProperties = true);
+	UNREALED_API void ReplaceSelectedActors(UActorFactory* Factory, const FAssetData& AssetData, bool bCopySourceProperties = true);
 
 	/**
 	 * Replaces specified Actors with the same number of a different kind of Actor using the specified factory to spawn the new Actors
@@ -1269,7 +1269,7 @@ public:
 	 * @param ActorsToReplace - Actors to replace
 	 * @param OutNewActors - Actors that were created
 	 */
-	void ReplaceActors(UActorFactory* Factory, const FAssetData& AssetData, const TArray<AActor*>& ActorsToReplace, TArray<AActor*>* OutNewActors = nullptr, bool bCopySourceProperties = true);
+	UNREALED_API void ReplaceActors(UActorFactory* Factory, const FAssetData& AssetData, const TArray<AActor*>& ActorsToReplace, TArray<AActor*>* OutNewActors = nullptr, bool bCopySourceProperties = true);
 
 	/**
 	 * Converts passed in brushes into a single static mesh actor. 
@@ -1280,7 +1280,7 @@ public:
 	 *
 	 * @return							Returns the newly created actor with the newly created static mesh.
 	 */
-	AActor* ConvertBrushesToStaticMesh(const FString& InStaticMeshPackageName, TArray<ABrush*>& InBrushesToConvert, const FVector& InPivotLocation);
+	UNREALED_API AActor* ConvertBrushesToStaticMesh(const FString& InStaticMeshPackageName, TArray<ABrush*>& InBrushesToConvert, const FVector& InPivotLocation);
 
 	/**
 	 * Converts passed in light actors into new actors of another type.
@@ -1291,7 +1291,7 @@ public:
 	 *
 	 * @param	ConvertToClass	The light class we are going to convert to.
 	 */
-	void ConvertLightActors( UClass* ConvertToClass );
+	UNREALED_API void ConvertLightActors( UClass* ConvertToClass );
 
 	/**
 	 * Converts passed in actors into new actors of the specified type.
@@ -1308,7 +1308,7 @@ public:
 	 * @param	bUseSpecialCases			If true, looks for classes that can be handled by hardcoded conversions
 	 * @param	InStaticMeshPackageName		The name to save the brushes to.
 	 */
-	void DoConvertActors( const TArray<AActor*>& ActorsToConvert, UClass* ConvertToClass, const TSet<FString>& ComponentsToConsider, bool bUseSpecialCases, const FString& InStaticMeshPackageName );
+	UNREALED_API void DoConvertActors( const TArray<AActor*>& ActorsToConvert, UClass* ConvertToClass, const TSet<FString>& ComponentsToConsider, bool bUseSpecialCases, const FString& InStaticMeshPackageName );
 
 	/**
 	 * Sets up for a potentially deferred ConvertActors call, based on if any brushes are being converted to a static mesh. If one (or more)
@@ -1319,24 +1319,24 @@ public:
 	 * @param	ComponentsToConsider	Names of components to consider for property copying as well
 	 * @param	bUseSpecialCases		If true, looks for classes that can be handled by hardcoded conversions
 	 */
-	void ConvertActors( const TArray<AActor*>& ActorsToConvert, UClass* ConvertToClass, const TSet<FString>& ComponentsToConsider, bool bUseSpecialCases = false );
+	UNREALED_API void ConvertActors( const TArray<AActor*>& ActorsToConvert, UClass* ConvertToClass, const TSet<FString>& ComponentsToConsider, bool bUseSpecialCases = false );
 
 	/**
 	 * Changes the state of preview mesh mode to on or off.
 	 *
 	 * @param	bState	Enables the preview mesh mode if true; Disables the preview mesh mode if false.
 	 */
-	void SetPreviewMeshMode( bool bState );
+	UNREALED_API void SetPreviewMeshMode( bool bState );
 
 	/**
 	 * Updates the position of the preview mesh in the level.
 	 */
-	void UpdatePreviewMesh();
+	UNREALED_API void UpdatePreviewMesh();
 
 	/**
 	 * Changes the preview mesh to the next one.
 	 */
-	void CyclePreviewMesh();
+	UNREALED_API void CyclePreviewMesh();
 
 	/**
 	 * Copy selected actors to the clipboard.
@@ -1391,7 +1391,7 @@ public:
 	*
 	* @param InWorld	The world in which the rebuild the CSG brushes for 
 	*/
-	virtual void csgRebuild( UWorld* InWorld );
+	UNREALED_API virtual void csgRebuild( UWorld* InWorld );
 
 	/**
 	*
@@ -1403,10 +1403,10 @@ public:
 	*
 	* returns true if poly not available
 	*/
-	static bool polyFindBrush(UModel* InModel, int32 iSurf, FPoly &Poly);
+	static UNREALED_API bool polyFindBrush(UModel* InModel, int32 iSurf, FPoly &Poly);
 
 	UE_DEPRECATED(5.1, "polyFindMaster is deprecated; please use polyFindBrush instead")
-	virtual bool polyFindMaster( UModel* InModel, int32 iSurf, FPoly& Poly );
+	UNREALED_API virtual bool polyFindMaster( UModel* InModel, int32 iSurf, FPoly& Poly );
 
 	/**
 	 * Update a the brush EdPoly corresponding to a newly-changed
@@ -1414,34 +1414,34 @@ public:
 	 *
 	 * Doesn't do any transaction tracking.
 	 */
-	static void polyUpdateBrush(UModel* Model, int32 iSurf, bool bUpdateTexCoords, bool bOnlyRefreshSurfaceMaterials);
+	static UNREALED_API void polyUpdateBrush(UModel* Model, int32 iSurf, bool bUpdateTexCoords, bool bOnlyRefreshSurfaceMaterials);
 
 	UE_DEPRECATED(5.1, "polyUpdateMaster is deprecated; please use polyUpdateBrush instead")
-	virtual void polyUpdateMaster( UModel* Model, int32 iSurf, bool bUpdateTexCoords, bool bOnlyRefreshSurfaceMaterials );
+	UNREALED_API virtual void polyUpdateMaster( UModel* Model, int32 iSurf, bool bUpdateTexCoords, bool bOnlyRefreshSurfaceMaterials );
 
 	/**
 	 * Populates a list with all polys that are linked to the specified poly.  The
 	 * resulting list includes the original poly.
 	 */
-	virtual void polyGetLinkedPolys( ABrush* InBrush, FPoly* InPoly, TArray<FPoly>* InPolyList );
+	UNREALED_API virtual void polyGetLinkedPolys( ABrush* InBrush, FPoly* InPoly, TArray<FPoly>* InPolyList );
 
 	/**
 	 * Takes a list of polygons and returns a list of the outside edges (edges which are not shared
 	 * by other polys in the list).
 	 */
-	virtual void polyGetOuterEdgeList( TArray<FPoly>* InPolyList, TArray<FEdge>* InEdgeList );
+	UNREALED_API virtual void polyGetOuterEdgeList( TArray<FPoly>* InPolyList, TArray<FEdge>* InEdgeList );
 
 	/**
 	 * Takes a list of polygons and creates a new list of polys which have no overlapping edges.  It splits
 	 * edges as necessary to achieve this.
 	 */
-	virtual void polySplitOverlappingEdges( TArray<FPoly>* InPolyList, TArray<FPoly>* InResult );
+	UNREALED_API virtual void polySplitOverlappingEdges( TArray<FPoly>* InPolyList, TArray<FPoly>* InResult );
 
 	/**
 	 * Sets and clears all Bsp node flags.  Affects all nodes, even ones that don't
 	 * really exist.
 	 */
-	virtual void polySetAndClearPolyFlags( UModel* Model, uint32 SetBits, uint32 ClearBits, bool SelectedOnly, bool UpdateBrush );
+	UNREALED_API virtual void polySetAndClearPolyFlags( UModel* Model, uint32 SetBits, uint32 ClearBits, bool SelectedOnly, bool UpdateBrush );
 
 	// Selection.
 	virtual void SelectActor(AActor* Actor, bool bInSelected, bool bNotify, bool bSelectEvenIfHidden = false, bool bForceRefresh = false) {}
@@ -1473,85 +1473,85 @@ public:
 	virtual void DeselectAllSurfaces() {}
 
 	// Bsp Poly selection virtuals from EditorCsg.cpp.
-	virtual void polySelectAll ( UModel* Model );
-	virtual void polySelectMatchingGroups( UModel* Model );
-	virtual void polySelectMatchingItems( UModel* Model );
-	virtual void polySelectCoplanars( UWorld* InWorld, UModel* Model );
-	virtual void polySelectAdjacents( UWorld* InWorld, UModel* Model );
-	virtual void polySelectAdjacentWalls( UWorld* InWorld, UModel* Model );
-	virtual void polySelectAdjacentFloors( UWorld* InWorld, UModel* Model );
-	virtual void polySelectAdjacentSlants( UWorld* InWorld, UModel* Model );
-	virtual void polySelectMatchingBrush( UModel* Model );
+	UNREALED_API virtual void polySelectAll ( UModel* Model );
+	UNREALED_API virtual void polySelectMatchingGroups( UModel* Model );
+	UNREALED_API virtual void polySelectMatchingItems( UModel* Model );
+	UNREALED_API virtual void polySelectCoplanars( UWorld* InWorld, UModel* Model );
+	UNREALED_API virtual void polySelectAdjacents( UWorld* InWorld, UModel* Model );
+	UNREALED_API virtual void polySelectAdjacentWalls( UWorld* InWorld, UModel* Model );
+	UNREALED_API virtual void polySelectAdjacentFloors( UWorld* InWorld, UModel* Model );
+	UNREALED_API virtual void polySelectAdjacentSlants( UWorld* InWorld, UModel* Model );
+	UNREALED_API virtual void polySelectMatchingBrush( UModel* Model );
 
 	/**
 	 * Selects surfaces whose material matches that of any selected surfaces.
 	 *
 	 * @param	bCurrentLevelOnly		If true, select
 	 */
-	virtual void polySelectMatchingMaterial(UWorld* InWorld, bool bCurrentLevelOnly);
+	UNREALED_API virtual void polySelectMatchingMaterial(UWorld* InWorld, bool bCurrentLevelOnly);
 
 	/**
 	 * Selects surfaces whose lightmap resolution matches that of any selected surfaces.
 	 *
 	 * @param	bCurrentLevelOnly		If true, select
 	 */
-	virtual void polySelectMatchingResolution(UWorld* InWorld, bool bCurrentLevelOnly);
+	UNREALED_API virtual void polySelectMatchingResolution(UWorld* InWorld, bool bCurrentLevelOnly);
 
-	virtual void polySelectReverse( UModel* Model );
-	virtual void polyMemorizeSet( UModel* Model );
-	virtual void polyRememberSet( UModel* Model );
-	virtual void polyXorSet( UModel* Model );
-	virtual void polyUnionSet( UModel* Model );
-	virtual void polyIntersectSet( UModel* Model );
-	virtual void polySelectZone( UModel *Model );
+	UNREALED_API virtual void polySelectReverse( UModel* Model );
+	UNREALED_API virtual void polyMemorizeSet( UModel* Model );
+	UNREALED_API virtual void polyRememberSet( UModel* Model );
+	UNREALED_API virtual void polyXorSet( UModel* Model );
+	UNREALED_API virtual void polyUnionSet( UModel* Model );
+	UNREALED_API virtual void polyIntersectSet( UModel* Model );
+	UNREALED_API virtual void polySelectZone( UModel *Model );
 
 	// Pan textures on selected polys.  Doesn't do transaction tracking.
-	virtual void polyTexPan( UModel* Model, int32 PanU, int32 PanV, int32 Absolute );
+	UNREALED_API virtual void polyTexPan( UModel* Model, int32 PanU, int32 PanV, int32 Absolute );
 
 	// Scale textures on selected polys. Doesn't do transaction tracking.
-	virtual void polyTexScale( UModel* Model,float UU, float UV, float VU, float VV, bool Absolute );
+	UNREALED_API virtual void polyTexScale( UModel* Model,float UU, float UV, float VU, float VV, bool Absolute );
 
 	// Map brush selection virtuals from EditorCsg.cpp.
-	virtual void MapSelectOperation( UWorld* InWorld, EBrushType BrushType );
-	virtual void MapSelectFlags( UWorld* InWorld, uint32 Flags );
+	UNREALED_API virtual void MapSelectOperation( UWorld* InWorld, EBrushType BrushType );
+	UNREALED_API virtual void MapSelectFlags( UWorld* InWorld, uint32 Flags );
 
 	// Put the first selected brush into the current Brush.
-	virtual void MapBrushGet(UWorld* InWorld);
+	UNREALED_API virtual void MapBrushGet(UWorld* InWorld);
 
 	// Replace all selected brushes with the current Brush.
-	virtual void mapBrushPut();
+	UNREALED_API virtual void mapBrushPut();
 
 	// Send all selected brushes in a level to the front of the hierarchy
-	virtual void mapSendToFirst(UWorld* InWorld);
+	UNREALED_API virtual void mapSendToFirst(UWorld* InWorld);
 
 	// Send all selected brushes in a level to the back of the hierarchy
-	virtual void mapSendToLast(UWorld* InWorld);
+	UNREALED_API virtual void mapSendToLast(UWorld* InWorld);
 
 	/**
 	 * Swaps position in the actor list for the first two selected actors in the current level
 	 */
-	virtual void mapSendToSwap(UWorld* InWorld);
-	virtual void MapSetBrush( UWorld* InWorld, EMapSetBrushFlags PropertiesMask, uint16 BrushColor, FName Group, uint32 SetPolyFlags, uint32 ClearPolyFlags, uint32 BrushType, int32 DrawType );
+	UNREALED_API virtual void mapSendToSwap(UWorld* InWorld);
+	UNREALED_API virtual void MapSetBrush( UWorld* InWorld, EMapSetBrushFlags PropertiesMask, uint16 BrushColor, FName Group, uint32 SetPolyFlags, uint32 ClearPolyFlags, uint32 BrushType, int32 DrawType );
 
 	// Bsp virtuals from Bsp.cpp.
-	virtual void bspRepartition( UWorld* InWorld, int32 iNode );
+	UNREALED_API virtual void bspRepartition( UWorld* InWorld, int32 iNode );
 
 	/** Convert a Bsp node to an EdPoly.  Returns number of vertices in Bsp node. */
-	virtual int32 bspNodeToFPoly( UModel* Model, int32 iNode, FPoly* EdPoly );
+	UNREALED_API virtual int32 bspNodeToFPoly( UModel* Model, int32 iNode, FPoly* EdPoly );
 
 	/**
 	 * Clean up all nodes after a CSG operation.  Resets temporary bit flags and unlinks
 	 * empty leaves.  Removes zero-vertex nodes which have nonzero-vertex coplanars.
 	 */
-	virtual void bspCleanup( UModel* Model );
+	UNREALED_API virtual void bspCleanup( UModel* Model );
 
 	/** 
 	 * Build EdPoly list from a model's Bsp. Not transactional.
 	 * @param DestArray helps build bsp FPolys in non-main threads. It also allows to perform this action without GUndo 
 	 *	interfering. Temporary results will be written to DestArray. Defaults to Model->Polys->Element
 	 */
-	virtual void bspBuildFPolys( UModel* Model, bool SurfLinks, int32 iNode, TArray<FPoly>* DestArray = NULL );
-	virtual void bspMergeCoplanars( UModel* Model, bool RemapLinks, bool MergeDisparateTextures );
+	UNREALED_API virtual void bspBuildFPolys( UModel* Model, bool SurfLinks, int32 iNode, TArray<FPoly>* DestArray = NULL );
+	UNREALED_API virtual void bspMergeCoplanars( UModel* Model, bool RemapLinks, bool MergeDisparateTextures );
 	/**
 	 * Performs any CSG operation between the brush and the world.
 	 *
@@ -1566,32 +1566,32 @@ public:
 	 * @param	bShowProgressBar				If true, display progress bar for complex brushes
 	 * @return									0 if nothing happened, 1 if the operation was error-free, or 1+N if N CSG errors occurred.
 	 */
-	virtual int32 bspBrushCSG( ABrush* Actor, UModel* Model, uint32 PolyFlags, EBrushType BrushType, ECsgOper CSGOper, bool bBuildBounds, bool bMergePolys, bool bReplaceNULLMaterialRefs, bool bShowProgressBar=true );
+	UNREALED_API virtual int32 bspBrushCSG( ABrush* Actor, UModel* Model, uint32 PolyFlags, EBrushType BrushType, ECsgOper CSGOper, bool bBuildBounds, bool bMergePolys, bool bReplaceNULLMaterialRefs, bool bShowProgressBar=true );
 
 	/**
 	 * Optimize a level's Bsp, eliminating T-joints where possible, and building side
 	 * links.  This does not always do a 100% perfect job, mainly due to imperfect 
 	 * levels, however it should never fail or return incorrect results.
 	 */
-	virtual void bspOptGeom( UModel* Model );
+	UNREALED_API virtual void bspOptGeom( UModel* Model );
 	
 	/**
 	 * Builds lighting information depending on passed in options.
 	 *
 	 * @param	Options		Options determining on what and how lighting is built
 	 */
-	void BuildLighting(const class FLightingBuildOptions& Options);
+	UNREALED_API void BuildLighting(const class FLightingBuildOptions& Options);
 
 	/** Updates the asynchronous static light building */
-	void UpdateBuildLighting();
+	UNREALED_API void UpdateBuildLighting();
 
 	/** Checks to see if the asynchronous lighting build is running or not */
-	bool IsLightingBuildCurrentlyRunning() const;
+	UNREALED_API bool IsLightingBuildCurrentlyRunning() const;
 
-	bool IsLightingBuildCurrentlyExporting() const;
+	UNREALED_API bool IsLightingBuildCurrentlyExporting() const;
 
 	/** Checks if asynchronous lighting is building, if so, it throws a warning notification and returns true */
-	bool WarnIfLightingBuildIsCurrentlyRunning();
+	UNREALED_API bool WarnIfLightingBuildIsCurrentlyRunning();
 
 	/**
 	 * Open a Fbx file with the given name, and import each sequence with the supplied Skeleton.
@@ -1602,7 +1602,7 @@ public:
 	 * @param Filename	The FBX filename
 	 * @param bImportMorphTracks	true to import any morph curve data.
 	 */
-	static UAnimSequence * ImportFbxAnimation( USkeleton* Skeleton, UObject* Outer, class UFbxAnimSequenceImportData* ImportData, const TCHAR* InFilename, const TCHAR* AnimName, bool bImportMorphTracks );
+	static UNREALED_API UAnimSequence * ImportFbxAnimation( USkeleton* Skeleton, UObject* Outer, class UFbxAnimSequenceImportData* ImportData, const TCHAR* InFilename, const TCHAR* AnimName, bool bImportMorphTracks );
 	/**
 	 * Reimport animation using SourceFilePath and SourceFileStamp 
 	 *
@@ -1614,19 +1614,19 @@ public:
 	 * @param bFactoryShowOptions	When true, create a UI popup asking the user for the reimport options.
 	 * @param ReimportUI			Optional parameter used to pass reimport options.
 	 */
-	static bool ReimportFbxAnimation( USkeleton* Skeleton, UAnimSequence* AnimSequence, class UFbxAnimSequenceImportData* ImportData, const TCHAR* InFilename, bool& bOutImportAll, const bool bFactoryShowOptions, UFbxImportUI* ReimportUI = nullptr);
+	static UNREALED_API bool ReimportFbxAnimation( USkeleton* Skeleton, UAnimSequence* AnimSequence, class UFbxAnimSequenceImportData* ImportData, const TCHAR* InFilename, bool& bOutImportAll, const bool bFactoryShowOptions, UFbxImportUI* ReimportUI = nullptr);
 
 
 	// Object management.
-	virtual void RenameObject(UObject* Object,UObject* NewOuter,const TCHAR* NewName, ERenameFlags Flags=REN_None);
+	UNREALED_API virtual void RenameObject(UObject* Object,UObject* NewOuter,const TCHAR* NewName, ERenameFlags Flags=REN_None);
 
 	// Level management.
-	void AnalyzeLevel(ULevel* Level,FOutputDevice& Ar);
+	UNREALED_API void AnalyzeLevel(ULevel* Level,FOutputDevice& Ar);
 
 	/**
 	 * Updates all components in the current level's scene.
 	 */
-	void EditorUpdateComponents();
+	UNREALED_API void EditorUpdateComponents();
 
 	/** 
 	 * Displays a modal message dialog 
@@ -1635,14 +1635,14 @@ public:
 	 * @param	InTitle		Title for the message box
 	 * @return	Returns the result of the modal message box
 	 */
-	EAppReturnType::Type OnModalMessageDialog(EAppMsgType::Type InMessage, const FText& InText, const FText& InTitle);
+	UNREALED_API EAppReturnType::Type OnModalMessageDialog(EAppMsgType::Type InMessage, const FText& InText, const FText& InTitle);
 
 	/** 
 	 * Returns whether an object should replace an exisiting one or not 
 	 * @param	Filename		Filename of the package
 	 * @return	Returns whether the objects should replace the already existing ones.
 	 */
-	bool OnShouldLoadOnTop(const FString& Filename);
+	UNREALED_API bool OnShouldLoadOnTop(const FString& Filename);
 
 	//@todo Slate Editor: Merge PlayMap and RequestSlatePlayMapSession
 	/**
@@ -1653,7 +1653,7 @@ public:
 	 * @param	InPlayInViewportIndex	Viewport index to play the game in, or -1 to spawn a standalone PIE window
 	 * @param	bUseMobilePreview		True to enable mobile preview mode (PC platform only)
 	 */
-	virtual void PlayMap( const FVector* StartLocation = NULL, const FRotator* StartRotation = NULL, int32 DestinationConsole = -1, int32 InPlayInViewportIndex = -1, bool bUseMobilePreview = false );
+	UNREALED_API virtual void PlayMap( const FVector* StartLocation = NULL, const FRotator* StartRotation = NULL, int32 DestinationConsole = -1, int32 InPlayInViewportIndex = -1, bool bUseMobilePreview = false );
 
 
 
@@ -1695,36 +1695,36 @@ public:
 	 * @param	bUseVRPreview			True to enable VR preview mode (PC platform only)
 	 */
 	UE_DEPRECATED(4.25, "Use the overload of RequestPlaySession which takes a FRequestPlaySessionParams instead.")
-	void RequestPlaySession( bool bAtPlayerStart, TSharedPtr<class IAssetViewport> DestinationViewport, bool bInSimulateInEditor, const FVector* StartLocation = NULL, const FRotator* StartRotation = NULL, int32 DestinationConsole = -1, bool bUseMobilePreview = false, bool bUseVRPreview = false, bool bUseVulkanPreview = false);
+	UNREALED_API void RequestPlaySession( bool bAtPlayerStart, TSharedPtr<class IAssetViewport> DestinationViewport, bool bInSimulateInEditor, const FVector* StartLocation = NULL, const FRotator* StartRotation = NULL, int32 DestinationConsole = -1, bool bUseMobilePreview = false, bool bUseVRPreview = false, bool bUseVulkanPreview = false);
 
 	// @todo gmp: temp hack for Rocket demo
 	UE_DEPRECATED(4.25, "Use the overload of RequestPlaySession which takes a FRequestPlaySessionParams instead.")
-	void RequestPlaySession(const FVector* StartLocation, const FRotator* StartRotation, bool MobilePreview, bool VulkanPreview, const FString& MobilePreviewTargetDevice, FString AdditionalStandaloneLaunchParameters = TEXT(""));
+	UNREALED_API void RequestPlaySession(const FVector* StartLocation, const FRotator* StartRotation, bool MobilePreview, bool VulkanPreview, const FString& MobilePreviewTargetDevice, FString AdditionalStandaloneLaunchParameters = TEXT(""));
 
 	/** Request to play a game on a remote device */
 	UE_DEPRECATED(4.25, "Use the overload of RequestPlaySession which takes a FRequestPlaySessionParams instead.")
-	void RequestPlaySession( const FString& DeviceId, const FString& DeviceName );
+	UNREALED_API void RequestPlaySession( const FString& DeviceId, const FString& DeviceName );
 	
 	/** Request a play session (Play in Editor, Play in New Process, Launcher) with the configured parameters. See FRequestPlaySessionParams for more details. */
-	void RequestPlaySession(const FRequestPlaySessionParams& InParams);
+	UNREALED_API void RequestPlaySession(const FRequestPlaySessionParams& InParams);
 
 	/** Cancel request to start a play session */
-	void CancelRequestPlaySession();
+	UNREALED_API void CancelRequestPlaySession();
 
 	/** Pause or unpause all PIE worlds. Returns true if successful */
-	bool SetPIEWorldsPaused(bool Paused);
+	UNREALED_API bool SetPIEWorldsPaused(bool Paused);
 
 	/** Makes a request to start a play from a Slate editor session */
 	void RequestToggleBetweenPIEandSIE() { bIsToggleBetweenPIEandSIEQueued = true; }
 
 	/** Called when the debugger has paused the active PIE or SIE session */
-	void PlaySessionPaused();
+	UNREALED_API void PlaySessionPaused();
 
 	/** Called when the debugger has resumed the active PIE or SIE session */
-	void PlaySessionResumed();
+	UNREALED_API void PlaySessionResumed();
 
 	/** Called when the debugger has single-stepped the active PIE or SIE session */
-	void PlaySessionSingleStepped();
+	UNREALED_API void PlaySessionSingleStepped();
 
 	/** Returns true if we are currently either PIE/SIE in the editor, false if we are not (even if we would start next tick). See IsPlaySessionInProgress() */
 	bool IsPlayingSessionInEditor() const { return PlayInEditorSessionInfo.IsSet(); }
@@ -1754,22 +1754,22 @@ public:
 	}
 
 	/** Called when game client received input key */
-	bool ProcessDebuggerCommands(const FKey InKey, const FModifierKeysState ModifierKeyState, EInputEvent EventType);
+	UNREALED_API bool ProcessDebuggerCommands(const FKey InKey, const FModifierKeysState ModifierKeyState, EInputEvent EventType);
 
 	/**
 	 * Kicks off a "Play From Here" request that was most likely made during a transaction
 	 */
 	UE_DEPRECATED(4.25, "Call StartQueuedPlaySessionRequest, or override StartQueuedPlaySessionRequestImpl instead.")
-	virtual void StartQueuedPlayMapRequest();
+	UNREALED_API virtual void StartQueuedPlayMapRequest();
 
 	/** Kicks off a "Request Play Session" request if there is one set. This should not be called mid-transaction. */
-	void StartQueuedPlaySessionRequest();
+	UNREALED_API void StartQueuedPlaySessionRequest();
 
 	/**
 	 * Request that the current PIE/SIE session should end.  
 	 * This should be used to end the game if its not safe to directly call EndPlayMap in your stack frame
 	 */
-	void RequestEndPlayMap();
+	UNREALED_API void RequestEndPlayMap();
 
 	/**
 	 * @return true if there is an end play map request queued 
@@ -1780,7 +1780,7 @@ public:
 	 * Request to create a new PIE window and join the currently running PIE session.
 	 * Deferred until the next tick.
 	 */
-	void RequestLateJoin();
+	UNREALED_API void RequestLateJoin();
 
 	/**
 	 * Builds a URL for game spawned by the editor (not including map name!). 
@@ -1790,7 +1790,7 @@ public:
 	 *
 	 * @return	The URL for the game
 	 */
-	virtual FString BuildPlayWorldURL(const TCHAR* MapName, bool bSpectatorMode = false, FString AdditionalURLOptions=FString());
+	UNREALED_API virtual FString BuildPlayWorldURL(const TCHAR* MapName, bool bSpectatorMode = false, FString AdditionalURLOptions=FString());
 
 	/**
 	 * Starts a Play In Editor session
@@ -1799,26 +1799,26 @@ public:
 	 * @param	bInSimulateInEditor	True to start an in-editor simulation session, or false to start a play-in-editor session
 	 */
 	UE_DEPRECATED(4.25, "Override StartPlayInEditorSession(FRequestPlaySessionParams& InRequestParams) instead. InWorld is GetEditorWorldContext().World(), and the other arguments come from InRequestParams now.") 
-	virtual void PlayInEditor( UWorld* InWorld, bool bInSimulateInEditor, FPlayInEditorOverrides Overrides = FPlayInEditorOverrides());
+	UNREALED_API virtual void PlayInEditor( UWorld* InWorld, bool bInSimulateInEditor, FPlayInEditorOverrides Overrides = FPlayInEditorOverrides());
 
 	UE_DEPRECATED(4.25, "Override CreateInnerProcessPIEGameInstance instead.")
-	virtual UGameInstance* CreatePIEGameInstance(int32 InPIEInstance, bool bInSimulateInEditor, bool bAnyBlueprintErrors, bool bStartInSpectatorMode, bool bPlayNetDedicated, bool bPlayStereoscopic, float PIEStartTime);
+	UNREALED_API virtual UGameInstance* CreatePIEGameInstance(int32 InPIEInstance, bool bInSimulateInEditor, bool bAnyBlueprintErrors, bool bStartInSpectatorMode, bool bPlayNetDedicated, bool bPlayStereoscopic, float PIEStartTime);
 
 	/**
 	 * Kills the Play From Here session
 	 */
-	virtual void EndPlayMap();
+	UNREALED_API virtual void EndPlayMap();
 
 	/** 
 	 * Destroy the current play session and perform miscellaneous cleanup
 	 */
-	virtual void TeardownPlaySession(FWorldContext &PieWorldContext);
+	UNREALED_API virtual void TeardownPlaySession(FWorldContext &PieWorldContext);
 
 	/**
 	 * Ends the current play on local pc session.
 	 * @todo gmp: temp hack for Rocket demo
 	 */
-	virtual void EndPlayOnLocalPc();
+	UNREALED_API virtual void EndPlayOnLocalPc();
 
 	/**
 	 * Overrides the realtime state all viewports until RemoveViewportsRealtimeOverride is called.
@@ -1827,22 +1827,22 @@ public:
 	 * @param bShouldBeRealtime	If true, this viewport will be realtime, if false this viewport will not be realtime
 	 * @param SystemDisplayName	This display name of whatever system is overriding realtime. This name is displayed to users in the viewport options menu
 	 */
-	void SetViewportsRealtimeOverride(bool bShouldBeRealtime, FText SystemDisplayName);
+	UNREALED_API void SetViewportsRealtimeOverride(bool bShouldBeRealtime, FText SystemDisplayName);
 
 	/**
 	 * Removes the current realtime override.  If there was another realtime override set it will restore that override
 	 */
-	void RemoveViewportsRealtimeOverride(FText SystemDisplayName);
+	UNREALED_API void RemoveViewportsRealtimeOverride(FText SystemDisplayName);
 
 	UE_DEPRECATED(4.26, "To remove realtime overrides, please now provide a system name to make sure you remove the correct override.")
-	void RemoveViewportsRealtimeOverride();
+	UNREALED_API void RemoveViewportsRealtimeOverride();
 
 	/**
 	 * Disables any realtime viewports that are currently viewing the level.  This will not disable
 	 * things like preview viewports in Cascade, etc. Typically called before running the game.
 	 */
 	UE_DEPRECATED(4.25, "To save and restore realtime state non-permanently use SetViewportsRealtimeOverride and RemoveViewportsRealtimeOverride")
-	void DisableRealtimeViewports();
+	UNREALED_API void DisableRealtimeViewports();
 
 	/**
 	 * Restores any realtime viewports that have been disabled by DisableRealtimeViewports. This won't
@@ -1850,24 +1850,24 @@ public:
 	 * latter toggled to be realtime.
 	 */
 	UE_DEPRECATED(4.25, "To save and restore realtime state non-permanently use SetViewportsRealtimeOverride and RemoveViewportsRealtimeOverride")
-	void RestoreRealtimeViewports();
+	UNREALED_API void RestoreRealtimeViewports();
 
 	/**
 	 * Checks to see if any viewport is set to update in realtime.
 	 */
-	bool IsAnyViewportRealtime();
+	UNREALED_API bool IsAnyViewportRealtime();
 
 
 	/**
 	 * @return true if all windows are hidden (including minimized)                                                         
 	 */
-	bool AreAllWindowsHidden() const;
+	UNREALED_API bool AreAllWindowsHidden() const;
 
 	/**
 	 *	Returns pointer to a temporary render target.
 	 *	If it has not already been created, does so here.
 	 */
-	UTextureRenderTarget2D* GetScratchRenderTarget( const uint32 MinSize );
+	UNREALED_API UTextureRenderTarget2D* GetScratchRenderTarget( const uint32 MinSize );
 
 	/**
 	 *  Returns the Editors timer manager instance.
@@ -1895,75 +1895,75 @@ public:
 	/**
 	 * Returns the number of currently selected actors.
 	 */
-	int32 GetSelectedActorCount() const;
+	UNREALED_API int32 GetSelectedActorCount() const;
 
 	/**
 	 * Returns the set of selected actors.
 	 */
-	class USelection* GetSelectedActors() const;
+	UNREALED_API class USelection* GetSelectedActors() const;
 
 	/** Returns True is a world info Actor is selected */
-	bool IsWorldSettingsSelected();
+	UNREALED_API bool IsWorldSettingsSelected();
 
 	/** Function to return unique list of the classes of the assets currently selected in content browser (loaded/not loaded) */
-	void GetContentBrowserSelectionClasses( TArray<UClass*>& Selection ) const;
+	UNREALED_API void GetContentBrowserSelectionClasses( TArray<UClass*>& Selection ) const;
 
 	/** Function to return list of assets currently selected in the content browser */
-	void GetContentBrowserSelections(TArray<FAssetData>& Selections) const;
+	UNREALED_API void GetContentBrowserSelections(TArray<FAssetData>& Selections) const;
 
 	/**
 	 * Returns an FSelectionIterator that iterates over the set of selected actors.
 	 */
-	class FSelectionIterator GetSelectedActorIterator() const;
+	UNREALED_API class FSelectionIterator GetSelectedActorIterator() const;
 
 	/**
 	* Returns an FSelectionIterator that iterates over the set of selected components.
 	*/
-	class FSelectionIterator GetSelectedComponentIterator() const;
+	UNREALED_API class FSelectionIterator GetSelectedComponentIterator() const;
 
-	class FSelectedEditableComponentIterator GetSelectedEditableComponentIterator() const;
+	UNREALED_API class FSelectedEditableComponentIterator GetSelectedEditableComponentIterator() const;
 
 	/**
 	* Returns the number of currently selected components.
 	*/
-	int32 GetSelectedComponentCount() const;
+	UNREALED_API int32 GetSelectedComponentCount() const;
 
 	/**
 	* @return the set of selected components.
 	*/
-	class USelection* GetSelectedComponents() const;
+	UNREALED_API class USelection* GetSelectedComponents() const;
 
 	/**
 	 * @return the set of selected non-actor objects.
 	 */
-	class USelection* GetSelectedObjects() const;
+	UNREALED_API class USelection* GetSelectedObjects() const;
 
 	/**
 	 * @return the appropriate selection set for the specified object class.
 	 */
-	class USelection* GetSelectedSet( const UClass* Class ) const;
+	UNREALED_API class USelection* GetSelectedSet( const UClass* Class ) const;
 
 	/**
 	 * @param	RequiredParentClass The parent class this type must implement, or null to accept anything
 	 * @return	The first selected class (either a UClass type itself, or the UClass generated by a blueprint), or null if there are no class or blueprint types selected
 	 */
-	const UClass* GetFirstSelectedClass( const UClass* const RequiredParentClass ) const;
+	UNREALED_API const UClass* GetFirstSelectedClass( const UClass* const RequiredParentClass ) const;
 
 	/**
 	 * Get the selection state of the current level (its actors and components) so that it might be restored later.
 	 */
-	void GetSelectionStateOfLevel(FSelectionStateOfLevel& OutSelectionStateOfLevel) const;
+	UNREALED_API void GetSelectionStateOfLevel(FSelectionStateOfLevel& OutSelectionStateOfLevel) const;
 
 	/**
 	 * Restore the selection state of the current level (its actors and components) from a previous state.
 	 */
-	void SetSelectionStateOfLevel(const FSelectionStateOfLevel& InSelectionStateOfLevel);
+	UNREALED_API void SetSelectionStateOfLevel(const FSelectionStateOfLevel& InSelectionStateOfLevel);
 
 	/**
 	 * Reset All Selection Sets (i.e. objects, actors, components)
 	 * @note each set is independent and a selected actor might not be in the object selection set.
 	 */
-	void ResetAllSelectionSets();
+	UNREALED_API void ResetAllSelectionSets();
 
 	/**
 	 * Clears out the current map, if any, and creates a new blank map.
@@ -1972,7 +1972,7 @@ public:
 	 * 
 	 * @return	Pointer to the map that was created.
 	 */
-	UWorld* NewMap(bool bIsPartitionedWorld = false);
+	UNREALED_API UWorld* NewMap(bool bIsPartitionedWorld = false);
 
 	/**
 	 * Exports the current map to the specified filename.
@@ -1980,7 +1980,7 @@ public:
 	 * @param	InFilename					Filename to export the map to.
 	 * @param	bExportSelectedActorsOnly	If true, export only the selected actors.
 	 */
-	void ExportMap(UWorld* InWorld, const TCHAR* InFilename, bool bExportSelectedActorsOnly);
+	UNREALED_API void ExportMap(UWorld* InWorld, const TCHAR* InFilename, bool bExportSelectedActorsOnly);
 
 	/**
 	 *	Returns list of all foliage types used in the world
@@ -1988,7 +1988,7 @@ public:
 	 * @param	InWorld	 The target world.
 	 * @return	List of all foliage types used in the world
 	 */
-	TArray<UFoliageType*> GetFoliageTypesInWorld(UWorld* InWorld);
+	UNREALED_API TArray<UFoliageType*> GetFoliageTypesInWorld(UWorld* InWorld);
 
 	/**
 	 * Checks to see whether it's possible to perform a copy operation on the selected actors.
@@ -1997,7 +1997,7 @@ public:
 	 * @param OutCopySelected	Can be NULL, copies the internal results of the copy check to this struct
 	 * @return			true if it's possible to do a copy operation based on the selection
 	 */
-	bool CanCopySelectedActorsToClipboard( UWorld* InWorld, FCopySelectedInfo* OutCopySelected = NULL );
+	UNREALED_API bool CanCopySelectedActorsToClipboard( UWorld* InWorld, FCopySelectedInfo* OutCopySelected = NULL );
 
 	/**
 	 * Copies selected actors to the clipboard.  Supports copying actors from multiple levels.
@@ -2009,7 +2009,7 @@ public:
 	 * @param bWarnAboutReferences	Whether or not to show a modal warning about referenced actors that may no longer function after being moved
 	 * @param DestinationData		Data that was copied to clipboard 
 	 */
-	void CopySelectedActorsToClipboard( UWorld* InWorld, const bool bShouldCut, const bool bIsMove = false, bool bWarnAboutReferences = true, FString* DestinationData = nullptr);
+	UNREALED_API void CopySelectedActorsToClipboard( UWorld* InWorld, const bool bShouldCut, const bool bIsMove = false, bool bWarnAboutReferences = true, FString* DestinationData = nullptr);
 
 	/**
 	 * Checks to see whether it's possible to perform a paste operation.
@@ -2017,7 +2017,7 @@ public:
 	 * @param InWorld		World to paste into
 	 * @return				true if it's possible to do a Paste operation
 	 */
-	bool CanPasteSelectedActorsFromClipboard( UWorld* InWorld );
+	UNREALED_API bool CanPasteSelectedActorsFromClipboard( UWorld* InWorld );
 
 	/**
 	 * Pastes selected actors from the clipboard.
@@ -2026,7 +2026,7 @@ public:
 	 * @param InWorld		World to get the selected actors from
 	 * @param PasteTo		Where to paste the content too
 	 */
-	void PasteSelectedActorsFromClipboard( UWorld* InWorld, const FText& TransDescription, const EPasteTo PasteTo );
+	UNREALED_API void PasteSelectedActorsFromClipboard( UWorld* InWorld, const FText& TransDescription, const EPasteTo PasteTo );
 
 	/**
 	 * Sets property value and property chain to be used for property-based coloration.
@@ -2036,7 +2036,7 @@ public:
 	 * @param	CommonBaseClass		The class of object to color.
 	 * @param	PropertyChain		The chain of properties from member to lowest property.
 	 */
-	virtual void SetPropertyColorationTarget(UWorld* InWorld, const FString& PropertyValue, class FProperty* Property, class UClass* CommonBaseClass, class FEditPropertyChain* PropertyChain);
+	UNREALED_API virtual void SetPropertyColorationTarget(UWorld* InWorld, const FString& PropertyValue, class FProperty* Property, class UClass* CommonBaseClass, class FEditPropertyChain* PropertyChain);
 
 	/**
 	 * Accessor for current property-based coloration settings.
@@ -2046,12 +2046,12 @@ public:
 	 * @param	OutCommonBaseClass	[out] The class of object to color.
 	 * @param	OutPropertyChain	[out] The chain of properties from member to lowest property.
 	 */
-	virtual void GetPropertyColorationTarget(FString& OutPropertyValue, FProperty*& OutProperty, UClass*& OutCommonBaseClass, FEditPropertyChain*& OutPropertyChain);
+	UNREALED_API virtual void GetPropertyColorationTarget(FString& OutPropertyValue, FProperty*& OutProperty, UClass*& OutCommonBaseClass, FEditPropertyChain*& OutPropertyChain);
 
 	/**
 	 * Selects actors that match the property coloration settings.
 	 */
-	void SelectByPropertyColoration(UWorld* InWorld);
+	UNREALED_API void SelectByPropertyColoration(UWorld* InWorld);
 
 	/**
 	 * Warns the user of any hidden levels, and prompts them with a Yes/No dialog
@@ -2063,11 +2063,11 @@ public:
 	 * @param	bIncludePersistentLvl	If true, the persistent level will also be checked for visibility
 	 * @return							false if the user selects "No", true otherwise.
 	 */
-	bool WarnAboutHiddenLevels( UWorld* InWorld, bool bIncludePersistentLvl) const;
+	UNREALED_API bool WarnAboutHiddenLevels( UWorld* InWorld, bool bIncludePersistentLvl) const;
 
-	void ApplyDeltaToActor(AActor* InActor, bool bDelta, const FVector* InTranslation, const FRotator* InRotation, const FVector* InScaling, bool bAltDown=false, bool bShiftDown=false, bool bControlDown=false) const;
+	UNREALED_API void ApplyDeltaToActor(AActor* InActor, bool bDelta, const FVector* InTranslation, const FRotator* InRotation, const FVector* InScaling, bool bAltDown=false, bool bShiftDown=false, bool bControlDown=false) const;
 
-	void ApplyDeltaToComponent(USceneComponent* InComponent, bool bDelta, const FVector* InTranslation, const FRotator* InRotation, const FVector* InScaling, const FVector& PivotLocation ) const;
+	UNREALED_API void ApplyDeltaToComponent(USceneComponent* InComponent, bool bDelta, const FVector* InTranslation, const FRotator* InRotation, const FVector* InScaling, const FVector& PivotLocation ) const;
 
 	/**
 	 * Disable actor/component modification during delta movement.
@@ -2102,7 +2102,7 @@ public:
 	/**
 	 * Auto merge all staticmeshes that are able to be merged
 	 */
-	void AutoMergeStaticMeshes();
+	UNREALED_API void AutoMergeStaticMeshes();
 
 	/**
 	 *	Check the InCmdParams for "MAPINISECTION=<name of section>".
@@ -2111,7 +2111,7 @@ public:
 	 *	@param	InCmdParams		The cmd line parameters for the application
 	 *	@param	OutMapList		The list of maps from the ini section, empty if not found
 	 */
-	void ParseMapSectionIni(const TCHAR* InCmdParams, TArray<FString>& OutMapList);
+	UNREALED_API void ParseMapSectionIni(const TCHAR* InCmdParams, TArray<FString>& OutMapList);
 
 	/**
 	 *	Load the list of maps from the given section of the Editor.ini file
@@ -2120,7 +2120,7 @@ public:
 	 *	@param	InSectionName		The name of the section to load
 	 *	@param	OutMapList			The list of maps from that section
 	 */
-	void LoadMapListFromIni(const FString& InSectionName, TArray<FString>& OutMapList);
+	UNREALED_API void LoadMapListFromIni(const FString& InSectionName, TArray<FString>& OutMapList);
 
 	/**
 	 *	Check whether the specified package file is a map
@@ -2129,93 +2129,93 @@ public:
 	 *	@param	OutErrorMsg			if an errors occurs, this is the localized reason
 	 *	@return	true if the package is not a map file
 	 */
-	bool PackageIsAMapFile( const TCHAR* PackageFilename, FText& OutNotMapReason );
+	UNREALED_API bool PackageIsAMapFile( const TCHAR* PackageFilename, FText& OutNotMapReason );
 
 	/**
 	 * Synchronizes the content or generic browser's selected objects to the collection specified.
 	 *
 	 * @param	InObjectsToSync	the list of objects to sync to
 	 */
-	void SyncBrowserToObjects( const TArray<UObject*>& InObjectsToSync, bool bFocusContentBrowser = true );
-	void SyncBrowserToObjects( const TArray<struct FAssetData>& InAssetsToSync, bool bFocusContentBrowser = true );
+	UNREALED_API void SyncBrowserToObjects( const TArray<UObject*>& InObjectsToSync, bool bFocusContentBrowser = true );
+	UNREALED_API void SyncBrowserToObjects( const TArray<struct FAssetData>& InAssetsToSync, bool bFocusContentBrowser = true );
 
 	/**
 	 * Synchronizes the content or generic browser's selected object to the one specified.
 	 *
 	 * @param	InObjectToSync	the object to sync to
 	 */
-	void SyncBrowserToObject( const UObject* InObjectToSync, bool bFocusContentBrowser = true );
-	void SyncBrowserToObject( const struct FAssetData& InAssetToSync, bool bFocusContentBrowser = true );
+	UNREALED_API void SyncBrowserToObject( const UObject* InObjectToSync, bool bFocusContentBrowser = true );
+	UNREALED_API void SyncBrowserToObject( const struct FAssetData& InAssetToSync, bool bFocusContentBrowser = true );
 
 	/**
 	 * Syncs the selected actors objects to the content browser
 	 * 
 	 * @param bAllowOverrideMetadata If true, allows an asset to define "BrowseToAssetOverride" in its metadata to sync to an asset other than itself
 	 */
-	void SyncToContentBrowser(bool bAllowOverrideMetadata = true);
+	UNREALED_API void SyncToContentBrowser(bool bAllowOverrideMetadata = true);
 
 	/**
 	 * Syncs the selected actors' levels to the content browser
 	 */
-	void SyncActorLevelsToContentBrowser();
+	UNREALED_API void SyncActorLevelsToContentBrowser();
 
 	/**
 	 * Checks if the slected objects contain something to browse to
 	 *
 	 * @return true if any of the selected objects contains something that can be browsed to
 	 */
-	bool CanSyncToContentBrowser();
+	UNREALED_API bool CanSyncToContentBrowser();
 
 	/**
 	 * Checks if the selected objects have levels which can be browsed to
 	 *
 	 * @return true if any of the selected objects contains something that can be browsed to
 	 */
-	bool CanSyncActorLevelsToContentBrowser();
+	UNREALED_API bool CanSyncActorLevelsToContentBrowser();
 
 	/**
 	 * Toggles the movement lock on selected actors so they may or may not be moved with the transform widgets
 	 */
-	void ToggleSelectedActorMovementLock();
+	UNREALED_API void ToggleSelectedActorMovementLock();
 
 	/**
 	 * @return true if there are selected locked actors
 	 */
-	bool HasLockedActors();
+	UNREALED_API bool HasLockedActors();
 
 	/**
 	 * Opens the objects specialized editor (Same as double clicking on the item in the content browser)
 	 *
 	 * @param ObjectToEdit	The object to open the editor for 
 	 */
-	void EditObject( UObject* ObjectToEdit );
+	UNREALED_API void EditObject( UObject* ObjectToEdit );
 
 	/**
 	 * Selects the currently selected actor(s) levels in the level browser
 	 *
 	 * @param bDeselectOthers	Whether or not to deselect the current level selection first
 	 */
-	void SelectLevelInLevelBrowser( bool bDeselectOthers );
+	UNREALED_API void SelectLevelInLevelBrowser( bool bDeselectOthers );
 
 	/**
 	 * Deselects the currently selected actor(s) levels in the level browser
 	 */
-	void DeselectLevelInLevelBrowser();
+	UNREALED_API void DeselectLevelInLevelBrowser();
 
 	/**
 	 * Selects all actors with the same class as the current selection
 	 *
 	 * @param bArchetype	true to only select actors of the same class AND same Archetype
 	 */
-	void SelectAllActorsWithClass( bool bArchetype );
+	UNREALED_API void SelectAllActorsWithClass( bool bArchetype );
 
 	/**
 	 * Finds all references to the currently selected actors, and reports results in a find message log
 	 */
-	void FindSelectedActorsInLevelScript();
+	UNREALED_API void FindSelectedActorsInLevelScript();
 
 	/** See if any selected actors are referenced in level script */
-	bool AreAnySelectedActorsInLevelScript();
+	UNREALED_API bool AreAnySelectedActorsInLevelScript();
 
 	/**
 	 * Checks if a provided package is valid to be auto-added to a default changelist based on several
@@ -2227,7 +2227,7 @@ public:
 	 *
 	 * @return	true if the provided package is valid to be auto-added to a default source control changelist
 	 */
-	bool IsPackageValidForAutoAdding(UPackage* InPackage, const FString& InFilename);
+	UNREALED_API bool IsPackageValidForAutoAdding(UPackage* InPackage, const FString& InFilename);
 
 	/**
 	 * Checks if a provided package is valid to be saved. For startup packages, it asks user for confirmation.
@@ -2238,73 +2238,73 @@ public:
 	 *
 	 * @return	true if the provided package is valid to be saved
 	 */
-	virtual bool IsPackageOKToSave(UPackage* InPackage, const FString& InFilename, FOutputDevice* Error);
+	UNREALED_API virtual bool IsPackageOKToSave(UPackage* InPackage, const FString& InFilename, FOutputDevice* Error);
 
 	/** The editor wrapper for UPackage::SavePackage. Auto-adds files to source control when necessary */
-	bool SavePackage(UPackage* InOuter, UObject* InAsset, const TCHAR* Filename, const FSavePackageArgs& SaveArgs);
+	UNREALED_API bool SavePackage(UPackage* InOuter, UObject* InAsset, const TCHAR* Filename, const FSavePackageArgs& SaveArgs);
 
 	/** The editor wrapper for UPackage::Save. Auto-adds files to source control when necessary */
-	FSavePackageResultStruct Save(UPackage* InOuter, UObject* InAsset, const TCHAR* Filename,
+	UNREALED_API FSavePackageResultStruct Save(UPackage* InOuter, UObject* InAsset, const TCHAR* Filename,
 		const FSavePackageArgs& SaveArgs);
 
-	virtual bool InitializePhysicsSceneForSaveIfNecessary(UWorld* World, bool &bOutForceInitialized);
-	void CleanupPhysicsSceneThatWasInitializedForSave(UWorld* World, bool bForceInitialized);
+	UNREALED_API virtual bool InitializePhysicsSceneForSaveIfNecessary(UWorld* World, bool &bOutForceInitialized);
+	UNREALED_API void CleanupPhysicsSceneThatWasInitializedForSave(UWorld* World, bool bForceInitialized);
 
 	/** Invoked before a UWorld is saved to update editor systems */
-	virtual void OnPreSaveWorld(UWorld* World, FObjectPreSaveContext ObjectSaveContext);
+	UNREALED_API virtual void OnPreSaveWorld(UWorld* World, FObjectPreSaveContext ObjectSaveContext);
 
 	/** Invoked after a UWorld is saved to update editor systems */
-	virtual void OnPostSaveWorld(UWorld* World, FObjectPostSaveContext ObjectSaveContext);
+	UNREALED_API virtual void OnPostSaveWorld(UWorld* World, FObjectPostSaveContext ObjectSaveContext);
 
 	/**
 	 * Adds provided package to a default changelist
 	 *
 	 * @param	PackageNames	Filenames of the packages after they will be saved (not necessarily the same as the filename of the first param in the event of a rename)
 	 */
-	void AddPackagesToDefaultChangelist(TArray<FString>& InPackageNames);
+	UNREALED_API void AddPackagesToDefaultChangelist(TArray<FString>& InPackageNames);
 
 	/**
 	 * Will run a batch mark for add source control operation for files recently saved if source control is enabled
 	 * @param	bool	Unused param to for the function to be used as a delegate when source control dialog is closed.
 	 */
-	void RunDeferredMarkForAddFiles(bool = false);
+	UNREALED_API void RunDeferredMarkForAddFiles(bool = false);
 
 	/** 
 	 * @return - returns the currently selected positional snap grid setting
 	 */
-	float GetGridSize() const;
+	UNREALED_API float GetGridSize() const;
 
 	/** 
 	 * @return - if the grid size is part of the 1,2,4,8,16,.. list or not
 	 */
-	bool IsGridSizePowerOfTwo() const;
+	UNREALED_API bool IsGridSizePowerOfTwo() const;
 
 	/** 
 	 * Sets the selected positional snap grid setting
 	 *
 	 * @param InIndex - The index of the selected grid setting
 	 */
-	void SetGridSize(int32 InIndex);
+	UNREALED_API void SetGridSize(int32 InIndex);
 
 	/** 
 	 * Increase the positional snap grid setting (will not increase passed the maximum value)
 	 */
-	void GridSizeIncrement();
+	UNREALED_API void GridSizeIncrement();
 	
 	/** 
 	 * Decreased the positional snap grid setting (will not decrease passed the minimum value)
 	 */
-	void GridSizeDecrement();
+	UNREALED_API void GridSizeDecrement();
 
 	/** 
 	 * Accesses the array of snap grid position options
 	 */
-	const TArray<float>& GetCurrentPositionGridArray() const;
+	UNREALED_API const TArray<float>& GetCurrentPositionGridArray() const;
 	
 	/** 
 	 * @return - returns the currently selected rotational snap grid setting
 	 */
-	FRotator GetRotGridSize();
+	UNREALED_API FRotator GetRotGridSize();
 
 	/** 
 	 * Sets the selected Rotational snap grid setting
@@ -2312,44 +2312,44 @@ public:
 	 * @param InIndex - The index of the selected grid setting
 	 * @param InGridMode - Selects which set of grid settings to use
 	 */
-	void SetRotGridSize(int32 InIndex, enum ERotationGridMode InGridMode);
+	UNREALED_API void SetRotGridSize(int32 InIndex, enum ERotationGridMode InGridMode);
 
 	/** 
 	 * Increase the rotational snap grid setting (will not increase passed the maximum value)
 	 */
-	void RotGridSizeIncrement();
+	UNREALED_API void RotGridSizeIncrement();
 	
 	/** 
 	 * Decreased the rotational snap grid setting (will not decrease passed the minimum value)
 	 */
-	void RotGridSizeDecrement();
+	UNREALED_API void RotGridSizeDecrement();
 	
 	/** 
 	 * Accesses the array of snap grid rotation options
 	 */
-	const TArray<float>& GetCurrentRotationGridArray() const;
+	UNREALED_API const TArray<float>& GetCurrentRotationGridArray() const;
 	
-	float GetScaleGridSize();
-	void SetScaleGridSize(int32 InIndex);
+	UNREALED_API float GetScaleGridSize();
+	UNREALED_API void SetScaleGridSize(int32 InIndex);
 
-	float GetGridInterval();
+	UNREALED_API float GetGridInterval();
 
 	/**
 	 * Get the location offset applied by the grid based on the grid size and active viewport.
 	 */
-	FVector GetGridLocationOffset(bool bUniformOffset) const;
+	UNREALED_API FVector GetGridLocationOffset(bool bUniformOffset) const;
 
 	/**
 	 * Access the array of grid interval options
 	 */
-	const TArray<float>& GetCurrentIntervalGridArray() const;
+	UNREALED_API const TArray<float>& GetCurrentIntervalGridArray() const;
 	
 	 /**
 	  * Function to convert selected brushes into volumes of the provided class.
 	  *
 	  * @param	VolumeClass	Class of volume that selected brushes should be converted into
 	  */
-	void ConvertSelectedBrushesToVolumes( UClass* VolumeClass );
+	UNREALED_API void ConvertSelectedBrushesToVolumes( UClass* VolumeClass );
 
 	/**
 	 * Called to convert actors of one class type to another
@@ -2357,18 +2357,18 @@ public:
 	 * @param FromClass The class converting from
 	 * @param ToClass	The class converting to
 	 */
-	void ConvertActorsFromClass( UClass* FromClass, UClass* ToClass );
+	UNREALED_API void ConvertActorsFromClass( UClass* FromClass, UClass* ToClass );
 
 	/**
 	* Update any outstanding reflection captures
 	*/
-	void BuildReflectionCaptures(UWorld* World = GWorld);
+	UNREALED_API void BuildReflectionCaptures(UWorld* World = GWorld);
 
 	/**
 	 * Convenience method for adding a Slate modal window that is parented to the main frame (if it exists)
 	 * This function does not return until the modal window is closed.
 	 */
-	void EditorAddModalWindow( TSharedRef<SWindow> InModalWindow ) const;
+	UNREALED_API void EditorAddModalWindow( TSharedRef<SWindow> InModalWindow ) const;
 
 	/**
 	 * Finds a brush builder of the provided class.  If it doesnt exist one will be created
@@ -2376,7 +2376,7 @@ public:
 	 * @param BrushBuilderClass	The class to get the builder brush from
 	 * @return The found or created brush builder
 	 */
-	UBrushBuilder* FindBrushBuilder( UClass* BrushBuilderClass );
+	UNREALED_API UBrushBuilder* FindBrushBuilder( UClass* BrushBuilderClass );
 
 	/**
 	 * Parents one actor to another
@@ -2387,14 +2387,14 @@ public:
 	 * @param SocketName	An optional socket name to attach to in the parent (use NAME_None when there is no socket)
 	 * @param Component		Actual Component included in the ParentActor which is a usually blueprint actor
 	 */
-	void ParentActors( AActor* ParentActor, AActor* ChildActor, const FName SocketName, USceneComponent* Component=NULL );
+	UNREALED_API void ParentActors( AActor* ParentActor, AActor* ChildActor, const FName SocketName, USceneComponent* Component=NULL );
 
 	/**
 	 * Detaches selected actors from there parents
 	 *
 	 * @return				true if a detachment occurred
 	 */
-	bool DetachSelectedActors();	
+	UNREALED_API bool DetachSelectedActors();	
 
 	/**
 	 * Checks the validity of parenting one actor to another
@@ -2404,12 +2404,12 @@ public:
 	 * @param ReasonText	Optional text to receive a description of why the function returned false.
 	 * @return				true if the parenting action is valid and will succeed, false if invalid.
 	 */
-	bool CanParentActors( const AActor* ParentActor, const AActor* ChildActor, FText* ReasonText = NULL);
+	UNREALED_API bool CanParentActors( const AActor* ParentActor, const AActor* ChildActor, FText* ReasonText = NULL);
 
 	/** 
 	 * turns all navigable static geometry of ULevel into polygon soup stored in passed Level (ULevel::StaticNavigableGeometry)
 	 */
-	virtual void RebuildStaticNavigableGeometry(ULevel* Level);
+	UNREALED_API virtual void RebuildStaticNavigableGeometry(ULevel* Level);
 
 	/**
 	 * Gets all objects which can be synced to in content browser for current selection
@@ -2418,21 +2418,21 @@ public:
 	 * @param bAllowOverrideMetadata If true, allows an asset to define "BrowseToAssetOverride" in its metadata to sync to an asset other than itself
 	 */
 	UE_DEPRECATED(5.1, "Use GetAssetsToSyncToContentBrowser instead")
-	void GetObjectsToSyncToContentBrowser(TArray<UObject*>& Objects, bool bAllowBrowseToAssetOverride = true);
+	UNREALED_API void GetObjectsToSyncToContentBrowser(TArray<UObject*>& Objects, bool bAllowBrowseToAssetOverride = true);
 	/**
 	 * Gets all assets which can be synced to in content browser for current selection
 	 *
 	 * @param Assets	Array to be filled with assets which can be browsed to
 	 * @param bAllowOverrideMetadata If true, allows an asset to define "BrowseToAssetOverride" in its metadata to sync to an asset other than itself
 	 */
-	void GetAssetsToSyncToContentBrowser(TArray<FAssetData>& Assets, bool bAllowBrowseToAssetOverride = true);
+	UNREALED_API void GetAssetsToSyncToContentBrowser(TArray<FAssetData>& Assets, bool bAllowBrowseToAssetOverride = true);
 
 	/**
 	 * Gets all levels which can be synced to in content browser for current selection
 	 *
 	 * @param Objects	Array to be filled with ULevel objects
 	 */
-	void GetLevelsToSyncToContentBrowser(TArray<UObject*>& Objects);
+	UNREALED_API void GetLevelsToSyncToContentBrowser(TArray<UObject*>& Objects);
 
 	/**
 	* Queries for a list of assets that are referenced by the current editor selection (actors, surfaces, etc.)
@@ -2440,14 +2440,14 @@ public:
 	* @param	Objects								Array to be filled with asset objects referenced by the current editor selection
 	* @param	bIgnoreOtherAssetsIfBPReferenced	If true, and a selected actor has a Blueprint asset, only that will be returned.
 	*/
-	void GetReferencedAssetsForEditorSelection(TArray<UObject*>& Objects, const bool bIgnoreOtherAssetsIfBPReferenced = false);
+	UNREALED_API void GetReferencedAssetsForEditorSelection(TArray<UObject*>& Objects, const bool bIgnoreOtherAssetsIfBPReferenced = false);
 
 	/**
 	* Queries for a list of assets that are soft referenced by the current editor selection (actors, surfaces, etc.)
 	*
 	* @param	SoftObjects							Array to be filled with asset objects referenced soft by the current editor selection
 	*/
-	void GetSoftReferencedAssetsForEditorSelection(TArray<FSoftObjectPath>& SoftObjects);
+	UNREALED_API void GetSoftReferencedAssetsForEditorSelection(TArray<FSoftObjectPath>& SoftObjects);
 
 	/** Returns a filter to restruct what assets show up in asset pickers based on what the selection is used for (i.e. what will reference the assets) */
 	DECLARE_DELEGATE_RetVal_OneParam(TSharedPtr<IAssetReferenceFilter>, FOnMakeAssetReferenceFilter, const FAssetReferenceFilterContext& /*Context*/);
@@ -2456,13 +2456,13 @@ public:
 
 	DECLARE_DELEGATE_RetVal(ULevelEditorDragDropHandler*, FOnCreateLevelEditorDragDropHandler);
 	FOnCreateLevelEditorDragDropHandler& OnCreateLevelEditorDragDropHandler() { return OnCreateLevelEditorDragDropHandlerDelegate; }
-	ULevelEditorDragDropHandler* GetLevelEditorDragDropHandler() const;
+	UNREALED_API ULevelEditorDragDropHandler* GetLevelEditorDragDropHandler() const;
 
 	/** 
 	 * Gets the interface to manage project references to external content
 	 * @note the returned pointer cannot be null
 	 */
-	IProjectExternalContentInterface* GetProjectExternalContentInterface();
+	UNREALED_API IProjectExternalContentInterface* GetProjectExternalContentInterface();
 
 	/** Delegate to override the IProjectExternalContentInterface */
 	DECLARE_DELEGATE_RetVal(IProjectExternalContentInterface*, FProjectExternalContentInterfaceGetter);
@@ -2481,13 +2481,13 @@ public:
 	 *
 	 * @param	bEnsureIsGWorld		Temporarily allow callers to validate that this maps to GWorld to help track down any conversion bugs
 	 */
-	FWorldContext &GetEditorWorldContext(bool bEnsureIsGWorld = false);
+	UNREALED_API FWorldContext &GetEditorWorldContext(bool bEnsureIsGWorld = false);
 
 	/** 
 	 * Returns the WorldContext for the PIE world, by default will get the first one which will be the server or simulate instance.
 	 * You need to iterate the context list if you want all the pie world contexts.
 	 */
-	FWorldContext* GetPIEWorldContext(int32 WorldPIEInstance = 0);
+	UNREALED_API FWorldContext* GetPIEWorldContext(int32 WorldPIEInstance = 0);
 
 	/** 
 	 * mostly done to check if PIE is being set up, go GWorld is going to change, and it's not really _the_G_World_
@@ -2500,43 +2500,43 @@ public:
 	 *
 	 *	@return		The currently active viewport.
 	 */
-	FViewport* GetActiveViewport();
+	UNREALED_API FViewport* GetActiveViewport();
 
 	/**
 	 *	Retrieves the PIE viewport from the editor.
 	 *
 	 *	@return		The PIE viewport. NULL if there is no PIE viewport active.
 	 */
-	FViewport* GetPIEViewport();
+	UNREALED_API FViewport* GetPIEViewport();
 
 	/** 
 	 *	Checks for any player starts and returns the first one found.
 	 *
 	 *	@return		The first player start location found.
 	 */
-	APlayerStart* CheckForPlayerStart();
+	UNREALED_API APlayerStart* CheckForPlayerStart();
 
 	/** Closes the popup created for GenericTextEntryModal or GenericTextEntryModeless*/
-	void CloseEntryPopupWindow();
+	UNREALED_API void CloseEntryPopupWindow();
 
 	/**
 	 * Prompts the user to save the current map if necessary, then creates a new (blank) map.
 	 */
-	void CreateNewMapForEditing(bool bPromptUserToSave = true, bool bIsPartitionedWorld = false);
+	UNREALED_API void CreateNewMapForEditing(bool bPromptUserToSave = true, bool bIsPartitionedWorld = false);
 
 	/**
 	 * If a PIE world exists, give the user the option to terminate it.
 	 *
 	 * @return				true if a PIE session exists and the user refused to end it, false otherwise.
 	 */
-	bool ShouldAbortBecauseOfPIEWorld();
+	UNREALED_API bool ShouldAbortBecauseOfPIEWorld();
 
 	/**
 	 * If an unsaved world exists that would be lost in a map transition, give the user the option to cancel a map load.
 	 *
 	 * @return				true if an unsaved world exists and the user refused to continue, false otherwise.
 	 */
-	bool ShouldAbortBecauseOfUnsavedWorld();
+	UNREALED_API bool ShouldAbortBecauseOfUnsavedWorld();
 
 	/**
 	 * Gets the user-friendly, localized (if exists) name of a property
@@ -2547,55 +2547,55 @@ public:
 	 *
 	 * @return	the friendly name for the property.  localized first, then metadata, then the property's name.
 	 */
-	static FString GetFriendlyName( const FProperty* Property, UStruct* OwnerStruct = NULL );
+	static UNREALED_API FString GetFriendlyName( const FProperty* Property, UStruct* OwnerStruct = NULL );
 
 	/**
 	 * Register a client tool to receive undo events 
 	 * @param UndoClient	An object wanting to receive PostUndo/PostRedo events
 	 */
-	void RegisterForUndo(class FEditorUndoClient* UndoClient );
+	UNREALED_API void RegisterForUndo(class FEditorUndoClient* UndoClient );
 
 	/**
 	 * Unregister a client from receiving undo events 
 	 * @param UndoClient	An object wanting to unsubscribe from PostUndo/PostRedo events
 	 */
-	void UnregisterForUndo( class FEditorUndoClient* UndoEditor );
+	UNREALED_API void UnregisterForUndo( class FEditorUndoClient* UndoEditor );
 
 	/**
 	 * Ensures the assets specified are loaded and adds them to the global selection set
 	 * @param Assets		An array of assets to load and select
 	 * @param TypeOfAsset	An optional class to restrict the types of assets loaded & selected
 	 */
-	void LoadAndSelectAssets( TArray<FAssetData>& Assets, UClass* TypeOfAsset=NULL );
+	UNREALED_API void LoadAndSelectAssets( TArray<FAssetData>& Assets, UClass* TypeOfAsset=NULL );
 
 	/**
 	 * @return True if percentage based scaling is enabled
 	 */
-	bool UsePercentageBasedScaling() const;
+	UNREALED_API bool UsePercentageBasedScaling() const;
 
 	/**
 	 * Returns the actor grouping utility class that performs all grouping related tasks
 	 * This will create the class instance if it doesn't exist.
 	 */
-	class UActorGroupingUtils* GetActorGroupingUtils();
+	UNREALED_API class UActorGroupingUtils* GetActorGroupingUtils();
 
 	/**
 	 * Query to tell if the editor is currently in a mode where it wants XR HMD
 	 * tracking to be used (like in the VR editor or VR PIE preview).
 	 */
-	bool IsHMDTrackingAllowed() const;
+	UNREALED_API bool IsHMDTrackingAllowed() const;
 
 private:
 	//
 	// Map execs.
 	//
 
-	bool Map_Select( UWorld* InWorld, const TCHAR* Str, FOutputDevice& Ar);
-	bool Map_Brush( UWorld* InWorld, const TCHAR* Str, FOutputDevice& Ar);
-	bool Map_Sendto( UWorld* InWorld, const TCHAR* Str, FOutputDevice& Ar);
-	bool Map_Rebuild(UWorld* InWorld, const TCHAR* Str, FOutputDevice& Ar);
-	bool Map_Load(const TCHAR* Str, FOutputDevice& Ar);
-	bool Map_Import( UWorld* InWorld, const TCHAR* Str, FOutputDevice& Ar );
+	UNREALED_API bool Map_Select( UWorld* InWorld, const TCHAR* Str, FOutputDevice& Ar);
+	UNREALED_API bool Map_Brush( UWorld* InWorld, const TCHAR* Str, FOutputDevice& Ar);
+	UNREALED_API bool Map_Sendto( UWorld* InWorld, const TCHAR* Str, FOutputDevice& Ar);
+	UNREALED_API bool Map_Rebuild(UWorld* InWorld, const TCHAR* Str, FOutputDevice& Ar);
+	UNREALED_API bool Map_Load(const TCHAR* Str, FOutputDevice& Ar);
+	UNREALED_API bool Map_Import( UWorld* InWorld, const TCHAR* Str, FOutputDevice& Ar );
 
 	struct EMapCheckNotification
 	{
@@ -2610,9 +2610,9 @@ private:
 	/**
 	 * Checks map for common errors.
 	 */
-	bool Map_Check(UWorld* InWorld, const TCHAR* Str, FOutputDevice& Ar, bool bCheckDeprecatedOnly, EMapCheckNotification::Type Notification = EMapCheckNotification::DisplayResults, bool bClearLog = true);
-	bool Map_Scale( UWorld* InWorld, const TCHAR* Str, FOutputDevice& Ar );
-	bool Map_Setbrush( UWorld* InWorld, const TCHAR* Str, FOutputDevice& Ar );
+	UNREALED_API bool Map_Check(UWorld* InWorld, const TCHAR* Str, FOutputDevice& Ar, bool bCheckDeprecatedOnly, EMapCheckNotification::Type Notification = EMapCheckNotification::DisplayResults, bool bClearLog = true);
+	UNREALED_API bool Map_Scale( UWorld* InWorld, const TCHAR* Str, FOutputDevice& Ar );
+	UNREALED_API bool Map_Setbrush( UWorld* InWorld, const TCHAR* Str, FOutputDevice& Ar );
 
 	/**
 	 * Attempts to load a preview static mesh from the array preview static meshes at the given index.
@@ -2620,7 +2620,7 @@ private:
 	 * @param	Index	The index of the name in the PlayerPreviewMeshNames array from the editor user settings.
 	 * @return	true if a static mesh was loaded; false, otherwise.
 	 */
-	bool LoadPreviewMesh( int32 Index );
+	UNREALED_API bool LoadPreviewMesh( int32 Index );
 
 	/**
 	 * Login PIE instances with the online platform before actually creating any PIE worlds
@@ -2630,20 +2630,20 @@ private:
 	 * @param PIEStartTime passthrough value of the time that PIE was initiated by the user
 	 */
 	UE_DEPRECATED(4.25, "This is now done as part of StartPlayInEditorSession as online-login/non-logged in flow is now combined.")
-	virtual void LoginPIEInstances(bool bAnyBlueprintErrors, bool bStartInSpectatorMode, double PIEStartTime);
+	UNREALED_API virtual void LoginPIEInstances(bool bAnyBlueprintErrors, bool bStartInSpectatorMode, double PIEStartTime);
 
 	/** Called when all PIE instances have been successfully logged in */
 	UE_DEPRECATED(4.25, "Override OnAllPIEInstancesStarted instead, which is called in both online and offline PIE sessions.")
-	virtual void OnLoginPIEAllComplete();
+	UNREALED_API virtual void OnLoginPIEAllComplete();
 	/** Called when a recompilation of a module is beginning */
-	void OnModuleCompileStarted(bool bIsAsyncCompile);
+	UNREALED_API void OnModuleCompileStarted(bool bIsAsyncCompile);
 
 	/** Called when a recompilation of a module is ending */
-	void OnModuleCompileFinished(const FString& CompilationOutput, ECompilationResult::Type CompilationResult, bool bShowLog);
+	UNREALED_API void OnModuleCompileFinished(const FString& CompilationOutput, ECompilationResult::Type CompilationResult, bool bShowLog);
 
 public:
 	/** Creates a pie world from the default entry map, used by clients that connect to a PIE server */
-	UWorld* CreatePIEWorldFromEntry(FWorldContext &WorldContext, UWorld* InWorld, FString &PlayWorldMapName);
+	UNREALED_API UWorld* CreatePIEWorldFromEntry(FWorldContext &WorldContext, UWorld* InWorld, FString &PlayWorldMapName);
 
 	/**
 	 * Continue the creation of a single PIE world after a login was successful
@@ -2654,20 +2654,20 @@ public:
 	 * @return	true if world created successfully
 	 */
 	UE_DEPRECATED(4.25, "This is now handled as part of OnLoginPIEComplete_Deferred.")
-	bool CreatePIEWorldFromLogin(FWorldContext& PieWorldContext, EPlayNetMode PlayNetMode, FPieLoginStruct& DataStruct);
+	UNREALED_API bool CreatePIEWorldFromLogin(FWorldContext& PieWorldContext, EPlayNetMode PlayNetMode, FPieLoginStruct& DataStruct);
 
 	/** Called before creating PIE instance(s). */
-	virtual FGameInstancePIEResult PreCreatePIEInstances(const bool bAnyBlueprintErrors, const bool bStartInSpectatorMode, const float PIEStartTime, const bool bSupportsOnlinePIE, int32& InNumOnlinePIEInstances);
+	UNREALED_API virtual FGameInstancePIEResult PreCreatePIEInstances(const bool bAnyBlueprintErrors, const bool bStartInSpectatorMode, const float PIEStartTime, const bool bSupportsOnlinePIE, int32& InNumOnlinePIEInstances);
 
 	/** Called before creating a PIE server instance. */
-	virtual FGameInstancePIEResult PreCreatePIEServerInstance(const bool bAnyBlueprintErrors, const bool bStartInSpectatorMode, const float PIEStartTime, const bool bSupportsOnlinePIE, int32& InNumOnlinePIEInstances);
+	UNREALED_API virtual FGameInstancePIEResult PreCreatePIEServerInstance(const bool bAnyBlueprintErrors, const bool bStartInSpectatorMode, const float PIEStartTime, const bool bSupportsOnlinePIE, int32& InNumOnlinePIEInstances);
 
 	/*
 	 * Handler for when viewport close request is made. 
 	 *
 	 * @param InViewport the viewport being closed.
 	 */
-	void OnViewportCloseRequested(FViewport* InViewport);
+	UNREALED_API void OnViewportCloseRequested(FViewport* InViewport);
 
 	/*
 	 * Fills level viewport client transform used in simulation mode 
@@ -2675,7 +2675,7 @@ public:
 	 * @param OutViewTransform filled view transform used in SIE
 	 * @return true if function succeeded, false if failed
 	 */
-	bool GetSimulateInEditorViewTransform(FTransform& OutViewTransform) const;
+	UNREALED_API bool GetSimulateInEditorViewTransform(FTransform& OutViewTransform) const;
 
 public:
 	/**
@@ -2685,7 +2685,7 @@ public:
 	 * @param bStartInSpectatorMode passthrough value if it is expected that PIE will start in spectator mode
 	 */
 	UE_DEPRECATED(4.25, "This non-online flow is now deprecated and is shared with normal online flow. See OnLoginPIEComplete_Deferred instead.")
-	void SpawnIntraProcessPIEWorlds(bool bAnyBlueprintErrors, bool bStartInSpectatorMode);
+	UNREALED_API void SpawnIntraProcessPIEWorlds(bool bAnyBlueprintErrors, bool bStartInSpectatorMode);
 
 	/*
 	 * Spawns a PlayFromHere playerstart in the given world
@@ -2696,7 +2696,7 @@ public:
 	 *
 	 * @return	true if spawn succeeded, false if failed
 	 */
-	bool SpawnPlayFromHereStart(UWorld* World, AActor*& PlayerStartPIE, const FVector& StartLocation, const FRotator& StartRotation);
+	UNREALED_API bool SpawnPlayFromHereStart(UWorld* World, AActor*& PlayerStartPIE, const FVector& StartLocation, const FRotator& StartRotation);
 	
 	/**
 	 * Spawns a PlayFromHere playerstart in the given world if there is a location set to spawn in.
@@ -2705,7 +2705,7 @@ public:
 	 *
 	 * @return	true if spawn succeeded, false if failed. Returns true even if no object was spawned (due to location not being set).
 	 */
-	bool SpawnPlayFromHereStart(UWorld* World, AActor*& PlayerStartPIE);
+	UNREALED_API bool SpawnPlayFromHereStart(UWorld* World, AActor*& PlayerStartPIE);
 
 
 private:
@@ -2721,8 +2721,8 @@ private:
 	 * @param InWorld			World containing the levels
 	 * @param InSelectCommand	Select command delegate
 	 */
-	void ExecuteCommandForAllLevelModels( UWorld* InWorld, FSelectCommand InSelectCommand, const FText& TransDesription );
-	void ExecuteCommandForAllLevelModels( UWorld* InWorld, FSelectInWorldCommand InSelectCommand, const FText& TransDesription );
+	UNREALED_API void ExecuteCommandForAllLevelModels( UWorld* InWorld, FSelectCommand InSelectCommand, const FText& TransDesription );
+	UNREALED_API void ExecuteCommandForAllLevelModels( UWorld* InWorld, FSelectInWorldCommand InSelectCommand, const FText& TransDesription );
 	
 public:
 
@@ -2731,11 +2731,11 @@ public:
 	 *
 	 * @param InWorld			World containing the levels
 	 */
-	 void FlagModifyAllSelectedSurfacesInLevels( UWorld* InWorld );
+	 UNREALED_API void FlagModifyAllSelectedSurfacesInLevels( UWorld* InWorld );
 	 
 private:
 	/** Checks for UWorld garbage collection leaks and reports any that are found */
-	void CheckForWorldGCLeaks( UWorld* NewWorld, UPackage* WorldPackage );
+	UNREALED_API void CheckForWorldGCLeaks( UWorld* NewWorld, UPackage* WorldPackage );
 
 	/**
 	 * This destroys the given world.
@@ -2746,52 +2746,52 @@ private:
 	 * @param	NewWorld	An optional new world to keep in memory after destroying the world referenced in the Context.
 	 *						This world and it's sublevels will remain in memory.
 	 */
-	void EditorDestroyWorld( FWorldContext & Context, const FText& CleanseText, UWorld* NewWorld = nullptr );
+	UNREALED_API void EditorDestroyWorld( FWorldContext & Context, const FText& CleanseText, UWorld* NewWorld = nullptr );
 
-	ULevel* CreateTransLevelMoveBuffer( UWorld* InWorld );
+	UNREALED_API ULevel* CreateTransLevelMoveBuffer( UWorld* InWorld );
 
 	/**	Broadcasts that an undo or redo has just occurred. */
-	void BroadcastPostUndoRedo(const FTransactionContext& UndoContext, bool bWasUndo);
+	UNREALED_API void BroadcastPostUndoRedo(const FTransactionContext& UndoContext, bool bWasUndo);
 
 	/** Helper function to show undo/redo notifications */
-	void ShowUndoRedoNotification(const FText& NotificationText, bool bSuccess);
+	UNREALED_API void ShowUndoRedoNotification(const FText& NotificationText, bool bSuccess);
 
 	/** Internal helper functions */
-	virtual void PostUndo (bool bSuccess);
+	UNREALED_API virtual void PostUndo (bool bSuccess);
 
 	/** Delegate callback: the world origin is going to be moved. */
-	void PreWorldOriginOffset(UWorld* InWorld, FIntVector InSrcOrigin, FIntVector InDstOrigin);
+	UNREALED_API void PreWorldOriginOffset(UWorld* InWorld, FIntVector InSrcOrigin, FIntVector InDstOrigin);
 
 	/** Delegate callback for when a streaming level is added to world. */
-	void OnLevelAddedToWorld(ULevel* InLevel, UWorld* InWorld);
+	UNREALED_API void OnLevelAddedToWorld(ULevel* InLevel, UWorld* InWorld);
 
 	/** Delegate callback for when a streaming level is removed from world. */
-	void OnLevelRemovedFromWorld(ULevel* InLevel, UWorld* InWorld);
+	UNREALED_API void OnLevelRemovedFromWorld(ULevel* InLevel, UWorld* InWorld);
 
 	/** Puts the currently loaded project file at the top of the recents list and trims and files that fell off the list */
-	void UpdateRecentlyLoadedProjectFiles();
+	UNREALED_API void UpdateRecentlyLoadedProjectFiles();
 
 	/** Updates the project file to auto load and initializes the bLoadTheMostRecentlyLoadedProjectAtStartup flag */
-	void UpdateAutoLoadProject();
+	UNREALED_API void UpdateAutoLoadProject();
 
 	/** Handles user setting changes. */
-	void HandleSettingChanged( FName Name );
+	UNREALED_API void HandleSettingChanged( FName Name );
 
 	/** Callback for handling undo and redo transactions before they happen. */
-	void HandleTransactorBeforeRedoUndo(const FTransactionContext& TransactionContext);
+	UNREALED_API void HandleTransactorBeforeRedoUndo(const FTransactionContext& TransactionContext);
 
 	/** Common code for finished undo and redo transactions. */
-	void HandleTransactorRedoUndo(const FTransactionContext& TransactionContext, bool Succeeded, bool WasUndo);
+	UNREALED_API void HandleTransactorRedoUndo(const FTransactionContext& TransactionContext, bool Succeeded, bool WasUndo);
 
 	/** Callback for finished redo transactions. */
-	void HandleTransactorRedo(const FTransactionContext& TransactionContext, bool Succeeded);
+	UNREALED_API void HandleTransactorRedo(const FTransactionContext& TransactionContext, bool Succeeded);
 
 	/** Callback for finished undo transactions. */
-	void HandleTransactorUndo(const FTransactionContext& TransactionContext, bool Succeeded);
+	UNREALED_API void HandleTransactorUndo(const FTransactionContext& TransactionContext, bool Succeeded);
 
 public:
 	/** Callback for object changes during undo/redo. */
-	void HandleObjectTransacted(UObject* InObject, const class FTransactionObjectEvent& InTransactionObjectEvent);
+	UNREALED_API void HandleObjectTransacted(UObject* InObject, const class FTransactionObjectEvent& InTransactionObjectEvent);
 
 private:
 
@@ -2914,7 +2914,7 @@ public:
 	 * @param NetMode				What net mode should this process launch under?
 	 * @param bIsDedicatedServer 	Is this instance a dedicated server? Overrides NetMode.
 	 */
-	virtual void LaunchNewProcess(const FRequestPlaySessionParams& InParams, const int32 InInstanceNum, EPlayNetMode NetMode, bool bIsDedicatedServer);
+	UNREALED_API virtual void LaunchNewProcess(const FRequestPlaySessionParams& InParams, const int32 InInstanceNum, EPlayNetMode NetMode, bool bIsDedicatedServer);
 
 	// This chunk is used for Play In New Process
 protected:
@@ -2927,10 +2927,10 @@ protected:
 	TArray<FPlayOnPCInfo> PlayOnLocalPCSessions;
 
 	/** Start a Play in New Process session with the given parameters. Called by StartQueuedPlaySessionRequestImpl based on request settings. */
-	virtual void StartPlayInNewProcessSession(FRequestPlaySessionParams& InRequestParams);
+	UNREALED_API virtual void StartPlayInNewProcessSession(FRequestPlaySessionParams& InRequestParams);
 
 	/** Does the actual late join process. Don't call this directly, use RequestLateJoin() which will defer it to the next frame. */
-	void AddPendingLateJoinClient();
+	UNREALED_API void AddPendingLateJoinClient();
 	
 	// This chunk is used for Play Via Launcher settings.
 public:
@@ -2942,26 +2942,26 @@ public:
 	/**
 	 * Cancel playing via the Launcher
 	 */
-	void CancelPlayingViaLauncher();
+	UNREALED_API void CancelPlayingViaLauncher();
 
 	/**
 	* Returns the last device that was used via the Launcher.
 	*/
-	FString GetPlayOnTargetPlatformName() const;
+	UNREALED_API FString GetPlayOnTargetPlatformName() const;
 
 	/**
 	* Launch a standalone instance using the Launcher to process required
 	* cooking, building, and deployment to a (possibly) remote device for testing.
 	*/
 	UE_DEPRECATED(4.25, "Use RequestPlaySession with parameters configured to provide a launcher target instead.")
-	void PlayUsingLauncher();
+	UNREALED_API void PlayUsingLauncher();
 
 	/** 
 	* Cancel Play using Launcher on error 
 	* 
 	* if the physical device is not authorized to be launched to, we need to pop an error instead of trying to launch
 	*/
-	void CancelPlayUsingLauncher();
+	UNREALED_API void CancelPlayUsingLauncher();
 
 protected:
 	struct FLauncherCachedInfo
@@ -2984,7 +2984,7 @@ protected:
 	FString LastPlayUsingLauncherDeviceId;
 	
 	/** Start a Launcher session with the given parameters. Called by StartQueuedPlaySessionRequestImpl based on request settings. */
-	virtual void StartPlayUsingLauncherSession(FRequestPlaySessionParams& InRequestParams);
+	UNREALED_API virtual void StartPlayUsingLauncherSession(FRequestPlaySessionParams& InRequestParams);
 
 	/** This flag is used to skip UAT\UBT compilation on every launch if it was successfully compiled once. */
 	bool bUATSuccessfullyCompiledOnce;
@@ -2992,7 +2992,7 @@ protected:
 	// This chunk is for Play in Editor
 public:
 	/** @return true if the editor is able to launch PIE with online platform support */
-	bool SupportsOnlinePIE() const;
+	UNREALED_API bool SupportsOnlinePIE() const;
 
 	/** @return true if there are active PIE instances logged into an online platform */
 	bool IsPlayingWithOnlinePIE() const
@@ -3015,15 +3015,15 @@ protected:
 	DECLARE_DELEGATE(FPIEInstanceWindowSwitch);
 
 	/** Sets the delegate for when the focused PIE window is changed */
-	void SetPIEInstanceWindowSwitchDelegate(FPIEInstanceWindowSwitch PIEInstanceWindowSwitchDelegate);
+	UNREALED_API void SetPIEInstanceWindowSwitchDelegate(FPIEInstanceWindowSwitch PIEInstanceWindowSwitchDelegate);
 
 	/** Gets the scene viewport for a viewport client */
-	FSceneViewport* GetGameSceneViewport(UGameViewportClient* ViewportClient) const;
+	UNREALED_API FSceneViewport* GetGameSceneViewport(UGameViewportClient* ViewportClient) const;
 
 	/**
 	 * Toggles PIE to SIE or vice-versa
 	 */
-	void ToggleBetweenPIEandSIE(bool bNewSession = false);
+	UNREALED_API void ToggleBetweenPIEandSIE(bool bNewSession = false);
 
 	/**
 	 * Hack to switch worlds for the PIE window before and after a slate event
@@ -3032,31 +3032,31 @@ protected:
 	 * @param PIEInstance	When switching to a PIE instance, this is the specific client/server instance to use
 	 * @return The ID of the world to restore later or -1 if no world to restore
 	 */
-	int32 OnSwitchWorldForSlatePieWindow(int32 WorldID, int32 WorldPIEInstance);
+	UNREALED_API int32 OnSwitchWorldForSlatePieWindow(int32 WorldID, int32 WorldPIEInstance);
 
 	/**
 	 * Called via a delegate to toggle between the editor and pie world
 	 */
-	void OnSwitchWorldsForPIE(bool bSwitchToPieWorld, UWorld* OverrideWorld = nullptr);
+	UNREALED_API void OnSwitchWorldsForPIE(bool bSwitchToPieWorld, UWorld* OverrideWorld = nullptr);
 
 	/**
 	 * Called to switch to a specific PIE instance, where -1 means the editor world
 	 */
-	void OnSwitchWorldsForPIEInstance(int32 WorldPIEInstance);
+	UNREALED_API void OnSwitchWorldsForPIEInstance(int32 WorldPIEInstance);
 
 	/** Call to enable/disable callbacks for PIE world switching when PIE starts/stops */
-	void EnableWorldSwitchCallbacks(bool bEnable);
+	UNREALED_API void EnableWorldSwitchCallbacks(bool bEnable);
 
 	/** Callback when script execution starts, might switch world */
-	void OnScriptExecutionStart(const struct FBlueprintContextTracker& ContextTracker, const UObject* ContextObject, const UFunction* ContextFunction);
+	UNREALED_API void OnScriptExecutionStart(const struct FBlueprintContextTracker& ContextTracker, const UObject* ContextObject, const UFunction* ContextFunction);
 
 	/** Callback when script execution starts, might switch world */
-	void OnScriptExecutionEnd(const struct FBlueprintContextTracker& ContextTracker);
+	UNREALED_API void OnScriptExecutionEnd(const struct FBlueprintContextTracker& ContextTracker);
 
 	/**
 	 * Gives focus to the server or first PIE client viewport
 	 */
-	void GiveFocusToLastClientPIEViewport();
+	UNREALED_API void GiveFocusToLastClientPIEViewport();
 
 	/**
 	 * Delegate called as each PIE instance login is complete, continues creating the PIE world for a single instance
@@ -3066,36 +3066,36 @@ protected:
 	 * @param ErrorString descriptive error when applicable
 	 * @param DataStruct data required to continue PIE creation, set at login time
 	 */
-	virtual void OnLoginPIEComplete(int32 LocalUserNum, bool bWasSuccessful, const FString& ErrorString, FPieLoginStruct DataStruct);
+	UNREALED_API virtual void OnLoginPIEComplete(int32 LocalUserNum, bool bWasSuccessful, const FString& ErrorString, FPieLoginStruct DataStruct);
 
 	/** Above function but called a frame later, to stop PIE login from happening from a network callback */
-	virtual void OnLoginPIEComplete_Deferred(int32 LocalUserNum, bool bWasSuccessful, FString ErrorString, FPieLoginStruct DataStruct);
+	UNREALED_API virtual void OnLoginPIEComplete_Deferred(int32 LocalUserNum, bool bWasSuccessful, FString ErrorString, FPieLoginStruct DataStruct);
 
 	/** allow for game specific override to determine if login should be treated as successful for pass-through handling instead */
 	virtual bool IsLoginPIESuccessful(int32 LocalUserNum, bool bWasSuccessful, const FString& ErrorString, const FPieLoginStruct& DataStruct) { return bWasSuccessful; }
 
 	/** Called when all PIE instances have been successfully logged in */
-	virtual void OnAllPIEInstancesStarted();
+	UNREALED_API virtual void OnAllPIEInstancesStarted();
 	
 	/** Backs up the current editor selection and then clears it. Optionally reselects the instances in the Play world. */
-	void TransferEditorSelectionToPlayInstances(const bool bInSelectInstances);
+	UNREALED_API void TransferEditorSelectionToPlayInstances(const bool bInSelectInstances);
 	
 	/** Create a new GameInstance for PIE with the specified parameters. */
-	virtual UGameInstance* CreateInnerProcessPIEGameInstance(FRequestPlaySessionParams& InParams, const FGameInstancePIEParameters& InPIEParameters, int32 InPIEInstanceIndex);
+	UNREALED_API virtual UGameInstance* CreateInnerProcessPIEGameInstance(FRequestPlaySessionParams& InParams, const FGameInstancePIEParameters& InPIEParameters, int32 InPIEInstanceIndex);
 	
 	/** Creates a SPIEViewport for the specified instance. This is used for Play in New Window, but not for instances that live in the level editor viewport. */
-	TSharedRef<class SPIEViewport> GeneratePIEViewportWindow(const FRequestPlaySessionParams& InSessionParams, int32 InViewportIndex, const FWorldContext& InWorldContext, EPlayNetMode InNetMode, UGameViewportClient* InViewportClient, FSlatePlayInEditorInfo& InSlateInfo);
+	UNREALED_API TSharedRef<class SPIEViewport> GeneratePIEViewportWindow(const FRequestPlaySessionParams& InSessionParams, int32 InViewportIndex, const FWorldContext& InWorldContext, EPlayNetMode InNetMode, UGameViewportClient* InViewportClient, FSlatePlayInEditorInfo& InSlateInfo);
 
 	/** Stores the position of the window for that instance index. Doesn't save it to the CDO until the play session ends. */
-	void StoreWindowSizeAndPositionForInstanceIndex(const int32 InInstanceIndex, const FIntPoint& InSize, const FIntPoint& InPosition);
+	UNREALED_API void StoreWindowSizeAndPositionForInstanceIndex(const int32 InInstanceIndex, const FIntPoint& InSize, const FIntPoint& InPosition);
 	/** Returns the stored position at that index. */
-	void GetWindowSizeAndPositionForInstanceIndex(ULevelEditorPlaySettings& InEditorPlaySettings, const int32 InInstanceIndex, const FWorldContext& InWorldContext, FIntPoint& OutSize, FIntPoint& OutPosition);
+	UNREALED_API void GetWindowSizeAndPositionForInstanceIndex(ULevelEditorPlaySettings& InEditorPlaySettings, const int32 InInstanceIndex, const FWorldContext& InWorldContext, FIntPoint& OutSize, FIntPoint& OutPosition);
 
 	/** Start the queued Play Session Request. After this is called the queued play session request will be cleared. */
-	virtual void StartQueuedPlaySessionRequestImpl();
+	UNREALED_API virtual void StartQueuedPlaySessionRequestImpl();
 	
 	/** Start a Play in Editor session with the given parameters. Called by StartQueuedPlaySessionRequestImpl based on request settings. */
-	virtual void StartPlayInEditorSession(FRequestPlaySessionParams& InRequestParams);
+	UNREALED_API virtual void StartPlayInEditorSession(FRequestPlaySessionParams& InRequestParams);
 
 	/** 
 	* Creates a new Play in Editor instance (which may be in a new process if not running under one process.
@@ -3105,11 +3105,11 @@ protected:
 	* @param	InNetMode				- The net mode to launch this instance with. Not pulled from the RequestParams because some net modes (such as ListenServer)
 										  need different net modes for subsequent instances.
 	*/
-	virtual void CreateNewPlayInEditorInstance(FRequestPlaySessionParams &InRequestParams, const bool bInDedicatedInstance, const EPlayNetMode InNetMode);
+	UNREALED_API virtual void CreateNewPlayInEditorInstance(FRequestPlaySessionParams &InRequestParams, const bool bInDedicatedInstance, const EPlayNetMode InNetMode);
 
 
 	/** Asks the player to save dirty maps, if this fails it will return false. */
-	bool SaveMapsForPlaySession();
+	UNREALED_API bool SaveMapsForPlaySession();
 	
 protected:
 
@@ -3136,41 +3136,41 @@ protected:
 	/**
 	 * Invalidates all editor viewports and hit proxies, used when global changes like Undo/Redo may have invalidated state everywhere
 	 */
-	void InvalidateAllViewportsAndHitProxies();
+	UNREALED_API void InvalidateAllViewportsAndHitProxies();
 
 	/** Initialize Portal RPC. */
-	void InitializePortal();
+	UNREALED_API void InitializePortal();
 
 	/** Destroy any online subsystems generated by PIE */
-	void CleanupPIEOnlineSessions(TArray<FName> OnlineIdentifiers);
+	UNREALED_API void CleanupPIEOnlineSessions(TArray<FName> OnlineIdentifiers);
 
 	// launch on callbacks
-	void HandleStageStarted(const FString& InStage, TWeakPtr<SNotificationItem> NotificationItemPtr);
-	void HandleStageCompleted(const FString& InStage, double StageTime, bool bHasCode, TWeakPtr<SNotificationItem> NotificationItemPtr);
-	void HandleLaunchCanceled(double TotalTime, bool bHasCode, TWeakPtr<SNotificationItem> NotificationItemPtr);
-	void HandleLaunchCompleted(bool Succeeded, double TotalTime, int32 ErrorCode, bool bHasCode, TWeakPtr<SNotificationItem> NotificationItemPtr);
+	UNREALED_API void HandleStageStarted(const FString& InStage, TWeakPtr<SNotificationItem> NotificationItemPtr);
+	UNREALED_API void HandleStageCompleted(const FString& InStage, double StageTime, bool bHasCode, TWeakPtr<SNotificationItem> NotificationItemPtr);
+	UNREALED_API void HandleLaunchCanceled(double TotalTime, bool bHasCode, TWeakPtr<SNotificationItem> NotificationItemPtr);
+	UNREALED_API void HandleLaunchCompleted(bool Succeeded, double TotalTime, int32 ErrorCode, bool bHasCode, TWeakPtr<SNotificationItem> NotificationItemPtr);
 
 	// Handle requests from slate application to open assets.
-	bool HandleOpenAsset(UObject* Asset);
+	UNREALED_API bool HandleOpenAsset(UObject* Asset);
 
 private:
 	/** Handler for when any asset is loaded in the editor */
-	void OnAssetLoaded( UObject* Asset );
+	UNREALED_API void OnAssetLoaded( UObject* Asset );
 
 	/** Handler for when an asset is created (used to detect world duplication after PostLoad) */
-	void OnAssetCreated( UObject* Asset );
+	UNREALED_API void OnAssetCreated( UObject* Asset );
 
 	/** Handler for when an asset finishes compiling (used to notify AssetRegistry to update tags after Load) */
-	void OnAssetPostCompile(const TArray<FAssetCompileData>& CompiledAssets);
+	UNREALED_API void OnAssetPostCompile(const TArray<FAssetCompileData>& CompiledAssets);
 
 	/** Handler for when a world is duplicated in the editor */
-	void InitializeNewlyCreatedInactiveWorld(UWorld* World);
+	UNREALED_API void InitializeNewlyCreatedInactiveWorld(UWorld* World);
 
 	/** Gets the init values for worlds opened via Map_Load in the editor */
-	UWorld::InitializationValues GetEditorWorldInitializationValues() const;
+	UNREALED_API UWorld::InitializationValues GetEditorWorldInitializationValues() const;
 
 	/** Logic that needs to run on change of effective shader platform (can happen via both SetPreviewPlatform and ToggleFeatureLevelPreview). */
-	void OnEffectivePreviewShaderPlatformChange();
+	UNREALED_API void OnEffectivePreviewShaderPlatformChange();
 
 public:
 	// Launcher Worker
@@ -3178,33 +3178,33 @@ public:
 	
 	/** Function to run the Play On command for automation testing. */
 	UE_DEPRECATED(4.25, "Use RequestPlaySession and StartQueuedPlaySessionRequest instead.")
-	void AutomationPlayUsingLauncher(const FString& InLauncherDeviceId);	
+	UNREALED_API void AutomationPlayUsingLauncher(const FString& InLauncherDeviceId);	
 
-	void AutomationLoadMap(const FString& MapName, bool bForceReload, FString* OutError);
+	UNREALED_API void AutomationLoadMap(const FString& MapName, bool bForceReload, FString* OutError);
 
 	/** This function should be called to notify the editor that new materials were added to our scene or some materials were modified */
-	void OnSceneMaterialsModified();
+	UNREALED_API void OnSceneMaterialsModified();
 
 	/** Call this function to change the feature level and to override the material quality platform of the editor and PIE worlds */
-	void SetPreviewPlatform(const FPreviewPlatformInfo& NewPreviewPlatform, bool bSaveSettings);
+	UNREALED_API void SetPreviewPlatform(const FPreviewPlatformInfo& NewPreviewPlatform, bool bSaveSettings);
 
 	/** Toggle the feature level preview */
-	void ToggleFeatureLevelPreview();
+	UNREALED_API void ToggleFeatureLevelPreview();
 
 	/** Return whether the feature level preview is enabled for use via the user accessing Settings->Preview Rendering Level. */
-	bool IsFeatureLevelPreviewEnabled() const;
+	UNREALED_API bool IsFeatureLevelPreviewEnabled() const;
 
 	/** Return whether the feature level preview enabled via Settings->Preview Rendering Level is currently active/displaying. */
-	bool IsFeatureLevelPreviewActive() const;
+	UNREALED_API bool IsFeatureLevelPreviewActive() const;
 
 	/** Return the current active shader platform in the editor*/
-	EShaderPlatform GetActiveShaderPlatform() const;
+	UNREALED_API EShaderPlatform GetActiveShaderPlatform() const;
 
 	/**
 	 * Return the active feature level. This will be the chosen Settings->Preview Rendering Level if the Preview Mode button
 	 * is highlighted or SM5 if the Preview Mode button has been clicked and isn't highlighted.
 	 */
-	ERHIFeatureLevel::Type GetActiveFeatureLevelPreviewType() const;
+	UNREALED_API ERHIFeatureLevel::Type GetActiveFeatureLevelPreviewType() const;
 
 	/** Return the delegate that is called when the preview feature level changes */
 	FPreviewFeatureLevelChanged& OnPreviewFeatureLevelChanged() { return PreviewFeatureLevelChanged; }
@@ -3218,14 +3218,14 @@ public:
 protected:
 
 	/** Function pair used to save and restore the global feature level */
-	void LoadEditorFeatureLevel();
-	void SaveEditorFeatureLevel();
+	UNREALED_API void LoadEditorFeatureLevel();
+	UNREALED_API void SaveEditorFeatureLevel();
 
 	/** Utility function that can determine whether some input world is using materials who's shaders are emulated in the editor */
-	bool IsEditorShaderPlatformEmulated(UWorld* World);
+	UNREALED_API bool IsEditorShaderPlatformEmulated(UWorld* World);
 
 	/** Utility function that checks if, for the current shader platform used by the editor, there's an available offline shader compiler */
-	bool IsOfflineShaderCompilerAvailable(UWorld* World);
+	UNREALED_API bool IsOfflineShaderCompilerAvailable(UWorld* World);
 
 protected:
 
