[UnrealEd]

diff --git a/Engine/Source/Editor/UnrealEd/Classes/Editor/TransBuffer.h b/Engine/Source/Editor/UnrealEd/Classes/Editor/TransBuffer.h
--- a/Engine/Source/Editor/UnrealEd/Classes/Editor/TransBuffer.h
+++ b/Engine/Source/Editor/UnrealEd/Classes/Editor/TransBuffer.h
@@ -11,8 +11,8 @@
 #include "Editor/Transactor.h"
 #include "TransBuffer.generated.h"
 
-UCLASS(transient)
-class UNREALED_API UTransBuffer
+UCLASS(transient, MinimalAPI)
+class UTransBuffer
 	: public UTransactor
 {
 public:
@@ -42,28 +42,28 @@ public:
 public:
 
 	// Constructor.
-	void Initialize(SIZE_T InMaxMemory);
+	UNREALED_API void Initialize(SIZE_T InMaxMemory);
 
 public:
 
 	/**
 	 * Validates the state of the transaction buffer.
 	 */
-	void CheckState() const;
-	static void AddReferencedObjects(UObject* InThis, FReferenceCollector& Collector);
+	UNREALED_API void CheckState() const;
+	static UNREALED_API void AddReferencedObjects(UObject* InThis, FReferenceCollector& Collector);
 
 public:
 
 	//~ Begin UObject Interface.
 
-	virtual void Serialize( FArchive& Ar ) override;
-	virtual void FinishDestroy() override;
+	UNREALED_API virtual void Serialize( FArchive& Ar ) override;
+	UNREALED_API virtual void FinishDestroy() override;
 
 	//~ End UObject Interface.
 
 protected:
 	
-	void OnObjectsReinstanced(const TMap<UObject*, UObject*>& OldToNewInstances);
+	UNREALED_API void OnObjectsReinstanced(const TMap<UObject*, UObject*>& OldToNewInstances);
 
 	/** Implementation of the begin function. Used to create a specific transaction type */
 	template<typename TTransaction>
@@ -145,32 +145,32 @@ public:
 
 	//~ Begin UTransactor Interface.
 
-	virtual int32 Begin( const TCHAR* SessionContext, const FText& Description ) override;
-	virtual int32 End() override;
-	virtual void Cancel( int32 StartIndex = 0 ) override;
-	virtual void Reset( const FText& Reason ) override;
-	virtual bool CanUndo( FText* Text=NULL ) override;
-	virtual bool CanRedo( FText* Text=NULL ) override;
+	UNREALED_API virtual int32 Begin( const TCHAR* SessionContext, const FText& Description ) override;
+	UNREALED_API virtual int32 End() override;
+	UNREALED_API virtual void Cancel( int32 StartIndex = 0 ) override;
+	UNREALED_API virtual void Reset( const FText& Reason ) override;
+	UNREALED_API virtual bool CanUndo( FText* Text=NULL ) override;
+	UNREALED_API virtual bool CanRedo( FText* Text=NULL ) override;
 	virtual int32 GetQueueLength( ) const override { return UndoBuffer.Num(); }
-	virtual int32 FindTransactionIndex(const FGuid& TransactionId) const override;
-	virtual const FTransaction* GetTransaction( int32 QueueIndex ) const override;
-	virtual FTransactionContext GetUndoContext( bool bCheckWhetherUndoPossible = true ) override;
-	virtual SIZE_T GetUndoSize() const override;
+	UNREALED_API virtual int32 FindTransactionIndex(const FGuid& TransactionId) const override;
+	UNREALED_API virtual const FTransaction* GetTransaction( int32 QueueIndex ) const override;
+	UNREALED_API virtual FTransactionContext GetUndoContext( bool bCheckWhetherUndoPossible = true ) override;
+	UNREALED_API virtual SIZE_T GetUndoSize() const override;
 	virtual int32 GetUndoCount( ) const override { return UndoCount; }
-	virtual FTransactionContext GetRedoContext() override;
-	virtual void SetUndoBarrier() override;
-	virtual void RemoveUndoBarrier() override;
-	virtual void ClearUndoBarriers() override;
-	virtual int32 GetCurrentUndoBarrier() const override;
-	virtual bool Undo(bool bCanRedo = true) override;
-	virtual bool Redo() override;
-	virtual bool EnableObjectSerialization() override;
-	virtual bool DisableObjectSerialization() override;
+	UNREALED_API virtual FTransactionContext GetRedoContext() override;
+	UNREALED_API virtual void SetUndoBarrier() override;
+	UNREALED_API virtual void RemoveUndoBarrier() override;
+	UNREALED_API virtual void ClearUndoBarriers() override;
+	UNREALED_API virtual int32 GetCurrentUndoBarrier() const override;
+	UNREALED_API virtual bool Undo(bool bCanRedo = true) override;
+	UNREALED_API virtual bool Redo() override;
+	UNREALED_API virtual bool EnableObjectSerialization() override;
+	UNREALED_API virtual bool DisableObjectSerialization() override;
 	virtual bool IsObjectSerializationEnabled() override { return DisallowObjectSerialization == 0; }
-	virtual void SetPrimaryUndoObject( UObject* Object ) override;
-	virtual bool IsObjectInTransactionBuffer( const UObject* Object ) const override;
-	virtual bool IsObjectTransacting(const UObject* Object) const override;
-	virtual bool ContainsPieObjects() const override;
+	UNREALED_API virtual void SetPrimaryUndoObject( UObject* Object ) override;
+	UNREALED_API virtual bool IsObjectInTransactionBuffer( const UObject* Object ) const override;
+	UNREALED_API virtual bool IsObjectTransacting(const UObject* Object) const override;
+	UNREALED_API virtual bool ContainsPieObjects() const override;
 	virtual bool IsActive() override
 	{
 		return ActiveCount > 0;
diff --git a/Engine/Source/Editor/UnrealEd/Classes/Editor/Transactor.h b/Engine/Source/Editor/UnrealEd/Classes/Editor/Transactor.h
--- a/Engine/Source/Editor/UnrealEd/Classes/Editor/Transactor.h
+++ b/Engine/Source/Editor/UnrealEd/Classes/Editor/Transactor.h
@@ -33,13 +33,13 @@
  * warning: UObject::Serialize implicitly assumes that class properties do not change
  * in between transaction resets.
  */
-class UNREALED_API FTransaction : public ITransaction
+class FTransaction : public ITransaction
 {
 	friend class FTransactionObjectEvent;
 
 protected:
 	// Record of an object.
-	class UNREALED_API FObjectRecord
+	class FObjectRecord
 	{
 	public:
 		struct FSerializedObject : public UE::Transaction::FSerializedObject
@@ -132,7 +132,7 @@ protected:
 
 		// Constructors.
 		FObjectRecord() = default;
-		FObjectRecord( FTransaction* Owner, UObject* InObject, TUniquePtr<FChange> InCustomChange, FScriptArray* InArray, int32 InIndex, int32 InCount, int32 InOper, int32 InElementSize, uint32 InElementAlignment, STRUCT_DC InDefaultConstructor, STRUCT_AR InSerializer, STRUCT_DTOR InDestructor );
+		UNREALED_API FObjectRecord( FTransaction* Owner, UObject* InObject, TUniquePtr<FChange> InCustomChange, FScriptArray* InArray, int32 InIndex, int32 InCount, int32 InOper, int32 InElementSize, uint32 InElementAlignment, STRUCT_DC InDefaultConstructor, STRUCT_AR InSerializer, STRUCT_DTOR InDestructor );
 
 	private:
 		// Non-copyable
@@ -140,28 +140,28 @@ protected:
 		FObjectRecord& operator=( const FObjectRecord& ) = delete;
 
 	public:
-		UE::Transaction::FDiffableObject GetDiffableObject(TArrayView<const FProperty*> PropertiesToSerialize = TArrayView<const FProperty*>(), UE::Transaction::DiffUtil::EGetDiffableObjectMode ObjectSerializationMode = UE::Transaction::DiffUtil::EGetDiffableObjectMode::SerializeObject) const;
+		UNREALED_API UE::Transaction::FDiffableObject GetDiffableObject(TArrayView<const FProperty*> PropertiesToSerialize = TArrayView<const FProperty*>(), UE::Transaction::DiffUtil::EGetDiffableObjectMode ObjectSerializationMode = UE::Transaction::DiffUtil::EGetDiffableObjectMode::SerializeObject) const;
 
 		// Functions.
-		void SerializeContents( FArchive& Ar, int32 InOper ) const;
-		void SerializeObject( FArchive& Ar ) const;
-		void Restore( FTransaction* Owner );
-		void Save( FTransaction* Owner );
-		void Load( FTransaction* Owner );
-		void Finalize( FTransaction* Owner, UE::Transaction::DiffUtil::FDiffableObjectArchetypeCache& ArchetypeCache, TSharedPtr<ITransactionObjectAnnotation>& OutFinalizedObjectAnnotation );
-		void Snapshot( FTransaction* Owner, UE::Transaction::DiffUtil::FDiffableObjectArchetypeCache& ArchetypeCache, TArrayView<const FProperty*> Properties );
+		UNREALED_API void SerializeContents( FArchive& Ar, int32 InOper ) const;
+		UNREALED_API void SerializeObject( FArchive& Ar ) const;
+		UNREALED_API void Restore( FTransaction* Owner );
+		UNREALED_API void Save( FTransaction* Owner );
+		UNREALED_API void Load( FTransaction* Owner );
+		UNREALED_API void Finalize( FTransaction* Owner, UE::Transaction::DiffUtil::FDiffableObjectArchetypeCache& ArchetypeCache, TSharedPtr<ITransactionObjectAnnotation>& OutFinalizedObjectAnnotation );
+		UNREALED_API void Snapshot( FTransaction* Owner, UE::Transaction::DiffUtil::FDiffableObjectArchetypeCache& ArchetypeCache, TArrayView<const FProperty*> Properties );
 
 		/** Used by GC to collect referenced objects. */
-		void AddReferencedObjects( FReferenceCollector& Collector );
+		UNREALED_API void AddReferencedObjects( FReferenceCollector& Collector );
 
 		/** @return True if this record contains a reference to a pie object */
-		bool ContainsPieObject() const;
+		UNREALED_API bool ContainsPieObject() const;
 
 		/** @return true if the record has a delta change or a custom change */
-		bool HasChanges() const;
+		UNREALED_API bool HasChanges() const;
 
 		/** @return true if the record has a custom change and the change has expired */
-		bool HasExpired() const;
+		UNREALED_API bool HasExpired() const;
 
 		/** Transfers data from an array. */
 		class FReader : public UE::Transaction::FSerializedObjectDataReader
@@ -327,33 +327,33 @@ private:
 	FTransaction( const FTransaction& ) = delete;
 	FTransaction& operator=( const FTransaction& ) = delete;
 
-	void SavePackage(UPackage* Package);
+	UNREALED_API void SavePackage(UPackage* Package);
 
 public:
 	
 	// FTransactionBase interface.
-	virtual void SaveObject( UObject* Object ) override;
-	virtual void SaveArray( UObject* Object, FScriptArray* Array, int32 Index, int32 Count, int32 Oper, int32 ElementSize, uint32 ElementAlignment, STRUCT_DC DefaultConstructor, STRUCT_AR Serializer, STRUCT_DTOR Destructor ) override;
-	virtual void StoreUndo( UObject* Object, TUniquePtr<FChange> UndoChange ) override;
-	virtual void SetPrimaryObject(UObject* InObject) override;
-	virtual void SnapshotObject( UObject* InObject, TArrayView<const FProperty*> Properties ) override;
-	virtual bool ContainsObject(const UObject* Object) const override;
+	UNREALED_API virtual void SaveObject( UObject* Object ) override;
+	UNREALED_API virtual void SaveArray( UObject* Object, FScriptArray* Array, int32 Index, int32 Count, int32 Oper, int32 ElementSize, uint32 ElementAlignment, STRUCT_DC DefaultConstructor, STRUCT_AR Serializer, STRUCT_DTOR Destructor ) override;
+	UNREALED_API virtual void StoreUndo( UObject* Object, TUniquePtr<FChange> UndoChange ) override;
+	UNREALED_API virtual void SetPrimaryObject(UObject* InObject) override;
+	UNREALED_API virtual void SnapshotObject( UObject* InObject, TArrayView<const FProperty*> Properties ) override;
+	UNREALED_API virtual bool ContainsObject(const UObject* Object) const override;
 
 	/** BeginOperation should be called when a transaction or undo/redo starts */
-	virtual void BeginOperation() override;
+	UNREALED_API virtual void BeginOperation() override;
 
 	/** EndOperation should be called when a transaction is finalized or canceled or undo/redo ends */
-	virtual void EndOperation() override;
+	UNREALED_API virtual void EndOperation() override;
 
 	/**
 	 * Enacts the transaction.
 	 */
-	virtual void Apply() override;
+	UNREALED_API virtual void Apply() override;
 
 	/**
 	 * Finalize the transaction (try and work out what's changed).
 	 */
-	virtual void Finalize() override;
+	UNREALED_API virtual void Finalize() override;
 
 	/**
 	 * Gets the full context for the transaction.
@@ -364,13 +364,13 @@ public:
 	}
 
 	/** @return true if the transaction is transient. */
-	virtual bool IsTransient() const override;
+	UNREALED_API virtual bool IsTransient() const override;
 
 	/** @return True if this record contains a reference to a pie object */
-	virtual bool ContainsPieObjects() const override;
+	UNREALED_API virtual bool ContainsPieObjects() const override;
 
 	/** @return true if the transaction contains custom changes and all the changes have expired */
-	virtual bool HasExpired() const;
+	UNREALED_API virtual bool HasExpired() const;
 
 
 	/** Returns a unique string to serve as a type ID for the FTranscationBase-derived type. */
@@ -380,7 +380,7 @@ public:
 	}
 
 	// FTransaction interface.
-	SIZE_T DataSize() const;
+	UNREALED_API SIZE_T DataSize() const;
 
 	/** Returns the unique identifier for this transaction, used to track it during its lifetime */
 	FGuid GetId() const
@@ -451,29 +451,29 @@ public:
 	}
 
 	/** Used by GC to collect referenced objects. */
-	void AddReferencedObjects( FReferenceCollector& Collector );
+	UNREALED_API void AddReferencedObjects( FReferenceCollector& Collector );
 
 	/**
 	 * Get all the objects that are part of this transaction.
 	 * @param	Objects		[out] Receives the object list.  Previous contents are cleared.
 	 */
-	void GetTransactionObjects(TArray<UObject*>& Objects) const;
-	int32 GetRecordCount() const;
+	UNREALED_API void GetTransactionObjects(TArray<UObject*>& Objects) const;
+	UNREALED_API int32 GetRecordCount() const;
 
 	const UObject* GetPrimaryObject() const { return PrimaryObject; }
 
 	/** Checks if a specific object is in the transaction currently underway */
-	bool IsObjectTransacting(const UObject* Object) const;
+	UNREALED_API bool IsObjectTransacting(const UObject* Object) const;
 
 	/**
 	 * Outputs the contents of the ObjectMap to the specified output device.
 	 */
-	void DumpObjectMap(FOutputDevice& Ar) const;
+	UNREALED_API void DumpObjectMap(FOutputDevice& Ar) const;
 
 	/**
 	 * Create a map that holds information about objects of a given transaction.
 	 */
-	FTransactionDiff GenerateDiff() const;
+	UNREALED_API FTransactionDiff GenerateDiff() const;
 
 	// Transaction friends.
 	friend FArchive& operator<<( FArchive& Ar, FTransaction::FObjectRecord& R );
@@ -483,8 +483,8 @@ public:
 	friend class FObjectRecord::FWriter;
 };
 
-UCLASS(abstract, transient)
-class UNREALED_API UTransactor : public UObject
+UCLASS(abstract, transient, MinimalAPI)
+class UTransactor : public UObject
 {
     GENERATED_UCLASS_BODY()
 	/**
@@ -500,7 +500,7 @@ class UNREALED_API UTransactor : public UObject
 	 * @return	Number of active actions when Begin() was called;  values greater than
 	 *			0 indicate that there was already an existing undo transaction in progress.
 	 */
-	virtual int32 Begin( const TCHAR* SessionContext, const FText& Description ) PURE_VIRTUAL(UTransactor::Begin,return 0;);
+	UNREALED_API virtual int32 Begin( const TCHAR* SessionContext, const FText& Description ) PURE_VIRTUAL(UTransactor::Begin,return 0;);
 
 	/**
 	 * Attempts to close an undo transaction.  Only successful if the transaction's action
@@ -509,25 +509,25 @@ class UNREALED_API UTransactor : public UObject
 	 * @return	Number of active actions when End() was called; a value of 1 indicates that the
 	 *			transaction was successfully closed
 	 */
-	virtual int32 End() PURE_VIRTUAL(UTransactor::End,return 0;);
+	UNREALED_API virtual int32 End() PURE_VIRTUAL(UTransactor::End,return 0;);
 
 	/**
 	 * Cancels the current transaction, no longer capture actions to be placed in the undo buffer.
 	 *
 	 * @param	StartIndex	the value of ActiveIndex when the transaction to be canceled was began. 
 	 */
-	virtual void Cancel( int32 StartIndex = 0 ) PURE_VIRTUAL(UTransactor::Cancel,);
+	UNREALED_API virtual void Cancel( int32 StartIndex = 0 ) PURE_VIRTUAL(UTransactor::Cancel,);
 
 	/**
 	 * Resets the entire undo buffer;  deletes all undo transactions.
 	 */
-	virtual void Reset( const FText& Reason ) PURE_VIRTUAL(UTransactor::Reset,);
+	UNREALED_API virtual void Reset( const FText& Reason ) PURE_VIRTUAL(UTransactor::Reset,);
 
 	/**
 	 * Returns whether there are any active actions; i.e. whether actions are currently
 	 * being captured into the undo buffer.
 	 */
-	virtual bool IsActive() PURE_VIRTUAL(UTransactor::IsActive,return false;);
+	UNREALED_API virtual bool IsActive() PURE_VIRTUAL(UTransactor::IsActive,return false;);
 
 	/**
 	 * Determines whether the undo option should be selectable.
@@ -536,7 +536,7 @@ class UNREALED_API UTransactor : public UObject
 	 *
 	 * @return	true if the "Undo" option should be selectable.
 	 */
-	virtual bool CanUndo( FText* Text=nullptr ) PURE_VIRTUAL(UTransactor::CanUndo,return false;);
+	UNREALED_API virtual bool CanUndo( FText* Text=nullptr ) PURE_VIRTUAL(UTransactor::CanUndo,return false;);
 
 	/**
 	 * Determines whether the redo option should be selectable.
@@ -545,14 +545,14 @@ class UNREALED_API UTransactor : public UObject
 	 *
 	 * @return	true if the "Redo" option should be selectable.
 	 */
-	virtual bool CanRedo( FText* Text=nullptr ) PURE_VIRTUAL(UTransactor::CanRedo,return false;);
+	UNREALED_API virtual bool CanRedo( FText* Text=nullptr ) PURE_VIRTUAL(UTransactor::CanRedo,return false;);
 
 	/**
 	 * Gets the current length of the transaction queue.
 	 *
 	 * @return Queue length.
 	 */
-	virtual int32 GetQueueLength( ) const PURE_VIRTUAL(UTransactor::GetQueueLength,return 0;);
+	UNREALED_API virtual int32 GetQueueLength( ) const PURE_VIRTUAL(UTransactor::GetQueueLength,return 0;);
 
 	/**
 	 * Gets the transaction queue index from its TransactionId.
@@ -561,7 +561,7 @@ class UNREALED_API UTransactor : public UObject
 	 *
 	 * @return The queue index or INDEX_NONE if not found.
 	 */
-	virtual int32 FindTransactionIndex(const FGuid& TransactionId ) const PURE_VIRTUAL(UTransactor::FindTransactionIndex, return INDEX_NONE;);
+	UNREALED_API virtual int32 FindTransactionIndex(const FGuid& TransactionId ) const PURE_VIRTUAL(UTransactor::FindTransactionIndex, return INDEX_NONE;);
 
 	/**
 	 * Gets the transaction at the specified queue index.
@@ -570,7 +570,7 @@ class UNREALED_API UTransactor : public UObject
 	 *
 	 * @return A read-only pointer to the transaction, or NULL if it does not exist.
 	 */
-	virtual const FTransaction* GetTransaction( int32 QueueIndex ) const PURE_VIRTUAL(UTransactor::GetTransaction,return nullptr;);
+	UNREALED_API virtual const FTransaction* GetTransaction( int32 QueueIndex ) const PURE_VIRTUAL(UTransactor::GetTransaction,return nullptr;);
 
 	/**
 	 * Returns the description of the undo action that will be performed next.
@@ -580,21 +580,21 @@ class UNREALED_API UTransactor : public UObject
 	 *
 	 * @return	text describing the next undo transaction
 	 */
-	virtual FTransactionContext GetUndoContext( bool bCheckWhetherUndoPossible = true ) PURE_VIRTUAL(UTransactor::GetUndoContext,return FTransactionContext(););
+	UNREALED_API virtual FTransactionContext GetUndoContext( bool bCheckWhetherUndoPossible = true ) PURE_VIRTUAL(UTransactor::GetUndoContext,return FTransactionContext(););
 
 	/**
 	 * Determines the amount of data currently stored by the transaction buffer.
 	 *
 	 * @return	number of bytes stored in the undo buffer
 	 */
-	virtual SIZE_T GetUndoSize() const PURE_VIRTUAL(UTransactor::GetUndoSize,return 0;);
+	UNREALED_API virtual SIZE_T GetUndoSize() const PURE_VIRTUAL(UTransactor::GetUndoSize,return 0;);
 
 	/**
 	 * Gets the number of transactions that were undone and can be redone.
 	 *
 	 * @return Undo count.
 	 */
-	virtual int32 GetUndoCount( ) const PURE_VIRTUAL(UTransactor::GetUndoCount,return 0;);
+	UNREALED_API virtual int32 GetUndoCount( ) const PURE_VIRTUAL(UTransactor::GetUndoCount,return 0;);
 
 	/**
 	 * Returns the description of the redo action that will be performed next.
@@ -602,28 +602,28 @@ class UNREALED_API UTransactor : public UObject
 	 * 
 	 * @return	text describing the next redo transaction
 	 */
-	virtual FTransactionContext GetRedoContext() PURE_VIRTUAL(UTransactor::GetRedoContext,return FTransactionContext(););
+	UNREALED_API virtual FTransactionContext GetRedoContext() PURE_VIRTUAL(UTransactor::GetRedoContext,return FTransactionContext(););
 
 	/**
 	 * Sets an undo barrier at the current point in the transaction buffer.
 	 * Undoing beyond this point will not be allowed until the barrier is removed.
 	 */
-	virtual void SetUndoBarrier() PURE_VIRTUAL(UTransactor::SetUndoBarrier, );
+	UNREALED_API virtual void SetUndoBarrier() PURE_VIRTUAL(UTransactor::SetUndoBarrier, );
 
 	/**
 	 * Removes the last set undo barrier from the transaction buffer.
 	 */
-	virtual void RemoveUndoBarrier() PURE_VIRTUAL(UTransactor::RemoveUndoBarrier, );
+	UNREALED_API virtual void RemoveUndoBarrier() PURE_VIRTUAL(UTransactor::RemoveUndoBarrier, );
 
 	/**
 	 * Clears all undo barriers.
 	 */
-	virtual void ClearUndoBarriers() PURE_VIRTUAL(UTransactor::ClearUndoBarriers, );
+	UNREALED_API virtual void ClearUndoBarriers() PURE_VIRTUAL(UTransactor::ClearUndoBarriers, );
 
 	/**
 	 * Gets the current undo barrier's position in transaction queue
 	 */
-	 virtual int32 GetCurrentUndoBarrier() const PURE_VIRTUAL(UTransactor::GetCurrentUndoBarrier(), { return INDEX_NONE; });
+	 UNREALED_API virtual int32 GetCurrentUndoBarrier() const PURE_VIRTUAL(UTransactor::GetCurrentUndoBarrier(), { return INDEX_NONE; });
 
 	/**
 	 * Executes an undo transaction, undoing all actions contained by that transaction.
diff --git a/Engine/Source/Editor/UnrealEd/Classes/Editor/UnrealEdEngine.h b/Engine/Source/Editor/UnrealEd/Classes/Editor/UnrealEdEngine.h
--- a/Engine/Source/Editor/UnrealEd/Classes/Editor/UnrealEdEngine.h
+++ b/Engine/Source/Editor/UnrealEd/Classes/Editor/UnrealEdEngine.h
@@ -83,8 +83,8 @@ struct FClassMoveInfo
 class FPerformanceMonitor;
 
 
-UCLASS(config=Engine, transient)
-class UNREALED_API UUnrealEdEngine : public UEditorEngine, public FNotifyHook
+UCLASS(config=Engine, transient, MinimalAPI)
+class UUnrealEdEngine : public UEditorEngine, public FNotifyHook
 {
 public:
 	GENERATED_BODY()
@@ -160,92 +160,92 @@ public:
 	FComponentVisualizerManager	ComponentVisManager;
 
 	//~ Begin UObject Interface.
-	~UUnrealEdEngine();
-	virtual void FinishDestroy() override;
-	virtual void Serialize( FArchive& Ar ) override;
+	UNREALED_API ~UUnrealEdEngine();
+	UNREALED_API virtual void FinishDestroy() override;
+	UNREALED_API virtual void Serialize( FArchive& Ar ) override;
 	//~ End UObject Interface.
 
 	//~ Begin FNotify Interface.
-	virtual void NotifyPreChange( FProperty* PropertyAboutToChange ) override;
-	virtual void NotifyPostChange( const FPropertyChangedEvent& PropertyChangedEvent, FProperty* PropertyThatChanged ) override;
+	UNREALED_API virtual void NotifyPreChange( FProperty* PropertyAboutToChange ) override;
+	UNREALED_API virtual void NotifyPostChange( const FPropertyChangedEvent& PropertyChangedEvent, FProperty* PropertyThatChanged ) override;
 	//~ End FNotify Interface.
 
 	//~ Begin UEditorEngine Interface
-	virtual void SelectActor(AActor* Actor, bool InSelected, bool bNotify, bool bSelectEvenIfHidden = false, bool bForceRefresh = false) override;
-	virtual bool CanSelectActor(AActor* Actor, bool InSelected, bool bSelectEvenIfHidden=false, bool bWarnIfLevelLocked=false) const override;
-	virtual void SelectGroup(AGroupActor* InGroupActor, bool bForceSelection=false, bool bInSelected=true, bool bNotify=true) override;
-	virtual void SelectComponent(class UActorComponent* Component, bool bInSelected, bool bNotify, bool bSelectEvenIfHidden = false) override;
-	virtual void SelectBSPSurf(UModel* InModel, int32 iSurf, bool bSelected, bool bNoteSelectionChange) override;
-	virtual void SelectNone(bool bNoteSelectionChange, bool bDeselectBSPSurfs, bool WarnAboutManyActors=true) override;
-	virtual void DeselectAllSurfaces() override;
-	virtual void NoteSelectionChange(bool bNotify = true) override;
-	virtual void NoteActorMovement() override;
-	virtual void FinishAllSnaps() override;
-	virtual void Cleanse( bool ClearSelection, bool Redraw, const FText& Reason, bool bResetTrans ) override;
-	virtual bool GetMapBuildCancelled() const override;
-	virtual void SetMapBuildCancelled( bool InCancelled ) override;
-	virtual FVector GetPivotLocation() override;
-	virtual void SetPivot(FVector NewPivot, bool bSnapPivotToGrid, bool bIgnoreAxis, bool bAssignPivot=false) override;
-	virtual void ResetPivot() override;
-	virtual void RedrawLevelEditingViewports(bool bInvalidateHitProxies=true) override;
-	virtual void TakeHighResScreenShots() override;
-	virtual void GetPackageList( TArray<UPackage*>* InPackages, UClass* InClass ) override;
-	virtual bool ShouldAbortActorDeletion() const override final; // @note Final - Override ShouldAbortComponentDeletion or ShouldAbortActorDeletion (with parameters) instead.
-	virtual void CloseEditor() override;
+	UNREALED_API virtual void SelectActor(AActor* Actor, bool InSelected, bool bNotify, bool bSelectEvenIfHidden = false, bool bForceRefresh = false) override;
+	UNREALED_API virtual bool CanSelectActor(AActor* Actor, bool InSelected, bool bSelectEvenIfHidden=false, bool bWarnIfLevelLocked=false) const override;
+	UNREALED_API virtual void SelectGroup(AGroupActor* InGroupActor, bool bForceSelection=false, bool bInSelected=true, bool bNotify=true) override;
+	UNREALED_API virtual void SelectComponent(class UActorComponent* Component, bool bInSelected, bool bNotify, bool bSelectEvenIfHidden = false) override;
+	UNREALED_API virtual void SelectBSPSurf(UModel* InModel, int32 iSurf, bool bSelected, bool bNoteSelectionChange) override;
+	UNREALED_API virtual void SelectNone(bool bNoteSelectionChange, bool bDeselectBSPSurfs, bool WarnAboutManyActors=true) override;
+	UNREALED_API virtual void DeselectAllSurfaces() override;
+	UNREALED_API virtual void NoteSelectionChange(bool bNotify = true) override;
+	UNREALED_API virtual void NoteActorMovement() override;
+	UNREALED_API virtual void FinishAllSnaps() override;
+	UNREALED_API virtual void Cleanse( bool ClearSelection, bool Redraw, const FText& Reason, bool bResetTrans ) override;
+	UNREALED_API virtual bool GetMapBuildCancelled() const override;
+	UNREALED_API virtual void SetMapBuildCancelled( bool InCancelled ) override;
+	UNREALED_API virtual FVector GetPivotLocation() override;
+	UNREALED_API virtual void SetPivot(FVector NewPivot, bool bSnapPivotToGrid, bool bIgnoreAxis, bool bAssignPivot=false) override;
+	UNREALED_API virtual void ResetPivot() override;
+	UNREALED_API virtual void RedrawLevelEditingViewports(bool bInvalidateHitProxies=true) override;
+	UNREALED_API virtual void TakeHighResScreenShots() override;
+	UNREALED_API virtual void GetPackageList( TArray<UPackage*>* InPackages, UClass* InClass ) override;
+	UNREALED_API virtual bool ShouldAbortActorDeletion() const override final; // @note Final - Override ShouldAbortComponentDeletion or ShouldAbortActorDeletion (with parameters) instead.
+	UNREALED_API virtual void CloseEditor() override;
 
 	
-	virtual bool IsAutosaving() const override;
+	UNREALED_API virtual bool IsAutosaving() const override;
 	//~ End UEditorEngine Interface 
 	
 	//~ Begin FExec Interface
 #if UE_ALLOW_EXEC_COMMANDS
-	virtual bool Exec( UWorld* InWorld, const TCHAR* Cmd, FOutputDevice& Ar=*GLog ) override;
+	UNREALED_API virtual bool Exec( UWorld* InWorld, const TCHAR* Cmd, FOutputDevice& Ar=*GLog ) override;
 #endif
 	//~ End FExec Interface
 
 
 	//~ Begin UEngine Interface.
-	virtual void Init(IEngineLoop* InEngineLoop) override;
+	UNREALED_API virtual void Init(IEngineLoop* InEngineLoop) override;
 
-	virtual void PreExit() override;
-	virtual void Tick(float DeltaSeconds, bool bIdleMode) override;
+	UNREALED_API virtual void PreExit() override;
+	UNREALED_API virtual void Tick(float DeltaSeconds, bool bIdleMode) override;
 	//~ End UEngine Interface.
 
 	/** Builds a list of sprite categories for use in menus */
-	static void MakeSortedSpriteInfo(TArray<struct FSpriteCategoryInfo>& OutSortedSpriteInfo);
+	static UNREALED_API void MakeSortedSpriteInfo(TArray<struct FSpriteCategoryInfo>& OutSortedSpriteInfo);
 
 	/** called when a package has has its dirty state updated */
-	void OnPackageDirtyStateUpdated( UPackage* Pkg);
+	UNREALED_API void OnPackageDirtyStateUpdated( UPackage* Pkg);
 	/** called when a package's source control state is updated */
-	void OnSourceControlStateUpdated(const FSourceControlOperationRef& SourceControlOp, ECommandResult::Type ResultType, TArray<TWeakObjectPtr<UPackage>> Packages);
+	UNREALED_API void OnSourceControlStateUpdated(const FSourceControlOperationRef& SourceControlOp, ECommandResult::Type ResultType, TArray<TWeakObjectPtr<UPackage>> Packages);
 	/** called when a package is automatically checked out from source control */
-	void OnPackagesCheckedOut(const FSourceControlOperationRef& SourceControlOp, ECommandResult::Type ResultType, TArray<TWeakObjectPtr<UPackage>> Packages);
+	UNREALED_API void OnPackagesCheckedOut(const FSourceControlOperationRef& SourceControlOp, ECommandResult::Type ResultType, TArray<TWeakObjectPtr<UPackage>> Packages);
 	/** caled by FCoreDelegate::PostGarbageCollect */
-	void OnPostGarbageCollect();
+	UNREALED_API void OnPostGarbageCollect();
 	/** called by color picker change event */
-	void OnColorPickerChanged();
+	UNREALED_API void OnColorPickerChanged();
 	/** called by the viewport client before a windows message is processed */
-	void OnPreWindowsMessage(FViewport* Viewport, uint32 Message);
+	UNREALED_API void OnPreWindowsMessage(FViewport* Viewport, uint32 Message);
 	/** called by the viewport client after a windows message is processed */
-	void OnPostWindowsMessage(FViewport* Viewport, uint32 Message);
+	UNREALED_API void OnPostWindowsMessage(FViewport* Viewport, uint32 Message);
 
 	/** Register a function to draw extra information when a particular component is selected */
-	void RegisterComponentVisualizer(FName ComponentClassName, TSharedPtr<class FComponentVisualizer> Visualizer);
+	UNREALED_API void RegisterComponentVisualizer(FName ComponentClassName, TSharedPtr<class FComponentVisualizer> Visualizer);
 	/** Unregister component visualizer function */
-	void UnregisterComponentVisualizer(FName ComponentClassName);
+	UNREALED_API void UnregisterComponentVisualizer(FName ComponentClassName);
 	/** Find a component visualizer for the given component class name */
-	TSharedPtr<class FComponentVisualizer> FindComponentVisualizer(FName ComponentClassName) const;
+	UNREALED_API TSharedPtr<class FComponentVisualizer> FindComponentVisualizer(FName ComponentClassName) const;
 
 	/** Find a component visualizer for the given component class (checking parent classes too) */
-	TSharedPtr<class FComponentVisualizer> FindComponentVisualizer(UClass* ComponentClass) const;
+	UNREALED_API TSharedPtr<class FComponentVisualizer> FindComponentVisualizer(UClass* ComponentClass) const;
 
 	/** Draw component visualizers for components for selected actors */
-	void DrawComponentVisualizers(const FSceneView* View, FPrimitiveDrawInterface* PDI);
+	UNREALED_API void DrawComponentVisualizers(const FSceneView* View, FPrimitiveDrawInterface* PDI);
 	/** Draw component visualizers HUD elements for components for selected actors */
-	void DrawComponentVisualizersHUD(const FViewport* Viewport, const FSceneView* View, FCanvas* Canvas);
+	UNREALED_API void DrawComponentVisualizersHUD(const FViewport* Viewport, const FSceneView* View, FCanvas* Canvas);
 
 	/** Updates the property windows of selected actors */
-	void UpdateFloatingPropertyWindows(bool bForceRefresh=false, bool bNotifyActorSelectionChanged=true);
+	UNREALED_API void UpdateFloatingPropertyWindows(bool bForceRefresh=false, bool bNotifyActorSelectionChanged=true);
 
 	/**
 	*	Updates the property windows to show the data of the supplied ActorList
@@ -253,29 +253,29 @@ public:
 	*	@param	ActorList	The list of actors to show the properties for
 	*
 	*/
-	void UpdateFloatingPropertyWindowsFromActorList(const TArray<AActor*>& ActorList, bool bForceRefresh=false);
+	UNREALED_API void UpdateFloatingPropertyWindowsFromActorList(const TArray<AActor*>& ActorList, bool bForceRefresh=false);
 
 	/**
 	 * Called whenever the actor selection has changed to invalidate any cached state
 	 */
-	void PostActorSelectionChanged();
+	UNREALED_API void PostActorSelectionChanged();
 
 	/**
 	 * Set whether the pivot has been moved independently or not
 	 */
-	void SetPivotMovedIndependently( bool bMovedIndependently );
+	UNREALED_API void SetPivotMovedIndependently( bool bMovedIndependently );
 
 	/**
 	 * Return whether the pivot has been moved independently or not
 	 */
-	bool IsPivotMovedIndependently() const;
+	UNREALED_API bool IsPivotMovedIndependently() const;
 
 	/**
 	 * Called to reset the editor's pivot (widget) location using the currently selected objects.  Usually
 	 * called when the selection changes.
 	 * @param bOnChange Set to true when we know for a fact the selected object has changed
 	 */
-	void UpdatePivotLocationForSelection( bool bOnChange = false );
+	UNREALED_API void UpdatePivotLocationForSelection( bool bOnChange = false );
 
 
 	/**
@@ -288,13 +288,13 @@ public:
 	 * @param	bNoteSelectionChange	If true, call NoteSelectionChange if the new actor was created successfully.
 	 * @return							The new actor.
 	 */
-	virtual AActor* ReplaceActor( AActor* CurrentActor, UClass* NewActorClass, UObject* Archetype, bool bNoteSelectionChange );
+	UNREALED_API virtual AActor* ReplaceActor( AActor* CurrentActor, UClass* NewActorClass, UObject* Archetype, bool bNoteSelectionChange );
 
 
 	/**
 	 * @return Returns the global instance of the editor options class.
 	 */
-	UUnrealEdOptions* GetUnrealEdOptions();
+	UNREALED_API UUnrealEdOptions* GetUnrealEdOptions();
 
 	/**
 	 * Iterate over all levels of the world and create a list of world infos, then
@@ -307,34 +307,34 @@ public:
 	 * @param	OutDeletableActors		Can be NULL, provides a list of all the actors, from the selection, that are deletable
 	 * @return							true if any of the selection can be deleted
 	 */
-	bool CanDeleteSelectedActors( const UWorld* InWorld, const bool bStopAtFirst, const bool bLogUndeletable, TArray<AActor*>* OutDeletableActors = NULL ) const;
+	UNREALED_API bool CanDeleteSelectedActors( const UWorld* InWorld, const bool bStopAtFirst, const bool bLogUndeletable, TArray<AActor*>* OutDeletableActors = NULL ) const;
 
 	// UnrealEdSrv stuff.
-	bool Exec_Edit( UWorld* InWorld, const TCHAR* Str, FOutputDevice& Ar );
-	bool Exec_Pivot( const TCHAR* Str, FOutputDevice& Ar );
-	bool Exec_Actor( UWorld* InWorld, const TCHAR* Str, FOutputDevice& Ar );
-	bool Exec_Element( UWorld* InWorld, const TCHAR* Str, FOutputDevice& Ar );
-	bool Exec_Mode( const TCHAR* Str, FOutputDevice& Ar );
-	bool Exec_Group( const TCHAR* Str, FOutputDevice& Ar );
+	UNREALED_API bool Exec_Edit( UWorld* InWorld, const TCHAR* Str, FOutputDevice& Ar );
+	UNREALED_API bool Exec_Pivot( const TCHAR* Str, FOutputDevice& Ar );
+	UNREALED_API bool Exec_Actor( UWorld* InWorld, const TCHAR* Str, FOutputDevice& Ar );
+	UNREALED_API bool Exec_Element( UWorld* InWorld, const TCHAR* Str, FOutputDevice& Ar );
+	UNREALED_API bool Exec_Mode( const TCHAR* Str, FOutputDevice& Ar );
+	UNREALED_API bool Exec_Group( const TCHAR* Str, FOutputDevice& Ar );
 
 
 	// Editor actor virtuals from EditorActor.cpp.
 	/**
 	 * Select all actors and BSP models, except those which are hidden.
 	 */
-	virtual void edactSelectAll( UWorld* InWorld );
+	UNREALED_API virtual void edactSelectAll( UWorld* InWorld );
 
 	/**
 	 * Invert the selection of all actors and BSP models.
 	 */
-	virtual void edactSelectInvert( UWorld* InWorld );
+	UNREALED_API virtual void edactSelectInvert( UWorld* InWorld );
 
 	/**
 	 * Select all children actors of the current selection.
 	 *
 	 * @param   bRecurseChildren	true to recurse through all descendants of the children
 	 */
-	virtual void edactSelectAllChildren( bool bRecurseChildren );
+	UNREALED_API virtual void edactSelectAllChildren( bool bRecurseChildren );
 
 	/**
 	 * Select all actors in a particular class.
@@ -342,7 +342,7 @@ public:
 	 * @param	InWorld		World context
 	 * @param	InClass		Class of actor to select
 	 */
-	virtual void edactSelectOfClass( UWorld* InWorld, UClass* Class );
+	UNREALED_API virtual void edactSelectOfClass( UWorld* InWorld, UClass* Class );
 
 	/**
 	 * Select all actors of a particular class and archetype.
@@ -351,52 +351,52 @@ public:
 	 * @param	InClass		Class of actor to select
 	 * @param	InArchetype	Archetype of actor to select
 	 */
-	virtual void edactSelectOfClassAndArchetype( UWorld* InWorld, const TSubclassOf<AActor> InClass, const UObject* InArchetype );
+	UNREALED_API virtual void edactSelectOfClassAndArchetype( UWorld* InWorld, const TSubclassOf<AActor> InClass, const UObject* InArchetype );
 
 	/**
 	 * Select all actors in a particular class and its subclasses.
 	 *
 	 * @param	InWorld		World context
 	 */
-	virtual void edactSelectSubclassOf( UWorld* InWorld, UClass* Class );
+	UNREALED_API virtual void edactSelectSubclassOf( UWorld* InWorld, UClass* Class );
 
 	/**
 	 * Select all actors in a level that are marked for deletion.
 	 *
 	 * @param	InWorld		World context
 	 */
-	virtual void edactSelectDeleted( UWorld* InWorld );
+	UNREALED_API virtual void edactSelectDeleted( UWorld* InWorld );
 
 	/**
 	 * Select all actors that have the same static mesh assigned to them as the selected ones.
 	 *
 	 * @param bAllClasses		If true, also select non-AStaticMeshActor actors whose meshes match.
 	 */
-	virtual void edactSelectMatchingStaticMesh(bool bAllClasses);
+	UNREALED_API virtual void edactSelectMatchingStaticMesh(bool bAllClasses);
 
 	/**
 	 * Select all actors that have the same skeletal mesh assigned to them as the selected ones.
 	 *
 	 * @param bAllClasses		If true, also select non-ASkeletalMeshActor actors whose meshes match.
 	 */
-	virtual void edactSelectMatchingSkeletalMesh(bool bAllClasses);
+	UNREALED_API virtual void edactSelectMatchingSkeletalMesh(bool bAllClasses);
 
 	/**
 	 * Select all material actors that have the same material assigned to them as the selected ones.
 	 */
-	virtual void edactSelectMatchingMaterial();
+	UNREALED_API virtual void edactSelectMatchingMaterial();
 
 	/**
 	 * Select all emitter actors that have the same particle system template assigned to them as the selected ones.
 	 */
-	virtual void edactSelectMatchingEmitter();
+	UNREALED_API virtual void edactSelectMatchingEmitter();
 
 	/**
 	 * Select the relevant lights for all selected actors
 	 *
 	 * @param	InWorld					World context
 	 */
-	virtual void edactSelectRelevantLights( UWorld* InWorld );
+	UNREALED_API virtual void edactSelectRelevantLights( UWorld* InWorld );
 
 	/**
 	 * Can the given component be deleted?
@@ -404,7 +404,7 @@ public:
 	 * @param InComponent				Component to check
 	 * @param OutReason					Optional value to fill with the reason the component cannot be deleted, if any
 	 */
-	virtual bool CanDeleteComponent(const UActorComponent* InComponent, FText* OutReason = nullptr) const;
+	UNREALED_API virtual bool CanDeleteComponent(const UActorComponent* InComponent, FText* OutReason = nullptr) const;
 
 	/**
 	 * Can the given actor be deleted?
@@ -412,7 +412,7 @@ public:
 	 * @param InActor					Actor to check
 	 * @param OutReason					Optional value to fill with the reason the actor cannot be deleted, if any
 	 */
-	virtual bool CanDeleteActor(const AActor* InActor, FText* OutReason = nullptr) const;
+	UNREALED_API virtual bool CanDeleteActor(const AActor* InActor, FText* OutReason = nullptr) const;
 
 	/**
 	 * Should the deletion of the given components be outright aborted?
@@ -420,7 +420,7 @@ public:
 	 * @param InComponentsToDelete		Components to check
 	 * @param OutReason					Optional value to fill with the reason the component deletion was aborted, if any
 	 */
-	virtual bool ShouldAbortComponentDeletion(const TArray<UActorComponent*>& InComponentsToDelete, FText* OutReason = nullptr) const;
+	UNREALED_API virtual bool ShouldAbortComponentDeletion(const TArray<UActorComponent*>& InComponentsToDelete, FText* OutReason = nullptr) const;
 
 	/**
 	 * Should the deletion of the given actors be outright aborted?
@@ -428,7 +428,7 @@ public:
 	 * @param InActorsToDelete			Actors to check
 	 * @param OutReason					Optional value to fill with the reason the actor deletion was aborted, if any
 	 */
-	virtual bool ShouldAbortActorDeletion(const TArray<AActor*>& InActorsToDelete, FText* OutReason = nullptr) const;
+	UNREALED_API virtual bool ShouldAbortActorDeletion(const TArray<AActor*>& InActorsToDelete, FText* OutReason = nullptr) const;
 
 	/**
 	 * Delete the given components.
@@ -440,7 +440,7 @@ public:
 	 * 
 	 * @return								true unless the delete operation was aborted.
 	 */
-	virtual bool DeleteComponents(const TArray<UActorComponent*>& InComponentsToDelete, UTypedElementSelectionSet* InSelectionSet, const bool bVerifyDeletionCanHappen = true);
+	UNREALED_API virtual bool DeleteComponents(const TArray<UActorComponent*>& InComponentsToDelete, UTypedElementSelectionSet* InSelectionSet, const bool bVerifyDeletionCanHappen = true);
 
 	/**
 	 * Deletes the given actors.
@@ -454,7 +454,7 @@ public:
 	 * 
 	 * @return								true unless the delete operation was aborted.
 	 */
-	virtual bool DeleteActors(const TArray<AActor*>& InActorsToDelete, UWorld* InWorld, UTypedElementSelectionSet* InSelectionSet, const bool bVerifyDeletionCanHappen = true, const bool bWarnAboutReferences = true, const bool bWarnAboutSoftReferences = true);
+	UNREALED_API virtual bool DeleteActors(const TArray<AActor*>& InActorsToDelete, UWorld* InWorld, UTypedElementSelectionSet* InSelectionSet, const bool bVerifyDeletionCanHappen = true, const bool bWarnAboutReferences = true, const bool bWarnAboutSoftReferences = true);
 
 	/**
 	 * Deletes all selected actors
@@ -466,7 +466,7 @@ public:
 	 * @param	bWarnAboutSoftReferences	[opt] If true (default), we prompt the user about soft references to actors they are about to delete
 	 * @return								true unless the delete operation was aborted.
 	 */
-	virtual bool edactDeleteSelected( UWorld* InWorld, bool bVerifyDeletionCanHappen=true, bool bWarnAboutReferences = true, bool bWarnAboutSoftReferences = true) override final;
+	UNREALED_API virtual bool edactDeleteSelected( UWorld* InWorld, bool bVerifyDeletionCanHappen=true, bool bWarnAboutReferences = true, bool bWarnAboutSoftReferences = true) override final;
 
 	/**
 	 * Copy selected actors to the clipboard.  Does not copy PrefabInstance actors or parts of Prefabs.
@@ -475,7 +475,7 @@ public:
 	 * @param	InWorld					World context
 	 * @param	DestinationData			If != NULL, fill instead of clipboard data
 	 */
-	virtual void edactCopySelected(UWorld* InWorld, FString* DestinationData = nullptr) override final;
+	UNREALED_API virtual void edactCopySelected(UWorld* InWorld, FString* DestinationData = nullptr) override final;
 
 	/**
 	 * Copy the given components to the clipboard.
@@ -483,7 +483,7 @@ public:
 	 * @param	InComponentsToCopy		Array of components to copy
 	 * @param	DestinationData			If != NULL, fill instead of clipboard data
 	 */
-	virtual void CopyComponents(const TArray<UActorComponent*>& InComponentsToCopy, FString* DestinationData = nullptr) const;
+	UNREALED_API virtual void CopyComponents(const TArray<UActorComponent*>& InComponentsToCopy, FString* DestinationData = nullptr) const;
 
 	/**
 	 * Copy the given actors to the clipboard.  Does not copy PrefabInstance actors or parts of Prefabs.
@@ -492,7 +492,7 @@ public:
 	 * @param	InWorld					World context
 	 * @param	DestinationData			If != NULL, fill instead of clipboard data
 	 */
-	virtual void CopyActors(const TArray<AActor*>& InActorsToCopy, UWorld* InWorld, FString* DestinationData = nullptr) const;
+	UNREALED_API virtual void CopyActors(const TArray<AActor*>& InActorsToCopy, UWorld* InWorld, FString* DestinationData = nullptr) const;
 
 	/**
 	 * Paste selected actors from the clipboard.
@@ -504,7 +504,7 @@ public:
 	 * @param	bWarnIfHidden		If true displays a warning if the destination level is hidden
 	 * @param	SourceData			If != NULL, use instead of clipboard data
 	 */
-	virtual void edactPasteSelected(UWorld* InWorld, bool bDuplicate, bool bOffsetLocations, bool bWarnIfHidden, const FString* SourceData = nullptr) override final;
+	UNREALED_API virtual void edactPasteSelected(UWorld* InWorld, bool bDuplicate, bool bOffsetLocations, bool bWarnIfHidden, const FString* SourceData = nullptr) override final;
 
 	/**
 	 * Paste the components from the clipboard.
@@ -514,7 +514,7 @@ public:
 	 * @param	bWarnIfHidden		If true displays a warning if the destination level is hidden
 	 * @param	SourceData			If != NULL, use instead of clipboard data
 	 */
-	virtual void PasteComponents(TArray<UActorComponent*>& OutPastedComponents, AActor* TargetActor, const bool bWarnIfHidden, const FString* SourceData = nullptr);
+	UNREALED_API virtual void PasteComponents(TArray<UActorComponent*>& OutPastedComponents, AActor* TargetActor, const bool bWarnIfHidden, const FString* SourceData = nullptr);
 
 	/**
 	 * Paste the actors from the clipboard.
@@ -526,7 +526,7 @@ public:
 	 * @param	bWarnIfHidden		If true displays a warning if the destination level is hidden
 	 * @param	SourceData			If != NULL, use instead of clipboard data
 	 */
-	virtual void PasteActors(TArray<AActor*>& OutPastedActors, UWorld* InWorld, const FVector& LocationOffset, bool bDuplicate, bool bWarnIfHidden, const FString* SourceData = nullptr);
+	UNREALED_API virtual void PasteActors(TArray<AActor*>& OutPastedActors, UWorld* InWorld, const FVector& LocationOffset, bool bDuplicate, bool bWarnIfHidden, const FString* SourceData = nullptr);
 
 	/**
 	 * Duplicates selected actors.  Handles the case where you are trying to duplicate PrefabInstance actors.
@@ -535,7 +535,7 @@ public:
 	 * @param	InLevel				Level to place duplicate
 	 * @param	bOffsetLocations	Should the actor locations be offset after they are created?
 	 */
-	virtual void edactDuplicateSelected(ULevel* InLevel, bool bOffsetLocations) override final;
+	UNREALED_API virtual void edactDuplicateSelected(ULevel* InLevel, bool bOffsetLocations) override final;
 
 	/**
 	 * Duplicate the given components.
@@ -543,7 +543,7 @@ public:
 	 * @param	InComponentsToDuplicate		Array of components to duplicate
 	 * @param	OutNewComponents			List of all the components that were duplicated
 	 */
-	virtual void DuplicateComponents(const TArray<UActorComponent*>& InComponentsToDuplicate, TArray<UActorComponent*>& OutNewComponents);
+	UNREALED_API virtual void DuplicateComponents(const TArray<UActorComponent*>& InComponentsToDuplicate, TArray<UActorComponent*>& OutNewComponents);
 
 	/**
 	 * Duplicates the given actors.  Handles the case where you are trying to duplicate PrefabInstance actors.
@@ -553,19 +553,19 @@ public:
 	 * @param	InLevel				Level to place duplicate
 	 * @param	LocationOffset		Offset to apply to actor locations after they're created
 	 */
-	virtual void DuplicateActors(const TArray<AActor*>& InActorsToDuplicate, TArray<AActor*>& OutNewActors, ULevel* InLevel, const FVector& LocationOffset);
+	UNREALED_API virtual void DuplicateActors(const TArray<AActor*>& InActorsToDuplicate, TArray<AActor*>& OutNewActors, ULevel* InLevel, const FVector& LocationOffset);
 
 	/**
 	 * Replace all selected brushes with the default brush.
 	 *
 	 * @param	InWorld					World context
 	 */
-	virtual void edactReplaceSelectedBrush( UWorld* InWorld );
+	UNREALED_API virtual void edactReplaceSelectedBrush( UWorld* InWorld );
 
 	/**
 	 * Replace all selected non-brush actors with the specified class.
 	 */
-	virtual void edactReplaceSelectedNonBrushWithClass(UClass* Class);
+	UNREALED_API virtual void edactReplaceSelectedNonBrushWithClass(UClass* Class);
 
 	/**
 	 * Replace all actors of the specified source class with actors of the destination class.
@@ -574,105 +574,105 @@ public:
 	 * @param	SrcClass	The class of actors to replace.
 	 * @param	DstClass	The class to replace with.
 	 */
-	virtual void edactReplaceClassWithClass(UWorld* InWorld, UClass* SrcClass, UClass* DstClass);
+	UNREALED_API virtual void edactReplaceClassWithClass(UWorld* InWorld, UClass* SrcClass, UClass* DstClass);
 
 	/**
 	* Align the origin with the current grid.
 	*/
-	virtual void edactAlignOrigin();
+	UNREALED_API virtual void edactAlignOrigin();
 
 	/**
 	 * Align all vertices with the current grid.
 	 */
-	virtual void edactAlignVertices();
+	UNREALED_API virtual void edactAlignVertices();
 
 	/**
 	 * Hide selected actors and BSP models by marking their bHiddenEdTemporary flags true. Will not
 	 * modify/dirty actors/BSP.
 	 */
-	virtual void edactHideSelected( UWorld* InWorld );
+	UNREALED_API virtual void edactHideSelected( UWorld* InWorld );
 
 	/**
 	 * Hide unselected actors and BSP models by marking their bHiddenEdTemporary flags true. Will not
 	 * modify/dirty actors/BSP.
 	 */
-	virtual void edactHideUnselected( UWorld* InWorld );
+	UNREALED_API virtual void edactHideUnselected( UWorld* InWorld );
 
 	/**
 	 * Attempt to unhide all actors and BSP models by setting their bHiddenEdTemporary flags to false if they
 	 * are true. Note: Will not unhide actors/BSP hidden by higher priority visibility settings, such as bHiddenEdGroup,
 	 * but also will not modify/dirty actors/BSP.
 	 */
-	virtual void edactUnHideAll( UWorld* InWorld );
+	UNREALED_API virtual void edactUnHideAll( UWorld* InWorld );
 
 	/**
 	 * Mark all selected actors and BSP models to be hidden upon editor startup, by setting their bHiddenEd flag to
 	 * true, if it is not already. This directly modifies/dirties the relevant actors/BSP.
 	 */
-	virtual void edactHideSelectedStartup( UWorld* InWorld );
+	UNREALED_API virtual void edactHideSelectedStartup( UWorld* InWorld );
 
 	/**
 	 * Mark all actors and BSP models to be shown upon editor startup, by setting their bHiddenEd flag to false, if it is
 	 * not already. This directly modifies/dirties the relevant actors/BSP.
 	 */
-	virtual void edactUnHideAllStartup( UWorld* InWorld );
+	UNREALED_API virtual void edactUnHideAllStartup( UWorld* InWorld );
 
 	/**
 	 * Mark all selected actors and BSP models to be shown upon editor startup, by setting their bHiddenEd flag to false, if it
 	 * not already. This directly modifies/dirties the relevant actors/BSP.
 	 */
-	virtual void edactUnHideSelectedStartup( UWorld* InWorld );
+	UNREALED_API virtual void edactUnHideSelectedStartup( UWorld* InWorld );
 
 	/**
 	 * Show selected actors and BSP models by marking their bHiddenEdTemporary flags false. Will not
 	 * modify/dirty actors/BSP.
 	 */
-	virtual void edactUnhideSelected( UWorld* InWorld );
+	UNREALED_API virtual void edactUnhideSelected( UWorld* InWorld );
 
 	/** Will create a map of currently visible BSP surfaces. */
-	virtual void CreateBSPVisibilityMap( UWorld* InWorld, TMap<AActor*, TArray<int32>>& OutBSPMap, bool& bOutAllVisible );
+	UNREALED_API virtual void CreateBSPVisibilityMap( UWorld* InWorld, TMap<AActor*, TArray<int32>>& OutBSPMap, bool& bOutAllVisible );
 
 	/** Go through a map of BSP and make only the requested objects visible. */
-	virtual void MakeBSPMapVisible(const TMap<AActor*, TArray<int32>>& InBSPMap, UWorld* InWorld );
+	UNREALED_API virtual void MakeBSPMapVisible(const TMap<AActor*, TArray<int32>>& InBSPMap, UWorld* InWorld );
 
 	/** Returns the configuration of attachment that would result from calling AttachSelectedActors at this point in time */
-	AActor* GetDesiredAttachmentState(TArray<AActor*>& OutNewChildren);
+	UNREALED_API AActor* GetDesiredAttachmentState(TArray<AActor*>& OutNewChildren);
 
 	/** Uses the current selection state to attach actors together. Last selected Actor becomes the base. */
-	void AttachSelectedActors();
+	UNREALED_API void AttachSelectedActors();
 
 
 	
 	/**
 	 * Can the editor do cook by the book in the editor process space
 	 */
-	virtual bool CanCookByTheBookInEditor(const FString& PlatformName) const override;
+	UNREALED_API virtual bool CanCookByTheBookInEditor(const FString& PlatformName) const override;
 
 	/**
 	 * Can the editor act as a cook on the fly server
 	 */
-	virtual bool CanCookOnTheFlyInEditor(const FString& PlatformName) const override;
+	UNREALED_API virtual bool CanCookOnTheFlyInEditor(const FString& PlatformName) const override;
 
 	/**
 	 * Start cook by the book in the editor process space
 	 */
-	virtual void StartCookByTheBookInEditor( const TArray<ITargetPlatform*> &TargetPlatforms, const TArray<FString> &CookMaps, const TArray<FString> &CookDirectories, const TArray<FString> &CookCultures, const TArray<FString> &IniMapSections ) override;
+	UNREALED_API virtual void StartCookByTheBookInEditor( const TArray<ITargetPlatform*> &TargetPlatforms, const TArray<FString> &CookMaps, const TArray<FString> &CookDirectories, const TArray<FString> &CookCultures, const TArray<FString> &IniMapSections ) override;
 
 	/**
 	 * Checks if the cook by the book is finished
 	 */
-	virtual bool IsCookByTheBookInEditorFinished() const override;
+	UNREALED_API virtual bool IsCookByTheBookInEditorFinished() const override;
 
 
 	/**
 	 * cancels the current cook by the book in editor
 	 */
-	virtual void CancelCookByTheBookInEditor() override;
+	UNREALED_API virtual void CancelCookByTheBookInEditor() override;
 
 
 
 	// Hook replacements.
-	void ShowActorProperties();
+	UNREALED_API void ShowActorProperties();
 
 	/**
 	 * Checks to see if any worlds are dirty (that is, they need to be saved.)
@@ -681,21 +681,21 @@ public:
 	 * 
 	 * @return true if any worlds are dirty
 	 */
-	bool AnyWorldsAreDirty( UWorld* InWorld ) const;
+	UNREALED_API bool AnyWorldsAreDirty( UWorld* InWorld ) const;
 
 	/**
 	 * Checks to see if any content packages are dirty (that is, they need to be saved.)
 	 *
 	 * @return true if any content packages are dirty
 	 */
-	bool AnyContentPackagesAreDirty() const;
+	UNREALED_API bool AnyContentPackagesAreDirty() const;
 
 	// Misc
 	/**
 	 * Attempts to prompt the user with a balloon notification to checkout modified packages from source control.
 	 * Will defer prompting the user if they are interacting with something
 	 */
-	void AttemptModifiedPackageNotification();
+	UNREALED_API void AttemptModifiedPackageNotification();
 
 	/**
 	 * Prompts the user with a modal checkout dialog to checkout packages from source control.
@@ -704,75 +704,75 @@ public:
 	 *
 	 * @param bPromptAll	If true we prompt for all packages in the PackageToNotifyState map.  If false only prompt about ones we have never prompted about before.
 	 */
-	void PromptToCheckoutModifiedPackages( bool bPromptAll = false );
+	UNREALED_API void PromptToCheckoutModifiedPackages( bool bPromptAll = false );
 
 	/**
 	 * Displays a toast notification or warning when a package is dirtied, indicating that it needs checking out (or that it cannot be checked out)
 	 */
-	void ShowPackageNotification();
+	UNREALED_API void ShowPackageNotification();
 
 	/**
 	 * @return Returns the number of dirty packages that require checkout.
 	 */
-	int32 GetNumDirtyPackagesThatNeedCheckout() const;
+	UNREALED_API int32 GetNumDirtyPackagesThatNeedCheckout() const;
 
 	/**
 	 * Checks to see if there are any packages in the PackageToNotifyState map that are not checked out by the user
 	 *
 	 * @return True if packages need to be checked out.
 	 */
-	bool DoDirtyPackagesNeedCheckout() const;
+	UNREALED_API bool DoDirtyPackagesNeedCheckout() const;
 
 	/**
 	 * Checks whether the specified map is a template map.
 	 *
 	 * @return true if the map is a template map, false otherwise.
 	 */
-	bool IsTemplateMap( const FString& MapName ) const;
+	UNREALED_API bool IsTemplateMap( const FString& MapName ) const;
 
-	void RebuildTemplateMapData();
+	UNREALED_API void RebuildTemplateMapData();
 
 	/**
 	 * Returns true if the user is currently interacting with a viewport.
 	 */
-	bool IsUserInteracting();
+	UNREALED_API bool IsUserInteracting();
 
-	void SetCurrentClass( UClass* InClass );
+	UNREALED_API void SetCurrentClass( UClass* InClass );
 
 	/**
 	 * @return true if selection of translucent objects in perspective viewports is allowed
 	 */
-	virtual bool AllowSelectTranslucent() const override;
+	UNREALED_API virtual bool AllowSelectTranslucent() const override;
 
 	/**
 	 * @return true if only editor-visible levels should be loaded in Play-In-Editor sessions
 	 */
-	virtual bool OnlyLoadEditorVisibleLevelsInPIE() const override;
+	UNREALED_API virtual bool OnlyLoadEditorVisibleLevelsInPIE() const override;
 
 	/**
 	 * @return true if level streaming should prefer to stream levels from disk instead of duplicating them from editor world
 	 */
-	virtual bool PreferToStreamLevelsInPIE() const override;
+	UNREALED_API virtual bool PreferToStreamLevelsInPIE() const override;
 
 	/**
 	 * Duplicate the currently selected actors.
 	 *
 	 * This is a high level routine which may ultimately call edactDuplicateSelected
 	 */
-	void DuplicateSelectedActors(UWorld* InWorld);
+	UNREALED_API void DuplicateSelectedActors(UWorld* InWorld);
 
 	/**
 	 * If all selected actors belong to the same level, that level is made the current level.
 	 */
-	void MakeSelectedActorsLevelCurrent();
+	UNREALED_API void MakeSelectedActorsLevelCurrent();
 
 	/** Returns the thumbnail manager and creates it if missing */
-	class UThumbnailManager* GetThumbnailManager();
+	UNREALED_API class UThumbnailManager* GetThumbnailManager();
 
 	/**
 	 * Returns whether saving the specified package is allowed
 	 */
-	virtual bool CanSavePackage( UPackage* PackageToSave );
+	UNREALED_API virtual bool CanSavePackage( UPackage* PackageToSave );
 
 	/**
 	 * Updates the volume actor visibility for all viewports based on the passed in volume class
@@ -780,12 +780,12 @@ public:
 	 * @param InVolumeActorClass	The type of volume actors to update.  If NULL is passed in all volume actor types are updated.
 	 * @param InViewport			The viewport where actor visibility should apply.  Pass NULL for all editor viewports.
 	 */
-	void UpdateVolumeActorVisibility( UClass* InVolumeActorClass = NULL , FLevelEditorViewportClient* InViewport = NULL);
+	UNREALED_API void UpdateVolumeActorVisibility( UClass* InVolumeActorClass = NULL , FLevelEditorViewportClient* InViewport = NULL);
 
 	/**
 	 * Identify any brushes whose sense is inverted and repair them
 	 */
-	void FixAnyInvertedBrushes(UWorld* World);
+	UNREALED_API void FixAnyInvertedBrushes(UWorld* World);
 
 	/**
 	 * Get the index of the provided sprite category
@@ -794,41 +794,41 @@ public:
 	 *
 	 * @return	Index of the provided sprite category, if possible; INDEX_NONE otherwise
 	 */
-	virtual int32 GetSpriteCategoryIndex( const FName& InSpriteCategory ) override;
+	UNREALED_API virtual int32 GetSpriteCategoryIndex( const FName& InSpriteCategory ) override;
 	
 	/**
 	 * Shows the LightingStaticMeshInfoWindow, creating it first if it hasn't been initialized.
 	 */
-	void ShowLightingStaticMeshInfoWindow();
+	UNREALED_API void ShowLightingStaticMeshInfoWindow();
 	
 	/**
 	 * Shows the SceneStatsWindow, creating it first if it hasn't been initialized.
 	 */
-	void OpenSceneStatsWindow();
+	UNREALED_API void OpenSceneStatsWindow();
 
 	/**
 	 * Shows the TextureStatsWindow, creating it first if it hasn't been initialized.
 	 */
-	void OpenTextureStatsWindow();
+	UNREALED_API void OpenTextureStatsWindow();
 
 	/**
 	* Puts all of the AVolume classes into the passed in array and sorts them by class name.
 	*
 	* @param	VolumeClasses		Array to populate with AVolume classes.
 	*/
-	static void GetSortedVolumeClasses( TArray< UClass* >* VolumeClasses );
+	static UNREALED_API void GetSortedVolumeClasses( TArray< UClass* >* VolumeClasses );
 
 	/**
 	 * Checks the destination level visibility and warns the user if they are trying to paste to a hidden level, offering the option to cancel the operation or unhide the level that is hidden
 	 * 
 	 * @param InWorld			World context
 	 */
-	bool WarnIfDestinationLevelIsHidden( UWorld* InWorld ) const;
+	UNREALED_API bool WarnIfDestinationLevelIsHidden( UWorld* InWorld ) const;
 
 	/**
 	 * Generate the package thumbails if they are needed. 
 	 */
-	UPackage* GeneratePackageThumbnailsIfRequired( const TCHAR* Str, FOutputDevice& Ar, TArray<FString>& ThumbNamesToUnload );
+	UNREALED_API UPackage* GeneratePackageThumbnailsIfRequired( const TCHAR* Str, FOutputDevice& Ar, TArray<FString>& ThumbNamesToUnload );
 
 	/** @return The package auto-saver instance used by the editor */
 	IPackageAutoSaver& GetPackageAutoSaver() const
@@ -839,48 +839,48 @@ public:
 	/**
 	 * Exec command handlers
 	 */
-	bool HandleDumpModelGUIDCommand( const TCHAR* Str, FOutputDevice& Ar );
-	bool HandleModalTestCommand( const TCHAR* Str, FOutputDevice& Ar );
-	bool HandleDisallowExportCommand( const TCHAR* Str, FOutputDevice& Ar );
-	bool HandleDumpBPClassesCommand( const TCHAR* Str, FOutputDevice& Ar );
-	bool HandleFindOutdateInstancesCommand( const TCHAR* Str, FOutputDevice& Ar );
-	bool HandleDumpSelectionCommand( const TCHAR* Str, FOutputDevice& Ar );
-	bool HandleBuildLightingCommand( const TCHAR* Str, FOutputDevice& Ar, UWorld* InWorld );
-	bool HandleBuildPathsCommand( const TCHAR* Str, FOutputDevice& Ar, UWorld* InWorld );
-	bool HandleRecreateLandscapeCollisionCommand(const TCHAR* Str, FOutputDevice& Ar, UWorld* InWorld);
-	bool HandleRemoveLandscapeXYOffsetsCommand(const TCHAR* Str, FOutputDevice& Ar, UWorld* InWorld);
-	bool HandleDisasmScriptCommand( const TCHAR* Str, FOutputDevice& Ar );	
+	UNREALED_API bool HandleDumpModelGUIDCommand( const TCHAR* Str, FOutputDevice& Ar );
+	UNREALED_API bool HandleModalTestCommand( const TCHAR* Str, FOutputDevice& Ar );
+	UNREALED_API bool HandleDisallowExportCommand( const TCHAR* Str, FOutputDevice& Ar );
+	UNREALED_API bool HandleDumpBPClassesCommand( const TCHAR* Str, FOutputDevice& Ar );
+	UNREALED_API bool HandleFindOutdateInstancesCommand( const TCHAR* Str, FOutputDevice& Ar );
+	UNREALED_API bool HandleDumpSelectionCommand( const TCHAR* Str, FOutputDevice& Ar );
+	UNREALED_API bool HandleBuildLightingCommand( const TCHAR* Str, FOutputDevice& Ar, UWorld* InWorld );
+	UNREALED_API bool HandleBuildPathsCommand( const TCHAR* Str, FOutputDevice& Ar, UWorld* InWorld );
+	UNREALED_API bool HandleRecreateLandscapeCollisionCommand(const TCHAR* Str, FOutputDevice& Ar, UWorld* InWorld);
+	UNREALED_API bool HandleRemoveLandscapeXYOffsetsCommand(const TCHAR* Str, FOutputDevice& Ar, UWorld* InWorld);
+	UNREALED_API bool HandleDisasmScriptCommand( const TCHAR* Str, FOutputDevice& Ar );	
 
-	bool IsComponentSelected(const UPrimitiveComponent* PrimComponent);
+	UNREALED_API bool IsComponentSelected(const UPrimitiveComponent* PrimComponent);
 
 	/** Return if we have write permission under the mount point this packages lives in. */
-	bool HasMountWritePermissionForPackage(const FString& PackageName);
+	UNREALED_API bool HasMountWritePermissionForPackage(const FString& PackageName);
 	UE_DEPRECATED(5.0, "Use HasMountWritePermissionForPackage instead")
-	bool HasMountWritePersmissionForPackage(const FString& PackageName);
+	UNREALED_API bool HasMountWritePersmissionForPackage(const FString& PackageName);
 
 	/* Delegate to override TemplateMapInfos */
 	DECLARE_DELEGATE_RetVal(const TArray<FTemplateMapInfo>&, FGetTemplateMapInfos);
 	FGetTemplateMapInfos& OnGetTemplateMapInfos() { return GetTemplateMapInfosDelegate; }
 
 	/** Gets the canonical list of map templates that should be visible in new level picker. This function calls OnGetTemplateMapInfos to allow runtime override of the default maps */
-	const TArray<FTemplateMapInfo>& GetTemplateMapInfos() const;
+	UNREALED_API const TArray<FTemplateMapInfo>& GetTemplateMapInfos() const;
 
 	/** Gets the project default map templates without any runtime overrides */
-	const TArray<FTemplateMapInfo>& GetProjectDefaultMapTemplates() const;
+	UNREALED_API const TArray<FTemplateMapInfo>& GetProjectDefaultMapTemplates() const;
 
 protected:
 
 	/** Called when global editor selection changes */
-	void OnEditorSelectionChanged(UObject* SelectionThatChanged);
+	UNREALED_API void OnEditorSelectionChanged(UObject* SelectionThatChanged);
 
 	/** Called when the element selection set pointer set on the global editor selection changes */
-	void OnEditorElementSelectionPtrChanged(USelection* Selection, UTypedElementSelectionSet* OldSelectionSet, UTypedElementSelectionSet* NewSelectionSet);
+	UNREALED_API void OnEditorElementSelectionPtrChanged(USelection* Selection, UTypedElementSelectionSet* OldSelectionSet, UTypedElementSelectionSet* NewSelectionSet);
 
 	/** Called when the element selection set associated with the global editor selection changes */
-	void OnEditorElementSelectionChanged(const UTypedElementSelectionSet* SelectionSet);
+	UNREALED_API void OnEditorElementSelectionChanged(const UTypedElementSelectionSet* SelectionSet);
 
 	/** Called when a HISM tree has finished building */
-	void OnHISMTreeBuilt(UHierarchicalInstancedStaticMeshComponent* Component, bool bWasAsyncBuild);
+	UNREALED_API void OnHISMTreeBuilt(UHierarchicalInstancedStaticMeshComponent* Component, bool bWasAsyncBuild);
 
 	/** The package auto-saver instance used by the editor */
 	TUniquePtr<IPackageAutoSaver> PackageAutoSaver;
@@ -907,13 +907,13 @@ protected:
 
 private:
 	/** Verify if we have write permission under the specified mount point. */
-	bool VerifyMountPointWritePermission(FName MountPoint);
+	UNREALED_API bool VerifyMountPointWritePermission(FName MountPoint);
 
 	/** Delegate when a new mount point is added, used to test writing permission. */
-	void OnContentPathMounted(const FString& AssetPath, const FString& FileSystemPath);
+	UNREALED_API void OnContentPathMounted(const FString& AssetPath, const FString& FileSystemPath);
 
 	/** Delegate when a new mount point is removed, used to test writing permission. */
-	void OnContentPathDismounted(const FString& AssetPath, const FString& FileSystemPath);
+	UNREALED_API void OnContentPathDismounted(const FString& AssetPath, const FString& FileSystemPath);
 
 	/** Map to track which mount point has write permissions. */
 	TMap<FName, bool> MountPointCheckedForWritePermission;
@@ -937,16 +937,16 @@ private:
 	*
 	* @return	Returns the number of dirty packages that require checkout. If bCheckIfAny is true, returns 1 if any packages will require checkout.
 	*/
-	int32 InternalGetNumDirtyPackagesThatNeedCheckout(bool bCheckIfAny) const;
+	UNREALED_API int32 InternalGetNumDirtyPackagesThatNeedCheckout(bool bCheckIfAny) const;
 
 	/**
 	* Internal function to validate free space on drives used by Unreal Engine
 	* The intent is to notify the user of situations where stability maybe impacted.
 	*/
-	void ValidateFreeDiskSpace() const;
+	UNREALED_API void ValidateFreeDiskSpace() const;
 
 	/** Internal function to filter and add visualizers to a specific list */
-	void AddVisualizers(AActor* Actor, TArray<FCachedComponentVisualizer>& Visualizers, TFunctionRef<bool(const TSharedPtr<FComponentVisualizer>&)> Condition);
+	UNREALED_API void AddVisualizers(AActor* Actor, TArray<FCachedComponentVisualizer>& Visualizers, TFunctionRef<bool(const TSharedPtr<FComponentVisualizer>&)> Condition);
 
 	/** Delegate Called after files have been deleted to perform necessary cleanups. */
 	FDelegateHandle SourceControlFilesDeletedHandle;
diff --git a/Engine/Source/Editor/UnrealEd/Classes/Exporters/RenderTargetExporterHDR.h b/Engine/Source/Editor/UnrealEd/Classes/Exporters/RenderTargetExporterHDR.h
--- a/Engine/Source/Editor/UnrealEd/Classes/Exporters/RenderTargetExporterHDR.h
+++ b/Engine/Source/Editor/UnrealEd/Classes/Exporters/RenderTargetExporterHDR.h
@@ -16,27 +16,27 @@
 //	export to HDR is discouraged, use PNG or EXR
 
 
-UCLASS()
-class UNREALED_API URenderTargetExporterPNG : public UExporter
+UCLASS(MinimalAPI)
+class URenderTargetExporterPNG : public UExporter
 {
 	GENERATED_UCLASS_BODY()
 	
 	//~ Begin UExporter Interface
-	virtual bool SupportsObject(UObject* Object) const override;
-	virtual bool ExportBinary( UObject* Object, const TCHAR* Type, FArchive& Ar, FFeedbackContext* Warn, int32 FileIndex = 0, uint32 PortFlags=0 ) override;
+	UNREALED_API virtual bool SupportsObject(UObject* Object) const override;
+	UNREALED_API virtual bool ExportBinary( UObject* Object, const TCHAR* Type, FArchive& Ar, FFeedbackContext* Warn, int32 FileIndex = 0, uint32 PortFlags=0 ) override;
 	//~ End UExporter Interface
 };
 
 
 
-UCLASS()
-class UNREALED_API URenderTargetExporterEXR : public UExporter
+UCLASS(MinimalAPI)
+class URenderTargetExporterEXR : public UExporter
 {
 	GENERATED_UCLASS_BODY()
 	
 	//~ Begin UExporter Interface
-	virtual bool SupportsObject(UObject* Object) const override;
-	virtual bool ExportBinary( UObject* Object, const TCHAR* Type, FArchive& Ar, FFeedbackContext* Warn, int32 FileIndex = 0, uint32 PortFlags=0 ) override;
+	UNREALED_API virtual bool SupportsObject(UObject* Object) const override;
+	UNREALED_API virtual bool ExportBinary( UObject* Object, const TCHAR* Type, FArchive& Ar, FFeedbackContext* Warn, int32 FileIndex = 0, uint32 PortFlags=0 ) override;
 	//~ End UExporter Interface
 };
 
diff --git a/Engine/Source/Editor/UnrealEd/Classes/Exporters/TextureCubeExporterHDR.h b/Engine/Source/Editor/UnrealEd/Classes/Exporters/TextureCubeExporterHDR.h
--- a/Engine/Source/Editor/UnrealEd/Classes/Exporters/TextureCubeExporterHDR.h
+++ b/Engine/Source/Editor/UnrealEd/Classes/Exporters/TextureCubeExporterHDR.h
@@ -11,18 +11,18 @@
 #include "Exporters/Exporter.h"
 #include "TextureCubeExporterHDR.generated.h"
 
-UCLASS()
-class UNREALED_API UTextureCubeExporterHDR : public UExporter
+UCLASS(MinimalAPI)
+class UTextureCubeExporterHDR : public UExporter
 {
 	GENERATED_UCLASS_BODY()
 
 	//~ Begin UExporter Interface
-	virtual bool ExportBinary(UObject* Object, const TCHAR* Type, FArchive& Ar, FFeedbackContext* Warn, int32 FileIndex = 0, uint32 PortFlags = 0) override;
+	UNREALED_API virtual bool ExportBinary(UObject* Object, const TCHAR* Type, FArchive& Ar, FFeedbackContext* Warn, int32 FileIndex = 0, uint32 PortFlags = 0) override;
 	//~ End UExporter Interface
 };
 
-UCLASS()
-class UNREALED_API URenderTargetCubeExporterHDR : public UTextureCubeExporterHDR
+UCLASS(MinimalAPI)
+class URenderTargetCubeExporterHDR : public UTextureCubeExporterHDR
 {
 	GENERATED_UCLASS_BODY()
 };
diff --git a/Engine/Source/Editor/UnrealEd/Classes/Exporters/TextureExporterPNG.h b/Engine/Source/Editor/UnrealEd/Classes/Exporters/TextureExporterPNG.h
--- a/Engine/Source/Editor/UnrealEd/Classes/Exporters/TextureExporterPNG.h
+++ b/Engine/Source/Editor/UnrealEd/Classes/Exporters/TextureExporterPNG.h
@@ -30,15 +30,15 @@ class UVirtualTextureBuilderExporterPNG : public UTextureExporterPNG
 
 // JPEG is not UTextureExporterGeneric
 // JPEG exporter writes stored JPEG bits in TextureSource compressed payloads
-UCLASS()
-class UNREALED_API UTextureExporterJPEG : public UExporter
+UCLASS(MinimalAPI)
+class UTextureExporterJPEG : public UExporter
 {
 	GENERATED_UCLASS_BODY()
 
 
 	//~ Begin UExporter Interface
-	virtual bool SupportsObject(UObject* Object) const override;
-	virtual bool ExportBinary( UObject* Object, const TCHAR* Type, FArchive& Ar, FFeedbackContext* Warn, int32 FileIndex = 0, uint32 PortFlags=0 ) override;
+	UNREALED_API virtual bool SupportsObject(UObject* Object) const override;
+	UNREALED_API virtual bool ExportBinary( UObject* Object, const TCHAR* Type, FArchive& Ar, FFeedbackContext* Warn, int32 FileIndex = 0, uint32 PortFlags=0 ) override;
 	//~ End UExporter Interface
 };
 
diff --git a/Engine/Source/Editor/UnrealEd/Classes/Exporters/TextureExporterTGA.h b/Engine/Source/Editor/UnrealEd/Classes/Exporters/TextureExporterTGA.h
--- a/Engine/Source/Editor/UnrealEd/Classes/Exporters/TextureExporterTGA.h
+++ b/Engine/Source/Editor/UnrealEd/Classes/Exporters/TextureExporterTGA.h
@@ -12,10 +12,10 @@
 #include "Exporters/TextureExporterGeneric.h"
 #include "TextureExporterTGA.generated.h"
 
-UCLASS()
-class UNREALED_API UTextureExporterTGA : public UTextureExporterGeneric
+UCLASS(MinimalAPI)
+class UTextureExporterTGA : public UTextureExporterGeneric
 {
 	GENERATED_UCLASS_BODY()
 	
-	virtual bool SupportsTexture(UTexture* Texture) const override;
+	UNREALED_API virtual bool SupportsTexture(UTexture* Texture) const override;
 };
diff --git a/Engine/Source/Editor/UnrealEd/Classes/Exporters/VectorFieldExporter.h b/Engine/Source/Editor/UnrealEd/Classes/Exporters/VectorFieldExporter.h
--- a/Engine/Source/Editor/UnrealEd/Classes/Exporters/VectorFieldExporter.h
+++ b/Engine/Source/Editor/UnrealEd/Classes/Exporters/VectorFieldExporter.h
@@ -11,15 +11,15 @@
 #include "Exporters/Exporter.h"
 #include "VectorFieldExporter.generated.h"
 
-UCLASS()
-class UNREALED_API UVectorFieldExporter : public UExporter
+UCLASS(MinimalAPI)
+class UVectorFieldExporter : public UExporter
 {
 	GENERATED_BODY()
 
 	UVectorFieldExporter(const FObjectInitializer& ObjectInitializer = FObjectInitializer());
 
 	//~ Begin UExporter Interface
-	virtual bool ExportText(const FExportObjectInnerContext* Context, UObject* Object, const TCHAR* Type, FOutputDevice& Ar, FFeedbackContext* Warn, uint32 PortFlags = 0) override;
+	UNREALED_API virtual bool ExportText(const FExportObjectInnerContext* Context, UObject* Object, const TCHAR* Type, FOutputDevice& Ar, FFeedbackContext* Warn, uint32 PortFlags = 0) override;
 	//~ End UExporter Interface
 };
 
diff --git a/Engine/Source/Editor/UnrealEd/Classes/Factories/AnimSequenceFactory.h b/Engine/Source/Editor/UnrealEd/Classes/Factories/AnimSequenceFactory.h
--- a/Engine/Source/Editor/UnrealEd/Classes/Factories/AnimSequenceFactory.h
+++ b/Engine/Source/Editor/UnrealEd/Classes/Factories/AnimSequenceFactory.h
@@ -14,8 +14,8 @@
 struct FAssetData;
 class SWindow;
 
-UCLASS(HideCategories=Object, BlueprintType)
-class UNREALED_API UAnimSequenceFactory : public UFactory
+UCLASS(HideCategories=Object, BlueprintType, MinimalAPI)
+class UAnimSequenceFactory : public UFactory
 {
 	GENERATED_UCLASS_BODY()
 
@@ -27,13 +27,13 @@ class UNREALED_API UAnimSequenceFactory : public UFactory
 	TObjectPtr<class USkeletalMesh> PreviewSkeletalMesh;
 
 	//~ Begin UFactory Interface
-	virtual bool ConfigureProperties() override;
-	virtual UObject* FactoryCreateNew(UClass* Class,UObject* InParent,FName Name,EObjectFlags Flags,UObject* Context,FFeedbackContext* Warn) override;
+	UNREALED_API virtual bool ConfigureProperties() override;
+	UNREALED_API virtual UObject* FactoryCreateNew(UClass* Class,UObject* InParent,FName Name,EObjectFlags Flags,UObject* Context,FFeedbackContext* Warn) override;
 	virtual bool ShouldShowInNewMenu() const override { return false; } 
 	//~ Begin UFactory Interface	
 
 private:
-	void OnTargetSkeletonSelected(const FAssetData& SelectedAsset);
+	UNREALED_API void OnTargetSkeletonSelected(const FAssetData& SelectedAsset);
 
 private:
 	TSharedPtr<SWindow> PickerWindow;
diff --git a/Engine/Source/Editor/UnrealEd/Classes/Factories/BlueprintFactory.h b/Engine/Source/Editor/UnrealEd/Classes/Factories/BlueprintFactory.h
--- a/Engine/Source/Editor/UnrealEd/Classes/Factories/BlueprintFactory.h
+++ b/Engine/Source/Editor/UnrealEd/Classes/Factories/BlueprintFactory.h
@@ -15,8 +15,8 @@
 
 class FClassViewerInitializationOptions;
 
-UCLASS(hidecategories=Object, collapsecategories)
-class UNREALED_API UBlueprintFactory : public UFactory
+UCLASS(hidecategories=Object, collapsecategories, MinimalAPI)
+class UBlueprintFactory : public UFactory
 {
 	GENERATED_UCLASS_BODY()
 
@@ -37,9 +37,9 @@ class UNREALED_API UBlueprintFactory : public UFactory
 	FOnConfigureProperties OnConfigurePropertiesDelegate;
 
 	// UFactory Interface
-	virtual bool ConfigureProperties() override;
-	virtual UObject* FactoryCreateNew(UClass* Class, UObject* InParent, FName Name, EObjectFlags Flags, UObject* Context, FFeedbackContext* Warn, FName CallingContext) override;
-	virtual UObject* FactoryCreateNew(UClass* Class, UObject* InParent, FName Name, EObjectFlags Flags, UObject* Context, FFeedbackContext* Warn) override;
+	UNREALED_API virtual bool ConfigureProperties() override;
+	UNREALED_API virtual UObject* FactoryCreateNew(UClass* Class, UObject* InParent, FName Name, EObjectFlags Flags, UObject* Context, FFeedbackContext* Warn, FName CallingContext) override;
+	UNREALED_API virtual UObject* FactoryCreateNew(UClass* Class, UObject* InParent, FName Name, EObjectFlags Flags, UObject* Context, FFeedbackContext* Warn) override;
 	// End of UFactory Interface
 };
 
diff --git a/Engine/Source/Editor/UnrealEd/Classes/Factories/CSVImportFactory.h b/Engine/Source/Editor/UnrealEd/Classes/Factories/CSVImportFactory.h
--- a/Engine/Source/Editor/UnrealEd/Classes/Factories/CSVImportFactory.h
+++ b/Engine/Source/Editor/UnrealEd/Classes/Factories/CSVImportFactory.h
@@ -53,36 +53,36 @@ struct FCSVImportSettings
 	FString DataToImport;
 };
 
-UCLASS(hidecategories=Object)
-class UNREALED_API UCSVImportFactory : public UFactory, public IImportSettingsParser
+UCLASS(hidecategories=Object, MinimalAPI)
+class UCSVImportFactory : public UFactory, public IImportSettingsParser
 {
 	GENERATED_UCLASS_BODY()
 
 public:
 	//~ Begin UFactory Interface
-	virtual FText GetDisplayName() const override;
-	virtual UObject* FactoryCreateFile(UClass* InClass, UObject* InParent, FName InName, EObjectFlags Flags,
+	UNREALED_API virtual FText GetDisplayName() const override;
+	UNREALED_API virtual UObject* FactoryCreateFile(UClass* InClass, UObject* InParent, FName InName, EObjectFlags Flags,
 		const FString& Filename, const TCHAR* Parms, FFeedbackContext* Warn, bool& bOutOperationCanceled) override;
-	virtual UObject* FactoryCreateText(UClass* InClass, UObject* InParent, FName InName, EObjectFlags Flags,
+	UNREALED_API virtual UObject* FactoryCreateText(UClass* InClass, UObject* InParent, FName InName, EObjectFlags Flags,
 		UObject* Context, const TCHAR* Type, const TCHAR*& Buffer, const TCHAR* BufferEnd, FFeedbackContext* Warn,
 		bool& bOutOperationCanceled) override;
-	virtual bool DoesSupportClass(UClass * Class) override;
-	virtual bool FactoryCanImport(const FString& Filename) override;
-	virtual	IImportSettingsParser* GetImportSettingsParser() override;
-	virtual void CleanUp() override;
+	UNREALED_API virtual bool DoesSupportClass(UClass * Class) override;
+	UNREALED_API virtual bool FactoryCanImport(const FString& Filename) override;
+	UNREALED_API virtual	IImportSettingsParser* GetImportSettingsParser() override;
+	UNREALED_API virtual void CleanUp() override;
 
 	/* Reimport an object that was created based on a CSV */
-	EReimportResult::Type ReimportCSV(UObject* Obj);
+	UNREALED_API EReimportResult::Type ReimportCSV(UObject* Obj);
 	
 	/**
 	 * IImportSettings interface
 	 */
-	virtual void ParseFromJson(TSharedRef<class FJsonObject> ImportSettingsJson) override;
+	UNREALED_API virtual void ParseFromJson(TSharedRef<class FJsonObject> ImportSettingsJson) override;
 
 protected:
-	virtual TArray<FString> DoImportDataTable(const FCSVImportSettings& ImportSettings, class UDataTable* TargetDataTable);
-	virtual TArray<FString> DoImportCurveTable(const FCSVImportSettings& ImportSettings, class UCurveTable* TargetCurveTable);
-	virtual TArray<FString> DoImportCurve(const FCSVImportSettings& InImportSettings, class UCurveBase* TargetCurve);
+	UNREALED_API virtual TArray<FString> DoImportDataTable(const FCSVImportSettings& ImportSettings, class UDataTable* TargetDataTable);
+	UNREALED_API virtual TArray<FString> DoImportCurveTable(const FCSVImportSettings& ImportSettings, class UCurveTable* TargetCurveTable);
+	UNREALED_API virtual TArray<FString> DoImportCurve(const FCSVImportSettings& InImportSettings, class UCurveBase* TargetCurve);
 
 private:
 	/* Reimport object from the given path*/
diff --git a/Engine/Source/Editor/UnrealEd/Classes/Factories/CompositeCurveTableFactory.h b/Engine/Source/Editor/UnrealEd/Classes/Factories/CompositeCurveTableFactory.h
--- a/Engine/Source/Editor/UnrealEd/Classes/Factories/CompositeCurveTableFactory.h
+++ b/Engine/Source/Editor/UnrealEd/Classes/Factories/CompositeCurveTableFactory.h
@@ -5,16 +5,16 @@
 #include "CSVImportFactory.h"
 #include "CompositeCurveTableFactory.generated.h"
 
-UCLASS(hidecategories = Object)
-class UNREALED_API UCompositeCurveTableFactory : public UFactory
+UCLASS(hidecategories = Object, MinimalAPI)
+class UCompositeCurveTableFactory : public UFactory
 {
 	GENERATED_UCLASS_BODY()
 
 	//~ Begin UFactory Interface
-	virtual UObject* FactoryCreateNew(UClass* Class, UObject* InParent, FName Name, EObjectFlags Flags, UObject* Context, FFeedbackContext* Warn) override;
+	UNREALED_API virtual UObject* FactoryCreateNew(UClass* Class, UObject* InParent, FName Name, EObjectFlags Flags, UObject* Context, FFeedbackContext* Warn) override;
 	//~ End UFactory Interface
 
 protected:
-	virtual UCurveTable* MakeNewCurveTable(UObject* InParent, FName Name, EObjectFlags Flags);
+	UNREALED_API virtual UCurveTable* MakeNewCurveTable(UObject* InParent, FName Name, EObjectFlags Flags);
 };
 
diff --git a/Engine/Source/Editor/UnrealEd/Classes/Factories/CompositeDataTableFactory.h b/Engine/Source/Editor/UnrealEd/Classes/Factories/CompositeDataTableFactory.h
--- a/Engine/Source/Editor/UnrealEd/Classes/Factories/CompositeDataTableFactory.h
+++ b/Engine/Source/Editor/UnrealEd/Classes/Factories/CompositeDataTableFactory.h
@@ -5,12 +5,12 @@
 #include "DataTableFactory.h"
 #include "CompositeDataTableFactory.generated.h"
 
-UCLASS(hidecategories = Object)
-class UNREALED_API UCompositeDataTableFactory : public UDataTableFactory
+UCLASS(hidecategories = Object, MinimalAPI)
+class UCompositeDataTableFactory : public UDataTableFactory
 {
 	GENERATED_UCLASS_BODY()
 
 protected:
-	virtual UDataTable* MakeNewDataTable(UObject* InParent, FName Name, EObjectFlags Flags) override;
+	UNREALED_API virtual UDataTable* MakeNewDataTable(UObject* InParent, FName Name, EObjectFlags Flags) override;
 };
 
diff --git a/Engine/Source/Editor/UnrealEd/Classes/Factories/CurveTableFactory.h b/Engine/Source/Editor/UnrealEd/Classes/Factories/CurveTableFactory.h
--- a/Engine/Source/Editor/UnrealEd/Classes/Factories/CurveTableFactory.h
+++ b/Engine/Source/Editor/UnrealEd/Classes/Factories/CurveTableFactory.h
@@ -12,19 +12,19 @@ class UCurveTable;
 
 /* Creates a CurveTable with Rich Curves */
 
-UCLASS(hidecategories=Object)
-class UNREALED_API UCurveTableFactory : public UFactory
+UCLASS(hidecategories=Object, MinimalAPI)
+class UCurveTableFactory : public UFactory
 {
 	GENERATED_UCLASS_BODY()
 
   public:
 	// UFactory Interface
-	virtual bool ConfigureProperties() override;
-	virtual UObject* FactoryCreateNew(UClass* Class, UObject* InParent, FName Name, EObjectFlags Flags, UObject* Context, FFeedbackContext* Warn) override;
+	UNREALED_API virtual bool ConfigureProperties() override;
+	UNREALED_API virtual UObject* FactoryCreateNew(UClass* Class, UObject* InParent, FName Name, EObjectFlags Flags, UObject* Context, FFeedbackContext* Warn) override;
 
 
 protected:
-	virtual UCurveTable* MakeNewCurveTable(UObject* InParent, FName Name, EObjectFlags Flags);
+	UNREALED_API virtual UCurveTable* MakeNewCurveTable(UObject* InParent, FName Name, EObjectFlags Flags);
 
 	ERichCurveInterpMode InterpMode;
 };
diff --git a/Engine/Source/Editor/UnrealEd/Classes/Factories/DataTableFactory.h b/Engine/Source/Editor/UnrealEd/Classes/Factories/DataTableFactory.h
--- a/Engine/Source/Editor/UnrealEd/Classes/Factories/DataTableFactory.h
+++ b/Engine/Source/Editor/UnrealEd/Classes/Factories/DataTableFactory.h
@@ -9,8 +9,8 @@
 
 class UDataTable;
 
-UCLASS(hidecategories=Object)
-class UNREALED_API UDataTableFactory : public UFactory
+UCLASS(hidecategories=Object, MinimalAPI)
+class UDataTableFactory : public UFactory
 {
 	GENERATED_UCLASS_BODY()
 
@@ -18,10 +18,10 @@ class UNREALED_API UDataTableFactory : public UFactory
 	TObjectPtr<const class UScriptStruct> Struct;
 
 	//~ Begin UFactory Interface
-	virtual bool ConfigureProperties() override;
-	virtual UObject* FactoryCreateNew(UClass* Class, UObject* InParent, FName Name, EObjectFlags Flags, UObject* Context, FFeedbackContext* Warn) override;
+	UNREALED_API virtual bool ConfigureProperties() override;
+	UNREALED_API virtual UObject* FactoryCreateNew(UClass* Class, UObject* InParent, FName Name, EObjectFlags Flags, UObject* Context, FFeedbackContext* Warn) override;
 	//~ Begin UFactory Interface
 
 protected:
-	virtual UDataTable* MakeNewDataTable(UObject* InParent, FName Name, EObjectFlags Flags);
+	UNREALED_API virtual UDataTable* MakeNewDataTable(UObject* InParent, FName Name, EObjectFlags Flags);
 };
diff --git a/Engine/Source/Editor/UnrealEd/Classes/Factories/Factory.h b/Engine/Source/Editor/UnrealEd/Classes/Factories/Factory.h
--- a/Engine/Source/Editor/UnrealEd/Classes/Factories/Factory.h
+++ b/Engine/Source/Editor/UnrealEd/Classes/Factories/Factory.h
@@ -18,8 +18,8 @@ class FBulkData;
  * An object responsible for creating and importing new objects.
  * 
  */
-UCLASS(abstract)
-class UNREALED_API UFactory : public UObject
+UCLASS(abstract, MinimalAPI)
+class UFactory : public UObject
 {
 	GENERATED_UCLASS_BODY()
 
@@ -38,7 +38,7 @@ public:
 	 *
 	 * @return true if the file is supported, false otherwise.
 	 */
-	virtual bool FactoryCanImport(const FString& Filename);
+	UNREALED_API virtual bool FactoryCanImport(const FString& Filename);
 
 	/**
 	 * Whether the factory is checking for SlowTask::ShouldCancel()
@@ -57,7 +57,7 @@ public:
 	 * @return true if the file is supported, false otherwise.
 	 */
 	UFUNCTION(BlueprintImplementableEvent, Category = "Miscellaneous")
-	bool ScriptFactoryCanImport(const FString& Filename);
+	UNREALED_API bool ScriptFactoryCanImport(const FString& Filename);
 
 	/**
 	 * Create a new object by importing it from a file name.
@@ -76,7 +76,7 @@ public:
 	 * @param bOutOperationCanceled Will indicate whether the user canceled the import.
 	 * @return The new object.
 	 */
-	virtual UObject* FactoryCreateFile(UClass* InClass, UObject* InParent, FName InName, EObjectFlags Flags, const FString& Filename, const TCHAR* Parms, FFeedbackContext* Warn, bool& bOutOperationCanceled);
+	UNREALED_API virtual UObject* FactoryCreateFile(UClass* InClass, UObject* InParent, FName InName, EObjectFlags Flags, const FString& Filename, const TCHAR* Parms, FFeedbackContext* Warn, bool& bOutOperationCanceled);
 
 	/**
 	 * Create a new object by class.
@@ -111,7 +111,7 @@ public:
 		return nullptr;
 	}
 
-	virtual UObject* ImportObject(UClass* InClass, UObject* InOuter, FName InName, EObjectFlags Flags, const FString& Filename, const TCHAR* Parms, bool& OutCanceled);
+	UNREALED_API virtual UObject* ImportObject(UClass* InClass, UObject* InOuter, FName InName, EObjectFlags Flags, const FString& Filename, const TCHAR* Parms, bool& OutCanceled);
 
 	/**
 	 * Returns an array of all the additional objects created during the last imports, as some factories may produce more than one object.
@@ -126,47 +126,47 @@ public:
 	 * @return True if script implements
 	 */
 	UFUNCTION(BlueprintImplementableEvent, Category = "Miscellaneous")
-	bool ScriptFactoryCreateFile(UAssetImportTask* InTask);
+	UNREALED_API bool ScriptFactoryCreateFile(UAssetImportTask* InTask);
 
 	/** Returns true if this factory should be shown in the New Asset menu (by default calls CanCreateNew). */
-	virtual bool ShouldShowInNewMenu() const;
+	UNREALED_API virtual bool ShouldShowInNewMenu() const;
 
 	/** Returns an optional override brush name for the new asset menu. If this is not specified, the thumbnail for the supported class will be used. */
-	virtual FName GetNewAssetThumbnailOverride() const;
+	UNREALED_API virtual FName GetNewAssetThumbnailOverride() const;
 
 	/** Returns the name of the factory for menus */
-	virtual FText GetDisplayName() const;
+	UNREALED_API virtual FText GetDisplayName() const;
 
 	/** When shown in menus, this is the category containing this factory. Return type is a BitFlag mask using EAssetTypeCategories. */
-	virtual uint32 GetMenuCategories() const;
+	UNREALED_API virtual uint32 GetMenuCategories() const;
 
 	/** Branch of sub-menus containing factory under each provided category. */
-	virtual const TArray<FText>& GetMenuCategorySubMenus() const;
+	UNREALED_API virtual const TArray<FText>& GetMenuCategorySubMenus() const;
 
 	/** Returns the tooltip text description of this factory */
-	virtual FText GetToolTip() const;
+	UNREALED_API virtual FText GetToolTip() const;
 
 	/** Returns the documentation page that should be use for the rich tool tip for this factory */
-	virtual FString GetToolTipDocumentationPage() const;
+	UNREALED_API virtual FString GetToolTipDocumentationPage() const;
 
 	/** Returns the documentation excerpt that should be use for the rich tool tip for this factory */
-	virtual FString GetToolTipDocumentationExcerpt() const;
+	UNREALED_API virtual FString GetToolTipDocumentationExcerpt() const;
 
 	/**
 	 * @return		The object class supported by this factory.
 	 */
-	UClass* GetSupportedClass() const;
+	UNREALED_API UClass* GetSupportedClass() const;
 
 	/**
 	 * @return true if it supports this class 
 	 */
-	virtual bool DoesSupportClass(UClass* Class);
+	UNREALED_API virtual bool DoesSupportClass(UClass* Class);
 
 	/**
 	 * Resolves SupportedClass for factories which support multiple classes.
 	 * Such factories will have a nullptr SupportedClass member.
 	 */
-	virtual UClass* ResolveSupportedClass();
+	UNREALED_API virtual UClass* ResolveSupportedClass();
 
 	/** Opens a dialog to configure the factory properties. Return false if user opted out of configuring properties */
 	virtual bool ConfigureProperties()
@@ -175,13 +175,13 @@ public:
 	}
 
 	// @todo document
-	virtual bool ImportUntypedBulkDataFromText(const TCHAR*& Buffer, FBulkData& BulkData);
+	UNREALED_API virtual bool ImportUntypedBulkDataFromText(const TCHAR*& Buffer, FBulkData& BulkData);
 
 	/** Creates a list of file extensions supported by this factory */
-	virtual void GetSupportedFileExtensions(TArray<FString>& OutExtensions) const;
+	UNREALED_API virtual void GetSupportedFileExtensions(TArray<FString>& OutExtensions) const;
 
 	/** Returns true if the provided file extension is supported */
-	virtual bool IsSupportedFileExtension(FStringView InExtension) const;
+	UNREALED_API virtual bool IsSupportedFileExtension(FStringView InExtension) const;
 
 	/** Do clean up after importing is done. Will be called once for multi batch import. */
 	virtual void CleanUp() { AdditionalImportedObjects.Empty(); }
@@ -196,10 +196,10 @@ public:
 	 *	               If nullptr, the class default object is used instead.
 	 * @return A pointer to a new asset of the specified type or null if the creation failed.
 	 */
-	UObject* CreateOrOverwriteAsset(UClass* InClass, UObject* InParent, FName InName, EObjectFlags InFlags, UObject* InTemplate = nullptr) const;
+	UNREALED_API UObject* CreateOrOverwriteAsset(UClass* InClass, UObject* InParent, FName InName, EObjectFlags InFlags, UObject* InTemplate = nullptr) const;
 
 	/** Returns a new starting point name for newly created assets in the content browser */
-	virtual FString GetDefaultNewAssetName() const;
+	UNREALED_API virtual FString GetDefaultNewAssetName() const;
 
 	/** @return the parser that is capable of parsing a json string of import settings for this factory */
 	virtual class IImportSettingsParser* GetImportSettingsParser() {return nullptr;}
@@ -209,19 +209,19 @@ public:
 	 *
 	 * @param Data	The automated import data or nullptr if it doesnt exist
 	 */
-	void SetAutomatedAssetImportData(const class UAutomatedAssetImportData* Data);
+	UNREALED_API void SetAutomatedAssetImportData(const class UAutomatedAssetImportData* Data);
 
 	/**
 	 * Sets the import task being used with this factory
 	 *
 	 * @param Task	The import task or nullptr if it does not exist
 	 */
-	void SetAssetImportTask(class UAssetImportTask* Task);
+	UNREALED_API void SetAssetImportTask(class UAssetImportTask* Task);
 
 	/**
 	 * @return true if this factory is being used for automated import.  Dialogs and user input should be disabled if this method returns true
 	 */
-	virtual bool IsAutomatedImport() const;
+	UNREALED_API virtual bool IsAutomatedImport() const;
 
 	/**
 	 * @return the supported factory formats list.
@@ -234,7 +234,7 @@ public:
 	 *
 	 * @param Message The message text.
 	 **/
-	void DisplayOverwriteOptionsDialog(const FText& Message);
+	UNREALED_API void DisplayOverwriteOptionsDialog(const FText& Message);
 
 	/** Get the name of the file currently being imported. */
 	static FString GetCurrentFilename()
@@ -262,20 +262,20 @@ public:
 	 * The states are used to suppress messages during multiple object import. 
 	 * It needs to be reset each time a new import is started
 	 */
-	void ResetState();
+	UNREALED_API void ResetState();
 
 	/** Helper function to sort an array of factories by their import priority - use as a predicate for Sort */
-	static bool SortFactoriesByPriority(const UFactory& A, const UFactory& B);
+	static UNREALED_API bool SortFactoriesByPriority(const UFactory& A, const UFactory& B);
 
 	// @todo document
-	static UObject* StaticImportObject(UClass* Class, UObject* InOuter, FName Name, EObjectFlags Flags, const TCHAR* Filename = TEXT(""), UObject* Context = nullptr, UFactory* Factory = nullptr, const TCHAR* Parms = nullptr, FFeedbackContext* Warn = GWarn, int32 MaxImportFileSize = 0xC100000);
-	static UObject* StaticImportObject(UClass* Class, UObject* InOuter, FName Name, EObjectFlags Flags, bool& bOutOperationCanceled, const TCHAR* Filename = TEXT(""), UObject* Context = nullptr, UFactory* Factory = nullptr, const TCHAR* Parms = nullptr, FFeedbackContext* Warn = GWarn, int32 MaxImportFileSize = 0xC100000);
+	static UNREALED_API UObject* StaticImportObject(UClass* Class, UObject* InOuter, FName Name, EObjectFlags Flags, const TCHAR* Filename = TEXT(""), UObject* Context = nullptr, UFactory* Factory = nullptr, const TCHAR* Parms = nullptr, FFeedbackContext* Warn = GWarn, int32 MaxImportFileSize = 0xC100000);
+	static UNREALED_API UObject* StaticImportObject(UClass* Class, UObject* InOuter, FName Name, EObjectFlags Flags, bool& bOutOperationCanceled, const TCHAR* Filename = TEXT(""), UObject* Context = nullptr, UFactory* Factory = nullptr, const TCHAR* Parms = nullptr, FFeedbackContext* Warn = GWarn, int32 MaxImportFileSize = 0xC100000);
 
 public:
 
 	//~ UObject interface
 
-	static void AddReferencedObjects(UObject* InThis, FReferenceCollector& Collector);
+	static UNREALED_API void AddReferencedObjects(UObject* InThis, FReferenceCollector& Collector);
 
 protected:
 
@@ -404,14 +404,14 @@ public:
 protected:
 
 	/** Name of the file currently being imported. */
-	static FString CurrentFilename;
+	static UNREALED_API FString CurrentFilename;
 
 	/** This is the import priority that all factories are given in the default constructor. */
-	static const int32 DefaultImportPriority;
+	static UNREALED_API const int32 DefaultImportPriority;
 
 	//@third party code BEGIN SIMPLYGON
 	/** This is the HASH for the file being imported */
-	static FMD5Hash FileHash;
+	static UNREALED_API FMD5Hash FileHash;
 	//@third party code END SIMPLYGON
 
 	/**
diff --git a/Engine/Source/Editor/UnrealEd/Classes/Factories/FbxAnimSequenceImportData.h b/Engine/Source/Editor/UnrealEd/Classes/Factories/FbxAnimSequenceImportData.h
--- a/Engine/Source/Editor/UnrealEd/Classes/Factories/FbxAnimSequenceImportData.h
+++ b/Engine/Source/Editor/UnrealEd/Classes/Factories/FbxAnimSequenceImportData.h
@@ -33,8 +33,8 @@ enum EFBXAnimationLengthImportType : int
 /**
 * Import data and options used when importing any mesh from FBX
 */
-UCLASS(BlueprintType, config = EditorPerProjectUserSettings, configdonotcheckdefaults)
-class UNREALED_API UFbxAnimSequenceImportData : public UFbxAssetImportData
+UCLASS(BlueprintType, config = EditorPerProjectUserSettings, configdonotcheckdefaults, MinimalAPI)
+class UFbxAnimSequenceImportData : public UFbxAssetImportData
 {
 	GENERATED_UCLASS_BODY()
 	
@@ -118,13 +118,13 @@ class UNREALED_API UFbxAnimSequenceImportData : public UFbxAssetImportData
 	bool bPreserveLocalTransform;
 
 	/** Gets or creates fbx import data for the specified anim sequence */
-	static UFbxAnimSequenceImportData* GetImportDataForAnimSequence(UAnimSequence* AnimSequence, UFbxAnimSequenceImportData* TemplateForCreation);
+	static UNREALED_API UFbxAnimSequenceImportData* GetImportDataForAnimSequence(UAnimSequence* AnimSequence, UFbxAnimSequenceImportData* TemplateForCreation);
 
-	virtual bool CanEditChange(const FProperty* InProperty) const override;
+	UNREALED_API virtual bool CanEditChange(const FProperty* InProperty) const override;
 
-	virtual void Serialize(FArchive& Ar) override;
+	UNREALED_API virtual void Serialize(FArchive& Ar) override;
 
-	virtual void PostEditChangeProperty(FPropertyChangedEvent& PropertyChangedEvent) override;
+	UNREALED_API virtual void PostEditChangeProperty(FPropertyChangedEvent& PropertyChangedEvent) override;
 
-	void CopyAnimationValues(const UFbxAnimSequenceImportData* Other);
+	UNREALED_API void CopyAnimationValues(const UFbxAnimSequenceImportData* Other);
 };
diff --git a/Engine/Source/Editor/UnrealEd/Classes/Factories/FbxAssetImportData.h b/Engine/Source/Editor/UnrealEd/Classes/Factories/FbxAssetImportData.h
--- a/Engine/Source/Editor/UnrealEd/Classes/Factories/FbxAssetImportData.h
+++ b/Engine/Source/Editor/UnrealEd/Classes/Factories/FbxAssetImportData.h
@@ -12,8 +12,8 @@ class UFbxSceneImportData;
 /**
  * Base class for import data and options used when importing any asset from FBX
  */
-UCLASS(BlueprintType, config=EditorPerProjectUserSettings, HideCategories=Object, abstract)
-class UNREALED_API UFbxAssetImportData : public UAssetImportData
+UCLASS(BlueprintType, config=EditorPerProjectUserSettings, HideCategories=Object, abstract, MinimalAPI)
+class UFbxAssetImportData : public UAssetImportData
 {
 	GENERATED_UCLASS_BODY()
 
