[UnrealEd]

diff --git a/Engine/Source/Editor/UnrealEd/Classes/Factories/FbxFactory.h b/Engine/Source/Editor/UnrealEd/Classes/Factories/FbxFactory.h
--- a/Engine/Source/Editor/UnrealEd/Classes/Factories/FbxFactory.h
+++ b/Engine/Source/Editor/UnrealEd/Classes/Factories/FbxFactory.h
@@ -13,8 +13,8 @@ namespace UnFbx {
 	class FFbxImporter;
 }
 
-UCLASS(hidecategories=Object)
-class UNREALED_API UFbxFactory : public UFactory
+UCLASS(hidecategories=Object, MinimalAPI)
+class UFbxFactory : public UFactory
 {
 	GENERATED_UCLASS_BODY()
 
@@ -30,22 +30,22 @@ class UNREALED_API UFbxFactory : public UFactory
 
 
 	//~ Begin UObject Interface
-	virtual bool ConfigureProperties() override;
-	virtual void PostInitProperties() override;
+	UNREALED_API virtual bool ConfigureProperties() override;
+	UNREALED_API virtual void PostInitProperties() override;
 	//~ End UObject Interface
 
 	//~ Begin UFactory Interface
-	virtual void CleanUp() override;
-	virtual bool DoesSupportClass(UClass * Class) override;
-	virtual UClass* ResolveSupportedClass() override;
-	virtual UObject* FactoryCreateFile(UClass* InClass, UObject* InParent, FName InName, EObjectFlags Flags, const FString& Filename, const TCHAR* Parms, FFeedbackContext* Warn, bool& bOutOperationCanceled) override;
-	virtual bool FactoryCanImport(const FString& Filename) override;
-	virtual bool CanImportBeCanceled() const override;
-	virtual IImportSettingsParser* GetImportSettingsParser() override;
-	virtual TArray<FString> GetFormats() const override;
+	UNREALED_API virtual void CleanUp() override;
+	UNREALED_API virtual bool DoesSupportClass(UClass * Class) override;
+	UNREALED_API virtual UClass* ResolveSupportedClass() override;
+	UNREALED_API virtual UObject* FactoryCreateFile(UClass* InClass, UObject* InParent, FName InName, EObjectFlags Flags, const FString& Filename, const TCHAR* Parms, FFeedbackContext* Warn, bool& bOutOperationCanceled) override;
+	UNREALED_API virtual bool FactoryCanImport(const FString& Filename) override;
+	UNREALED_API virtual bool CanImportBeCanceled() const override;
+	UNREALED_API virtual IImportSettingsParser* GetImportSettingsParser() override;
+	UNREALED_API virtual TArray<FString> GetFormats() const override;
 	//~ End UFactory Interface
 	
-	static TArray<FString> GetFbxFormats(const UFactory* Factory);
+	static UNREALED_API TArray<FString> GetFbxFormats(const UFactory* Factory);
 
 	/**
 	 * Detect mesh type to import: Static Mesh or Skeletal Mesh.
@@ -54,7 +54,7 @@ class UNREALED_API UFbxFactory : public UFactory
 	 * @param InFilename	FBX file name
 	 * @return bool	return true if parse the file successfully
 	 */
-	bool DetectImportType(const FString& InFilename);
+	UNREALED_API bool DetectImportType(const FString& InFilename);
 
 	void SetDetectImportTypeOnImport(bool bDetectState) { bDetectImportTypeOnImport = bDetectState; }
 
@@ -63,14 +63,14 @@ private:
 	/**
 	 * Delete the newly created objects after an import cancelation.
 	 */
-	void CancelObjectCreation(UnFbx::FFbxImporter* FbxImporter) const;
+	UNREALED_API void CancelObjectCreation(UnFbx::FFbxImporter* FbxImporter) const;
 
 protected:
 	// @todo document
-	UObject* RecursiveImportNode(UnFbx::FFbxImporter* FFbxImporter, void* Node, UObject* InParent, FName InName, EObjectFlags Flags, struct FScopedSlowTask& SlowTask, TArray<UObject*>& OutNewAssets);
+	UNREALED_API UObject* RecursiveImportNode(UnFbx::FFbxImporter* FFbxImporter, void* Node, UObject* InParent, FName InName, EObjectFlags Flags, struct FScopedSlowTask& SlowTask, TArray<UObject*>& OutNewAssets);
 
 	// @todo document
-	UObject* ImportANode(UnFbx::FFbxImporter* FbxImporter, TArray<void*> VoidNodes, UObject* InParent, FName InName, EObjectFlags Flags, struct FScopedSlowTask& SlowTask, UObject* InMesh = NULL, int LODIndex = 0);
+	UNREALED_API UObject* ImportANode(UnFbx::FFbxImporter* FbxImporter, TArray<void*> VoidNodes, UObject* InParent, FName InName, EObjectFlags Flags, struct FScopedSlowTask& SlowTask, UObject* InMesh = NULL, int LODIndex = 0);
 
 	/**
 	 * Used to display special task progression text if needs be (like import being cancelled).
@@ -78,7 +78,7 @@ protected:
 	 * @param TaskText	The current import task text we want to display.
 	 * @return			The task text after taking into account the factory's current state.
 	 */
-	FText GetImportTaskText(const FText& TaskText) const;
+	UNREALED_API FText GetImportTaskText(const FText& TaskText) const;
 
 	bool bShowOption;
 	bool bDetectImportTypeOnImport;
diff --git a/Engine/Source/Editor/UnrealEd/Classes/Factories/FbxSceneImportData.h b/Engine/Source/Editor/UnrealEd/Classes/Factories/FbxSceneImportData.h
--- a/Engine/Source/Editor/UnrealEd/Classes/Factories/FbxSceneImportData.h
+++ b/Engine/Source/Editor/UnrealEd/Classes/Factories/FbxSceneImportData.h
@@ -23,8 +23,8 @@ enum class EFbxSceneReimportStatusFlags : uint8
 ENUM_CLASS_FLAGS(EFbxSceneReimportStatusFlags);
 
 
-UCLASS()
-class UNREALED_API UFbxSceneImportData : public UObject
+UCLASS(MinimalAPI)
+class UFbxSceneImportData : public UObject
 {
 	GENERATED_UCLASS_BODY()
 
@@ -55,16 +55,16 @@ public:
 	ImportOptionsNameMap NameOptionsMap;
 
 	/** Convert this import information to JSON */
-	FString ToJson() const;
+	UNREALED_API FString ToJson() const;
 
 	/** Attempt to parse an asset import structure from the specified json string. */
-	void FromJson(FString InJsonString);
+	UNREALED_API void FromJson(FString InJsonString);
 
-	virtual void GetAssetRegistryTags(TArray<FAssetRegistryTag>& OutTags) const override;
+	UNREALED_API virtual void GetAssetRegistryTags(TArray<FAssetRegistryTag>& OutTags) const override;
 
 protected:
 	/** Overridden serialize function to write out the underlying data as json */
-	virtual void Serialize(FArchive& Ar) override;
+	UNREALED_API virtual void Serialize(FArchive& Ar) override;
 #endif
 };
 
diff --git a/Engine/Source/Editor/UnrealEd/Classes/Factories/FbxSceneImportFactory.h b/Engine/Source/Editor/UnrealEd/Classes/Factories/FbxSceneImportFactory.h
--- a/Engine/Source/Editor/UnrealEd/Classes/Factories/FbxSceneImportFactory.h
+++ b/Engine/Source/Editor/UnrealEd/Classes/Factories/FbxSceneImportFactory.h
@@ -341,17 +341,17 @@ namespace UnFbx
 typedef TMap<FString, UnFbx::FBXImportOptions*> ImportOptionsNameMap;
 typedef ImportOptionsNameMap* ImportOptionsNameMapPtr;
 
-UCLASS(BlueprintType, hidecategories=Object)
-class UNREALED_API UFbxSceneImportFactory : public USceneImportFactory
+UCLASS(BlueprintType, hidecategories=Object, MinimalAPI)
+class UFbxSceneImportFactory : public USceneImportFactory
 {
 	GENERATED_UCLASS_BODY()
 
 	/** UFactory Interface */
-	virtual UObject* FactoryCreateFile(UClass* InClass, UObject* InParent, FName InName, EObjectFlags Flags, const FString& Filename, const TCHAR* Parms, FFeedbackContext* Warn, bool& bOutOperationCanceled) override;
-	virtual UObject* FactoryCreateBinary(UClass* InClass, UObject* InParent, FName InName, EObjectFlags Flags, UObject* Context, const TCHAR* Type, const uint8*& Buffer, const uint8* BufferEnd, FFeedbackContext* Warn) override;
-	virtual UObject* FactoryCreateBinary(UClass* InClass, UObject* InParent, FName InName, EObjectFlags Flags, UObject* Context, const TCHAR* Type, const uint8*& Buffer, const uint8* BufferEnd, FFeedbackContext* Warn, bool& bOutOperationCanceled) override;
-	virtual bool FactoryCanImport(const FString& Filename) override;
-	virtual TArray<FString> GetFormats() const override;
+	UNREALED_API virtual UObject* FactoryCreateFile(UClass* InClass, UObject* InParent, FName InName, EObjectFlags Flags, const FString& Filename, const TCHAR* Parms, FFeedbackContext* Warn, bool& bOutOperationCanceled) override;
+	UNREALED_API virtual UObject* FactoryCreateBinary(UClass* InClass, UObject* InParent, FName InName, EObjectFlags Flags, UObject* Context, const TCHAR* Type, const uint8*& Buffer, const uint8* BufferEnd, FFeedbackContext* Warn) override;
+	UNREALED_API virtual UObject* FactoryCreateBinary(UClass* InClass, UObject* InParent, FName InName, EObjectFlags Flags, UObject* Context, const TCHAR* Type, const uint8*& Buffer, const uint8* BufferEnd, FFeedbackContext* Warn, bool& bOutOperationCanceled) override;
+	UNREALED_API virtual bool FactoryCanImport(const FString& Filename) override;
+	UNREALED_API virtual TArray<FString> GetFormats() const override;
 
 	/** USceneImportFactory Interface */
 	virtual bool ImportsAssets() const override { return true; }
@@ -385,60 +385,60 @@ class UNREALED_API UFbxSceneImportFactory : public USceneImportFactory
 	TObjectPtr<class UFbxTextureImportData> TextureImportData;
 	
 	/* Default Options always have the same name "Default" */
-	static FString DefaultOptionName;
+	static UNREALED_API FString DefaultOptionName;
 
 public:
-	static TSharedPtr<FFbxSceneInfo> ConvertSceneInfo(void* VoidFbxImporter, void* VoidFbxSceneInfo);
-	static void ExtractMaterialInfo(void* FbxImporterVoid, TSharedPtr<FFbxSceneInfo> SceneInfoPtr);
+	static UNREALED_API TSharedPtr<FFbxSceneInfo> ConvertSceneInfo(void* VoidFbxImporter, void* VoidFbxSceneInfo);
+	static UNREALED_API void ExtractMaterialInfo(void* FbxImporterVoid, TSharedPtr<FFbxSceneInfo> SceneInfoPtr);
 
 protected:
 	/** Convert the scene and remake all the transform for the SceneInfo pass in parameter.
 	 *  We need this because EvaluateGlobal and EvaluateLocal are dependent of the scene conversion.
 	 */
-	void ChangeFrontAxis(void* VoidFbxImporter, void* VoidSceneInfo, TSharedPtr<FFbxSceneInfo> SceneInfoPtr);
+	UNREALED_API void ChangeFrontAxis(void* VoidFbxImporter, void* VoidSceneInfo, TSharedPtr<FFbxSceneInfo> SceneInfoPtr);
 
 	/** Make sure GlobalImportSettings is pointing to the correct options */
-	void ApplyMeshInfoFbxOptions(TSharedPtr<FFbxMeshInfo> MeshInfo);
+	UNREALED_API void ApplyMeshInfoFbxOptions(TSharedPtr<FFbxMeshInfo> MeshInfo);
 
 	/* Compute the path of every node and fill the result in the node. This data will be use by the reimport
 	*  as a unique key for for the reimport status of the node hierarchy.
 	*/
-	static void FillSceneHierarchyPath(TSharedPtr<FFbxSceneInfo> SceneInfo);
+	static UNREALED_API void FillSceneHierarchyPath(TSharedPtr<FFbxSceneInfo> SceneInfo);
 
 	/** Create a hierarchy of actor in the current level */
-	void CreateLevelActorHierarchy(TSharedPtr<FFbxSceneInfo> SceneInfoPtr);
+	UNREALED_API void CreateLevelActorHierarchy(TSharedPtr<FFbxSceneInfo> SceneInfoPtr);
 
 	/** Create a hierarchy of actor in the current level */
-	AActor *CreateActorComponentsHierarchy(TSharedPtr<FFbxSceneInfo> SceneInfoPtr);
+	UNREALED_API AActor *CreateActorComponentsHierarchy(TSharedPtr<FFbxSceneInfo> SceneInfoPtr);
 
 	/** Apply the LocalTransform to the SceneComponent and if PreMultiplyTransform is not null do a pre multiplication
 	* SceneComponent: Must be a valid pointer
 	* LocalTransform: Must be a valid pointer
 	* PreMultiplyTransform: Can be nullptr
 	*/
-	void ApplyTransformToComponent(USceneComponent *SceneComponent, FTransform *LocalTransform, FTransform *PreMultiplyTransform, FVector &PivotLocation, FVector &ParentPivotAccumulation);
+	UNREALED_API void ApplyTransformToComponent(USceneComponent *SceneComponent, FTransform *LocalTransform, FTransform *PreMultiplyTransform, FVector &PivotLocation, FVector &ParentPivotAccumulation);
 
 	/** Import all skeletal mesh from the fbx scene */
-	void ImportAllSkeletalMesh(void* VoidRootNodeToImport, void* VoidFbxImporter, EObjectFlags Flags, int32& NodeIndex, int32& InterestingNodeCount , TSharedPtr<FFbxSceneInfo> SceneInfo);
+	UNREALED_API void ImportAllSkeletalMesh(void* VoidRootNodeToImport, void* VoidFbxImporter, EObjectFlags Flags, int32& NodeIndex, int32& InterestingNodeCount , TSharedPtr<FFbxSceneInfo> SceneInfo);
 	
-	UObject* ImportOneSkeletalMesh(void* VoidRootNodeToImport, void* VoidFbxImporter, TSharedPtr<FFbxSceneInfo> SceneInfo, EObjectFlags Flags, TArray<void*> &VoidNodeArray, int32 &TotalNumNodes);
+	UNREALED_API UObject* ImportOneSkeletalMesh(void* VoidRootNodeToImport, void* VoidFbxImporter, TSharedPtr<FFbxSceneInfo> SceneInfo, EObjectFlags Flags, TArray<void*> &VoidNodeArray, int32 &TotalNumNodes);
 
 	/** Import all static mesh from the fbx scene */
-	void ImportAllStaticMesh(void* VoidRootNodeToImport, void* VoidFbxImporter, EObjectFlags Flags, int32& NodeIndex, int32& InterestingNodeCount, TSharedPtr<FFbxSceneInfo> SceneInfo);
+	UNREALED_API void ImportAllStaticMesh(void* VoidRootNodeToImport, void* VoidFbxImporter, EObjectFlags Flags, int32& NodeIndex, int32& InterestingNodeCount, TSharedPtr<FFbxSceneInfo> SceneInfo);
 
 	// @todo document
-	UObject* RecursiveImportNode(void* FFbxImporter, void* VoidNode, EObjectFlags Flags, int32& Index, int32 Total, TSharedPtr<FFbxSceneInfo> SceneInfo, FString PackagePath);
+	UNREALED_API UObject* RecursiveImportNode(void* FFbxImporter, void* VoidNode, EObjectFlags Flags, int32& Index, int32 Total, TSharedPtr<FFbxSceneInfo> SceneInfo, FString PackagePath);
 
 	// @todo document
-	UObject* ImportANode(void* VoidFbxImporter, TArray<void*> &VoidNode, EObjectFlags Flags, int32& NodeIndex, TSharedPtr<FFbxSceneInfo> SceneInfo, TSharedPtr<FFbxNodeInfo> &OutNodeInfo, FString PackagePath, int32 Total = 0, UObject* InMesh = NULL, int LODIndex = 0);
+	UNREALED_API UObject* ImportANode(void* VoidFbxImporter, TArray<void*> &VoidNode, EObjectFlags Flags, int32& NodeIndex, TSharedPtr<FFbxSceneInfo> SceneInfo, TSharedPtr<FFbxNodeInfo> &OutNodeInfo, FString PackagePath, int32 Total = 0, UObject* InMesh = NULL, int LODIndex = 0);
 
 	/** Find the FFbxNodeInfo in the hierarchy. */
-	bool FindSceneNodeInfo(TSharedPtr<FFbxSceneInfo> SceneInfo, uint64 NodeInfoUniqueId, TSharedPtr<FFbxNodeInfo> &OutNodeInfo);
+	UNREALED_API bool FindSceneNodeInfo(TSharedPtr<FFbxSceneInfo> SceneInfo, uint64 NodeInfoUniqueId, TSharedPtr<FFbxNodeInfo> &OutNodeInfo);
 
 	/** Create a package for the specified node. Package will be the concatenation of UFbxSceneImportFactory::Path and Node->GetName(). */
-	UPackage *CreatePackageForNode(FString PackageName, FString &StaticMeshName);
+	UNREALED_API UPackage *CreatePackageForNode(FString PackageName, FString &StaticMeshName);
 
-	bool SetStaticMeshComponentOverrideMaterial(class UStaticMeshComponent* StaticMeshComponent, TSharedPtr<FFbxNodeInfo> NodeInfo);
+	UNREALED_API bool SetStaticMeshComponentOverrideMaterial(class UStaticMeshComponent* StaticMeshComponent, TSharedPtr<FFbxNodeInfo> NodeInfo);
 
 	/** The path of the asset to import */
 	FString Path;
@@ -460,7 +460,7 @@ protected:
 	ImportOptionsNameMap NameOptionsMap;
 
 	/* Return the Options from the NameOptionMap Map. return nulptr if the options are not found*/
-	UnFbx::FBXImportOptions *GetOptionsFromName(FString OptionName);
+	UNREALED_API UnFbx::FBXImportOptions *GetOptionsFromName(FString OptionName);
 
 	/** Is the import was cancel*/
 	bool ImportWasCancel;
diff --git a/Engine/Source/Editor/UnrealEd/Classes/Factories/FbxTextureImportData.h b/Engine/Source/Editor/UnrealEd/Classes/Factories/FbxTextureImportData.h
--- a/Engine/Source/Editor/UnrealEd/Classes/Factories/FbxTextureImportData.h
+++ b/Engine/Source/Editor/UnrealEd/Classes/Factories/FbxTextureImportData.h
@@ -12,8 +12,8 @@
 /**
  * Import data and options used when importing any mesh from FBX
  */
-UCLASS(BlueprintType, AutoExpandCategories=(Texture))
-class UNREALED_API UFbxTextureImportData : public UFbxAssetImportData
+UCLASS(BlueprintType, AutoExpandCategories=(Texture), MinimalAPI)
+class UFbxTextureImportData : public UFbxAssetImportData
 {
 	GENERATED_UCLASS_BODY()
 
@@ -53,5 +53,5 @@ class UNREALED_API UFbxTextureImportData : public UFbxAssetImportData
 	UPROPERTY(BlueprintReadWrite, config, Category = Material, meta = (ImportType = "Mesh"))
 	FString BaseOpacityTextureName;
 
-	bool CanEditChange( const FProperty* InProperty ) const override;
+	UNREALED_API bool CanEditChange( const FProperty* InProperty ) const override;
 };
diff --git a/Engine/Source/Editor/UnrealEd/Classes/Factories/FontFactory.h b/Engine/Source/Editor/UnrealEd/Classes/Factories/FontFactory.h
--- a/Engine/Source/Editor/UnrealEd/Classes/Factories/FontFactory.h
+++ b/Engine/Source/Editor/UnrealEd/Classes/Factories/FontFactory.h
@@ -11,13 +11,13 @@
 #include "Factories/Factory.h"
 #include "FontFactory.generated.h"
 
-UCLASS()
-class UNREALED_API UFontFactory : public UFactory
+UCLASS(MinimalAPI)
+class UFontFactory : public UFactory
 {
 	GENERATED_UCLASS_BODY()
 
 	//~ Begin UFactory Interface
-	virtual UObject* FactoryCreateNew(UClass* InClass, UObject* InParent, FName InName, EObjectFlags Flags, UObject* Context, FFeedbackContext* Warn) override;
+	UNREALED_API virtual UObject* FactoryCreateNew(UClass* InClass, UObject* InParent, FName InName, EObjectFlags Flags, UObject* Context, FFeedbackContext* Warn) override;
 	//~ Begin UFactory Interface	
 };
 
diff --git a/Engine/Source/Editor/UnrealEd/Classes/Factories/FontFileImportFactory.h b/Engine/Source/Editor/UnrealEd/Classes/Factories/FontFileImportFactory.h
--- a/Engine/Source/Editor/UnrealEd/Classes/Factories/FontFileImportFactory.h
+++ b/Engine/Source/Editor/UnrealEd/Classes/Factories/FontFileImportFactory.h
@@ -12,21 +12,21 @@
 #include "EditorReimportHandler.h"
 #include "FontFileImportFactory.generated.h"
 
-UCLASS()
-class UNREALED_API UFontFileImportFactory : public UFactory, public FReimportHandler
+UCLASS(MinimalAPI)
+class UFontFileImportFactory : public UFactory, public FReimportHandler
 {
 	GENERATED_UCLASS_BODY()
 
 	//~ Begin UFactory Interface
-	virtual bool ConfigureProperties() override;
-	virtual UObject* FactoryCreateBinary(UClass* Class, UObject* InParent, FName Name, EObjectFlags Flags, UObject* Context, const TCHAR* Type, const uint8*& Buffer, const uint8* BufferEnd, FFeedbackContext* Warn) override;
+	UNREALED_API virtual bool ConfigureProperties() override;
+	UNREALED_API virtual UObject* FactoryCreateBinary(UClass* Class, UObject* InParent, FName Name, EObjectFlags Flags, UObject* Context, const TCHAR* Type, const uint8*& Buffer, const uint8* BufferEnd, FFeedbackContext* Warn) override;
 	//~ Begin UFactory Interface
 
 	//~ Begin FReimportHandler Interface
-	virtual bool CanReimport(UObject* Obj, TArray<FString>& OutFilenames) override;
-	virtual void SetReimportPaths(UObject* Obj, const TArray<FString>& NewReimportPaths) override;
-	virtual EReimportResult::Type Reimport(UObject* Obj) override;
-	virtual int32 GetPriority() const override;
+	UNREALED_API virtual bool CanReimport(UObject* Obj, TArray<FString>& OutFilenames) override;
+	UNREALED_API virtual void SetReimportPaths(UObject* Obj, const TArray<FString>& NewReimportPaths) override;
+	UNREALED_API virtual EReimportResult::Type Reimport(UObject* Obj) override;
+	UNREALED_API virtual int32 GetPriority() const override;
 	//~ End FReimportHandler Interface
 
 private:
diff --git a/Engine/Source/Editor/UnrealEd/Classes/Factories/ImportSettings.h b/Engine/Source/Editor/UnrealEd/Classes/Factories/ImportSettings.h
--- a/Engine/Source/Editor/UnrealEd/Classes/Factories/ImportSettings.h
+++ b/Engine/Source/Editor/UnrealEd/Classes/Factories/ImportSettings.h
@@ -9,7 +9,7 @@
  * 
  * Supports filling the settings from a json file for automated importing
  */
-class UNREALED_API IImportSettingsParser
+class IImportSettingsParser
 {
 public:
 	/**
diff --git a/Engine/Source/Editor/UnrealEd/Classes/Factories/LightWeightInstanceFactory.h b/Engine/Source/Editor/UnrealEd/Classes/Factories/LightWeightInstanceFactory.h
--- a/Engine/Source/Editor/UnrealEd/Classes/Factories/LightWeightInstanceFactory.h
+++ b/Engine/Source/Editor/UnrealEd/Classes/Factories/LightWeightInstanceFactory.h
@@ -9,18 +9,18 @@
 
 class ALightWeightInstanceManager;
 
-UCLASS(hidecategories=Object)
-class UNREALED_API ULightWeightInstanceFactory : public UFactory
+UCLASS(hidecategories=Object, MinimalAPI)
+class ULightWeightInstanceFactory : public UFactory
 {
 	GENERATED_UCLASS_BODY()
 
 	//~ Begin UFactory Interface
-	virtual bool ConfigureProperties() override;
-	virtual UObject* FactoryCreateNew(UClass* Class, UObject* InParent, FName Name, EObjectFlags Flags, UObject* Context, FFeedbackContext* Warn) override;
+	UNREALED_API virtual bool ConfigureProperties() override;
+	UNREALED_API virtual UObject* FactoryCreateNew(UClass* Class, UObject* InParent, FName Name, EObjectFlags Flags, UObject* Context, FFeedbackContext* Warn) override;
 	//~ Begin UFactory Interface
 
 protected:
-	virtual ALightWeightInstanceManager* MakeNewLightWeightInstanceManager(UObject* InParent, FName Name, EObjectFlags Flags);
+	UNREALED_API virtual ALightWeightInstanceManager* MakeNewLightWeightInstanceManager(UObject* InParent, FName Name, EObjectFlags Flags);
 
 	// The parent class of the created blueprint
 	UPROPERTY()
diff --git a/Engine/Source/Editor/UnrealEd/Classes/Factories/MirrorDataTableFactory.h b/Engine/Source/Editor/UnrealEd/Classes/Factories/MirrorDataTableFactory.h
--- a/Engine/Source/Editor/UnrealEd/Classes/Factories/MirrorDataTableFactory.h
+++ b/Engine/Source/Editor/UnrealEd/Classes/Factories/MirrorDataTableFactory.h
@@ -21,8 +21,8 @@ public:
 	TArray<FMirrorFindReplaceExpression> FindReplaceExpressions;
 };
 
-UCLASS(hidecategories=Object)
-class UNREALED_API UMirrorDataTableFactory : public UFactory
+UCLASS(hidecategories=Object, MinimalAPI)
+class UMirrorDataTableFactory : public UFactory
 {
 	GENERATED_BODY()
 	UMirrorDataTableFactory(const FObjectInitializer& ObjectInitializer);
@@ -38,18 +38,18 @@ class UNREALED_API UMirrorDataTableFactory : public UFactory
 
 public:
 	//~ Begin UFactory Interface
-	virtual bool ConfigureProperties() override;
-	virtual UObject* FactoryCreateNew(UClass* Class, UObject* InParent, FName Name, EObjectFlags Flags, UObject* Context, FFeedbackContext* Warn) override;
+	UNREALED_API virtual bool ConfigureProperties() override;
+	UNREALED_API virtual UObject* FactoryCreateNew(UClass* Class, UObject* InParent, FName Name, EObjectFlags Flags, UObject* Context, FFeedbackContext* Warn) override;
 	//~ Begin UFactory Interface
 
 		/** Returns the name of the factory for menus */
-	virtual FText GetDisplayName() const override;
+	UNREALED_API virtual FText GetDisplayName() const override;
 
 	/** When shown in menus, this is the category containing this factory. Return type is a BitFlag mask using EAssetTypeCategories. */
-	virtual uint32 GetMenuCategories() const override;
+	UNREALED_API virtual uint32 GetMenuCategories() const override;
 protected:
 
-	void OnWindowUserActionDelegate(bool bCreate, USkeleton* InSkeleton, UMirrorTableFindReplaceExpressions* InFindReplaceExpressions, const UScriptStruct* InResultStruct);
-	virtual UMirrorDataTable* MakeNewMirrorDataTable(UObject* InParent, FName Name, EObjectFlags Flags);
+	UNREALED_API void OnWindowUserActionDelegate(bool bCreate, USkeleton* InSkeleton, UMirrorTableFindReplaceExpressions* InFindReplaceExpressions, const UScriptStruct* InResultStruct);
+	UNREALED_API virtual UMirrorDataTable* MakeNewMirrorDataTable(UObject* InParent, FName Name, EObjectFlags Flags);
 };
 
diff --git a/Engine/Source/Editor/UnrealEd/Classes/Factories/ReimportDataTableFactory.h b/Engine/Source/Editor/UnrealEd/Classes/Factories/ReimportDataTableFactory.h
--- a/Engine/Source/Editor/UnrealEd/Classes/Factories/ReimportDataTableFactory.h
+++ b/Engine/Source/Editor/UnrealEd/Classes/Factories/ReimportDataTableFactory.h
@@ -8,17 +8,17 @@
 #include "EditorReimportHandler.h"
 #include "ReimportDataTableFactory.generated.h"
 
-UCLASS()
-class UNREALED_API UReimportDataTableFactory : public UCSVImportFactory, public FReimportHandler
+UCLASS(MinimalAPI)
+class UReimportDataTableFactory : public UCSVImportFactory, public FReimportHandler
 {
 	GENERATED_UCLASS_BODY()
 
 	//~ Begin FReimportHandler Interface
-	virtual bool FactoryCanImport( const FString& Filename ) override;
-	virtual bool CanReimport( UObject* Obj, TArray<FString>& OutFilenames ) override;
-	virtual void SetReimportPaths( UObject* Obj, const TArray<FString>& NewReimportPaths ) override;
-	virtual EReimportResult::Type Reimport( UObject* Obj ) override;
-	virtual int32 GetPriority() const override;
+	UNREALED_API virtual bool FactoryCanImport( const FString& Filename ) override;
+	UNREALED_API virtual bool CanReimport( UObject* Obj, TArray<FString>& OutFilenames ) override;
+	UNREALED_API virtual void SetReimportPaths( UObject* Obj, const TArray<FString>& NewReimportPaths ) override;
+	UNREALED_API virtual EReimportResult::Type Reimport( UObject* Obj ) override;
+	UNREALED_API virtual int32 GetPriority() const override;
 	//~ End FReimportHandler Interface
 };
 
diff --git a/Engine/Source/Editor/UnrealEd/Classes/Factories/StringTableFactory.h b/Engine/Source/Editor/UnrealEd/Classes/Factories/StringTableFactory.h
--- a/Engine/Source/Editor/UnrealEd/Classes/Factories/StringTableFactory.h
+++ b/Engine/Source/Editor/UnrealEd/Classes/Factories/StringTableFactory.h
@@ -7,12 +7,12 @@
 #include "Factories/Factory.h"
 #include "StringTableFactory.generated.h"
 
-UCLASS(hidecategories=Object)
-class UNREALED_API UStringTableFactory : public UFactory
+UCLASS(hidecategories=Object, MinimalAPI)
+class UStringTableFactory : public UFactory
 {
 	GENERATED_UCLASS_BODY()
 
 	//~ Begin UFactory Interface
-	virtual UObject* FactoryCreateNew(UClass* Class, UObject* InParent, FName Name, EObjectFlags Flags, UObject* Context, FFeedbackContext* Warn) override;
+	UNREALED_API virtual UObject* FactoryCreateNew(UClass* Class, UObject* InParent, FName Name, EObjectFlags Flags, UObject* Context, FFeedbackContext* Warn) override;
 	//~ Begin UFactory Interface
 };
diff --git a/Engine/Source/Editor/UnrealEd/Classes/Factories/TextureFactory.h b/Engine/Source/Editor/UnrealEd/Classes/Factories/TextureFactory.h
--- a/Engine/Source/Editor/UnrealEd/Classes/Factories/TextureFactory.h
+++ b/Engine/Source/Editor/UnrealEd/Classes/Factories/TextureFactory.h
@@ -46,8 +46,8 @@ enum class ETextureSourceColorSpace
 	SRGB
 };
 
-UCLASS(customconstructor, collapsecategories, hidecategories=Object)
-class UNREALED_API UTextureFactory : public UFactory, public IImportSettingsParser
+UCLASS(customconstructor, collapsecategories, hidecategories=Object, MinimalAPI)
+class UTextureFactory : public UFactory, public IImportSettingsParser
 {
 	GENERATED_UCLASS_BODY()
 
@@ -159,33 +159,33 @@ class UNREALED_API UTextureFactory : public UFactory, public IImportSettingsPars
 	TEnumAsByte<EAppReturnType::Type> HDRImportShouldBeLongLatCubeMap = EAppReturnType::Retry;
 
 public:
-	UTextureFactory(const FObjectInitializer& ObjectInitializer = FObjectInitializer::Get());
+	UNREALED_API UTextureFactory(const FObjectInitializer& ObjectInitializer = FObjectInitializer::Get());
 
 	//~ Begin UObject Interface
-	virtual void PostInitProperties() override;
+	UNREALED_API virtual void PostInitProperties() override;
 	//~ End UObject Interface
 
 	//~ Begin UFactory Interface
-	virtual bool DoesSupportClass(UClass* Class) override;
-	virtual UObject* FactoryCreateBinary( UClass* Class, UObject* InParent, FName Name, EObjectFlags Flags, UObject* Context, const TCHAR* Type, const uint8*& Buffer, const uint8* BufferEnd, FFeedbackContext* Warn ) override;
-	virtual bool FactoryCanImport(const FString& Filename) override;
-	virtual IImportSettingsParser* GetImportSettingsParser() override;
+	UNREALED_API virtual bool DoesSupportClass(UClass* Class) override;
+	UNREALED_API virtual UObject* FactoryCreateBinary( UClass* Class, UObject* InParent, FName Name, EObjectFlags Flags, UObject* Context, const TCHAR* Type, const uint8*& Buffer, const uint8* BufferEnd, FFeedbackContext* Warn ) override;
+	UNREALED_API virtual bool FactoryCanImport(const FString& Filename) override;
+	UNREALED_API virtual IImportSettingsParser* GetImportSettingsParser() override;
 
 	//~ End UFactory Interface
 	
 	/** IImportSettingsParser interface */
-	virtual void ParseFromJson(TSharedRef<class FJsonObject> ImportSettingsJson) override;
+	UNREALED_API virtual void ParseFromJson(TSharedRef<class FJsonObject> ImportSettingsJson) override;
 
 
 	/** Create a texture given the appropriate input parameters	*/
-	virtual UTexture2D* CreateTexture2D( UObject* InParent, FName Name, EObjectFlags Flags );
-	virtual UTextureCube* CreateTextureCube( UObject* InParent, FName Name, EObjectFlags Flags );
-	virtual UTexture2DArray* CreateTexture2DArray(UObject* InParent, FName Name, EObjectFlags Flags);
+	UNREALED_API virtual UTexture2D* CreateTexture2D( UObject* InParent, FName Name, EObjectFlags Flags );
+	UNREALED_API virtual UTextureCube* CreateTextureCube( UObject* InParent, FName Name, EObjectFlags Flags );
+	UNREALED_API virtual UTexture2DArray* CreateTexture2DArray(UObject* InParent, FName Name, EObjectFlags Flags);
 	/**
 	 * Suppresses the dialog box that, when importing over an existing texture, asks if the users wishes to overwrite its settings.
 	 * This is primarily for reimporting textures.
 	 */
-	static void SuppressImportOverwriteDialog(bool bOverwriteExistingSettings = false);
+	static UNREALED_API void SuppressImportOverwriteDialog(bool bOverwriteExistingSettings = false);
 
 	/**
 	 *	Initializes the given texture from the TextureData text block supplied.
@@ -197,16 +197,16 @@ public:
 	 *
 	 *	@return	bool		true if successful, false if not
 	 */
-	bool InitializeFromT3DTextureDataText(UTexture* InTexture, const FString& Text, FFeedbackContext* Warn);
+	UNREALED_API bool InitializeFromT3DTextureDataText(UTexture* InTexture, const FString& Text, FFeedbackContext* Warn);
 	
 	// @todo document
-	bool InitializeFromT3DTexture2DDataText(UTexture2D* InTexture2D, const TCHAR*& Buffer, FFeedbackContext* Warn);
+	UNREALED_API bool InitializeFromT3DTexture2DDataText(UTexture2D* InTexture2D, const TCHAR*& Buffer, FFeedbackContext* Warn);
 
 	// @todo document
-	void FindCubeMapFace(const FString& ParsedText, const FString& FaceString, UTextureCube& TextureCube, UTexture2D*& TextureFace);
+	UNREALED_API void FindCubeMapFace(const FString& ParsedText, const FString& FaceString, UTextureCube& TextureCube, UTexture2D*& TextureFace);
 
 	// @todo document
-	bool InitializeFromT3DTextureCubeDataText(UTextureCube* InTextureCube, const TCHAR*& Buffer, FFeedbackContext* Warn);
+	UNREALED_API bool InitializeFromT3DTextureCubeDataText(UTextureCube* InTextureCube, const TCHAR*& Buffer, FFeedbackContext* Warn);
 
 protected:
 	/** Keep track of if we are doing a reimport */
diff --git a/Engine/Source/Editor/UnrealEd/Classes/Factories/TouchInterfaceFactory.h b/Engine/Source/Editor/UnrealEd/Classes/Factories/TouchInterfaceFactory.h
--- a/Engine/Source/Editor/UnrealEd/Classes/Factories/TouchInterfaceFactory.h
+++ b/Engine/Source/Editor/UnrealEd/Classes/Factories/TouchInterfaceFactory.h
@@ -7,14 +7,14 @@
 #include "Factories/Factory.h"
 #include "TouchInterfaceFactory.generated.h"
 
-UCLASS(hidecategories=Object)
-class UNREALED_API UTouchInterfaceFactory : public UFactory
+UCLASS(hidecategories=Object, MinimalAPI)
+class UTouchInterfaceFactory : public UFactory
 {
 	GENERATED_UCLASS_BODY()
 
 
 	//~ Begin UFactory Interface
-	virtual UObject* FactoryCreateNew(UClass* Class,UObject* InParent,FName Name,EObjectFlags Flags,UObject* Context,FFeedbackContext* Warn) override;
+	UNREALED_API virtual UObject* FactoryCreateNew(UClass* Class,UObject* InParent,FName Name,EObjectFlags Flags,UObject* Context,FFeedbackContext* Warn) override;
 	//~ Begin UFactory Interface	
 };
 
diff --git a/Engine/Source/Editor/UnrealEd/Classes/Factories/TrueTypeFontFactory.h b/Engine/Source/Editor/UnrealEd/Classes/Factories/TrueTypeFontFactory.h
--- a/Engine/Source/Editor/UnrealEd/Classes/Factories/TrueTypeFontFactory.h
+++ b/Engine/Source/Editor/UnrealEd/Classes/Factories/TrueTypeFontFactory.h
@@ -15,8 +15,8 @@
 class UFont;
 class UTexture2D;
 
-UCLASS(hidecategories=Object, collapsecategories)
-class UNREALED_API UTrueTypeFontFactory : public UTextureFactory, public FReimportHandler
+UCLASS(hidecategories=Object, collapsecategories, MinimalAPI)
+class UTrueTypeFontFactory : public UTextureFactory, public FReimportHandler
 {
 	GENERATED_UCLASS_BODY()
 
@@ -33,12 +33,12 @@ class UNREALED_API UTrueTypeFontFactory : public UTextureFactory, public FReimpo
 	bool bFontSelected;
 
 	//~ Begin UObject Interface
-	virtual void PostInitProperties() override;
+	UNREALED_API virtual void PostInitProperties() override;
 	//~ End UObject Interface
 
 	//~ Begin UFactory Interface
-	virtual bool ConfigureProperties() override;
-	virtual UObject* FactoryCreateNew(UClass* Class,UObject* InParent,FName Name,EObjectFlags Flags,UObject* Context,FFeedbackContext* Warn) override;
+	UNREALED_API virtual bool ConfigureProperties() override;
+	UNREALED_API virtual UObject* FactoryCreateNew(UClass* Class,UObject* InParent,FName Name,EObjectFlags Flags,UObject* Context,FFeedbackContext* Warn) override;
 	virtual bool ShouldShowInNewMenu() const override
 	{
 		// Don't show this factory in the content browser menu; it's invoked manually when changing the UFont cache type to "Offline" 
@@ -47,14 +47,14 @@ class UNREALED_API UTrueTypeFontFactory : public UTextureFactory, public FReimpo
 	//~ Begin UFactory Interface	
 
 	//~ Begin FReimportHandler Interface
-	virtual bool CanReimport( UObject* Obj, TArray<FString>& OutFilenames ) override;
-	virtual void SetReimportPaths( UObject* Obj, const TArray<FString>& NewReimportPaths ) override;
-	virtual EReimportResult::Type Reimport( UObject* Obj ) override;
-	virtual int32 GetPriority() const override;
+	UNREALED_API virtual bool CanReimport( UObject* Obj, TArray<FString>& OutFilenames ) override;
+	UNREALED_API virtual void SetReimportPaths( UObject* Obj, const TArray<FString>& NewReimportPaths ) override;
+	UNREALED_API virtual EReimportResult::Type Reimport( UObject* Obj ) override;
+	UNREALED_API virtual int32 GetPriority() const override;
 	//~ End FReimportHandler Interface
 
 	/** Creates the import options structure for this font */
-	void SetupFontImportOptions();
+	UNREALED_API void SetupFontImportOptions();
 
 #if PLATFORM_WINDOWS
 	/**
@@ -67,17 +67,17 @@ class UNREALED_API UTrueTypeFontFactory : public UTextureFactory, public FReimpo
 	 *
 	 * @return Returns the newly created texture, if successful, otherwise NULL
 	 */
-	UTexture2D* CreateTextureFromDC( UFont* Font, Windows::HDC dc, int32 RowHeight, int32 TextureNum );
+	UNREALED_API UTexture2D* CreateTextureFromDC( UFont* Font, Windows::HDC dc, int32 RowHeight, int32 TextureNum );
 #endif
 
 #if PLATFORM_WINDOWS || PLATFORM_MAC || PLATFORM_LINUX
-	void* LoadFontFace( void* FTLibrary, int32 Height, FFeedbackContext* Warn, void** OutFontData );
-	UTexture2D* CreateTextureFromBitmap( UFont* Font, uint8* BitmapData, int32 Height, int32 TextureNum );
-	bool CreateFontTexture( UFont* Font, FFeedbackContext* Warn, const int32 NumResolutions, const int32 CharsPerPage,
+	UNREALED_API void* LoadFontFace( void* FTLibrary, int32 Height, FFeedbackContext* Warn, void** OutFontData );
+	UNREALED_API UTexture2D* CreateTextureFromBitmap( UFont* Font, uint8* BitmapData, int32 Height, int32 TextureNum );
+	UNREALED_API bool CreateFontTexture( UFont* Font, FFeedbackContext* Warn, const int32 NumResolutions, const int32 CharsPerPage,
 		const TMap<TCHAR,TCHAR>& InverseMap, const TArray< float >& ResHeights );
 
 #if PLATFORM_WINDOWS
-	FString FindBitmapFontFile();
+	UNREALED_API FString FindBitmapFontFile();
 #endif
 
 	/**
@@ -90,7 +90,7 @@ class UNREALED_API UTrueTypeFontFactory : public UTextureFactory, public FReimpo
 	 *
 	 * @return true if successful
 	 */
-	bool ImportTrueTypeFont( UFont* Font, FFeedbackContext* Warn, const int32 NumResolutions, const TArray< float >& ResHeights );
+	UNREALED_API bool ImportTrueTypeFont( UFont* Font, FFeedbackContext* Warn, const int32 NumResolutions, const TArray< float >& ResHeights );
 #endif
 };
 
diff --git a/Engine/Source/Editor/UnrealEd/Classes/MaterialEditor/DEditorDoubleVectorParameterValue.h b/Engine/Source/Editor/UnrealEd/Classes/MaterialEditor/DEditorDoubleVectorParameterValue.h
--- a/Engine/Source/Editor/UnrealEd/Classes/MaterialEditor/DEditorDoubleVectorParameterValue.h
+++ b/Engine/Source/Editor/UnrealEd/Classes/MaterialEditor/DEditorDoubleVectorParameterValue.h
@@ -8,8 +8,8 @@
 #include "Materials/MaterialExpressionVectorParameter.h"
 #include "DEditorDoubleVectorParameterValue.generated.h"
 
-UCLASS(hidecategories = Object, collapsecategories, editinlinenew)
-class UNREALED_API UDEditorDoubleVectorParameterValue : public UDEditorParameterValue
+UCLASS(hidecategories = Object, collapsecategories, editinlinenew, MinimalAPI)
+class UDEditorDoubleVectorParameterValue : public UDEditorParameterValue
 {
 	GENERATED_UCLASS_BODY()
 
diff --git a/Engine/Source/Editor/UnrealEd/Classes/MaterialEditor/DEditorFontParameterValue.h b/Engine/Source/Editor/UnrealEd/Classes/MaterialEditor/DEditorFontParameterValue.h
--- a/Engine/Source/Editor/UnrealEd/Classes/MaterialEditor/DEditorFontParameterValue.h
+++ b/Engine/Source/Editor/UnrealEd/Classes/MaterialEditor/DEditorFontParameterValue.h
@@ -9,7 +9,7 @@
 #include "DEditorFontParameterValue.generated.h"
 
 USTRUCT()
-struct UNREALED_API FDFontParameters
+struct FDFontParameters
 {
 	GENERATED_USTRUCT_BODY()
 
@@ -20,8 +20,8 @@ struct UNREALED_API FDFontParameters
 	int32 FontPage = 0;
 };
 
-UCLASS(hidecategories = Object, collapsecategories)
-class UNREALED_API UDEditorFontParameterValue : public UDEditorParameterValue
+UCLASS(hidecategories = Object, collapsecategories, MinimalAPI)
+class UDEditorFontParameterValue : public UDEditorParameterValue
 {
 	GENERATED_UCLASS_BODY()
 
diff --git a/Engine/Source/Editor/UnrealEd/Classes/MaterialEditor/DEditorMaterialLayersParameterValue.h b/Engine/Source/Editor/UnrealEd/Classes/MaterialEditor/DEditorMaterialLayersParameterValue.h
--- a/Engine/Source/Editor/UnrealEd/Classes/MaterialEditor/DEditorMaterialLayersParameterValue.h
+++ b/Engine/Source/Editor/UnrealEd/Classes/MaterialEditor/DEditorMaterialLayersParameterValue.h
@@ -9,8 +9,8 @@
 #include "DEditorMaterialLayersParameterValue.generated.h"
 
 // FMaterialLayersFunctions are no longer treated as material parameters, so this should maybe be refactored at some point
-UCLASS(hidecategories = Object, collapsecategories)
-class UNREALED_API UDEditorMaterialLayersParameterValue : public UDEditorParameterValue
+UCLASS(hidecategories = Object, collapsecategories, MinimalAPI)
+class UDEditorMaterialLayersParameterValue : public UDEditorParameterValue
 {
 	GENERATED_UCLASS_BODY()
 
diff --git a/Engine/Source/Editor/UnrealEd/Classes/MaterialEditor/DEditorParameterValue.h b/Engine/Source/Editor/UnrealEd/Classes/MaterialEditor/DEditorParameterValue.h
--- a/Engine/Source/Editor/UnrealEd/Classes/MaterialEditor/DEditorParameterValue.h
+++ b/Engine/Source/Editor/UnrealEd/Classes/MaterialEditor/DEditorParameterValue.h
@@ -11,12 +11,12 @@
 
 #include "DEditorParameterValue.generated.h"
 
-UCLASS(hidecategories=Object, collapsecategories, editinlinenew)
-class UNREALED_API UDEditorParameterValue : public UObject
+UCLASS(hidecategories=Object, collapsecategories, editinlinenew, MinimalAPI)
+class UDEditorParameterValue : public UObject
 {
 	GENERATED_UCLASS_BODY()
 
-	static UDEditorParameterValue* Create(UObject* Owner,
+	static UNREALED_API UDEditorParameterValue* Create(UObject* Owner,
 		EMaterialParameterType Type,
 		const FMaterialParameterInfo& ParameterInfo,
 		const FMaterialParameterMetadata& Meta);
diff --git a/Engine/Source/Editor/UnrealEd/Classes/MaterialEditor/DEditorRuntimeVirtualTextureParameterValue.h b/Engine/Source/Editor/UnrealEd/Classes/MaterialEditor/DEditorRuntimeVirtualTextureParameterValue.h
--- a/Engine/Source/Editor/UnrealEd/Classes/MaterialEditor/DEditorRuntimeVirtualTextureParameterValue.h
+++ b/Engine/Source/Editor/UnrealEd/Classes/MaterialEditor/DEditorRuntimeVirtualTextureParameterValue.h
@@ -8,8 +8,8 @@
 
 #include "DEditorRuntimeVirtualTextureParameterValue.generated.h"
 
-UCLASS(hidecategories=Object, collapsecategories)
-class UNREALED_API UDEditorRuntimeVirtualTextureParameterValue : public UDEditorParameterValue
+UCLASS(hidecategories=Object, collapsecategories, MinimalAPI)
+class UDEditorRuntimeVirtualTextureParameterValue : public UDEditorParameterValue
 {
 	GENERATED_UCLASS_BODY()
 
diff --git a/Engine/Source/Editor/UnrealEd/Classes/MaterialEditor/DEditorScalarParameterValue.h b/Engine/Source/Editor/UnrealEd/Classes/MaterialEditor/DEditorScalarParameterValue.h
--- a/Engine/Source/Editor/UnrealEd/Classes/MaterialEditor/DEditorScalarParameterValue.h
+++ b/Engine/Source/Editor/UnrealEd/Classes/MaterialEditor/DEditorScalarParameterValue.h
@@ -23,8 +23,8 @@ public:
 	TSoftObjectPtr<class UCurveLinearColorAtlas> Atlas;
 };
 
-UCLASS(hidecategories=Object, collapsecategories, editinlinenew)
-class UNREALED_API UDEditorScalarParameterValue : public UDEditorParameterValue
+UCLASS(hidecategories=Object, collapsecategories, editinlinenew, MinimalAPI)
+class UDEditorScalarParameterValue : public UDEditorParameterValue
 {
 	GENERATED_UCLASS_BODY()
 
diff --git a/Engine/Source/Editor/UnrealEd/Classes/MaterialEditor/DEditorSparseVolumeTextureParameterValue.h b/Engine/Source/Editor/UnrealEd/Classes/MaterialEditor/DEditorSparseVolumeTextureParameterValue.h
--- a/Engine/Source/Editor/UnrealEd/Classes/MaterialEditor/DEditorSparseVolumeTextureParameterValue.h
+++ b/Engine/Source/Editor/UnrealEd/Classes/MaterialEditor/DEditorSparseVolumeTextureParameterValue.h
@@ -8,8 +8,8 @@
 
 #include "DEditorSparseVolumeTextureParameterValue.generated.h"
 
-UCLASS(hidecategories=Object, collapsecategories)
-class UNREALED_API UDEditorSparseVolumeTextureParameterValue : public UDEditorParameterValue
+UCLASS(hidecategories=Object, collapsecategories, MinimalAPI)
+class UDEditorSparseVolumeTextureParameterValue : public UDEditorParameterValue
 {
 	GENERATED_UCLASS_BODY()
 
diff --git a/Engine/Source/Editor/UnrealEd/Classes/MaterialEditor/DEditorStaticComponentMaskParameterValue.h b/Engine/Source/Editor/UnrealEd/Classes/MaterialEditor/DEditorStaticComponentMaskParameterValue.h
--- a/Engine/Source/Editor/UnrealEd/Classes/MaterialEditor/DEditorStaticComponentMaskParameterValue.h
+++ b/Engine/Source/Editor/UnrealEd/Classes/MaterialEditor/DEditorStaticComponentMaskParameterValue.h
@@ -9,7 +9,7 @@
 #include "DEditorStaticComponentMaskParameterValue.generated.h"
 
 USTRUCT()
-struct UNREALED_API FDComponentMaskParameter
+struct FDComponentMaskParameter
 {
 	GENERATED_USTRUCT_BODY()
 
@@ -45,8 +45,8 @@ struct UNREALED_API FDComponentMaskParameter
 	
 };
 
-UCLASS(hidecategories=Object, collapsecategories)
-class UNREALED_API UDEditorStaticComponentMaskParameterValue : public UDEditorParameterValue
+UCLASS(hidecategories=Object, collapsecategories, MinimalAPI)
+class UDEditorStaticComponentMaskParameterValue : public UDEditorParameterValue
 {
 	GENERATED_UCLASS_BODY()
 
diff --git a/Engine/Source/Editor/UnrealEd/Classes/MaterialEditor/DEditorStaticSwitchParameterValue.h b/Engine/Source/Editor/UnrealEd/Classes/MaterialEditor/DEditorStaticSwitchParameterValue.h
--- a/Engine/Source/Editor/UnrealEd/Classes/MaterialEditor/DEditorStaticSwitchParameterValue.h
+++ b/Engine/Source/Editor/UnrealEd/Classes/MaterialEditor/DEditorStaticSwitchParameterValue.h
@@ -8,8 +8,8 @@
 
 #include "DEditorStaticSwitchParameterValue.generated.h"
 
-UCLASS(hidecategories=Object, collapsecategories)
-class UNREALED_API UDEditorStaticSwitchParameterValue : public UDEditorParameterValue
+UCLASS(hidecategories=Object, collapsecategories, MinimalAPI)
+class UDEditorStaticSwitchParameterValue : public UDEditorParameterValue
 {
 	GENERATED_UCLASS_BODY()
 
diff --git a/Engine/Source/Editor/UnrealEd/Classes/MaterialEditor/DEditorTextureParameterValue.h b/Engine/Source/Editor/UnrealEd/Classes/MaterialEditor/DEditorTextureParameterValue.h
--- a/Engine/Source/Editor/UnrealEd/Classes/MaterialEditor/DEditorTextureParameterValue.h
+++ b/Engine/Source/Editor/UnrealEd/Classes/MaterialEditor/DEditorTextureParameterValue.h
@@ -8,8 +8,8 @@
 
 #include "DEditorTextureParameterValue.generated.h"
 
-UCLASS(hidecategories=Object, collapsecategories)
-class UNREALED_API UDEditorTextureParameterValue : public UDEditorParameterValue
+UCLASS(hidecategories=Object, collapsecategories, MinimalAPI)
+class UDEditorTextureParameterValue : public UDEditorParameterValue
 {
 	GENERATED_UCLASS_BODY()
 
diff --git a/Engine/Source/Editor/UnrealEd/Classes/MaterialEditor/DEditorVectorParameterValue.h b/Engine/Source/Editor/UnrealEd/Classes/MaterialEditor/DEditorVectorParameterValue.h
--- a/Engine/Source/Editor/UnrealEd/Classes/MaterialEditor/DEditorVectorParameterValue.h
+++ b/Engine/Source/Editor/UnrealEd/Classes/MaterialEditor/DEditorVectorParameterValue.h
@@ -8,8 +8,8 @@
 #include "Materials/MaterialExpressionVectorParameter.h"
 #include "DEditorVectorParameterValue.generated.h"
 
-UCLASS(hidecategories=Object, collapsecategories, editinlinenew)
-class UNREALED_API UDEditorVectorParameterValue : public UDEditorParameterValue
+UCLASS(hidecategories=Object, collapsecategories, editinlinenew, MinimalAPI)
+class UDEditorVectorParameterValue : public UDEditorParameterValue
 {
 	GENERATED_UCLASS_BODY()
 
diff --git a/Engine/Source/Editor/UnrealEd/Classes/MaterialEditor/MaterialEditorInstanceConstant.h b/Engine/Source/Editor/UnrealEd/Classes/MaterialEditor/MaterialEditorInstanceConstant.h
--- a/Engine/Source/Editor/UnrealEd/Classes/MaterialEditor/MaterialEditorInstanceConstant.h
+++ b/Engine/Source/Editor/UnrealEd/Classes/MaterialEditor/MaterialEditorInstanceConstant.h
@@ -217,8 +217,8 @@ struct FEditorStaticComponentMaskParameterValue : public FEditorParameterValue
 	}
 };
 
-UCLASS(hidecategories=Object, collapsecategories)
-class UNREALED_API UMaterialEditorInstanceConstant : public UObject
+UCLASS(hidecategories=Object, collapsecategories, MinimalAPI)
+class UMaterialEditorInstanceConstant : public UObject
 {
 	GENERATED_UCLASS_BODY()
 
@@ -283,56 +283,56 @@ class UNREALED_API UMaterialEditorInstanceConstant : public UObject
 	TSoftObjectPtr<UMaterialInterface> NaniteOverrideMaterial;
 
 	//~ Begin UObject Interface.
-	virtual void PostEditChangeProperty(FPropertyChangedEvent& PropertyChangedEvent) override;
+	UNREALED_API virtual void PostEditChangeProperty(FPropertyChangedEvent& PropertyChangedEvent) override;
 #if WITH_EDITOR
-	virtual void PostEditUndo() override;
+	UNREALED_API virtual void PostEditUndo() override;
 #endif
 	//~ End UObject Interface.
 
 	/** Regenerates the parameter arrays. */
-	void RegenerateArrays();
+	UNREALED_API void RegenerateArrays();
 #if WITH_EDITOR
 	/** Sets back to zero the overrides for any parameters copied out of the layer stack */
-	void CleanParameterStack(int32 Index, EMaterialParameterAssociation MaterialType);
+	UNREALED_API void CleanParameterStack(int32 Index, EMaterialParameterAssociation MaterialType);
 	/** Copies the overrides for any parameters copied out of the layer stack from the layer or blend */
-	void ResetOverrides(int32 Index, EMaterialParameterAssociation MaterialType);
+	UNREALED_API void ResetOverrides(int32 Index, EMaterialParameterAssociation MaterialType);
 #endif
 	/** Copies the parameter array values back to the source instance. */
-	void CopyToSourceInstance(const bool bForceStaticPermutationUpdate = false);
+	UNREALED_API void CopyToSourceInstance(const bool bForceStaticPermutationUpdate = false);
 
-	void ApplySourceFunctionChanges();
+	UNREALED_API void ApplySourceFunctionChanges();
 
 	/** 
 	 * Sets the source instance for this object and regenerates arrays. 
 	 *
 	 * @param MaterialInterface		Instance to use as the source for this material editor instance.
 	 */
-	void SetSourceInstance(UMaterialInstanceConstant* MaterialInterface);
+	UNREALED_API void SetSourceInstance(UMaterialInstanceConstant* MaterialInterface);
 
-	void CopyBasePropertiesFromParent();
+	UNREALED_API void CopyBasePropertiesFromParent();
 
-	void SetSourceFunction(UMaterialFunctionInstance* MaterialFunction);
+	UNREALED_API void SetSourceFunction(UMaterialFunctionInstance* MaterialFunction);
 
 	/** 
 	 * Update the source instance parent to match this
 	 */
-	void UpdateSourceInstanceParent();
+	UNREALED_API void UpdateSourceInstanceParent();
 
 	/** 
 	 *  Returns group for parameter. Creates one if needed. 
 	 *
 	 * @param ParameterGroup		Name to be looked for.
 	 */
-	FEditorParameterGroup & GetParameterGroup(FName& ParameterGroup);
+	UNREALED_API FEditorParameterGroup & GetParameterGroup(FName& ParameterGroup);
 	/** 
 	 *  Creates/adds value to group retrieved from parent material . 
 	 *
 	 * @param ParameterValue		Current data to be grouped
 	 * @param GroupName				Name of the group
 	 */
-	void AssignParameterToGroup(UDEditorParameterValue* ParameterValue, const FName& GroupName);
+	UNREALED_API void AssignParameterToGroup(UDEditorParameterValue* ParameterValue, const FName& GroupName);
 
-	static FName GlobalGroupPrefix;
+	static UNREALED_API FName GlobalGroupPrefix;
 
 	TWeakPtr<class IDetailsView> DetailsView;
 
diff --git a/Engine/Source/Editor/UnrealEd/Classes/MaterialEditor/MaterialEditorMeshComponent.h b/Engine/Source/Editor/UnrealEd/Classes/MaterialEditor/MaterialEditorMeshComponent.h
--- a/Engine/Source/Editor/UnrealEd/Classes/MaterialEditor/MaterialEditorMeshComponent.h
+++ b/Engine/Source/Editor/UnrealEd/Classes/MaterialEditor/MaterialEditorMeshComponent.h
@@ -8,12 +8,12 @@
 #include "Components/StaticMeshComponent.h"
 #include "MaterialEditorMeshComponent.generated.h"
 
-UCLASS()
-class UNREALED_API UMaterialEditorMeshComponent : public UStaticMeshComponent
+UCLASS(MinimalAPI)
+class UMaterialEditorMeshComponent : public UStaticMeshComponent
 {
 	GENERATED_UCLASS_BODY()
 
 	// USceneComponent Interface
-	virtual FBoxSphereBounds CalcBounds(const FTransform& LocalToWorld) const override;
+	UNREALED_API virtual FBoxSphereBounds CalcBounds(const FTransform& LocalToWorld) const override;
 };
 
diff --git a/Engine/Source/Editor/UnrealEd/Classes/MaterialEditor/MaterialEditorPreviewParameters.h b/Engine/Source/Editor/UnrealEd/Classes/MaterialEditor/MaterialEditorPreviewParameters.h
--- a/Engine/Source/Editor/UnrealEd/Classes/MaterialEditor/MaterialEditorPreviewParameters.h
+++ b/Engine/Source/Editor/UnrealEd/Classes/MaterialEditor/MaterialEditorPreviewParameters.h
@@ -21,8 +21,8 @@ class UMaterial;
 class UMaterialInstanceConstant;
 struct FPropertyChangedEvent;
 
-UCLASS(hidecategories = Object, collapsecategories)
-class UNREALED_API UMaterialEditorPreviewParameters : public UObject
+UCLASS(hidecategories = Object, collapsecategories, MinimalAPI)
+class UMaterialEditorPreviewParameters : public UObject
 {
 	GENERATED_BODY()
 
@@ -48,20 +48,20 @@ public:
 #endif
 
 	//~ Begin UObject Interface.
-	virtual void PostEditChangeProperty(FPropertyChangedEvent& PropertyChangedEvent) override;
+	UNREALED_API virtual void PostEditChangeProperty(FPropertyChangedEvent& PropertyChangedEvent) override;
 #if WITH_EDITOR
-	virtual void PostEditUndo() override;
+	UNREALED_API virtual void PostEditUndo() override;
 #endif
 	//~ End UObject Interface.
 
 	/** Regenerates the parameter arrays. */
-	void RegenerateArrays();
+	UNREALED_API void RegenerateArrays();
 
 protected:
 	/** Copies the parameter array values back to the source instance. */
-	void CopyToSourceInstance();
+	UNREALED_API void CopyToSourceInstance();
 
-	void ApplySourceFunctionChanges();
+	UNREALED_API void ApplySourceFunctionChanges();
 
 	/**
 	*  Creates/adds value to group retrieved from parent material .
@@ -69,7 +69,7 @@ protected:
 	* @param ParameterValue		Current data to be grouped
 	* @param GroupName			Name of the group
 	*/
-	void AssignParameterToGroup(UDEditorParameterValue* ParameterValue, const FName& GroupName);
+	UNREALED_API void AssignParameterToGroup(UDEditorParameterValue* ParameterValue, const FName& GroupName);
 
-	static FName GlobalGroupPrefix;
+	static UNREALED_API FName GlobalGroupPrefix;
 };
diff --git a/Engine/Source/Editor/UnrealEd/Classes/MaterialEditor/PreviewMaterial.h b/Engine/Source/Editor/UnrealEd/Classes/MaterialEditor/PreviewMaterial.h
--- a/Engine/Source/Editor/UnrealEd/Classes/MaterialEditor/PreviewMaterial.h
+++ b/Engine/Source/Editor/UnrealEd/Classes/MaterialEditor/PreviewMaterial.h
@@ -8,14 +8,14 @@
 
 #include "PreviewMaterial.generated.h"
 
-UCLASS()
-class UNREALED_API UPreviewMaterial : public UMaterial
+UCLASS(MinimalAPI)
+class UPreviewMaterial : public UMaterial
 {
 	GENERATED_UCLASS_BODY()
 
 
 	//~ Begin UMaterial Interface.
-	virtual FMaterialResource* AllocateResource() override;
+	UNREALED_API virtual FMaterialResource* AllocateResource() override;
 	virtual bool IsAsset()  const override  { return false; }
 	//~ End UMaterial Interface.
 };
diff --git a/Engine/Source/Editor/UnrealEd/Classes/MaterialGraph/MaterialGraph.h b/Engine/Source/Editor/UnrealEd/Classes/MaterialGraph/MaterialGraph.h
--- a/Engine/Source/Editor/UnrealEd/Classes/MaterialGraph/MaterialGraph.h
+++ b/Engine/Source/Editor/UnrealEd/Classes/MaterialGraph/MaterialGraph.h
@@ -120,8 +120,8 @@ private:
 
 };
 
-UCLASS(Optional)
-class UNREALED_API UMaterialGraph : public UEdGraph
+UCLASS(Optional, MinimalAPI)
+class UMaterialGraph : public UEdGraph
 {
 	GENERATED_UCLASS_BODY()
 
@@ -158,20 +158,20 @@ class UNREALED_API UMaterialGraph : public UEdGraph
 	FString	OriginalMaterialFullName;
 
 	//~ Begin UEdGraph interface
-	virtual void NotifyGraphChanged() override;
+	UNREALED_API virtual void NotifyGraphChanged() override;
 protected:
-	virtual void NotifyGraphChanged(const FEdGraphEditAction& Action) override;
+	UNREALED_API virtual void NotifyGraphChanged(const FEdGraphEditAction& Action) override;
 	//~ End UEdGraph interface
 
 public:
-	UMaterialGraph();
-	UMaterialGraph(FVTableHelper& Helper);
-	virtual ~UMaterialGraph();
+	UNREALED_API UMaterialGraph();
+	UNREALED_API UMaterialGraph(FVTableHelper& Helper);
+	UNREALED_API virtual ~UMaterialGraph();
 
 	/**
 	 * Completely rebuild the graph from the material, removing all old nodes
 	 */
-	void RebuildGraph();
+	UNREALED_API void RebuildGraph();
 
 	/**
 	 * Add an Expression to the Graph
@@ -180,7 +180,7 @@ public:
 	 *
 	 * @return	UMaterialGraphNode*	Newly created Graph node to represent expression
 	 */
-	class UMaterialGraphNode*			AddExpression(UMaterialExpression* Expression, bool bUserInvoked);
+	UNREALED_API class UMaterialGraphNode*			AddExpression(UMaterialExpression* Expression, bool bUserInvoked);
 
 	/**
 	 * Add a Comment to the Graph
@@ -190,7 +190,7 @@ public:
 	 *
 	 * @return	UMaterialGraphNode_Comment*	Newly created Graph node to represent comment
 	 */
-	class UMaterialGraphNode_Comment*	AddComment(UMaterialExpressionComment* Comment, bool bIsUserInvoked = false);
+	UNREALED_API class UMaterialGraphNode_Comment*	AddComment(UMaterialExpressionComment* Comment, bool bIsUserInvoked = false);
 
 	/**
 	 * Add a Subgraph to the Graph
@@ -199,32 +199,32 @@ public:
 	 *
 	 * @return	UMaterialGraph* Newly created subgraph
 	 */
-	UMaterialGraph* AddSubGraph(UMaterialExpression* InSubgraphExpression);
+	UNREALED_API UMaterialGraph* AddSubGraph(UMaterialExpression* InSubgraphExpression);
 
 	/** Link all of the Graph nodes using the Material's connections */
-	void LinkGraphNodesFromMaterial();
+	UNREALED_API void LinkGraphNodesFromMaterial();
 
 	/** Link the Material using the Graph node's connections */
-	void LinkMaterialExpressionsFromGraph();
+	UNREALED_API void LinkMaterialExpressionsFromGraph();
 
 	/**
 	 * Check whether a material input should be marked as active
 	 *
 	 * @param	GraphPin	Pin representing the material input
 	 */
-	bool IsInputActive(class UEdGraphPin* GraphPin) const;
+	UNREALED_API bool IsInputActive(class UEdGraphPin* GraphPin) const;
 
 	/** Returns the input index associated with the given property */
-	int32 GetInputIndexForProperty(EMaterialProperty Property) const;
+	UNREALED_API int32 GetInputIndexForProperty(EMaterialProperty Property) const;
 
 	/**
 	 * Get a list of nodes representing expressions that are not used in the Material
 	 *
 	 * @param	UnusedNodes	Array to contain nodes representing unused expressions
 	 */
-	void GetUnusedExpressions(TArray<class UEdGraphNode*>& UnusedNodes) const;
+	UNREALED_API void GetUnusedExpressions(TArray<class UEdGraphNode*>& UnusedNodes) const;
 
-	void UpdatePinTypes();
+	UNREALED_API void UpdatePinTypes();
 
 private:
 	/**
diff --git a/Engine/Source/Editor/UnrealEd/Classes/MaterialGraph/MaterialGraphSchema.h b/Engine/Source/Editor/UnrealEd/Classes/MaterialGraph/MaterialGraphSchema.h
--- a/Engine/Source/Editor/UnrealEd/Classes/MaterialGraph/MaterialGraphSchema.h
+++ b/Engine/Source/Editor/UnrealEd/Classes/MaterialGraph/MaterialGraphSchema.h
@@ -12,7 +12,7 @@ struct FToolMenuContext;
 
 /** Action to add an expression node to the graph */
 USTRUCT()
-struct UNREALED_API FMaterialGraphSchemaAction_NewNode : public FEdGraphSchemaAction
+struct FMaterialGraphSchemaAction_NewNode : public FEdGraphSchemaAction
 {
 	GENERATED_USTRUCT_BODY();
 
@@ -35,7 +35,7 @@ struct UNREALED_API FMaterialGraphSchemaAction_NewNode : public FEdGraphSchemaAc
 	{}
 
 	//~ Begin FEdGraphSchemaAction Interface
-	virtual UEdGraphNode* PerformAction(class UEdGraph* ParentGraph, UEdGraphPin* FromPin, const FVector2D Location, bool bSelectNewNode = true) override;
+	UNREALED_API virtual UEdGraphNode* PerformAction(class UEdGraph* ParentGraph, UEdGraphPin* FromPin, const FVector2D Location, bool bSelectNewNode = true) override;
 	//~ End FEdGraphSchemaAction Interface
 
 	/**
@@ -44,12 +44,12 @@ struct UNREALED_API FMaterialGraphSchemaAction_NewNode : public FEdGraphSchemaAc
 	 * @param	FunctionInput		The function input to set.
 	 * @param	MaterialValueType	Value type we want input to accept.
 	 */
-	void SetFunctionInputType(class UMaterialExpressionFunctionInput* FunctionInput, uint32 MaterialValueType) const;
+	UNREALED_API void SetFunctionInputType(class UMaterialExpressionFunctionInput* FunctionInput, uint32 MaterialValueType) const;
 };
 
 /** Action to add a Material Function call to the graph */
 USTRUCT()
-struct UNREALED_API FMaterialGraphSchemaAction_NewFunctionCall : public FEdGraphSchemaAction
+struct FMaterialGraphSchemaAction_NewFunctionCall : public FEdGraphSchemaAction
 {
 	GENERATED_USTRUCT_BODY();
 
@@ -70,13 +70,13 @@ struct UNREALED_API FMaterialGraphSchemaAction_NewFunctionCall : public FEdGraph
 	{}
 
 	//~ Begin FEdGraphSchemaAction Interface
-	virtual UEdGraphNode* PerformAction(class UEdGraph* ParentGraph, UEdGraphPin* FromPin, const FVector2D Location, bool bSelectNewNode = true) override;
+	UNREALED_API virtual UEdGraphNode* PerformAction(class UEdGraph* ParentGraph, UEdGraphPin* FromPin, const FVector2D Location, bool bSelectNewNode = true) override;
 	//~ End FEdGraphSchemaAction Interface
 };
 
 /** Action to add a composite node to the graph */
 USTRUCT()
-struct UNREALED_API FMaterialGraphSchemaAction_NewComposite : public FEdGraphSchemaAction
+struct FMaterialGraphSchemaAction_NewComposite : public FEdGraphSchemaAction
 {
 	GENERATED_USTRUCT_BODY();
 
@@ -93,15 +93,15 @@ struct UNREALED_API FMaterialGraphSchemaAction_NewComposite : public FEdGraphSch
 	{}
 
 	//~ Begin FEdGraphSchemaAction Interface
-	virtual UEdGraphNode* PerformAction(class UEdGraph* ParentGraph, UEdGraphPin* FromPin, const FVector2D Location, bool bSelectNewNode = true) override;
+	UNREALED_API virtual UEdGraphNode* PerformAction(class UEdGraph* ParentGraph, UEdGraphPin* FromPin, const FVector2D Location, bool bSelectNewNode = true) override;
 	//~ End FEdGraphSchemaAction Interface
 
-	static UEdGraphNode* SpawnNode(class UEdGraph* ParentGraph, const FVector2D Location);
+	static UNREALED_API UEdGraphNode* SpawnNode(class UEdGraph* ParentGraph, const FVector2D Location);
 };
 
 /** Action to add a comment node to the graph */
 USTRUCT()
-struct UNREALED_API FMaterialGraphSchemaAction_NewComment : public FEdGraphSchemaAction
+struct FMaterialGraphSchemaAction_NewComment : public FEdGraphSchemaAction
 {
 	GENERATED_USTRUCT_BODY();
 
@@ -118,13 +118,13 @@ struct UNREALED_API FMaterialGraphSchemaAction_NewComment : public FEdGraphSchem
 	{}
 
 	//~ Begin FEdGraphSchemaAction Interface
-	virtual UEdGraphNode* PerformAction(class UEdGraph* ParentGraph, UEdGraphPin* FromPin, const FVector2D Location, bool bSelectNewNode = true) override;
+	UNREALED_API virtual UEdGraphNode* PerformAction(class UEdGraph* ParentGraph, UEdGraphPin* FromPin, const FVector2D Location, bool bSelectNewNode = true) override;
 	//~ End FEdGraphSchemaAction Interface
 };
 
 /** Action to add a local variable usage to the graph */
 USTRUCT()
-struct UNREALED_API FMaterialGraphSchemaAction_NewNamedRerouteUsage : public FEdGraphSchemaAction
+struct FMaterialGraphSchemaAction_NewNamedRerouteUsage : public FEdGraphSchemaAction
 {
 	GENERATED_USTRUCT_BODY();
 
@@ -145,13 +145,13 @@ struct UNREALED_API FMaterialGraphSchemaAction_NewNamedRerouteUsage : public FEd
 	{}
 
 	//~ Begin FEdGraphSchemaAction Interface
-	virtual UEdGraphNode* PerformAction(class UEdGraph* ParentGraph, UEdGraphPin* FromPin, const FVector2D Location, bool bSelectNewNode = true) override;
+	UNREALED_API virtual UEdGraphNode* PerformAction(class UEdGraph* ParentGraph, UEdGraphPin* FromPin, const FVector2D Location, bool bSelectNewNode = true) override;
 	//~ End FEdGraphSchemaAction Interface
 };
 
 /** Action to paste clipboard contents into the graph */
 USTRUCT()
-struct UNREALED_API FMaterialGraphSchemaAction_Paste : public FEdGraphSchemaAction
+struct FMaterialGraphSchemaAction_Paste : public FEdGraphSchemaAction
 {
 	GENERATED_USTRUCT_BODY();
 
@@ -168,7 +168,7 @@ struct UNREALED_API FMaterialGraphSchemaAction_Paste : public FEdGraphSchemaActi
 	{}
 
 	//~ Begin FEdGraphSchemaAction Interface
-	virtual UEdGraphNode* PerformAction(class UEdGraph* ParentGraph, UEdGraphPin* FromPin, const FVector2D Location, bool bSelectNewNode = true) override;
+	UNREALED_API virtual UEdGraphNode* PerformAction(class UEdGraph* ParentGraph, UEdGraphPin* FromPin, const FVector2D Location, bool bSelectNewNode = true) override;
 	//~ End FEdGraphSchemaAction Interface
 };
 
diff --git a/Engine/Source/Editor/UnrealEd/Classes/Preferences/AnimationBlueprintEditorOptions.h b/Engine/Source/Editor/UnrealEd/Classes/Preferences/AnimationBlueprintEditorOptions.h
--- a/Engine/Source/Editor/UnrealEd/Classes/Preferences/AnimationBlueprintEditorOptions.h
+++ b/Engine/Source/Editor/UnrealEd/Classes/Preferences/AnimationBlueprintEditorOptions.h
@@ -12,8 +12,8 @@
 #include "UObject/Object.h"
 #include "AnimationBlueprintEditorOptions.generated.h"
 
-UCLASS(hidecategories=Object, config=EditorPerProjectUserSettings)
-class UNREALED_API UAnimationBlueprintEditorOptions : public UObject
+UCLASS(hidecategories=Object, config=EditorPerProjectUserSettings, MinimalAPI)
+class UAnimationBlueprintEditorOptions : public UObject
 {
 	GENERATED_UCLASS_BODY()
 
diff --git a/Engine/Source/Editor/UnrealEd/Classes/Preferences/CascadeOptions.h b/Engine/Source/Editor/UnrealEd/Classes/Preferences/CascadeOptions.h
--- a/Engine/Source/Editor/UnrealEd/Classes/Preferences/CascadeOptions.h
+++ b/Engine/Source/Editor/UnrealEd/Classes/Preferences/CascadeOptions.h
@@ -14,8 +14,8 @@
 #include "UObject/Object.h"
 #include "CascadeOptions.generated.h"
 
-UCLASS(hidecategories=Object, config=EditorPerProjectUserSettings)
-class UNREALED_API UCascadeOptions : public UObject
+UCLASS(hidecategories=Object, config=EditorPerProjectUserSettings, MinimalAPI)
+class UCascadeOptions : public UObject
 {
 	GENERATED_UCLASS_BODY()
 
diff --git a/Engine/Source/Editor/UnrealEd/Classes/Preferences/CurveEdOptions.h b/Engine/Source/Editor/UnrealEd/Classes/Preferences/CurveEdOptions.h
--- a/Engine/Source/Editor/UnrealEd/Classes/Preferences/CurveEdOptions.h
+++ b/Engine/Source/Editor/UnrealEd/Classes/Preferences/CurveEdOptions.h
@@ -14,8 +14,8 @@
 #include "UObject/Object.h"
 #include "CurveEdOptions.generated.h"
 
-UCLASS(hidecategories=Object, config=EditorPerProjectUserSettings)
-class UNREALED_API UCurveEdOptions : public UObject
+UCLASS(hidecategories=Object, config=EditorPerProjectUserSettings, MinimalAPI)
+class UCurveEdOptions : public UObject
 {
 	GENERATED_UCLASS_BODY()
 
diff --git a/Engine/Source/Editor/UnrealEd/Classes/Preferences/MaterialEditorOptions.h b/Engine/Source/Editor/UnrealEd/Classes/Preferences/MaterialEditorOptions.h
--- a/Engine/Source/Editor/UnrealEd/Classes/Preferences/MaterialEditorOptions.h
+++ b/Engine/Source/Editor/UnrealEd/Classes/Preferences/MaterialEditorOptions.h
@@ -13,8 +13,8 @@
 #include "UObject/Object.h"
 #include "MaterialEditorOptions.generated.h"
 
-UCLASS(hidecategories=Object, config=EditorPerProjectUserSettings)
-class UNREALED_API UMaterialEditorOptions : public UObject
+UCLASS(hidecategories=Object, config=EditorPerProjectUserSettings, MinimalAPI)
+class UMaterialEditorOptions : public UObject
 {
 	GENERATED_UCLASS_BODY()
 
diff --git a/Engine/Source/Editor/UnrealEd/Classes/Preferences/MaterialStatsOptions.h b/Engine/Source/Editor/UnrealEd/Classes/Preferences/MaterialStatsOptions.h
--- a/Engine/Source/Editor/UnrealEd/Classes/Preferences/MaterialStatsOptions.h
+++ b/Engine/Source/Editor/UnrealEd/Classes/Preferences/MaterialStatsOptions.h
@@ -17,8 +17,8 @@
 #include "SceneTypes.h"
 #include "MaterialStatsOptions.generated.h"
 
-UCLASS(hidecategories = Object, config = EditorPerProjectUserSettings)
-class UNREALED_API UMaterialStatsOptions : public UObject
+UCLASS(hidecategories = Object, config = EditorPerProjectUserSettings, MinimalAPI)
+class UMaterialStatsOptions : public UObject
 {
 	GENERATED_UCLASS_BODY()
 
@@ -27,4 +27,4 @@ class UNREALED_API UMaterialStatsOptions : public UObject
 
 	UPROPERTY(EditAnywhere, config, Category = Options)
 	int32 bMaterialQualityUsed[EMaterialQualityLevel::Num];
-};
\ No newline at end of file
+};
diff --git a/Engine/Source/Editor/UnrealEd/Classes/Preferences/PersonaOptions.h b/Engine/Source/Editor/UnrealEd/Classes/Preferences/PersonaOptions.h
--- a/Engine/Source/Editor/UnrealEd/Classes/Preferences/PersonaOptions.h
+++ b/Engine/Source/Editor/UnrealEd/Classes/Preferences/PersonaOptions.h
@@ -104,8 +104,8 @@ struct FAssetEditorOptions
 	void SetViewportConfigsToDefault();
 };
 
-UCLASS(hidecategories=Object, config=EditorPerProjectUserSettings)
-class UNREALED_API UPersonaOptions : public UObject
+UCLASS(hidecategories=Object, config=EditorPerProjectUserSettings, MinimalAPI)
+class UPersonaOptions : public UObject
 {
 	GENERATED_UCLASS_BODY()
 
@@ -251,25 +251,25 @@ class UNREALED_API UPersonaOptions : public UObject
 	UPROPERTY(EditAnywhere, config, Category = "Timeline")
 	FString AnimationCurveGroupingDelimiters = TEXT("._/|\\");
 public:
-	void SetShowGrid( bool bInShowGrid );
-	void SetHighlightOrigin( bool bInHighlightOrigin );
-	void SetAutoAlignFloorToMesh(bool bInAutoAlignFloorToMesh);
-	void SetMuteAudio( bool bInMuteAudio );
-	void SetUseAudioAttenuation( bool bInUseAudioAttenuation );
-	void SetViewModeIndex( FName InContext, EViewModeIndex InViewModeIndex, int32 InViewportIndex );
-	void SetViewFOV( FName InContext, float InViewFOV, int32 InViewportIndex );
-	void SetCameraSpeed(FName InContext, int32 InCameraSpeed, int32 InViewportIndex);
-	void SetCameraSpeedScalar(FName InContext, float InCameraSpeedScalar, int32 InViewportIndex);
-	void SetViewCameraFollow( FName InContext, EAnimationViewportCameraFollowMode InCameraFollowMode, FName InCameraFollowBoneName, int32 InViewportIndex );
-	void SetDefaultLocalAxesSelection( uint32 InDefaultLocalAxesSelection );
-	void SetDefaultBoneDrawSelection(uint32 InDefaultBoneAxesSelection);
-	void SetShowMeshStats( int32 InShowMeshStats );
-	void SetSectionTimingNodeColor(const FLinearColor& InColor);
-	void SetNotifyTimingNodeColor(const FLinearColor& InColor);
-	void SetBranchingPointTimingNodeColor(const FLinearColor& InColor);
-	FAssetEditorOptions& GetAssetEditorOptions(const FName& InContext);
-	bool GetAllowIncompatibleSkeletonSelection() const;
-	void SetAllowIncompatibleSkeletonSelection(bool bState);
+	UNREALED_API void SetShowGrid( bool bInShowGrid );
+	UNREALED_API void SetHighlightOrigin( bool bInHighlightOrigin );
+	UNREALED_API void SetAutoAlignFloorToMesh(bool bInAutoAlignFloorToMesh);
+	UNREALED_API void SetMuteAudio( bool bInMuteAudio );
+	UNREALED_API void SetUseAudioAttenuation( bool bInUseAudioAttenuation );
+	UNREALED_API void SetViewModeIndex( FName InContext, EViewModeIndex InViewModeIndex, int32 InViewportIndex );
+	UNREALED_API void SetViewFOV( FName InContext, float InViewFOV, int32 InViewportIndex );
+	UNREALED_API void SetCameraSpeed(FName InContext, int32 InCameraSpeed, int32 InViewportIndex);
+	UNREALED_API void SetCameraSpeedScalar(FName InContext, float InCameraSpeedScalar, int32 InViewportIndex);
+	UNREALED_API void SetViewCameraFollow( FName InContext, EAnimationViewportCameraFollowMode InCameraFollowMode, FName InCameraFollowBoneName, int32 InViewportIndex );
+	UNREALED_API void SetDefaultLocalAxesSelection( uint32 InDefaultLocalAxesSelection );
+	UNREALED_API void SetDefaultBoneDrawSelection(uint32 InDefaultBoneAxesSelection);
+	UNREALED_API void SetShowMeshStats( int32 InShowMeshStats );
+	UNREALED_API void SetSectionTimingNodeColor(const FLinearColor& InColor);
+	UNREALED_API void SetNotifyTimingNodeColor(const FLinearColor& InColor);
+	UNREALED_API void SetBranchingPointTimingNodeColor(const FLinearColor& InColor);
+	UNREALED_API FAssetEditorOptions& GetAssetEditorOptions(const FName& InContext);
+	UNREALED_API bool GetAllowIncompatibleSkeletonSelection() const;
+	UNREALED_API void SetAllowIncompatibleSkeletonSelection(bool bState);
 
 	DECLARE_MULTICAST_DELEGATE_TwoParams(FOnUpdateSettingsMulticaster, const UPersonaOptions*, EPropertyChangeType::Type);
 	FOnUpdateSettingsMulticaster OnSettingsChange;
@@ -320,7 +320,7 @@ public:
 
 protected:
 	// UObject interface
-	virtual void PostEditChangeProperty(struct FPropertyChangedEvent& PropertyChangedEvent) override;
+	UNREALED_API virtual void PostEditChangeProperty(struct FPropertyChangedEvent& PropertyChangedEvent) override;
 
 private:
 	FNamePermissionList AllowedAnimationEditorTracks;
diff --git a/Engine/Source/Editor/UnrealEd/Classes/Preferences/PhysicsAssetEditorOptions.h b/Engine/Source/Editor/UnrealEd/Classes/Preferences/PhysicsAssetEditorOptions.h
--- a/Engine/Source/Editor/UnrealEd/Classes/Preferences/PhysicsAssetEditorOptions.h
+++ b/Engine/Source/Editor/UnrealEd/Classes/Preferences/PhysicsAssetEditorOptions.h
@@ -34,8 +34,8 @@ enum class EPhysicsAssetEditorConstraintViewMode : uint8
 	AllLimits
 };
 
-UCLASS(hidecategories=Object, config=EditorPerProjectUserSettings)
-class UNREALED_API UPhysicsAssetEditorOptions : public UObject
+UCLASS(hidecategories=Object, config=EditorPerProjectUserSettings, MinimalAPI)
+class UPhysicsAssetEditorOptions : public UObject
 {
 	GENERATED_UCLASS_BODY()
 
diff --git a/Engine/Source/Editor/UnrealEd/Classes/Preferences/UnrealEdOptions.h b/Engine/Source/Editor/UnrealEd/Classes/Preferences/UnrealEdOptions.h
--- a/Engine/Source/Editor/UnrealEd/Classes/Preferences/UnrealEdOptions.h
+++ b/Engine/Source/Editor/UnrealEd/Classes/Preferences/UnrealEdOptions.h
@@ -81,8 +81,8 @@ struct FClassPickerDefaults
 	FText GetDescription() const;
 };
 
-UCLASS(Config=Editor)
-class UNREALED_API UUnrealEdOptions : public UObject
+UCLASS(Config=Editor, MinimalAPI)
+class UUnrealEdOptions : public UObject
 {
 	GENERATED_BODY()
 public:
@@ -128,13 +128,13 @@ public:
 	TMap<FName, int32>	CommandMap;
 
 	//~ Begin UObject Interface
-	virtual void PostInitProperties() override;
+	UNREALED_API virtual void PostInitProperties() override;
 	//~ End UObject Interface
 
 	/**
 	 * Generates a mapping from commnands to their parent sets for quick lookup.
 	 */
-	void GenerateCommandMap();
+	UNREALED_API void GenerateCommandMap();
 
 	/**
 	 * Attempts to locate a exec command bound to a hotkey.
@@ -145,7 +145,7 @@ public:
 	 * @param bShiftDown	Whether or not SHIFT is pressed.
 	 * @param EditorSet		Set of bindings to search in.
 	 */
-	FString GetExecCommand(FKey Key, bool bAltDown, bool bCtrlDown, bool bShiftDown, FName EditorSet);
+	UNREALED_API FString GetExecCommand(FKey Key, bool bAltDown, bool bCtrlDown, bool bShiftDown, FName EditorSet);
 
 private:
 
diff --git a/Engine/Source/Editor/UnrealEd/Classes/Settings/AnimBlueprintSettings.h b/Engine/Source/Editor/UnrealEd/Classes/Settings/AnimBlueprintSettings.h
--- a/Engine/Source/Editor/UnrealEd/Classes/Settings/AnimBlueprintSettings.h
+++ b/Engine/Source/Editor/UnrealEd/Classes/Settings/AnimBlueprintSettings.h
@@ -9,8 +9,8 @@
 /**
  * Implements Editor settings for animation blueprints
  */
-UCLASS(config = EditorPerProjectUserSettings)
-class UNREALED_API UAnimBlueprintSettings : public UObject
+UCLASS(config = EditorPerProjectUserSettings, MinimalAPI)
+class UAnimBlueprintSettings : public UObject
 {
 	GENERATED_BODY()
 
@@ -58,4 +58,4 @@ public:
 	/** The set of allowed base functions if restricted */
 	UPROPERTY()
 	TArray<FName> BaseFunctionOverrideAllowList;
-};
\ No newline at end of file
+};
diff --git a/Engine/Source/Editor/UnrealEd/Classes/Settings/ClassViewerSettings.h b/Engine/Source/Editor/UnrealEd/Classes/Settings/ClassViewerSettings.h
--- a/Engine/Source/Editor/UnrealEd/Classes/Settings/ClassViewerSettings.h
+++ b/Engine/Source/Editor/UnrealEd/Classes/Settings/ClassViewerSettings.h
@@ -27,8 +27,8 @@ enum class EClassViewerDeveloperType : uint8
 /**
  * Implements the Class Viewer's loading and saving settings.
  */
-UCLASS(config=EditorPerProjectUserSettings)
-class UNREALED_API UClassViewerSettings
+UCLASS(config=EditorPerProjectUserSettings, MinimalAPI)
+class UClassViewerSettings
 	: public UObject
 {
 	GENERATED_UCLASS_BODY()
@@ -60,12 +60,12 @@ public:
 protected:
 
 	// UObject overrides
-	virtual void PostEditChangeProperty(struct FPropertyChangedEvent& PropertyChangedEvent) override;
-	virtual void PostInitProperties() override;
-	virtual void PostLoad() override;
+	UNREALED_API virtual void PostEditChangeProperty(struct FPropertyChangedEvent& PropertyChangedEvent) override;
+	UNREALED_API virtual void PostInitProperties() override;
+	UNREALED_API virtual void PostLoad() override;
 
-	void FixupShortNames();
+	UNREALED_API void FixupShortNames();
 
 	// Holds an event delegate that is executed when a setting has changed.
-	static FSettingChangedEvent SettingChangedEvent;
+	static UNREALED_API FSettingChangedEvent SettingChangedEvent;
 };
diff --git a/Engine/Source/Editor/UnrealEd/Classes/Settings/ContentBrowserSettings.h b/Engine/Source/Editor/UnrealEd/Classes/Settings/ContentBrowserSettings.h
--- a/Engine/Source/Editor/UnrealEd/Classes/Settings/ContentBrowserSettings.h
+++ b/Engine/Source/Editor/UnrealEd/Classes/Settings/ContentBrowserSettings.h
@@ -14,8 +14,8 @@
 /**
  * Implements content browser settings.  These are global not per-project
  */
-UCLASS(config=EditorSettings)
-class UNREALED_API UContentBrowserSettings : public UObject
+UCLASS(config=EditorSettings, MinimalAPI)
+class UContentBrowserSettings : public UObject
 {
 	GENERATED_BODY()
 
@@ -235,7 +235,7 @@ protected:
 
 	// UObject overrides
 
-	virtual void PostEditChangeProperty( struct FPropertyChangedEvent& PropertyChangedEvent ) override;
+	UNREALED_API virtual void PostEditChangeProperty( struct FPropertyChangedEvent& PropertyChangedEvent ) override;
 
 private:
 
@@ -272,5 +272,5 @@ private:
 	bool IncludeCollectionNames;
 
 	// Holds an event delegate that is executed when a setting has changed.
-	static FSettingChangedEvent SettingChangedEvent;
-};
\ No newline at end of file
+	static UNREALED_API FSettingChangedEvent SettingChangedEvent;
+};
diff --git a/Engine/Source/Editor/UnrealEd/Classes/Settings/EditorExperimentalSettings.h b/Engine/Source/Editor/UnrealEd/Classes/Settings/EditorExperimentalSettings.h
--- a/Engine/Source/Editor/UnrealEd/Classes/Settings/EditorExperimentalSettings.h
+++ b/Engine/Source/Editor/UnrealEd/Classes/Settings/EditorExperimentalSettings.h
@@ -13,8 +13,8 @@
 /**
  * Implements Editor settings for experimental features.
  */
-UCLASS(config=EditorPerProjectUserSettings)
-class UNREALED_API UEditorExperimentalSettings
+UCLASS(config=EditorPerProjectUserSettings, MinimalAPI)
+class UEditorExperimentalSettings
 	: public UObject
 {
 	GENERATED_UCLASS_BODY()
@@ -157,8 +157,8 @@ public:
 protected:
 
 	// UObject overrides
-	virtual void PostEditChangeProperty(struct FPropertyChangedEvent& PropertyChangedEvent) override;
-	virtual void PostInitProperties() override;
+	UNREALED_API virtual void PostEditChangeProperty(struct FPropertyChangedEvent& PropertyChangedEvent) override;
+	UNREALED_API virtual void PostInitProperties() override;
 
 private:
 
diff --git a/Engine/Source/Editor/UnrealEd/Classes/Settings/EditorLoadingSavingSettings.h b/Engine/Source/Editor/UnrealEd/Classes/Settings/EditorLoadingSavingSettings.h
--- a/Engine/Source/Editor/UnrealEd/Classes/Settings/EditorLoadingSavingSettings.h
+++ b/Engine/Source/Editor/UnrealEd/Classes/Settings/EditorLoadingSavingSettings.h
@@ -76,11 +76,11 @@ struct FAutoReimportDirectoryConfig
 	UPROPERTY(EditAnywhere, config, Category=AutoReimport, meta=(DisplayName="Include/Exclude Wildcards", ToolTip="(Optional) Specify a set of wildcards to include or exclude files from this auto-reimporter."))
 	TArray<FAutoReimportWildcard> Wildcards;
 
-	struct UNREALED_API FParseContext
+	struct FParseContext
 	{
 		TArray<TPair<FString, FString>> MountedPaths;
 		bool bEnableLogging;
-		FParseContext(bool bInEnableLogging = true);
+		UNREALED_API FParseContext(bool bInEnableLogging = true);
 	};
 
 	/** Parse and validate the specified source directory / mount point combination */
@@ -91,8 +91,8 @@ struct FAutoReimportDirectoryConfig
 /**
  * Implements the Level Editor's loading and saving settings.
  */
-UCLASS(config=EditorPerProjectUserSettings, autoexpandcategories=(AutoSave, AutoReimport, Blueprints))
-class UNREALED_API UEditorLoadingSavingSettings
+UCLASS(config=EditorPerProjectUserSettings, autoexpandcategories=(AutoSave, AutoReimport, Blueprints), MinimalAPI)
+class UEditorLoadingSavingSettings
 	: public UObject
 {
 	GENERATED_UCLASS_BODY()
@@ -223,13 +223,13 @@ public:
 public:
 
 	// @todo thomass: proper settings support for source control module
-	void SccHackInitialize( );
+	UNREALED_API void SccHackInitialize( );
 
-	bool GetAutomaticallyCheckoutOnAssetModification() const;
+	UNREALED_API bool GetAutomaticallyCheckoutOnAssetModification() const;
 
-	void SetAutomaticallyCheckoutOnAssetModificationOverride(bool InValue);
+	UNREALED_API void SetAutomaticallyCheckoutOnAssetModificationOverride(bool InValue);
 
-	void ResetAutomaticallyCheckoutOnAssetModificationOverride();
+	UNREALED_API void ResetAutomaticallyCheckoutOnAssetModificationOverride();
 
 public:
 
@@ -245,8 +245,8 @@ protected:
 
 	// UObject overrides
 
-	virtual void PostEditChangeProperty( struct FPropertyChangedEvent& PropertyChangedEvent ) override;
-	virtual void PostInitProperties() override;
+	UNREALED_API virtual void PostEditChangeProperty( struct FPropertyChangedEvent& PropertyChangedEvent ) override;
+	UNREALED_API virtual void PostInitProperties() override;
 
 private:
 
diff --git a/Engine/Source/Editor/UnrealEd/Classes/Settings/EditorMiscSettings.h b/Engine/Source/Editor/UnrealEd/Classes/Settings/EditorMiscSettings.h
--- a/Engine/Source/Editor/UnrealEd/Classes/Settings/EditorMiscSettings.h
+++ b/Engine/Source/Editor/UnrealEd/Classes/Settings/EditorMiscSettings.h
@@ -12,8 +12,8 @@
 /**
  * Implements the miscellaneous Editor settings.
  */
-UCLASS(config=EditorPerProjectUserSettings)
-class UNREALED_API UEditorMiscSettings
+UCLASS(config=EditorPerProjectUserSettings, MinimalAPI)
+class UEditorMiscSettings
 	: public UObject
 {
 	GENERATED_UCLASS_BODY()
diff --git a/Engine/Source/Editor/UnrealEd/Classes/Settings/EditorStyleSettings.h b/Engine/Source/Editor/UnrealEd/Classes/Settings/EditorStyleSettings.h
--- a/Engine/Source/Editor/UnrealEd/Classes/Settings/EditorStyleSettings.h
+++ b/Engine/Source/Editor/UnrealEd/Classes/Settings/EditorStyleSettings.h
@@ -32,14 +32,14 @@ enum class EAssetEditorOpenLocation : uint8
 /**
  * Implements the Editor style settings.
  */
-UCLASS(config=EditorPerProjectUserSettings)
-class UNREALED_API UEditorStyleSettings : public UObject
+UCLASS(config=EditorPerProjectUserSettings, MinimalAPI)
+class UEditorStyleSettings : public UObject
 {
 public:
 
 	GENERATED_UCLASS_BODY()
 
-	void Init();
+	UNREALED_API void Init();
 public:
 	/**
 	 * Enables high dpi support in the editor which will adjust the scale of elements in the UI to account for high DPI monitors
@@ -159,15 +159,15 @@ public:
 	FSettingChangedEvent& OnSettingChanged( ) { return SettingChangedEvent; }
 
 	/** @return A subdued version of the users selection color (for use with inactive selection)*/
-	FLinearColor GetSubduedSelectionColor() const;
+	UNREALED_API FLinearColor GetSubduedSelectionColor() const;
 	
-	bool OnImportBegin(const FString& ImportFromPath);
-	bool OnExportBegin(const FString& ExportToPath);
+	UNREALED_API bool OnImportBegin(const FString& ImportFromPath);
+	UNREALED_API bool OnExportBegin(const FString& ExportToPath);
 
 protected:
 
 	// UObject overrides
-	virtual void PostEditChangeProperty( struct FPropertyChangedEvent& PropertyChangedEvent ) override;
+	UNREALED_API virtual void PostEditChangeProperty( struct FPropertyChangedEvent& PropertyChangedEvent ) override;
 
 private:
 
diff --git a/Engine/Source/Editor/UnrealEd/Classes/Settings/LevelEditorMiscSettings.h b/Engine/Source/Editor/UnrealEd/Classes/Settings/LevelEditorMiscSettings.h
--- a/Engine/Source/Editor/UnrealEd/Classes/Settings/LevelEditorMiscSettings.h
+++ b/Engine/Source/Editor/UnrealEd/Classes/Settings/LevelEditorMiscSettings.h
@@ -14,8 +14,8 @@ class ULevelStreaming;
 /**
  * Configure miscellaneous settings for the Level Editor.
  */
-UCLASS(config=EditorPerProjectUserSettings, meta=( DisplayName="Miscellaneous" ))
-class UNREALED_API ULevelEditorMiscSettings : public UDeveloperSettings
+UCLASS(config=EditorPerProjectUserSettings, meta=( DisplayName="Miscellaneous" ), MinimalAPI)
+class ULevelEditorMiscSettings : public UDeveloperSettings
 {
 	GENERATED_UCLASS_BODY()
 
@@ -96,5 +96,5 @@ protected:
 
 	// UObject overrides
 
-	virtual void PostEditChangeProperty( struct FPropertyChangedEvent& PropertyChangedEvent ) override;
+	UNREALED_API virtual void PostEditChangeProperty( struct FPropertyChangedEvent& PropertyChangedEvent ) override;
 };
diff --git a/Engine/Source/Editor/UnrealEd/Classes/Settings/LevelEditorPlaySettings.h b/Engine/Source/Editor/UnrealEd/Classes/Settings/LevelEditorPlaySettings.h
--- a/Engine/Source/Editor/UnrealEd/Classes/Settings/LevelEditorPlaySettings.h
+++ b/Engine/Source/Editor/UnrealEd/Classes/Settings/LevelEditorPlaySettings.h
@@ -195,8 +195,8 @@ public:
 	void PostInitProperties();
 };
 
-UCLASS()
-class UNREALED_API UCommonResolutionMenuContext
+UCLASS(MinimalAPI)
+class UCommonResolutionMenuContext
 	: public UToolMenuContextBase
 {
 	GENERATED_BODY()
@@ -209,8 +209,8 @@ public:
 /**
  * Implements the Editor's play settings.
  */
-UCLASS(config=EditorPerProjectUserSettings)
-class UNREALED_API ULevelEditorPlaySettings
+UCLASS(config=EditorPerProjectUserSettings, MinimalAPI)
+class ULevelEditorPlaySettings
 	: public UObject
 {
 	GENERATED_UCLASS_BODY()
@@ -444,7 +444,7 @@ private:
 	FLinearColor ServerDebugDrawingColorTint;
 
 private:
-	void PushDebugDrawingSettings();
+	UNREALED_API void PushDebugDrawingSettings();
 
 public:
 	bool ShowServerDebugDrawingByDefault() const
@@ -629,27 +629,27 @@ public:
 	TArray<FVector2D> CustomUnsafeZoneDimensions;
 
 	// UObject interface
-	virtual void PostEditChangeProperty(struct FPropertyChangedEvent& PropertyChangedEvent) override;
-	virtual void PostInitProperties() override;
-	virtual bool CanEditChange(const FProperty* InProperty) const override;
+	UNREALED_API virtual void PostEditChangeProperty(struct FPropertyChangedEvent& PropertyChangedEvent) override;
+	UNREALED_API virtual void PostInitProperties() override;
+	UNREALED_API virtual bool CanEditChange(const FProperty* InProperty) const override;
 	// End of UObject interface
 
 #if WITH_EDITOR
 	// Recalculates and broadcasts safe zone size changes based on device to emulate and r.DebugSafeZone.TitleRatio values.
-	void UpdateCustomSafeZones();
+	UNREALED_API void UpdateCustomSafeZones();
 #endif
 
-	FMargin CalculateCustomUnsafeZones(TArray<FVector2D>& CustomSafeZoneStarts, TArray<FVector2D>& CustomSafeZoneDimensions, FString& DeviceType, FVector2D PreviewSize);
-	FMargin FlipCustomUnsafeZones(TArray<FVector2D>& CustomSafeZoneStarts, TArray<FVector2D>& CustomSafeZoneDimensions, FString& DeviceType, FVector2D PreviewSize);
-	void RescaleForMobilePreview(const class UDeviceProfile* DeviceProfile, int32 &PreviewWidth, int32 &PreviewHeight, float &ScaleFactor);
+	UNREALED_API FMargin CalculateCustomUnsafeZones(TArray<FVector2D>& CustomSafeZoneStarts, TArray<FVector2D>& CustomSafeZoneDimensions, FString& DeviceType, FVector2D PreviewSize);
+	UNREALED_API FMargin FlipCustomUnsafeZones(TArray<FVector2D>& CustomSafeZoneStarts, TArray<FVector2D>& CustomSafeZoneDimensions, FString& DeviceType, FVector2D PreviewSize);
+	UNREALED_API void RescaleForMobilePreview(const class UDeviceProfile* DeviceProfile, int32 &PreviewWidth, int32 &PreviewHeight, float &ScaleFactor);
 
 	/**
      * Creates a widget for the resolution picker.
      *
      * @return The widget.
      */
-	void RegisterCommonResolutionsMenu();
-	static FName GetCommonResolutionsMenuName();
+	UNREALED_API void RegisterCommonResolutionsMenu();
+	static UNREALED_API FName GetCommonResolutionsMenuName();
 
 protected:
 	/**
@@ -659,5 +659,5 @@ protected:
 	 * @param Resolutions The collection of screen resolutions to add.
 	 * @param SectionName The name of the section to add.
 	 */
-	static void AddScreenResolutionSection( UToolMenu* InToolMenu, const TArray<FPlayScreenResolution>* Resolutions, const FString SectionName );
+	static UNREALED_API void AddScreenResolutionSection( UToolMenu* InToolMenu, const TArray<FPlayScreenResolution>* Resolutions, const FString SectionName );
 };
diff --git a/Engine/Source/Editor/UnrealEd/Classes/Settings/LevelEditorViewportSettings.h b/Engine/Source/Editor/UnrealEd/Classes/Settings/LevelEditorViewportSettings.h
--- a/Engine/Source/Editor/UnrealEd/Classes/Settings/LevelEditorViewportSettings.h
+++ b/Engine/Source/Editor/UnrealEd/Classes/Settings/LevelEditorViewportSettings.h
@@ -86,7 +86,7 @@ enum class EMaterialKind : uint8
  * Implements the Level Editor's per-instance view port settings.
  */
 USTRUCT()
-struct UNREALED_API FLevelEditorViewportInstanceSettings
+struct FLevelEditorViewportInstanceSettings
 {
 	GENERATED_USTRUCT_BODY()
 
@@ -209,7 +209,7 @@ struct UNREALED_API FLevelEditorViewportInstanceSettings
  * Implements a key -> value pair for the per-instance view port settings
  */
 USTRUCT()
-struct UNREALED_API FLevelEditorViewportInstanceSettingsKeyValuePair
+struct FLevelEditorViewportInstanceSettingsKeyValuePair
 {
 	GENERATED_USTRUCT_BODY()
 
@@ -227,7 +227,7 @@ struct UNREALED_API FLevelEditorViewportInstanceSettingsKeyValuePair
  * Settings that control the behavior of the "snap to surface" feature
  */
 USTRUCT()
-struct UNREALED_API FSnapToSurfaceSettings
+struct FSnapToSurfaceSettings
 {
 	GENERATED_USTRUCT_BODY()
 
@@ -251,7 +251,7 @@ struct UNREALED_API FSnapToSurfaceSettings
 };
 
 USTRUCT()
-struct UNREALED_API FLevelEditorViewporEditorViews
+struct FLevelEditorViewporEditorViews
 {
 	GENERATED_USTRUCT_BODY()
 
@@ -262,8 +262,8 @@ struct UNREALED_API FLevelEditorViewporEditorViews
 /**
  * Implements the Level Editor's view port settings.
  */
-UCLASS(config=EditorPerProjectUserSettings)
-class UNREALED_API ULevelEditorViewportSettings
+UCLASS(config=EditorPerProjectUserSettings, MinimalAPI)
+class ULevelEditorViewportSettings
 	: public UObject
 {
 	GENERATED_UCLASS_BODY()
@@ -667,8 +667,8 @@ protected:
 
 	// UObject overrides
 
-	virtual void PostInitProperties() override;
-	virtual void PostEditChangeProperty( struct FPropertyChangedEvent& PropertyChangedEvent ) override;
+	UNREALED_API virtual void PostInitProperties() override;
+	UNREALED_API virtual void PostEditChangeProperty( struct FPropertyChangedEvent& PropertyChangedEvent ) override;
 
 private:
 
diff --git a/Engine/Source/Editor/UnrealEd/Classes/Settings/SkeletalMeshEditorSettings.h b/Engine/Source/Editor/UnrealEd/Classes/Settings/SkeletalMeshEditorSettings.h
--- a/Engine/Source/Editor/UnrealEd/Classes/Settings/SkeletalMeshEditorSettings.h
+++ b/Engine/Source/Editor/UnrealEd/Classes/Settings/SkeletalMeshEditorSettings.h
@@ -16,8 +16,8 @@
 /**
  * Implements the settings for the skeletal mesh editor.
  */
-UCLASS(config=EditorPerProjectUserSettings)
-class UNREALED_API USkeletalMeshEditorSettings
+UCLASS(config=EditorPerProjectUserSettings, MinimalAPI)
+class USkeletalMeshEditorSettings
 	: public UObject
 {
 	GENERATED_UCLASS_BODY()
@@ -41,4 +41,4 @@ public:
 
 	UPROPERTY(EditAnywhere, config, Category=AnimationPreview, meta=(DisplayName="Viewport Directional Color"))
 	FColor AnimPreviewDirectionalColor;
-};
\ No newline at end of file
+};
diff --git a/Engine/Source/Editor/UnrealEd/Classes/Settings/StructViewerSettings.h b/Engine/Source/Editor/UnrealEd/Classes/Settings/StructViewerSettings.h
--- a/Engine/Source/Editor/UnrealEd/Classes/Settings/StructViewerSettings.h
+++ b/Engine/Source/Editor/UnrealEd/Classes/Settings/StructViewerSettings.h
@@ -21,8 +21,8 @@ enum class EStructViewerDeveloperType : uint8
 /**
  * Implements the Struct Viewer's loading and saving settings.
  */
-UCLASS(config=EditorPerProjectUserSettings)
-class UNREALED_API UStructViewerSettings : public UObject
+UCLASS(config=EditorPerProjectUserSettings, MinimalAPI)
+class UStructViewerSettings : public UObject
 {
 	GENERATED_BODY()
 
@@ -46,8 +46,8 @@ public:
 
 protected:
 	//~ UObject overrides
-	virtual void PostEditChangeProperty(struct FPropertyChangedEvent& PropertyChangedEvent) override;
+	UNREALED_API virtual void PostEditChangeProperty(struct FPropertyChangedEvent& PropertyChangedEvent) override;
 
 	// Holds an event delegate that is executed when a setting has changed.
-	static FSettingChangedEvent SettingChangedEvent;
+	static UNREALED_API FSettingChangedEvent SettingChangedEvent;
 };
diff --git a/Engine/Source/Editor/UnrealEd/Classes/ThumbnailRendering/SkeletalMeshThumbnailRenderer.h b/Engine/Source/Editor/UnrealEd/Classes/ThumbnailRendering/SkeletalMeshThumbnailRenderer.h
--- a/Engine/Source/Editor/UnrealEd/Classes/ThumbnailRendering/SkeletalMeshThumbnailRenderer.h
+++ b/Engine/Source/Editor/UnrealEd/Classes/ThumbnailRendering/SkeletalMeshThumbnailRenderer.h
@@ -16,20 +16,20 @@
 class FCanvas;
 class FRenderTarget;
 
-UCLASS(config=Editor)
-class UNREALED_API USkeletalMeshThumbnailRenderer : public UDefaultSizedThumbnailRenderer
+UCLASS(config=Editor, MinimalAPI)
+class USkeletalMeshThumbnailRenderer : public UDefaultSizedThumbnailRenderer
 {
 	GENERATED_UCLASS_BODY()
 
 
 	// Begin UThumbnailRenderer Object
-	virtual void Draw(UObject* Object, int32 X, int32 Y, uint32 Width, uint32 Height, FRenderTarget*, FCanvas* Canvas, bool bAdditionalViewFamily) override;
-	virtual EThumbnailRenderFrequency GetThumbnailRenderFrequency(UObject* Object) const override;
-	virtual bool CanVisualizeAsset(UObject* Object) override;
+	UNREALED_API virtual void Draw(UObject* Object, int32 X, int32 Y, uint32 Width, uint32 Height, FRenderTarget*, FCanvas* Canvas, bool bAdditionalViewFamily) override;
+	UNREALED_API virtual EThumbnailRenderFrequency GetThumbnailRenderFrequency(UObject* Object) const override;
+	UNREALED_API virtual bool CanVisualizeAsset(UObject* Object) override;
 	// End UThumbnailRenderer Object
 
 	// UObject implementation
-	virtual void BeginDestroy() override;
+	UNREALED_API virtual void BeginDestroy() override;
 
 	virtual void AddAdditionalPreviewSceneContent(UObject* Object, UWorld* PreviewWorld) {}
 
diff --git a/Engine/Source/Editor/UnrealEd/Classes/ThumbnailRendering/SkeletonThumbnailRenderer.h b/Engine/Source/Editor/UnrealEd/Classes/ThumbnailRendering/SkeletonThumbnailRenderer.h
--- a/Engine/Source/Editor/UnrealEd/Classes/ThumbnailRendering/SkeletonThumbnailRenderer.h
+++ b/Engine/Source/Editor/UnrealEd/Classes/ThumbnailRendering/SkeletonThumbnailRenderer.h
@@ -16,19 +16,19 @@
 class FCanvas;
 class FRenderTarget;
 
-UCLASS(config=Editor)
-class UNREALED_API USkeletonThumbnailRenderer : public UDefaultSizedThumbnailRenderer
+UCLASS(config=Editor, MinimalAPI)
+class USkeletonThumbnailRenderer : public UDefaultSizedThumbnailRenderer
 {
 	GENERATED_UCLASS_BODY()
 
 
 	// Begin UThumbnailRenderer Object
-	virtual void Draw(UObject* Object, int32 X, int32 Y, uint32 Width, uint32 Height, FRenderTarget*, FCanvas* Canvas, bool bAdditionalViewFamily) override;
-	virtual EThumbnailRenderFrequency GetThumbnailRenderFrequency(UObject* Object) const override;
+	UNREALED_API virtual void Draw(UObject* Object, int32 X, int32 Y, uint32 Width, uint32 Height, FRenderTarget*, FCanvas* Canvas, bool bAdditionalViewFamily) override;
+	UNREALED_API virtual EThumbnailRenderFrequency GetThumbnailRenderFrequency(UObject* Object) const override;
 	// End UThumbnailRenderer Object
 
 	// UObject implementation
-	virtual void BeginDestroy() override;
+	UNREALED_API virtual void BeginDestroy() override;
 
 	virtual void AddAdditionalPreviewSceneContent(UObject* Object, UWorld* PreviewWorld) {}
 
diff --git a/Engine/Source/Editor/UnrealEd/Classes/ThumbnailRendering/ThumbnailManager.h b/Engine/Source/Editor/UnrealEd/Classes/ThumbnailRendering/ThumbnailManager.h
--- a/Engine/Source/Editor/UnrealEd/Classes/ThumbnailRendering/ThumbnailManager.h
+++ b/Engine/Source/Editor/UnrealEd/Classes/ThumbnailRendering/ThumbnailManager.h
@@ -156,13 +156,13 @@ public:
 	/** 
 	 * Returns the thumbnail pool which should be used for most thumbnails in the editor.
 	 */
-	UNREALED_API TSharedPtr<FAssetThumbnailPool> GetSharedThumbnailPool() const { return SharedThumbnailPool; }
+	TSharedPtr<FAssetThumbnailPool> GetSharedThumbnailPool() const { return SharedThumbnailPool; }
 
 	/**
 	 * Event that is being broadcasted when a thumbnail gets dirtied.
 	 * Parameter is the object soft object path associated with the thumbnail.
 	 */
-	UNREALED_API FOnThumbnailDirtied& GetOnThumbnailDirtied() { return OnThumbnailDirtied; }
+	FOnThumbnailDirtied& GetOnThumbnailDirtied() { return OnThumbnailDirtied; }
 
 protected:
 	/**
diff --git a/Engine/Source/Editor/UnrealEd/Classes/UserDefinedStructure/UserDefinedStructEditorData.h b/Engine/Source/Editor/UnrealEd/Classes/UserDefinedStructure/UserDefinedStructEditorData.h
--- a/Engine/Source/Editor/UnrealEd/Classes/UserDefinedStructure/UserDefinedStructEditorData.h
+++ b/Engine/Source/Editor/UnrealEd/Classes/UserDefinedStructure/UserDefinedStructEditorData.h
@@ -111,8 +111,8 @@ struct TStructOpsTypeTraits< FStructVariableDescription > : public TStructOpsTyp
 	};
 };
 
-UCLASS()
-class UNREALED_API UUserDefinedStructEditorData : public UObject, public FEditorUndoClient
+UCLASS(MinimalAPI)
+class UUserDefinedStructEditorData : public UObject, public FEditorUndoClient
 {
 	GENERATED_UCLASS_BODY()
 
@@ -130,28 +130,28 @@ public:
 
 public:
 	// UObject interface.
-	virtual TSharedPtr<ITransactionObjectAnnotation> FactoryTransactionAnnotation(const ETransactionAnnotationCreationMode InCreationMode) const override;
-	virtual void Serialize(FArchive& Ar) override;
-	virtual void PostEditUndo() override;
-	virtual void PostEditUndo(TSharedPtr<ITransactionObjectAnnotation> TransactionAnnotation) override;
-	virtual void PostLoadSubobjects(struct FObjectInstancingGraph* OuterInstanceGraph) override;
-	virtual void PreSave(FObjectPreSaveContext ObjectSaveContext) override;
+	UNREALED_API virtual TSharedPtr<ITransactionObjectAnnotation> FactoryTransactionAnnotation(const ETransactionAnnotationCreationMode InCreationMode) const override;
+	UNREALED_API virtual void Serialize(FArchive& Ar) override;
+	UNREALED_API virtual void PostEditUndo() override;
+	UNREALED_API virtual void PostEditUndo(TSharedPtr<ITransactionObjectAnnotation> TransactionAnnotation) override;
+	UNREALED_API virtual void PostLoadSubobjects(struct FObjectInstancingGraph* OuterInstanceGraph) override;
+	UNREALED_API virtual void PreSave(FObjectPreSaveContext ObjectSaveContext) override;
 	// End of UObject interface.
 
 	// FEditorUndoClient interface
-	virtual void PostUndo(bool bSuccess) override;
+	UNREALED_API virtual void PostUndo(bool bSuccess) override;
 	virtual void PostRedo(bool bSuccess) override { PostUndo(bSuccess); }
 	// End of FEditorUndoClient interface.
 
 
-	uint32 GenerateUniqueNameIdForMemberVariable();
-	class UUserDefinedStruct* GetOwnerStruct() const;
+	UNREALED_API uint32 GenerateUniqueNameIdForMemberVariable();
+	UNREALED_API class UUserDefinedStruct* GetOwnerStruct() const;
 
-	const uint8* GetDefaultInstance() const;
-	void RecreateDefaultInstance(FString* OutLog = nullptr);
-	void ReinitializeDefaultInstance(FString* OutLog = nullptr);
-	void CleanDefaultInstance();
-	void RefreshValuesFromDefaultInstance();
+	UNREALED_API const uint8* GetDefaultInstance() const;
+	UNREALED_API void RecreateDefaultInstance(FString* OutLog = nullptr);
+	UNREALED_API void ReinitializeDefaultInstance(FString* OutLog = nullptr);
+	UNREALED_API void CleanDefaultInstance();
+	UNREALED_API void RefreshValuesFromDefaultInstance();
 
 private:
 
@@ -159,6 +159,6 @@ private:
 	FStructureEditorUtils::EStructureEditorChangeInfo CachedStructureChange;
 
 	// Utility function for both PostEditUndo to route through
-	void ConsolidatedPostEditUndo(FStructureEditorUtils::EStructureEditorChangeInfo ActiveChange);
+	UNREALED_API void ConsolidatedPostEditUndo(FStructureEditorUtils::EStructureEditorChangeInfo ActiveChange);
 
 };
diff --git a/Engine/Source/Editor/UnrealEd/Private/EditorViewportLayoutEntity.h b/Engine/Source/Editor/UnrealEd/Private/EditorViewportLayoutEntity.h
--- a/Engine/Source/Editor/UnrealEd/Private/EditorViewportLayoutEntity.h
+++ b/Engine/Source/Editor/UnrealEd/Private/EditorViewportLayoutEntity.h
@@ -9,20 +9,20 @@
 
 class FEditorViewportClient;
 
-class UNREALED_API FEditorViewportLayoutEntity : public IEditorViewportLayoutEntity
+class FEditorViewportLayoutEntity : public IEditorViewportLayoutEntity
 {
 public:
 
-	FEditorViewportLayoutEntity(TSharedPtr<SAssetEditorViewport>& InViewport);
-	virtual TSharedRef<SWidget> AsWidget() const override;
-	virtual TSharedPtr<SAssetEditorViewport> AsAssetEditorViewport() const;
+	UNREALED_API FEditorViewportLayoutEntity(TSharedPtr<SAssetEditorViewport>& InViewport);
+	UNREALED_API virtual TSharedRef<SWidget> AsWidget() const override;
+	UNREALED_API virtual TSharedPtr<SAssetEditorViewport> AsAssetEditorViewport() const;
 
-	void SetKeyboardFocus() override;
-	void OnLayoutDestroyed() override;
-	void SaveConfig(const FString& ConfigSection) override;
-	TSharedPtr<FEditorViewportClient> GetViewportClient() const;
-	FName GetType() const override;
-	void TakeHighResScreenShot() const override;
+	UNREALED_API void SetKeyboardFocus() override;
+	UNREALED_API void OnLayoutDestroyed() override;
+	UNREALED_API void SaveConfig(const FString& ConfigSection) override;
+	UNREALED_API TSharedPtr<FEditorViewportClient> GetViewportClient() const;
+	UNREALED_API FName GetType() const override;
+	UNREALED_API void TakeHighResScreenShot() const override;
 
 protected:
 
