[UnrealEd]

diff --git a/Engine/Source/Editor/UnrealEd/Private/FbxExporter.h b/Engine/Source/Editor/UnrealEd/Private/FbxExporter.h
--- a/Engine/Source/Editor/UnrealEd/Private/FbxExporter.h
+++ b/Engine/Source/Editor/UnrealEd/Private/FbxExporter.h
@@ -48,7 +48,7 @@ struct FMovieSceneSequenceTransform;
 namespace UnFbx
 {
 	/** Adapter interface which allows ExportAnimTrack to act on sequencer without a tight coupling. */
-	class UNREALED_API IAnimTrackAdapter
+	class IAnimTrackAdapter
 	{
 	public:
 		virtual ~IAnimTrackAdapter() {};
@@ -65,17 +65,17 @@ namespace UnFbx
 	};
 
 	/** An anim track adapter for a level sequence. */
-	class UNREALED_API FLevelSequenceAnimTrackAdapter : public IAnimTrackAdapter
+	class FLevelSequenceAnimTrackAdapter : public IAnimTrackAdapter
 	{
 	public:
-		FLevelSequenceAnimTrackAdapter(IMovieScenePlayer* InMovieScenePlayer, UMovieScene* InMovieScene, const FMovieSceneSequenceTransform& InRootToLocalTransform, UMovieSceneSkeletalAnimationTrack* InAnimTrack = nullptr);
-		virtual int32 GetLocalStartFrame() const override;
-		virtual int32 GetStartFrame() const override;
-		virtual int32 GetLength() const override;
-		virtual void UpdateAnimation(int32 LocalFrame) override;
-		virtual float GetFrameRate() const override;
-		virtual UAnimSequence* GetAnimSequence(int32 LocalFrame) const override;
-		virtual float GetAnimTime(int32 LocalFrame) const override;
+		UNREALED_API FLevelSequenceAnimTrackAdapter(IMovieScenePlayer* InMovieScenePlayer, UMovieScene* InMovieScene, const FMovieSceneSequenceTransform& InRootToLocalTransform, UMovieSceneSkeletalAnimationTrack* InAnimTrack = nullptr);
+		UNREALED_API virtual int32 GetLocalStartFrame() const override;
+		UNREALED_API virtual int32 GetStartFrame() const override;
+		UNREALED_API virtual int32 GetLength() const override;
+		UNREALED_API virtual void UpdateAnimation(int32 LocalFrame) override;
+		UNREALED_API virtual float GetFrameRate() const override;
+		UNREALED_API virtual UAnimSequence* GetAnimSequence(int32 LocalFrame) const override;
+		UNREALED_API virtual float GetAnimTime(int32 LocalFrame) const override;
 
 	private:
 		IMovieScenePlayer* MovieScenePlayer;
@@ -86,18 +86,18 @@ namespace UnFbx
 /**
  * Main FBX Exporter class.
  */
-class UNREALED_API FFbxExporter : public FCinematicExporter, public FGCObject
+class FFbxExporter : public FCinematicExporter, public FGCObject
 {
 public:
 	/**
 	 * Returns the exporter singleton. It will be created on the first request.
 	 */
-	static FFbxExporter* GetInstance();
-	static void DeleteInstance();
-	~FFbxExporter();
+	static UNREALED_API FFbxExporter* GetInstance();
+	static UNREALED_API void DeleteInstance();
+	UNREALED_API ~FFbxExporter();
 	
 	//~ FGCObject
-	virtual void AddReferencedObjects(FReferenceCollector& Collector) override;
+	UNREALED_API virtual void AddReferencedObjects(FReferenceCollector& Collector) override;
 	virtual FString GetReferencerName() const override
 	{
 		return TEXT("FFbxExporter");
@@ -111,52 +111,52 @@ public:
 	*
 	* The function is saving the dialog state in a user ini file and reload it from there. It is not changing the CDO.
 	*/
-	void FillExportOptions(bool BatchMode, bool bShowOptionDialog, const FString& FullPath, bool& OutOperationCanceled, bool& bOutExportAll);
+	UNREALED_API void FillExportOptions(bool BatchMode, bool bShowOptionDialog, const FString& FullPath, bool& OutOperationCanceled, bool& bOutExportAll);
 
 	/**
 	* Custom set of export options instead of UI dialog. For automation.
 	*/
-	void SetExportOptionsOverride(UFbxExportOption* OverrideOptions);
+	UNREALED_API void SetExportOptionsOverride(UFbxExportOption* OverrideOptions);
 
 	/**
 	 * Creates and readies an empty document for export.
 	 */
-	virtual void CreateDocument();
+	UNREALED_API virtual void CreateDocument();
 	
 	/**
 	 * Closes the FBX document, releasing its memory.
 	 */
-	virtual void CloseDocument();
+	UNREALED_API virtual void CloseDocument();
 	
 	/**
 	 * Writes the FBX document to disk and releases it by calling the CloseDocument() function.
 	 */
-	virtual void WriteToFile(const TCHAR* Filename);
+	UNREALED_API virtual void WriteToFile(const TCHAR* Filename);
 	
 	/**
 	 * Exports the light-specific information for a light actor.
 	 */
-	virtual void ExportLight( ALight* Actor, INodeNameAdapter& NodeNameAdapter );
+	UNREALED_API virtual void ExportLight( ALight* Actor, INodeNameAdapter& NodeNameAdapter );
 
 	/**
 	 * Exports the camera-specific information for a camera actor.
 	 */
-	virtual void ExportCamera( ACameraActor* Actor, bool bExportComponents, INodeNameAdapter& NodeNameAdapter );
+	UNREALED_API virtual void ExportCamera( ACameraActor* Actor, bool bExportComponents, INodeNameAdapter& NodeNameAdapter );
 
 	/**
 	 * Exports the mesh and the actor information for a brush actor.
 	 */
-	virtual void ExportBrush(ABrush* Actor, UModel* InModel, bool bConvertToStaticMesh, INodeNameAdapter& NodeNameAdapter );
+	UNREALED_API virtual void ExportBrush(ABrush* Actor, UModel* InModel, bool bConvertToStaticMesh, INodeNameAdapter& NodeNameAdapter );
 
 	/**
 	 * Exports the basic scene information to the FBX document.
 	 */
-	virtual void ExportLevelMesh( ULevel* InLevel, bool bSelectedOnly, INodeNameAdapter& NodeNameAdapter , bool bSaveAnimSeq = true);
+	UNREALED_API virtual void ExportLevelMesh( ULevel* InLevel, bool bSelectedOnly, INodeNameAdapter& NodeNameAdapter , bool bSaveAnimSeq = true);
 
 	/**
 	* Exports the basic scene information to the FBX document, using the passed in Actors
 	*/
-	virtual void ExportLevelMesh(ULevel* InLevel, bool bExportLevelGeometry, TArray<AActor*>& ActorToExport, INodeNameAdapter& NodeNameAdapter, bool bSaveAnimSeq = true);
+	UNREALED_API virtual void ExportLevelMesh(ULevel* InLevel, bool bExportLevelGeometry, TArray<AActor*>& ActorToExport, INodeNameAdapter& NodeNameAdapter, bool bSaveAnimSeq = true);
 
 
 	/**
@@ -164,74 +164,74 @@ public:
 	 *
 	 * @return	true, if successful
 	 */
-	bool ExportLevelSequence(UMovieScene* MovieScene, const TArray<FGuid>& InBindings, IMovieScenePlayer* MovieScenePlayer, INodeNameAdapter& NodeNameAdapter, FMovieSceneSequenceIDRef SequenceID, const FMovieSceneSequenceTransform& RootToLocalTransform);
+	UNREALED_API bool ExportLevelSequence(UMovieScene* MovieScene, const TArray<FGuid>& InBindings, IMovieScenePlayer* MovieScenePlayer, INodeNameAdapter& NodeNameAdapter, FMovieSceneSequenceIDRef SequenceID, const FMovieSceneSequenceTransform& RootToLocalTransform);
 
 	/** Add timecode attributes to the given fbx node and add a key at the sequence playback start, using the timecode of the source section */
-	void AddTimecodeAttributesAndSetKey(const UMovieSceneSection* InSection, FbxNode* InFbxNode, const FMovieSceneSequenceTransform& RootToLocalTransform);
+	UNREALED_API void AddTimecodeAttributesAndSetKey(const UMovieSceneSection* InSection, FbxNode* InFbxNode, const FMovieSceneSequenceTransform& RootToLocalTransform);
 
 	/** Export an FBX from the given control rig section. Optionally remapping and filtering controls. */
-	bool ExportControlRigSection(const UMovieSceneSection* Section, const TArray<FControlRigFbxNodeMapping>& ChannelsMapping, const TArray<FName>& FilterControls, const FMovieSceneSequenceTransform& RootToLocalTransform);
+	UNREALED_API bool ExportControlRigSection(const UMovieSceneSection* Section, const TArray<FControlRigFbxNodeMapping>& ChannelsMapping, const TArray<FName>& FilterControls, const FMovieSceneSequenceTransform& RootToLocalTransform);
 	
 	/**
 	 * Exports the given level sequence track information into a FBX document.
 	 *
 	 * @return	true, if successful
 	 */
-	bool ExportLevelSequenceTracks(UMovieScene* MovieScene, IMovieScenePlayer* MovieScenePlayer, FMovieSceneSequenceIDRef InSequenceID, FbxNode* FbxActor, UObject* BoundObject, const TArray<UMovieSceneTrack*>& Tracks, const FMovieSceneSequenceTransform& RootToLocalTransform);
+	UNREALED_API bool ExportLevelSequenceTracks(UMovieScene* MovieScene, IMovieScenePlayer* MovieScenePlayer, FMovieSceneSequenceIDRef InSequenceID, FbxNode* FbxActor, UObject* BoundObject, const TArray<UMovieSceneTrack*>& Tracks, const FMovieSceneSequenceTransform& RootToLocalTransform);
 
 
 	/**
 	 * Exports the mesh and the actor information for a static mesh actor.
 	 */
-	virtual void ExportStaticMesh( AActor* Actor, UStaticMeshComponent* StaticMeshComponent, INodeNameAdapter& NodeNameAdapter );
+	UNREALED_API virtual void ExportStaticMesh( AActor* Actor, UStaticMeshComponent* StaticMeshComponent, INodeNameAdapter& NodeNameAdapter );
 
 	/**
 	 * Exports a static mesh
 	 * @param StaticMesh	The static mesh to export
 	 * @param MaterialOrder	Optional ordering of materials to set up correct material ID's across multiple meshes being export such as BSP surfaces which share common materials. Should be used sparingly
 	 */
-	virtual void ExportStaticMesh( UStaticMesh* StaticMesh, const TArray<FStaticMaterial>* MaterialOrder = nullptr );
+	UNREALED_API virtual void ExportStaticMesh( UStaticMesh* StaticMesh, const TArray<FStaticMaterial>* MaterialOrder = nullptr );
 
 	/**
 	 * Exports BSP
 	 * @param Model			 The model with BSP to export
 	 * @param bSelectedOnly  true to export only selected surfaces (or brushes)
 	 */
-	virtual void ExportBSP( UModel* Model, bool bSelectedOnly );
+	UNREALED_API virtual void ExportBSP( UModel* Model, bool bSelectedOnly );
 
 	/**
 	 * Exports a static mesh light map
 	 */
-	virtual void ExportStaticMeshLightMap( UStaticMesh* StaticMesh, int32 LODIndex, int32 UVChannel );
+	UNREALED_API virtual void ExportStaticMeshLightMap( UStaticMesh* StaticMesh, int32 LODIndex, int32 UVChannel );
 
 	/**
 	 * Exports a skeletal mesh
 	 */
-	virtual void ExportSkeletalMesh( USkeletalMesh* SkeletalMesh );
+	UNREALED_API virtual void ExportSkeletalMesh( USkeletalMesh* SkeletalMesh );
 
 	/**
 	 * Exports the mesh and the actor information for a skeletal mesh actor.
 	 */
-	virtual void ExportSkeletalMesh( AActor* Actor, USkeletalMeshComponent* SkeletalMeshComponent, INodeNameAdapter& NodeNameAdapter );
+	UNREALED_API virtual void ExportSkeletalMesh( AActor* Actor, USkeletalMeshComponent* SkeletalMeshComponent, INodeNameAdapter& NodeNameAdapter );
 	
 	/**
 	 * Exports the mesh and the actor information for a landscape actor.
 	 */
-	void ExportLandscape(ALandscapeProxy* Landscape, bool bSelectedOnly, INodeNameAdapter& NodeNameAdapter);
+	UNREALED_API void ExportLandscape(ALandscapeProxy* Landscape, bool bSelectedOnly, INodeNameAdapter& NodeNameAdapter);
 
 	/**
 	 * Exports a single UAnimSequence, and optionally a skeletal mesh
 	 */
-	FbxNode* ExportAnimSequence( const UAnimSequence* AnimSeq, const USkeletalMesh* SkelMesh, bool bExportSkelMesh, const TCHAR* MeshNames=nullptr, FbxNode* ActorRootNode=nullptr, const TArray<UMaterialInterface*>* OverrideMaterials = nullptr);
+	UNREALED_API FbxNode* ExportAnimSequence( const UAnimSequence* AnimSeq, const USkeletalMesh* SkelMesh, bool bExportSkelMesh, const TCHAR* MeshNames=nullptr, FbxNode* ActorRootNode=nullptr, const TArray<UMaterialInterface*>* OverrideMaterials = nullptr);
 
 	/** A node name adapter for a level sequence. */
-	class UNREALED_API FLevelSequenceNodeNameAdapter : public INodeNameAdapter
+	class FLevelSequenceNodeNameAdapter : public INodeNameAdapter
 	{
 	public:
-		FLevelSequenceNodeNameAdapter( UMovieScene* InMovieScene, IMovieScenePlayer* InMovieScenePlayer, FMovieSceneSequenceIDRef InSequenceID);
-		virtual FString GetActorNodeName(const AActor* InActor) override;
-		virtual void AddFbxNode(UObject* InObject, FbxNode* InFbxNode) override;
-		virtual FbxNode* GetFbxNode(UObject* InObject) override;
+		UNREALED_API FLevelSequenceNodeNameAdapter( UMovieScene* InMovieScene, IMovieScenePlayer* InMovieScenePlayer, FMovieSceneSequenceIDRef InSequenceID);
+		UNREALED_API virtual FString GetActorNodeName(const AActor* InActor) override;
+		UNREALED_API virtual void AddFbxNode(UObject* InObject, FbxNode* InFbxNode) override;
+		UNREALED_API virtual FbxNode* GetFbxNode(UObject* InObject) override;
 	private:
 		UMovieScene* MovieScene;
 		IMovieScenePlayer* MovieScenePlayer;
@@ -240,18 +240,18 @@ public:
 	};
 
 	/* Get a valid unique name from a name */
-	FString GetFbxObjectName(const FString &FbxObjectNode, INodeNameAdapter& NodeNameAdapter);
+	UNREALED_API FString GetFbxObjectName(const FString &FbxObjectNode, INodeNameAdapter& NodeNameAdapter);
 
 	/**
 	 * Exports the basic information about an actor and buffers it.
 	 * This function creates one FBX node for the actor with its placement.
 	 */
-	FbxNode* ExportActor(AActor* Actor, bool bExportComponents, INodeNameAdapter& NodeNameAdapter, bool bSaveAnimSeq = true);
+	UNREALED_API FbxNode* ExportActor(AActor* Actor, bool bExportComponents, INodeNameAdapter& NodeNameAdapter, bool bSaveAnimSeq = true);
 
 private:
-	FFbxExporter();
+	UNREALED_API FFbxExporter();
 
-	static TSharedPtr<FFbxExporter> StaticInstance;
+	static UNREALED_API TSharedPtr<FFbxExporter> StaticInstance;
 
 	FbxManager* SdkManager;
 	FbxScene* Scene;
@@ -269,10 +269,10 @@ private:
 	TMap<const UStaticMesh*, FbxMesh*> FbxCollisionMeshes;
 
 	/** The frames-per-second (FPS) used when baking transforms */
-	static const float BakeTransformsFPS;
+	static UNREALED_API const float BakeTransformsFPS;
 	
 	/** Whether or not to export vertices unwelded */
-	static bool bStaticMeshExportUnWeldedVerts;
+	static UNREALED_API bool bStaticMeshExportUnWeldedVerts;
 
 	TObjectPtr<UFbxExportOption> ExportOptionsUI;
 	TObjectPtr<UFbxExportOption> ExportOptionsOverride;
@@ -280,11 +280,11 @@ private:
 	/**
 	* Export Anim Track of the given SkeletalMeshComponent
 	*/
-	void ExportAnimTrack( IAnimTrackAdapter& AnimTrackAdapter, AActor* Actor, USkeletalMeshComponent* SkeletalMeshComponent, float SamplingRate );
+	UNREALED_API void ExportAnimTrack( IAnimTrackAdapter& AnimTrackAdapter, AActor* Actor, USkeletalMeshComponent* SkeletalMeshComponent, float SamplingRate );
 
-	void ExportModel(UModel* Model, FbxNode* Node, const char* Name);
+	UNREALED_API void ExportModel(UModel* Model, FbxNode* Node, const char* Name);
 
-	FbxNode* ExportCollisionMesh(const UStaticMesh* StaticMesh, const TCHAR* MeshName, FbxNode* ParentActor);
+	UNREALED_API FbxNode* ExportCollisionMesh(const UStaticMesh* StaticMesh, const TCHAR* MeshName, FbxNode* ParentActor);
 
 	/**
 	 * Exports a static mesh
@@ -297,9 +297,9 @@ private:
 	 * @param MaterialOrderOverride	Optional ordering of materials to set up correct material ID's across multiple meshes being export such as BSP surfaces which share common materials. Should be used sparingly
 	 * @param OverrideMaterials	Optional array of materials to be used instead of the static mesh materials. Used for material overrides in static mesh components.
 	 */
-	FbxNode* ExportStaticMeshToFbx(const UStaticMesh* StaticMesh, int32 ExportLOD, const TCHAR* MeshName, FbxNode* FbxActor, int32 LightmapUVChannel = -1, const FColorVertexBuffer* ColorBuffer = nullptr, const TArray<FStaticMaterial>* MaterialOrderOverride = nullptr, const TArray<UMaterialInterface*>* OverrideMaterials = nullptr);
+	UNREALED_API FbxNode* ExportStaticMeshToFbx(const UStaticMesh* StaticMesh, int32 ExportLOD, const TCHAR* MeshName, FbxNode* FbxActor, int32 LightmapUVChannel = -1, const FColorVertexBuffer* ColorBuffer = nullptr, const TArray<FStaticMaterial>* MaterialOrderOverride = nullptr, const TArray<UMaterialInterface*>* OverrideMaterials = nullptr);
 
-	bool ExportStaticMeshFromMeshDescription(FbxMesh* Mesh
+	UNREALED_API bool ExportStaticMeshFromMeshDescription(FbxMesh* Mesh
 		, const UStaticMesh* StaticMesh
 		, const FMeshDescription* MeshDescription
 		, FbxNode* FbxActor
@@ -307,7 +307,7 @@ private:
 		, const TArray<FStaticMaterial>* MaterialOrderOverride
 		, const TArray<UMaterialInterface*>* OverrideMaterials);
 
-	bool ExportStaticMeshFromRenderData(FbxMesh* Mesh
+	UNREALED_API bool ExportStaticMeshFromRenderData(FbxMesh* Mesh
 		, const UStaticMesh* StaticMesh
 		, const FStaticMeshLODResources& RenderMesh
 		, FbxNode* FbxActor
@@ -322,7 +322,7 @@ private:
 	 * @param MeshName		The name of the mesh for the FBX file
 	 * @param FbxActor		The fbx node representing the mesh
 	 */
-	void ExportSplineMeshToFbx(const USplineMeshComponent* SplineMeshComp, const TCHAR* MeshName, FbxNode* FbxActor);
+	UNREALED_API void ExportSplineMeshToFbx(const USplineMeshComponent* SplineMeshComp, const TCHAR* MeshName, FbxNode* FbxActor);
 
 	/**
 	 * Exports an instanced mesh
@@ -330,7 +330,7 @@ private:
 	 * @param MeshName		The name of the mesh for the FBX file
 	 * @param FbxActor		The fbx node representing the mesh
 	 */
-	void ExportInstancedMeshToFbx(const UInstancedStaticMeshComponent* InstancedMeshComp, const TCHAR* MeshName, FbxNode* FbxActor);
+	UNREALED_API void ExportInstancedMeshToFbx(const UInstancedStaticMeshComponent* InstancedMeshComp, const TCHAR* MeshName, FbxNode* FbxActor);
 
 	/**
 	* Exports a landscape
@@ -338,7 +338,7 @@ private:
 	* @param MeshName		The name of the mesh for the FBX file
 	* @param FbxActor		The fbx node representing the mesh
 	*/
-	void ExportLandscapeToFbx(ALandscapeProxy* Landscape, const TCHAR* MeshName, FbxNode* FbxActor, bool bSelectedOnly);
+	UNREALED_API void ExportLandscapeToFbx(ALandscapeProxy* Landscape, const TCHAR* MeshName, FbxNode* FbxActor, bool bSelectedOnly);
 
 	/**
 	* Fill an fbx light with from a unreal light component
@@ -346,7 +346,7 @@ private:
 	* @param Camera				Fbx light object
 	* @param CameraComponent	Unreal light component
 	*/
-	void FillFbxLightAttribute(FbxLight* Light, FbxNode* FbxParentNode, ULightComponent* BaseLight);
+	UNREALED_API void FillFbxLightAttribute(FbxLight* Light, FbxNode* FbxParentNode, ULightComponent* BaseLight);
 
 	/**
 	* Fill an fbx camera with from a unreal camera component
@@ -354,13 +354,13 @@ private:
 	* @param Camera				Fbx camera object
 	* @param CameraComponent	Unreal camera component
 	*/
-	void FillFbxCameraAttribute(FbxNode* ParentNode, FbxCamera* Camera, UCameraComponent *CameraComponent);
+	UNREALED_API void FillFbxCameraAttribute(FbxNode* ParentNode, FbxCamera* Camera, UCameraComponent *CameraComponent);
 
 	/**
 	 * Adds FBX skeleton nodes to the FbxScene based on the skeleton in the given USkeletalMesh, and fills
 	 * the given array with the nodes created
 	 */
-	FbxNode* CreateSkeleton(const USkeletalMesh* SkelMesh, TArray<FbxNode*>& BoneNodes);
+	UNREALED_API FbxNode* CreateSkeleton(const USkeletalMesh* SkelMesh, TArray<FbxNode*>& BoneNodes);
 
 	/**
 	 * Adds an Fbx Mesh to the FBX scene based on the data in the given FSkeletalMeshLODModel
@@ -370,77 +370,77 @@ private:
 	 * @param AnimSeq			If an AnimSeq is provided and are exporting MorphTarget, the MorphTarget animation will be exported as well.
 	 * @param OverrideMaterials Optional array of materials to be used instead of the skeletal mesh materials. Used for material overrides in skeletal mesh components.
 	 */
-	FbxNode* CreateMesh(const USkeletalMesh* SkelMesh, const TCHAR* MeshName, int32 LODIndex, const UAnimSequence* AnimSeq = nullptr, const TArray<UMaterialInterface*>* OverrideMaterials = nullptr);
+	UNREALED_API FbxNode* CreateMesh(const USkeletalMesh* SkelMesh, const TCHAR* MeshName, int32 LODIndex, const UAnimSequence* AnimSeq = nullptr, const TArray<UMaterialInterface*>* OverrideMaterials = nullptr);
 
 	/**
 	 * Adds Fbx Clusters necessary to skin a skeletal mesh to the bones in the BoneNodes list
 	 */
-	void BindMeshToSkeleton(const USkeletalMesh* SkelMesh, FbxNode* MeshRootNode, TArray<FbxNode*>& BoneNodes, int32 LODIndex);
+	UNREALED_API void BindMeshToSkeleton(const USkeletalMesh* SkelMesh, FbxNode* MeshRootNode, TArray<FbxNode*>& BoneNodes, int32 LODIndex);
 
 	/**
 	 * Add a bind pose to the scene based on the FbxMesh and skinning settings of the given node
 	 */
-	void CreateBindPose(FbxNode* MeshRootNode);
+	UNREALED_API void CreateBindPose(FbxNode* MeshRootNode);
 
 	/**
 	 * Add the given skeletal mesh to the Fbx scene in preparation for exporting.  Makes all new nodes a child of the given node
 	 */
-	FbxNode* ExportSkeletalMeshToFbx(const USkeletalMesh* SkelMesh, const UAnimSequence* AnimSeq, const TCHAR* MeshName, FbxNode* ActorRootNode, const TArray<UMaterialInterface*>* OverrideMaterials = nullptr);
+	UNREALED_API FbxNode* ExportSkeletalMeshToFbx(const USkeletalMesh* SkelMesh, const UAnimSequence* AnimSeq, const TCHAR* MeshName, FbxNode* ActorRootNode, const TArray<UMaterialInterface*>* OverrideMaterials = nullptr);
 
 	/** Export SkeletalMeshComponent */
-	void ExportSkeletalMeshComponent(USkeletalMeshComponent* SkelMeshComp, const TCHAR* MeshName, FbxNode* ActorRootNode, INodeNameAdapter& NodeNameAdapter, bool bSaveAnimSeq = true);
+	UNREALED_API void ExportSkeletalMeshComponent(USkeletalMeshComponent* SkelMeshComp, const TCHAR* MeshName, FbxNode* ActorRootNode, INodeNameAdapter& NodeNameAdapter, bool bSaveAnimSeq = true);
 
 	/** Initializing the AnimStack playrate from the AnimSequence */
-	bool SetupAnimStack(const UAnimSequence* AnimSeq);
+	UNREALED_API bool SetupAnimStack(const UAnimSequence* AnimSeq);
 
 	/**
 	 * Add the given animation sequence as rotation and translation tracks to the given list of bone nodes
 	 */
-	void ExportAnimSequenceToFbx(const UAnimSequence* AnimSeq, const USkeletalMesh* SkelMesh, TArray<FbxNode*>& BoneNodes, FbxAnimLayer* AnimLayer, FFrameTime StartFrameTime, FFrameTime EndFrameTime, float FrameRateScale, float StartTime);
+	UNREALED_API void ExportAnimSequenceToFbx(const UAnimSequence* AnimSeq, const USkeletalMesh* SkelMesh, TArray<FbxNode*>& BoneNodes, FbxAnimLayer* AnimLayer, FFrameTime StartFrameTime, FFrameTime EndFrameTime, float FrameRateScale, float StartTime);
 
 	UE_DEPRECATED(5.1, "ExportAnimSequenceToFbx is deprecated, use different signature")
-	void ExportAnimSequenceToFbx(const UAnimSequence* AnimSeq, const USkeletalMesh* SkelMesh, TArray<FbxNode*>& BoneNodes, FbxAnimLayer* AnimLayer,
+	UNREALED_API void ExportAnimSequenceToFbx(const UAnimSequence* AnimSeq, const USkeletalMesh* SkelMesh, TArray<FbxNode*>& BoneNodes, FbxAnimLayer* AnimLayer,
 		float AnimStartOffset, float AnimEndOffset, float AnimPlayRate, float StartTime);
 
 	/**
 	 * Add the custom Curve data to the FbxAnimCurves passed in parameter by matching their name to the skeletal mesh custom curves.
 	 */
-	void ExportCustomAnimCurvesToFbx(const TMap<FName, FbxAnimCurve*>& CustomCurves, const UAnimSequence* AnimSeq,
+	UNREALED_API void ExportCustomAnimCurvesToFbx(const TMap<FName, FbxAnimCurve*>& CustomCurves, const UAnimSequence* AnimSeq,
 		FFrameTime AnimStartOffset, FFrameTime AnimEndOffset, float FrameRateScale, float StartTime, float ValueScale = 1.f);
 	
 	UE_DEPRECATED(5.1, "ExportCustomAnimCurvesToFbx is deprecated, use different signature")
-	void ExportCustomAnimCurvesToFbx(const TMap<FName, FbxAnimCurve*>& CustomCurves, const UAnimSequence* AnimSeq,
+	UNREALED_API void ExportCustomAnimCurvesToFbx(const TMap<FName, FbxAnimCurve*>& CustomCurves, const UAnimSequence* AnimSeq,
 		float AnimStartOffset, float AnimEndOffset, float AnimPlayRate, float StartTime, float ValueScale = 1.f);
 
 	/**
 	 * Used internally to reuse the AnimSequence iteration code when exporting various kind of curves.
 	 */
-	void IterateInsideAnimSequence(const UAnimSequence* AnimSeq, FFrameTime StartFrameTime, FFrameTime EndFrameTime, float FrameRateScale, float StartTime, TFunctionRef<void(double, FbxTime, bool)> IterationLambda);
+	UNREALED_API void IterateInsideAnimSequence(const UAnimSequence* AnimSeq, FFrameTime StartFrameTime, FFrameTime EndFrameTime, float FrameRateScale, float StartTime, TFunctionRef<void(double, FbxTime, bool)> IterationLambda);
 
 	UE_DEPRECATED(5.1, "IterateInsideAnimSequence is deprecated, use different signature")
-	void IterateInsideAnimSequence(const UAnimSequence* AnimSeq, float AnimStartOffset, float AnimEndOffset, float AnimPlayRate, float StartTime, TFunctionRef<void(float, FbxTime, bool)> IterationLambda);
+	UNREALED_API void IterateInsideAnimSequence(const UAnimSequence* AnimSeq, float AnimStartOffset, float AnimEndOffset, float AnimPlayRate, float StartTime, TFunctionRef<void(float, FbxTime, bool)> IterationLambda);
 
 	/** 
 	 * The curve code doesn't differentiate between angles and other data, so an interpolation from 179 to -179
 	 * will cause the bone to rotate all the way around through 0 degrees.  So here we make a second pass over the 
 	 * rotation tracks to convert the angles into a more interpolation-friendly format.  
 	 */
-	void CorrectAnimTrackInterpolation( TArray<FbxNode*>& BoneNodes, FbxAnimLayer* AnimLayer );
+	UNREALED_API void CorrectAnimTrackInterpolation( TArray<FbxNode*>& BoneNodes, FbxAnimLayer* AnimLayer );
 
 	/**
 	 * Exports a level sequence 3D transform track into the FBX animation stack.
 	 */
-	void ExportLevelSequence3DTransformTrack(FbxNode* FbxActor, IMovieScenePlayer* MovieScenePlayer, FMovieSceneSequenceIDRef InSequenceID, UMovieScene3DTransformTrack& TransformTrack, UObject* BoundObject, const TRange<FFrameNumber>& InPlaybackRange, const FMovieSceneSequenceTransform& RootToLocalTransform);
+	UNREALED_API void ExportLevelSequence3DTransformTrack(FbxNode* FbxActor, IMovieScenePlayer* MovieScenePlayer, FMovieSceneSequenceIDRef InSequenceID, UMovieScene3DTransformTrack& TransformTrack, UObject* BoundObject, const TRange<FFrameNumber>& InPlaybackRange, const FMovieSceneSequenceTransform& RootToLocalTransform);
 
 	/**
 	 * Exports a level sequence 3D transform track that's getting baked (sample all sections) onto the FBX animation stack.
 	 */
-	void ExportLevelSequenceBaked3DTransformTrack(IAnimTrackAdapter& AnimTrackAdapter, FbxNode* FbxActor, IMovieScenePlayer* MovieScenePlayer, FMovieSceneSequenceIDRef InSequenceID, TArray<TWeakObjectPtr<UMovieScene3DTransformTrack> > TransformTracks, UObject* BoundObject, const TRange<FFrameNumber>& InPlaybackRange, const FMovieSceneSequenceTransform& RootToLocalTransform);
+	UNREALED_API void ExportLevelSequenceBaked3DTransformTrack(IAnimTrackAdapter& AnimTrackAdapter, FbxNode* FbxActor, IMovieScenePlayer* MovieScenePlayer, FMovieSceneSequenceIDRef InSequenceID, TArray<TWeakObjectPtr<UMovieScene3DTransformTrack> > TransformTracks, UObject* BoundObject, const TRange<FFrameNumber>& InPlaybackRange, const FMovieSceneSequenceTransform& RootToLocalTransform);
 
 	/** 
 	 * Exports a level sequence property track into the FBX animation stack. 
 	 */
-	void ExportLevelSequenceTrackChannels( FbxNode* FbxActor, UMovieSceneTrack& Track, const TRange<FFrameNumber>& InPlaybackRange, const FMovieSceneSequenceTransform& RootToLocalTransform);
+	UNREALED_API void ExportLevelSequenceTrackChannels( FbxNode* FbxActor, UMovieSceneTrack& Track, const TRange<FFrameNumber>& InPlaybackRange, const FMovieSceneSequenceTransform& RootToLocalTransform);
 
 	/** Defines value export modes for the EportRichCurveToFbxCurve method. */
 	enum class ERichCurveValueMode
@@ -456,47 +456,47 @@ private:
 	void ExportBezierChannelToFbxCurve(FbxAnimCurve& InFbxCurve, const ChannelType& InChannel, FFrameRate TickResolution, ERichCurveValueMode ValueMode, bool bNegative, const FMovieSceneSequenceTransform& RootToLocalTransform);
 
 	/** Exports a movie scene float channel to an fbx animation curve. */
-	void ExportChannelToFbxCurve(FbxAnimCurve& InFbxCurve, const FMovieSceneFloatChannel& InChannel, FFrameRate TickResolution, ERichCurveValueMode ValueMode = ERichCurveValueMode::Default, bool bNegative = false, const FMovieSceneSequenceTransform& RootToLocalTransform = FMovieSceneSequenceTransform());
+	UNREALED_API void ExportChannelToFbxCurve(FbxAnimCurve& InFbxCurve, const FMovieSceneFloatChannel& InChannel, FFrameRate TickResolution, ERichCurveValueMode ValueMode = ERichCurveValueMode::Default, bool bNegative = false, const FMovieSceneSequenceTransform& RootToLocalTransform = FMovieSceneSequenceTransform());
 
 	/** Exports a movie scene double channel to an fbx animation curve. */
-	void ExportChannelToFbxCurve(FbxAnimCurve& InFbxCurve, const FMovieSceneDoubleChannel& InChannel, FFrameRate TickResolution, ERichCurveValueMode ValueMode = ERichCurveValueMode::Default, bool bNegative = false, const FMovieSceneSequenceTransform& RootToLocalTransform = FMovieSceneSequenceTransform());
+	UNREALED_API void ExportChannelToFbxCurve(FbxAnimCurve& InFbxCurve, const FMovieSceneDoubleChannel& InChannel, FFrameRate TickResolution, ERichCurveValueMode ValueMode = ERichCurveValueMode::Default, bool bNegative = false, const FMovieSceneSequenceTransform& RootToLocalTransform = FMovieSceneSequenceTransform());
 
 	/** Exports a movie scene integer channel to an fbx animation curve. */
-	void ExportChannelToFbxCurve(FbxAnimCurve& InFbxCurve, const FMovieSceneIntegerChannel& InChannel, FFrameRate TickResolution, const FMovieSceneSequenceTransform& RootToLocalTransform = FMovieSceneSequenceTransform());
+	UNREALED_API void ExportChannelToFbxCurve(FbxAnimCurve& InFbxCurve, const FMovieSceneIntegerChannel& InChannel, FFrameRate TickResolution, const FMovieSceneSequenceTransform& RootToLocalTransform = FMovieSceneSequenceTransform());
 
-	void ExportChannelToFbxCurve(FbxAnimCurve& InFbxCurve, const FMovieSceneBoolChannel& InChannel, FFrameRate TickResolution, const FMovieSceneSequenceTransform& RootToLocalTransform);
+	UNREALED_API void ExportChannelToFbxCurve(FbxAnimCurve& InFbxCurve, const FMovieSceneBoolChannel& InChannel, FFrameRate TickResolution, const FMovieSceneSequenceTransform& RootToLocalTransform);
 	
-	void ExportChannelToFbxCurve(FbxAnimCurve& InFbxCurve, const FMovieSceneByteChannel& InChannel, FFrameRate TickResolution, const FMovieSceneSequenceTransform& RootToLocalTransform);
+	UNREALED_API void ExportChannelToFbxCurve(FbxAnimCurve& InFbxCurve, const FMovieSceneByteChannel& InChannel, FFrameRate TickResolution, const FMovieSceneSequenceTransform& RootToLocalTransform);
 
 	template <class ChannelType, typename T>
 	void ExportConstantChannelToFbxCurve(FbxAnimCurve& InFbxCurve, const ChannelType& InChannel, FFrameRate TickResolution, const FMovieSceneSequenceTransform& RootToLocalTransform);
 
-	void ExportTransformChannelsToFbxCurve(FbxNode* InFbxNode, FMovieSceneFloatChannel* ChannelX, FMovieSceneFloatChannel* ChannelY, FMovieSceneFloatChannel* ChannelZ, int TmPropertyIndex, const UMovieSceneTrack* Track, const FMovieSceneSequenceTransform& RootToLocalTransform);
+	UNREALED_API void ExportTransformChannelsToFbxCurve(FbxNode* InFbxNode, FMovieSceneFloatChannel* ChannelX, FMovieSceneFloatChannel* ChannelY, FMovieSceneFloatChannel* ChannelZ, int TmPropertyIndex, const UMovieSceneTrack* Track, const FMovieSceneSequenceTransform& RootToLocalTransform);
 
 	/**
 	 * Finds the given actor in the already-exported list of structures
 	 * @return FbxNode* the FBX node created from the UE actor
 	 */
-	FbxNode* FindActor(AActor* Actor, INodeNameAdapter* NodeNameAdapter = nullptr);
+	UNREALED_API FbxNode* FindActor(AActor* Actor, INodeNameAdapter* NodeNameAdapter = nullptr);
 
 	/** Create fbx node with the given name */
-	FbxNode* CreateNode(const FString& NodeName);
+	UNREALED_API FbxNode* CreateNode(const FString& NodeName);
 
 	/**
 	 * Find bone array of FbxNOdes of the given skeletalmeshcomponent  
 	 */
-	bool FindSkeleton(USkeletalMeshComponent* SkelComp, TArray<FbxNode*>& BoneNodes, INodeNameAdapter* NodeNameAdapter = nullptr);
+	UNREALED_API bool FindSkeleton(USkeletalMeshComponent* SkelComp, TArray<FbxNode*>& BoneNodes, INodeNameAdapter* NodeNameAdapter = nullptr);
 
 	/** recursively get skeleton */
-	void GetSkeleton(FbxNode* RootNode, TArray<FbxNode*>& BoneNodes);
+	UNREALED_API void GetSkeleton(FbxNode* RootNode, TArray<FbxNode*>& BoneNodes);
 
-	bool FillFbxTextureProperty(const char *PropertyName, const FExpressionInput& MaterialInput, FbxSurfaceMaterial* FbxMaterial);
+	UNREALED_API bool FillFbxTextureProperty(const char *PropertyName, const FExpressionInput& MaterialInput, FbxSurfaceMaterial* FbxMaterial);
 	/**
 	 * Exports the profile_COMMON information for a material.
 	 */
-	FbxSurfaceMaterial* ExportMaterial(UMaterialInterface* Material);
+	UNREALED_API FbxSurfaceMaterial* ExportMaterial(UMaterialInterface* Material);
 	
-	FbxSurfaceMaterial* CreateDefaultMaterial();
+	UNREALED_API FbxSurfaceMaterial* CreateDefaultMaterial();
 	
 	/**
 	 * Create user property in Fbx Node.
@@ -511,11 +511,11 @@ private:
 	void CreateAnimatableUserProperty(FbxNode* Node, T Value, const char* Name, const char* Label, FbxDataType DataType = FbxFloatDT);
 
 	/** Exports all the object's FBX metadata to the FBX node */
-	void ExportObjectMetadata(const UObject* ObjectToExport, FbxNode* Node);
+	UNREALED_API void ExportObjectMetadata(const UObject* ObjectToExport, FbxNode* Node);
 
 public:
 	/** Returns currently active FBX export options. Automation or UI dialog based options. */
-	UFbxExportOption* GetExportOptions();
+	UNREALED_API UFbxExportOption* GetExportOptions();
 
 	bool bSceneGlobalTimeLineSet = false;
 };
diff --git a/Engine/Source/Editor/UnrealEd/Private/Tools/LegacyEdMode.h b/Engine/Source/Editor/UnrealEd/Private/Tools/LegacyEdMode.h
--- a/Engine/Source/Editor/UnrealEd/Private/Tools/LegacyEdMode.h
+++ b/Engine/Source/Editor/UnrealEd/Private/Tools/LegacyEdMode.h
@@ -18,109 +18,109 @@ class FEditorViewportClient;
 class FViewport;
 class UPrimitiveComponent;
 
-UCLASS()
-class UNREALED_API ULegacyEdModeWrapper final : public UEdMode, public ILegacyEdModeWidgetInterface, public ILegacyEdModeToolInterface, public ILegacyEdModeSelectInterface, public ILegacyEdModeDrawHelperInterface, public ILegacyEdModeViewportInterface
+UCLASS(MinimalAPI)
+class ULegacyEdModeWrapper final : public UEdMode, public ILegacyEdModeWidgetInterface, public ILegacyEdModeToolInterface, public ILegacyEdModeSelectInterface, public ILegacyEdModeDrawHelperInterface, public ILegacyEdModeViewportInterface
 {
 	GENERATED_BODY()
 
 public:
 	ULegacyEdModeWrapper();
 
-	bool CreateLegacyMode(FEditorModeID ModeID, FEditorModeTools& ModeManager);
+	UNREALED_API bool CreateLegacyMode(FEditorModeID ModeID, FEditorModeTools& ModeManager);
 
 	// Begin UEdMode overrides
-	virtual void Initialize() override;
-	virtual bool IsSelectionAllowed(AActor* InActor, bool bInSelected) const override;
-	virtual bool Select(AActor* InActor, bool bInSelected) override;
-	virtual bool ProcessEditDuplicate() override;
-	virtual bool ProcessEditDelete() override;
-	virtual bool ProcessEditCut() override;
-	virtual bool ProcessEditCopy() override;
-	virtual bool ProcessEditPaste() override;
-	virtual EEditAction::Type GetActionEditDuplicate() override;
-	virtual EEditAction::Type GetActionEditDelete() override;
-	virtual EEditAction::Type GetActionEditCut() override;
-	virtual EEditAction::Type GetActionEditCopy() override;
-	virtual EEditAction::Type GetActionEditPaste() override;
-	virtual bool IsSnapRotationEnabled() override;
-	virtual bool SnapRotatorToGridOverride(FRotator& Rotation) override;
-	virtual void ActorsDuplicatedNotify(TArray<AActor*>& PreDuplicateSelection, TArray<AActor*>& PostDuplicateSelection, bool bOffsetLocations) override;
-	virtual void ActorMoveNotify() override;
-	virtual void ActorSelectionChangeNotify() override;
-	virtual void ActorPropChangeNotify() override;
-	virtual void UpdateInternalData() override;
-	virtual void MapChangeNotify() override;
-	virtual void SelectNone() override;
-	virtual bool GetPivotForOrbit(FVector& OutPivot) const override;
-	virtual void PostUndo() override;
-	virtual bool DisallowMouseDeltaTracking() const override;
-	virtual bool GetCursor(EMouseCursor::Type& OutCursor) const override;
-	virtual bool GetOverrideCursorVisibility(bool& bWantsOverride, bool& bHardwareCursorVisible, bool bSoftwareCursorVisible) const;
-	virtual bool CanAutoSave() const override;
-	virtual bool IsCompatibleWith(FEditorModeID OtherModeID) const override;
-	virtual bool ComputeBoundingBoxForViewportFocus(AActor* Actor, UPrimitiveComponent* PrimitiveComponent, FBox& InOutBox) const override;
-	virtual bool AllowsViewportDragTool() const override;
-	virtual bool UsesToolkits() const override;
+	UNREALED_API virtual void Initialize() override;
+	UNREALED_API virtual bool IsSelectionAllowed(AActor* InActor, bool bInSelected) const override;
+	UNREALED_API virtual bool Select(AActor* InActor, bool bInSelected) override;
+	UNREALED_API virtual bool ProcessEditDuplicate() override;
+	UNREALED_API virtual bool ProcessEditDelete() override;
+	UNREALED_API virtual bool ProcessEditCut() override;
+	UNREALED_API virtual bool ProcessEditCopy() override;
+	UNREALED_API virtual bool ProcessEditPaste() override;
+	UNREALED_API virtual EEditAction::Type GetActionEditDuplicate() override;
+	UNREALED_API virtual EEditAction::Type GetActionEditDelete() override;
+	UNREALED_API virtual EEditAction::Type GetActionEditCut() override;
+	UNREALED_API virtual EEditAction::Type GetActionEditCopy() override;
+	UNREALED_API virtual EEditAction::Type GetActionEditPaste() override;
+	UNREALED_API virtual bool IsSnapRotationEnabled() override;
+	UNREALED_API virtual bool SnapRotatorToGridOverride(FRotator& Rotation) override;
+	UNREALED_API virtual void ActorsDuplicatedNotify(TArray<AActor*>& PreDuplicateSelection, TArray<AActor*>& PostDuplicateSelection, bool bOffsetLocations) override;
+	UNREALED_API virtual void ActorMoveNotify() override;
+	UNREALED_API virtual void ActorSelectionChangeNotify() override;
+	UNREALED_API virtual void ActorPropChangeNotify() override;
+	UNREALED_API virtual void UpdateInternalData() override;
+	UNREALED_API virtual void MapChangeNotify() override;
+	UNREALED_API virtual void SelectNone() override;
+	UNREALED_API virtual bool GetPivotForOrbit(FVector& OutPivot) const override;
+	UNREALED_API virtual void PostUndo() override;
+	UNREALED_API virtual bool DisallowMouseDeltaTracking() const override;
+	UNREALED_API virtual bool GetCursor(EMouseCursor::Type& OutCursor) const override;
+	UNREALED_API virtual bool GetOverrideCursorVisibility(bool& bWantsOverride, bool& bHardwareCursorVisible, bool bSoftwareCursorVisible) const;
+	UNREALED_API virtual bool CanAutoSave() const override;
+	UNREALED_API virtual bool IsCompatibleWith(FEditorModeID OtherModeID) const override;
+	UNREALED_API virtual bool ComputeBoundingBoxForViewportFocus(AActor* Actor, UPrimitiveComponent* PrimitiveComponent, FBox& InOutBox) const override;
+	UNREALED_API virtual bool AllowsViewportDragTool() const override;
+	UNREALED_API virtual bool UsesToolkits() const override;
 
-	virtual bool ProcessCapturedMouseMoves(FEditorViewportClient* InViewportClient, FViewport* InViewport, const TArrayView<FIntPoint>& CapturedMouseMoves) override;
-	virtual bool InputKey(FEditorViewportClient* ViewportClient, FViewport* Viewport, FKey Key, EInputEvent Event) override;
-	virtual bool InputAxis(FEditorViewportClient* InViewportClient, FViewport* Viewport, int32 ControllerId, FKey Key, float Delta, float DeltaTime) override;
-	virtual bool InputDelta(FEditorViewportClient* InViewportClient, FViewport* InViewport, FVector& InDrag, FRotator& InRot, FVector& InScale) override;
-	virtual bool StartTracking(FEditorViewportClient* InViewportClient, FViewport* InViewport) override;
-	virtual bool EndTracking(FEditorViewportClient* InViewportClient, FViewport* InViewport) override;
-	virtual bool HandleClick(FEditorViewportClient* InViewportClient, HHitProxy* HitProxy, const FViewportClick& Click) override;
-	virtual void Tick(FEditorViewportClient* ViewportClient, float DeltaTime) override;
-	virtual bool CapturedMouseMove(FEditorViewportClient* InViewportClient, FViewport* InViewport, int32 InMouseX, int32 InMouseY) override;
-	virtual bool MouseEnter(FEditorViewportClient* ViewportClient, FViewport* Viewport, int32 x, int32 y) override;
-	virtual bool MouseLeave(FEditorViewportClient* ViewportClient, FViewport* Viewport) override;
-	virtual bool MouseMove(FEditorViewportClient* ViewportClient, FViewport* Viewport, int32 x, int32 y) override;
-	virtual bool ReceivedFocus(FEditorViewportClient* ViewportClient, FViewport* Viewport) override;
-	virtual bool LostFocus(FEditorViewportClient* ViewportClient, FViewport* Viewport) override;
-	virtual void Render(const FSceneView* View, FViewport* Viewport, FPrimitiveDrawInterface* PDI) override;
-	virtual void DrawHUD(FEditorViewportClient* ViewportClient, FViewport* Viewport, const FSceneView* View, FCanvas* Canvas) override;
-	virtual bool PreConvertMouseMovement(FEditorViewportClient* InViewportClient) override;
-	virtual bool PostConvertMouseMovement(FEditorViewportClient* InViewportClient) override;
+	UNREALED_API virtual bool ProcessCapturedMouseMoves(FEditorViewportClient* InViewportClient, FViewport* InViewport, const TArrayView<FIntPoint>& CapturedMouseMoves) override;
+	UNREALED_API virtual bool InputKey(FEditorViewportClient* ViewportClient, FViewport* Viewport, FKey Key, EInputEvent Event) override;
+	UNREALED_API virtual bool InputAxis(FEditorViewportClient* InViewportClient, FViewport* Viewport, int32 ControllerId, FKey Key, float Delta, float DeltaTime) override;
+	UNREALED_API virtual bool InputDelta(FEditorViewportClient* InViewportClient, FViewport* InViewport, FVector& InDrag, FRotator& InRot, FVector& InScale) override;
+	UNREALED_API virtual bool StartTracking(FEditorViewportClient* InViewportClient, FViewport* InViewport) override;
+	UNREALED_API virtual bool EndTracking(FEditorViewportClient* InViewportClient, FViewport* InViewport) override;
+	UNREALED_API virtual bool HandleClick(FEditorViewportClient* InViewportClient, HHitProxy* HitProxy, const FViewportClick& Click) override;
+	UNREALED_API virtual void Tick(FEditorViewportClient* ViewportClient, float DeltaTime) override;
+	UNREALED_API virtual bool CapturedMouseMove(FEditorViewportClient* InViewportClient, FViewport* InViewport, int32 InMouseX, int32 InMouseY) override;
+	UNREALED_API virtual bool MouseEnter(FEditorViewportClient* ViewportClient, FViewport* Viewport, int32 x, int32 y) override;
+	UNREALED_API virtual bool MouseLeave(FEditorViewportClient* ViewportClient, FViewport* Viewport) override;
+	UNREALED_API virtual bool MouseMove(FEditorViewportClient* ViewportClient, FViewport* Viewport, int32 x, int32 y) override;
+	UNREALED_API virtual bool ReceivedFocus(FEditorViewportClient* ViewportClient, FViewport* Viewport) override;
+	UNREALED_API virtual bool LostFocus(FEditorViewportClient* ViewportClient, FViewport* Viewport) override;
+	UNREALED_API virtual void Render(const FSceneView* View, FViewport* Viewport, FPrimitiveDrawInterface* PDI) override;
+	UNREALED_API virtual void DrawHUD(FEditorViewportClient* ViewportClient, FViewport* Viewport, const FSceneView* View, FCanvas* Canvas) override;
+	UNREALED_API virtual bool PreConvertMouseMovement(FEditorViewportClient* InViewportClient) override;
+	UNREALED_API virtual bool PostConvertMouseMovement(FEditorViewportClient* InViewportClient) override;
 
-	virtual bool ShouldDrawBrushWireframe(AActor* InActor) const override;
-	virtual void Enter() override;
-	virtual void Exit() override;
+	UNREALED_API virtual bool ShouldDrawBrushWireframe(AActor* InActor) const override;
+	UNREALED_API virtual void Enter() override;
+	UNREALED_API virtual void Exit() override;
 
-	virtual FEdMode* AsLegacyMode() override;
-	virtual UTexture2D* GetVertexTexture() override;
+	UNREALED_API virtual FEdMode* AsLegacyMode() override;
+	UNREALED_API virtual UTexture2D* GetVertexTexture() override;
 	// End UEdMode overrides
 
 	// ILegacyEdModeWidgetInterface
-	virtual bool AllowWidgetMove() override;
-	virtual bool CanCycleWidgetMode() const override;
-	virtual bool ShowModeWidgets() const override;
-	virtual EAxisList::Type GetWidgetAxisToDraw(UE::Widget::EWidgetMode InWidgetMode) const override;
-	virtual FVector GetWidgetLocation() const override;
-	virtual bool ShouldDrawWidget() const override;
-	virtual bool UsesTransformWidget() const override;
-	virtual bool UsesTransformWidget(UE::Widget::EWidgetMode CheckMode) const override;
-	virtual FVector GetWidgetNormalFromCurrentAxis(void* InData) override;
-	virtual bool BoxSelect(FBox& InBox, bool InSelect = true) override;
-	virtual bool FrustumSelect(const FConvexVolume& InFrustum, FEditorViewportClient* InViewportClient, bool InSelect = true) override;
-	virtual void SetCurrentWidgetAxis(EAxisList::Type InAxis) override;
-	virtual EAxisList::Type GetCurrentWidgetAxis() const override;
-	virtual bool UsesPropertyWidgets() const override;
-	virtual bool GetCustomDrawingCoordinateSystem(FMatrix& InMatrix, void* InData) override;
-	virtual bool GetCustomInputCoordinateSystem(FMatrix& InMatrix, void* InData) override;
+	UNREALED_API virtual bool AllowWidgetMove() override;
+	UNREALED_API virtual bool CanCycleWidgetMode() const override;
+	UNREALED_API virtual bool ShowModeWidgets() const override;
+	UNREALED_API virtual EAxisList::Type GetWidgetAxisToDraw(UE::Widget::EWidgetMode InWidgetMode) const override;
+	UNREALED_API virtual FVector GetWidgetLocation() const override;
+	UNREALED_API virtual bool ShouldDrawWidget() const override;
+	UNREALED_API virtual bool UsesTransformWidget() const override;
+	UNREALED_API virtual bool UsesTransformWidget(UE::Widget::EWidgetMode CheckMode) const override;
+	UNREALED_API virtual FVector GetWidgetNormalFromCurrentAxis(void* InData) override;
+	UNREALED_API virtual bool BoxSelect(FBox& InBox, bool InSelect = true) override;
+	UNREALED_API virtual bool FrustumSelect(const FConvexVolume& InFrustum, FEditorViewportClient* InViewportClient, bool InSelect = true) override;
+	UNREALED_API virtual void SetCurrentWidgetAxis(EAxisList::Type InAxis) override;
+	UNREALED_API virtual EAxisList::Type GetCurrentWidgetAxis() const override;
+	UNREALED_API virtual bool UsesPropertyWidgets() const override;
+	UNREALED_API virtual bool GetCustomDrawingCoordinateSystem(FMatrix& InMatrix, void* InData) override;
+	UNREALED_API virtual bool GetCustomInputCoordinateSystem(FMatrix& InMatrix, void* InData) override;
 	// End ILegacyEdModeWidgetInterface
 
 	// ILegacyEdModeToolInterface overrides
-	virtual void SetCurrentTool(EModeTools InID) override;
-	virtual void SetCurrentTool(FModeTool* InModeTool) override;
-	virtual FModeTool* FindTool(EModeTools InID) override;
-	virtual const TArray<FModeTool*>& GetTools() const override;
-	virtual FModeTool* GetCurrentTool() override;
-	virtual const FModeTool* GetCurrentTool() const override;
+	UNREALED_API virtual void SetCurrentTool(EModeTools InID) override;
+	UNREALED_API virtual void SetCurrentTool(FModeTool* InModeTool) override;
+	UNREALED_API virtual FModeTool* FindTool(EModeTools InID) override;
+	UNREALED_API virtual const TArray<FModeTool*>& GetTools() const override;
+	UNREALED_API virtual FModeTool* GetCurrentTool() override;
+	UNREALED_API virtual const FModeTool* GetCurrentTool() const override;
 	// End ILegacyEdModeToolInterface overrides
 
 	// Start FCommonDrawHelper overrides
-	virtual void Draw(const FSceneView* View, FPrimitiveDrawInterface* PDI) override;
+	UNREALED_API virtual void Draw(const FSceneView* View, FPrimitiveDrawInterface* PDI) override;
 	// End FCommonDrawHelper overrides
 
 private:
 	TSharedPtr<FEdMode> LegacyEditorMode;
-};
\ No newline at end of file
+};
diff --git a/Engine/Source/Editor/UnrealEd/Public/ActorGroupingUtils.h b/Engine/Source/Editor/UnrealEd/Public/ActorGroupingUtils.h
--- a/Engine/Source/Editor/UnrealEd/Public/ActorGroupingUtils.h
+++ b/Engine/Source/Editor/UnrealEd/Public/ActorGroupingUtils.h
@@ -9,8 +9,8 @@
 /**
  * Helper class for grouping actors in the level editor
  */
-UCLASS(transient)
-class UNREALED_API UActorGroupingUtils : public UObject
+UCLASS(transient, MinimalAPI)
+class UActorGroupingUtils : public UObject
 {
 	GENERATED_BODY()
 public:
@@ -18,66 +18,66 @@ public:
 	static bool IsGroupingActive() { return bGroupingActive; }
 
 	UFUNCTION(BlueprintCallable, Category = "Editor Scripting | Actor Grouping")
-	static void SetGroupingActive(bool bInGroupingActive);
+	static UNREALED_API void SetGroupingActive(bool bInGroupingActive);
 
 	/**
 	 * Convenience method for accessing grouping utils in a blueprint or script
 	 */
 	UFUNCTION(BlueprintCallable, Category = "Editor Scripting | Actor Grouping", DisplayName="Get Actor Grouping Utils")
-	static UActorGroupingUtils* Get();
+	static UNREALED_API UActorGroupingUtils* Get();
 
 	/**
 	 * Creates a new group from the current selection removing the actors from any existing groups they are already in
 	 */
 	UFUNCTION(BlueprintCallable, Category = "Editor Scripting | Actor Grouping")
-	virtual AGroupActor* GroupSelected();
+	UNREALED_API virtual AGroupActor* GroupSelected();
 
 	/**
 	 * Creates a new group from the provided list of actors removing the actors from any existing groups they are already in
 	 */
 	UFUNCTION(BlueprintCallable, Category = "Editor Scripting | Actor Grouping")
-	virtual AGroupActor* GroupActors(const TArray<AActor*>& ActorsToGroup);
+	UNREALED_API virtual AGroupActor* GroupActors(const TArray<AActor*>& ActorsToGroup);
 
 	/**
 	 * Disbands any groups in the current selection, does not attempt to maintain any hierarchy
 	 */
 	UFUNCTION(BlueprintCallable, Category = "Editor Scripting | Actor Grouping")
-	virtual void UngroupSelected();
+	UNREALED_API virtual void UngroupSelected();
 
 	/**
 	 * Disbands any groups that the provided actors belong to, does not attempt to maintain any hierarchy
 	 */
 	UFUNCTION(BlueprintCallable, Category = "Editor Scripting | Actor Grouping")
-	virtual void UngroupActors(const TArray<AActor*>& ActorsToUngroup);
+	UNREALED_API virtual void UngroupActors(const TArray<AActor*>& ActorsToUngroup);
 
 	/**
 	* Locks any groups in the current selection
 	*/
 	UFUNCTION(BlueprintCallable, Category = "Editor Scripting | Actor Grouping")
-	virtual void LockSelectedGroups();
+	UNREALED_API virtual void LockSelectedGroups();
 
 	/**
 	 * Unlocks any groups in the current selection
 	 */
 	UFUNCTION(BlueprintCallable, Category = "Editor Scripting | Actor Grouping")
-	virtual void UnlockSelectedGroups();
+	UNREALED_API virtual void UnlockSelectedGroups();
 
 	/**
 	 * Activates "Add to Group" mode which allows the user to select a group to append current selection
 	 */
 	UFUNCTION(BlueprintCallable, Category = "Editor Scripting | Actor Grouping")
-	virtual void AddSelectedToGroup();
+	UNREALED_API virtual void AddSelectedToGroup();
 
 	/**
 	 * Removes any groups or actors in the current selection from their immediate parent.
 	 * If all actors/subgroups are removed, the parent group will be destroyed.
 	 */
 	UFUNCTION(BlueprintCallable, Category = "Editor Scripting | Actor Grouping")
-	virtual void RemoveSelectedFromGroup();
+	UNREALED_API virtual void RemoveSelectedFromGroup();
 
 protected:
-	static bool bGroupingActive;
-	static FSoftClassPath ClassToUse;
+	static UNREALED_API bool bGroupingActive;
+	static UNREALED_API FSoftClassPath ClassToUse;
 };
 
 
diff --git a/Engine/Source/Editor/UnrealEd/Public/AnimationEditorUtils.h b/Engine/Source/Editor/UnrealEd/Public/AnimationEditorUtils.h
--- a/Engine/Source/Editor/UnrealEd/Public/AnimationEditorUtils.h
+++ b/Engine/Source/Editor/UnrealEd/Public/AnimationEditorUtils.h
@@ -88,7 +88,7 @@ struct FAnimationCompressionSelectionDialogConfig
 };
 
 /** Dialog to prompt user to select an animation compression settings asset. */
-class UNREALED_API SAnimationCompressionSelectionDialog : public SCompoundWidget
+class SAnimationCompressionSelectionDialog : public SCompoundWidget
 {
 public:
 	/** Called from the Dialog when an asset has been selected. */
@@ -98,13 +98,13 @@ public:
 
 	SLATE_END_ARGS()
 
-	SAnimationCompressionSelectionDialog();
-	virtual ~SAnimationCompressionSelectionDialog();
+	UNREALED_API SAnimationCompressionSelectionDialog();
+	UNREALED_API virtual ~SAnimationCompressionSelectionDialog();
 
-	virtual void Construct(const FArguments& InArgs, const FAnimationCompressionSelectionDialogConfig& InConfig);
+	UNREALED_API virtual void Construct(const FArguments& InArgs, const FAnimationCompressionSelectionDialogConfig& InConfig);
 
 	/** Sets the delegate handler for when an open operation is committed */
-	void SetOnAssetSelected(const FOnAssetSelected& InHandler);
+	UNREALED_API void SetOnAssetSelected(const FOnAssetSelected& InHandler);
 
 private:
 	void DoSelectAsset(const FAssetData& SelectedAsset);
diff --git a/Engine/Source/Editor/UnrealEd/Public/AssetDefinitionDefault.h b/Engine/Source/Editor/UnrealEd/Public/AssetDefinitionDefault.h
--- a/Engine/Source/Editor/UnrealEd/Public/AssetDefinitionDefault.h
+++ b/Engine/Source/Editor/UnrealEd/Public/AssetDefinitionDefault.h
@@ -7,15 +7,15 @@
 
 #include "AssetDefinitionDefault.generated.h"
 
-UCLASS(Abstract)
-class UNREALED_API UAssetDefinitionDefault : public UAssetDefinition
+UCLASS(Abstract, MinimalAPI)
+class UAssetDefinitionDefault : public UAssetDefinition
 {
 	GENERATED_BODY()
 
 public:
 	// UAssetDefinition Begin
-	virtual EAssetCommandResult OpenAssets(const FAssetOpenArgs& OpenArgs) const override;
-	virtual EAssetCommandResult PerformAssetDiff(const FAssetDiffArgs& DiffArgs) const override;
+	UNREALED_API virtual EAssetCommandResult OpenAssets(const FAssetOpenArgs& OpenArgs) const override;
+	UNREALED_API virtual EAssetCommandResult PerformAssetDiff(const FAssetDiffArgs& DiffArgs) const override;
 	// UAssetDefinition End
 };
 
diff --git a/Engine/Source/Editor/UnrealEd/Public/AssetDeleteModel.h b/Engine/Source/Editor/UnrealEd/Public/AssetDeleteModel.h
--- a/Engine/Source/Editor/UnrealEd/Public/AssetDeleteModel.h
+++ b/Engine/Source/Editor/UnrealEd/Public/AssetDeleteModel.h
@@ -82,7 +82,7 @@ private:
  * that must be performed against the GC for UObjects, and looking up references for assets through
  * the asset registry.
  */
-class UNREALED_API FAssetDeleteModel
+class FAssetDeleteModel
 {
 public:
 	/** States used to manage the async deletion process. */
@@ -103,13 +103,13 @@ public:
 public:
 
 	/** Constructor */
-	FAssetDeleteModel( const TArray<UObject*>& InObjectsToDelete );
+	UNREALED_API FAssetDeleteModel( const TArray<UObject*>& InObjectsToDelete );
 
 	/** Destructor */
-	~FAssetDeleteModel();
+	UNREALED_API ~FAssetDeleteModel();
 
 	/** Add an object to the list of pending deleted assets, this will invalidate the scanning state. */
-	void AddObjectToDelete(UObject* InObject);
+	UNREALED_API void AddObjectToDelete(UObject* InObject);
 
 	/** Returns the pending deleted assets. */
 	const TArray< TSharedPtr< FPendingDelete > >* GetPendingDeletedAssets() const { return &PendingDeletes; };
@@ -124,43 +124,43 @@ public:
 	const TSet< FName >& GetAssetReferences() const { return OnDiskReferences; };
 
 	/** Ticks the delete model which does a little work before returning so that we don't completely block when deleting a lot of things. */
-	void Tick( const float InDeltaTime );
+	UNREALED_API void Tick( const float InDeltaTime );
 
 	/** Returns true if the object is one of the pending deleted assets. */
-	bool IsObjectInPendingDeletes( const UObject* InObject ) const;
+	UNREALED_API bool IsObjectInPendingDeletes( const UObject* InObject ) const;
 
 	/** Returns true if the package is one of the pending deleted assets. */
-	bool IsAssetInPendingDeletes( const FName& PackageName ) const;
+	UNREALED_API bool IsAssetInPendingDeletes( const FName& PackageName ) const;
 
 	/** Deletes any source content files referenced by the assets */
-	void DeleteSourceContentFiles();
+	UNREALED_API void DeleteSourceContentFiles();
 
 	/** Returns true if it is valid to delete the current objects with no problems. */
-	bool CanDelete() const;
+	UNREALED_API bool CanDelete() const;
 
 	/** Performs the delete if it's possible. */
-	bool DoDelete();
+	UNREALED_API bool DoDelete();
 
 	/** Returns true if it is valid to force the delete of the current assets. */
-	bool CanForceDelete() const;
+	UNREALED_API bool CanForceDelete() const;
 
 	/** Performs a force delete on the pending deleted assets if possible. */
-	bool DoForceDelete();
+	UNREALED_API bool DoForceDelete();
 
 	/** Returns true if it's valid to replace the references of the pending deleted objects. */
-	bool CanReplaceReferences() const;
+	UNREALED_API bool CanReplaceReferences() const;
 
 	/** Returns true if it is valid to force the delete of the current assets with the provided asset. */
-	bool CanReplaceReferencesWith( const FAssetData& InAssetData ) const;
+	UNREALED_API bool CanReplaceReferencesWith( const FAssetData& InAssetData ) const;
 
 	/** Performs the replace references action if possible with the provided asset. */
-	bool DoReplaceReferences( const FAssetData& ReplaceReferencesWith );
+	UNREALED_API bool DoReplaceReferences( const FAssetData& ReplaceReferencesWith );
 
 	/** Gets the 0..1 progress of the scanning. */
-	float GetProgress() const;
+	UNREALED_API float GetProgress() const;
 
 	/** Gets the current text to display for the current progress of the scanning. */
-	FText GetProgressText() const;
+	UNREALED_API FText GetProgressText() const;
 
 	/** Is any of the pending deleted assets being referenced in memory. */
 	bool IsAnythingReferencedInMemoryByNonUndo() const { return bIsAnythingReferencedInMemoryByNonUndo; }
@@ -169,13 +169,13 @@ public:
 	bool IsAnythingReferencedInMemoryByUndo() const { return bIsAnythingReferencedInMemoryByUndo; }
 
 	/** Check whether we have any source files residing under monitored, mounted paths to delete */
-	bool HasAnySourceContentFilesToDelete() const;
+	UNREALED_API bool HasAnySourceContentFilesToDelete() const;
 
 	/** Goes to the next actor in the loaded level if it is available */
-	bool GoToNextReferenceInLevel() const;
+	UNREALED_API bool GoToNextReferenceInLevel() const;
 
 	/** Gets the number of objects successfully deleted. */
-	int32 GetDeletedObjectCount() const;
+	UNREALED_API int32 GetDeletedObjectCount() const;
 	
 	/** Fires whenever the state changes. */
 	DECLARE_EVENT_OneParam( FAssetDeleteModel, FOnStateChanged, EState /*NewState*/ );
@@ -185,16 +185,16 @@ public:
 	}
 
 private:
-	void PrepareToDelete(UObject* InObject);
+	UNREALED_API void PrepareToDelete(UObject* InObject);
 
 	/** Sets the current state of the model. */
-	void SetState( EState NewState );
+	UNREALED_API void SetState( EState NewState );
 
 	/** Computes the value that should be used for CanReplaceReferences */
-	bool ComputeCanReplaceReferences();
+	UNREALED_API bool ComputeCanReplaceReferences();
 
 	/** Discover source file references for the specified object */
-	void DiscoverSourceFileReferences(FPendingDelete& PendingDelete);
+	UNREALED_API void DiscoverSourceFileReferences(FPendingDelete& PendingDelete);
 
 private:
 
diff --git a/Engine/Source/Editor/UnrealEd/Public/AssetEditorModeManager.h b/Engine/Source/Editor/UnrealEd/Public/AssetEditorModeManager.h
--- a/Engine/Source/Editor/UnrealEd/Public/AssetEditorModeManager.h
+++ b/Engine/Source/Editor/UnrealEd/Public/AssetEditorModeManager.h
@@ -10,21 +10,21 @@ class FPreviewScene;
 //////////////////////////////////////////////////////////////////////////
 // FAssetEditorModeManager
 
-class UNREALED_API FAssetEditorModeManager : public FEditorModeTools
+class FAssetEditorModeManager : public FEditorModeTools
 {
 public:
-	FAssetEditorModeManager();
-	virtual ~FAssetEditorModeManager() override;
+	UNREALED_API FAssetEditorModeManager();
+	UNREALED_API virtual ~FAssetEditorModeManager() override;
 
 	// FEditorModeTools interface
-	virtual USelection* GetSelectedActors() const override;
-	virtual USelection* GetSelectedObjects() const override;
-	virtual USelection* GetSelectedComponents() const override;
-	virtual UWorld* GetWorld() const override;
+	UNREALED_API virtual USelection* GetSelectedActors() const override;
+	UNREALED_API virtual USelection* GetSelectedObjects() const override;
+	UNREALED_API virtual USelection* GetSelectedComponents() const override;
+	UNREALED_API virtual UWorld* GetWorld() const override;
 	// End of FEditorModeTools interface
 
-	virtual void SetPreviewScene(FPreviewScene* NewPreviewScene);
-	FPreviewScene* GetPreviewScene() const;
+	UNREALED_API virtual void SetPreviewScene(FPreviewScene* NewPreviewScene);
+	UNREALED_API FPreviewScene* GetPreviewScene() const;
 
 protected:
 	USelection* ActorSet = nullptr;
diff --git a/Engine/Source/Editor/UnrealEd/Public/AssetEditorViewportLayout.h b/Engine/Source/Editor/UnrealEd/Public/AssetEditorViewportLayout.h
--- a/Engine/Source/Editor/UnrealEd/Public/AssetEditorViewportLayout.h
+++ b/Engine/Source/Editor/UnrealEd/Public/AssetEditorViewportLayout.h
@@ -23,7 +23,7 @@ class FEditorViewportTabContent;
 class FViewportTabContent;
 
 /** Arguments for constructing a viewport */
-struct UNREALED_API FAssetEditorViewportConstructionArgs
+struct FAssetEditorViewportConstructionArgs
 {
 	FAssetEditorViewportConstructionArgs()
 		: ViewportType(LVT_Perspective)
@@ -65,7 +65,7 @@ namespace EditorViewportConfigurationNames
 * It will also store the ViewportLayout data because that data can't be stored
 * per app; it must be stored per viewport overlay in case the app that made it closes.
 */
-class UNREALED_API SAssetEditorViewportsOverlay : public SCompoundWidget
+class SAssetEditorViewportsOverlay : public SCompoundWidget
 {
 
 public:
@@ -75,7 +75,7 @@ public:
 		SLATE_ARGUMENT(TSharedPtr<FViewportTabContent>, ViewportTab)
 	SLATE_END_ARGS()
 
-		void Construct(const FArguments& InArgs);
+		UNREALED_API void Construct(const FArguments& InArgs);
 
 	/** Default constructor */
 	SAssetEditorViewportsOverlay()
@@ -83,23 +83,23 @@ public:
 	{}
 
 	/** Overridden from SWidget */
-	virtual void Tick(const FGeometry& AllottedGeometry, const double InCurrentTime, const float InDeltaTime) override;
+	UNREALED_API virtual void Tick(const FGeometry& AllottedGeometry, const double InCurrentTime, const float InDeltaTime) override;
 
 	/** Wraps SOverlay::AddSlot() */
-	SOverlay::FScopedWidgetSlotArguments AddSlot();
+	UNREALED_API SOverlay::FScopedWidgetSlotArguments AddSlot();
 
 	/** Wraps SOverlay::RemoveSlot() */
-	void RemoveSlot();
+	UNREALED_API void RemoveSlot();
 
 	/**
 	* Returns the cached size of this viewport overlay
 	*
 	* @return	The size that was cached
 	*/
-	const FVector2D& GetCachedSize() const;
+	UNREALED_API const FVector2D& GetCachedSize() const;
 
 	/** Gets the  Viewport Tab that created this overlay */
-	TSharedPtr<FViewportTabContent> GetViewportTab() const;
+	UNREALED_API TSharedPtr<FViewportTabContent> GetViewportTab() const;
 
 private:
 
@@ -113,7 +113,7 @@ private:
 	FVector2D CachedSize;
 };
 
-class UNREALED_API FAssetEditorViewportPaneLayout
+class FAssetEditorViewportPaneLayout
 {
 public:
 	virtual ~FAssetEditorViewportPaneLayout() = default;
@@ -132,14 +132,14 @@ public:
 	 */
 	virtual const FName& GetLayoutTypeName() const = 0;
 
-	void LoadConfig(const FString& LayoutString, TFunction<void(const FString&, const FName)> LoadAdditionalLayoutInfoCallback = nullptr);
+	UNREALED_API void LoadConfig(const FString& LayoutString, TFunction<void(const FString&, const FName)> LoadAdditionalLayoutInfoCallback = nullptr);
 
-	void SaveConfig(const FString& LayoutString, TFunction<void(const FString&)> SaveAdditionalLayoutInfoCallback = nullptr) const;
+	UNREALED_API void SaveConfig(const FString& LayoutString, TFunction<void(const FString&)> SaveAdditionalLayoutInfoCallback = nullptr) const;
 
 	virtual void ReplaceWidget(TSharedRef<SWidget> OriginalWidget, TSharedRef<SWidget> ReplacementWidget) = 0;
 
 protected:
-	FString GetTypeSpecificLayoutString(const FString& LayoutString) const;
+	UNREALED_API FString GetTypeSpecificLayoutString(const FString& LayoutString) const;
 	virtual void SaveLayoutString(const FString& LayoutString) const {}
 	virtual void LoadLayoutString(const FString& LayoutString) {}
 
@@ -150,29 +150,29 @@ protected:
  * Base class for viewport layout configurations
  * Handles maximizing and restoring well as visibility of specific viewports.
  */
-class UNREALED_API FAssetEditorViewportLayout : public TSharedFromThis<FAssetEditorViewportLayout>, public FEditorViewportLayout, public FTickableEditorObject
+class FAssetEditorViewportLayout : public TSharedFromThis<FAssetEditorViewportLayout>, public FEditorViewportLayout, public FTickableEditorObject
 {
 public:
 	/**
 	 * Constructor
 	 */
-	FAssetEditorViewportLayout();
+	UNREALED_API FAssetEditorViewportLayout();
 
 	/**
 	 * Destructor
 	 */
-	virtual ~FAssetEditorViewportLayout();
+	UNREALED_API virtual ~FAssetEditorViewportLayout();
 
 	/** Create an instance of a custom viewport from the specified viewport type name */
-	virtual TSharedRef<SWidget> FactoryViewport(FName InTypeName, const FAssetEditorViewportConstructionArgs& ConstructionArgs);
+	UNREALED_API virtual TSharedRef<SWidget> FactoryViewport(FName InTypeName, const FAssetEditorViewportConstructionArgs& ConstructionArgs);
 
 	/** FTickableEditorObject interface */
 	virtual void Tick(float DeltaTime) override {}
 	virtual bool IsTickable() const override { return false; }
  	virtual TStatId GetStatId() const override { return TStatId(); }
 
-	virtual void FactoryPaneConfigurationFromTypeName(const FName& InLayoutConfigTypeName) override;
-	virtual const FName GetActivePaneConfigurationTypeName() const override;
+	UNREALED_API virtual void FactoryPaneConfigurationFromTypeName(const FName& InLayoutConfigTypeName) override;
+	UNREALED_API virtual const FName GetActivePaneConfigurationTypeName() const override;
 
 	/**
 	 * Builds a viewport layout and returns the widget containing the layout
@@ -181,7 +181,7 @@ public:
 	 * @param InParentTab			The parent tab object
 	 * @param LayoutString			The layout string loaded from file to custom build the layout with
 	 */
- 	virtual TSharedRef<SWidget> BuildViewportLayout(TSharedPtr<SDockTab> InParentDockTab, TSharedPtr<FEditorViewportTabContent> InParentTab, const FString& LayoutString );
+ 	UNREALED_API virtual TSharedRef<SWidget> BuildViewportLayout(TSharedPtr<SDockTab> InParentDockTab, TSharedPtr<FEditorViewportTabContent> InParentTab, const FString& LayoutString );
 
 	/** Returns the parent tab content object */
 	TWeakPtr< FEditorViewportTabContent > GetParentTabContent() const { return ParentTabContent; }
@@ -193,10 +193,10 @@ public:
 	 * @return The base widget representing the layout.  This is commonly a splitter.
 	 */
 	UE_DEPRECATED(5.0, "This functionality has moved to the layout configurations. Use BuildViewportLayout.")
-	virtual TSharedRef<SWidget> MakeViewportLayout(const FString& LayoutString) final;
+	UNREALED_API virtual TSharedRef<SWidget> MakeViewportLayout(const FString& LayoutString) final;
 
-	virtual void LoadConfig(const FString& LayoutString);
-	virtual void SaveConfig(const FString& LayoutString) const;
+	UNREALED_API virtual void LoadConfig(const FString& LayoutString);
+	UNREALED_API virtual void SaveConfig(const FString& LayoutString) const;
 
 protected:
 	/**
@@ -205,7 +205,7 @@ protected:
 	 *
 	 * @param EVisibility::Visible when visible, EVisibility::Collapsed otherwise
 	 */
-	virtual EVisibility OnGetNonMaximizedVisibility() const;
+	UNREALED_API virtual EVisibility OnGetNonMaximizedVisibility() const;
 
 	/** The overlay widget that handles what viewports should be on top (non-maximized or maximized) */
 	TWeakPtr< SAssetEditorViewportsOverlay > ViewportsOverlayPtr;
diff --git a/Engine/Source/Editor/UnrealEd/Public/AssetImportTask.h b/Engine/Source/Editor/UnrealEd/Public/AssetImportTask.h
--- a/Engine/Source/Editor/UnrealEd/Public/AssetImportTask.h
+++ b/Engine/Source/Editor/UnrealEd/Public/AssetImportTask.h
@@ -21,13 +21,13 @@ class UFactory;
 /**
  * Contains data for a group of assets to import
  */ 
-UCLASS(Transient, BlueprintType)
-class UNREALED_API UAssetImportTask : public UObject
+UCLASS(Transient, BlueprintType, MinimalAPI)
+class UAssetImportTask : public UObject
 {
 	GENERATED_BODY()
 
 public:
-	UAssetImportTask();
+	UNREALED_API UAssetImportTask();
 
 // Task Options
 public:
@@ -79,14 +79,14 @@ public:
 	 * To test whether asynchronous results are ready or not, use IsAsyncImportComplete().
 	 */
 	UFUNCTION(BlueprintCallable, Category = "Asset Import Task")
-	const TArray<UObject*>& GetObjects() const;
+	UNREALED_API const TArray<UObject*>& GetObjects() const;
 
 	/**
 	 * Query whether this asynchronous import task is complete, and the results are ready to read.
 	 * This will always return true in the case of a blocking import.
 	 */
 	UFUNCTION(BlueprintCallable, Category = "Asset Import Task")
-	bool IsAsyncImportComplete() const;
+	UNREALED_API bool IsAsyncImportComplete() const;
 
 	/** Paths to objects created or updated after import */
 	UPROPERTY(BlueprintReadWrite, Category = "Asset Import Task")
diff --git a/Engine/Source/Editor/UnrealEd/Public/AssetPrivatizeModel.h b/Engine/Source/Editor/UnrealEd/Public/AssetPrivatizeModel.h
--- a/Engine/Source/Editor/UnrealEd/Public/AssetPrivatizeModel.h
+++ b/Engine/Source/Editor/UnrealEd/Public/AssetPrivatizeModel.h
@@ -55,7 +55,7 @@ private:
 	bool bReferencesChecked;
 };
 
-class UNREALED_API FAssetPrivatizeModel
+class FAssetPrivatizeModel
 {
 public:
 	enum EState
@@ -70,10 +70,10 @@ public:
 		Finished,
 	};
 
-	FAssetPrivatizeModel(const TArray<UObject*>& InObjectsToPrivatize);
+	UNREALED_API FAssetPrivatizeModel(const TArray<UObject*>& InObjectsToPrivatize);
 
 	/** Add an object to the list of pending private assets, this will invalidate the scanning state */
-	void AddObjectToPrivatize(UObject* InObject, UPackage* InOwningPackage);
+	UNREALED_API void AddObjectToPrivatize(UObject* InObject, UPackage* InOwningPackage);
 
 	/** Returns the pending private assets */
 	const TArray<TSharedPtr<FPendingPrivateAsset>>* GetPendingPrivateAssets() const { return &PendingPrivateAssets; }
@@ -82,16 +82,16 @@ public:
 	const TSet<FName>& GetIllegalAssetReferences() const { return IllegalOnDiskReferences; }
 
 	/** Ticks the privatize model which does a timeslice of work before returning */
-	void Tick(const float InDeltaTime);
+	UNREALED_API void Tick(const float InDeltaTime);
 
 	/** Returns the current state of the scanning process */
 	EState GetState() { return State; }
 
 	/** Gets the 0..1 progress of the scan */
-	float GetProgress() const;
+	UNREALED_API float GetProgress() const;
 
 	/** Gets the current text to display for the current progress of the scan */
-	FText GetProgressText() const;
+	UNREALED_API FText GetProgressText() const;
 
 	/** Are any of the pending assets being referenced in memory. */
 	bool IsAnythingReferencedInMemoryByNonUndo() const { return bIsAnythingReferencedInMemoryByNonUndo; }
@@ -103,16 +103,16 @@ public:
 	int32 GetObjectsPrivatizedCount() const { return ObjectsPrivatized; }
 
 	/** Returns true if it is valid to mark private the pending objects with no problems (there are no illegal references detected) */
-	bool CanPrivatize();
+	UNREALED_API bool CanPrivatize();
 
 	/** Performs the privatization action if possible */
-	bool DoPrivatize();
+	UNREALED_API bool DoPrivatize();
 
 	/** Returns true if we can't cleanly mark the pending objects as private (there are illegal references detected) */
-	bool CanForcePrivatize();
+	UNREALED_API bool CanForcePrivatize();
 
 	/** Marks the pending private objects as private and nulls out all detected illegal references */
-	bool DoForcePrivatize();
+	UNREALED_API bool DoForcePrivatize();
 
 	/** Fires whenever the state changes */
 	DECLARE_EVENT_OneParam(FAssetPrivatizeModel, FOnStateChanged, EState /*NewState*/);
@@ -124,10 +124,10 @@ public:
 private:
 
 	/** Sets the current state of the model */
-	void SetState(EState NewState);
+	UNREALED_API void SetState(EState NewState);
 
 	/** Helper that does the reference detection during scanning */
-	void ScanForReferences();
+	UNREALED_API void ScanForReferences();
 
 private:
 
@@ -157,4 +157,4 @@ private:
 
 	/** The number of objects successfully marked private */
 	int32 ObjectsPrivatized;
-};
\ No newline at end of file
+};
diff --git a/Engine/Source/Editor/UnrealEd/Public/AssetSelection.h b/Engine/Source/Editor/UnrealEd/Public/AssetSelection.h
--- a/Engine/Source/Editor/UnrealEd/Public/AssetSelection.h
+++ b/Engine/Source/Editor/UnrealEd/Public/AssetSelection.h
@@ -221,7 +221,7 @@ namespace ActorPlacementUtils
 	UNREALED_API bool IsLevelValidForActorPlacement(ULevel* Level, TArray<FTransform>& InActorTransforms);
 }
 
-class UNREALED_API FActorFactoryAssetProxy
+class FActorFactoryAssetProxy
 {
 public:
 
@@ -250,7 +250,7 @@ public:
 	 * @param	OutMenuItems				receives the list of menu items to use for populating an actor factory menu.
 	 * @param	ExcludeStandAloneFactories	if true, only factories that can create actors with the currently selected assets will be added
 	 */
-	static void GenerateActorFactoryMenuItems( const FAssetData& AssetData, TArray<FMenuItem>* OutMenuItems, bool ExcludeStandAloneFactories );
+	static UNREALED_API void GenerateActorFactoryMenuItems( const FAssetData& AssetData, TArray<FMenuItem>* OutMenuItems, bool ExcludeStandAloneFactories );
 
 	/**
 	 * Find the appropriate actor factory for an asset by type.
@@ -261,7 +261,7 @@ public:
 	 *
 	 * @return	the factory that is responsible for creating actors for the specified asset type.
 	 */
-	static UActorFactory* GetFactoryForAsset( const FAssetData& DropData, bool bRequireValidObject=false );
+	static UNREALED_API UActorFactory* GetFactoryForAsset( const FAssetData& DropData, bool bRequireValidObject=false );
 
 	/**
 	 * Find the appropriate actor factory for an asset.
@@ -270,7 +270,7 @@ public:
 	 *
 	 * @return	The factory that is responsible for creating actors for the specified asset
 	 */
-	static UActorFactory* GetFactoryForAssetObject( UObject* AssetObj );
+	static UNREALED_API UActorFactory* GetFactoryForAssetObject( UObject* AssetObj );
 
 	/**
 	 * Places an actor instance using the factory appropriate for the type of asset
@@ -283,7 +283,7 @@ public:
 	 * @return	the actor that was created by the factory, or NULL if there aren't any factories for this asset (or
 	 *			the actor couldn't be created for some other reason)
 	 */
-	static AActor* AddActorForAsset( UObject* AssetObj, bool SelectActor = true, EObjectFlags ObjectFlags = RF_Transactional, UActorFactory* FactoryToUse = NULL, const FName Name = NAME_None );
+	static UNREALED_API AActor* AddActorForAsset( UObject* AssetObj, bool SelectActor = true, EObjectFlags ObjectFlags = RF_Transactional, UActorFactory* FactoryToUse = NULL, const FName Name = NAME_None );
 
 	/**
 	 * Places an actor instance using the factory appropriate for the type of asset using the current object selection as the asset
@@ -298,7 +298,7 @@ public:
 	 * @return	the actor that was created by the factory, or NULL if there aren't any factories for this asset (or
 	 *			the actor couldn't be created for some other reason)
 	 */
-	static AActor* AddActorFromSelection( UClass* ActorClass, const FVector* ActorLocation=NULL, bool SelectActor = true, EObjectFlags ObjectFlags = RF_Transactional, UActorFactory* ActorFactory = NULL, const FName Name = NAME_None );
+	static UNREALED_API AActor* AddActorFromSelection( UClass* ActorClass, const FVector* ActorLocation=NULL, bool SelectActor = true, EObjectFlags ObjectFlags = RF_Transactional, UActorFactory* ActorFactory = NULL, const FName Name = NAME_None );
 
 	/**
 	 * Determines if the provided actor is capable of having a material applied to it.
@@ -307,7 +307,7 @@ public:
 	 *
 	 * @return	true if the actor is valid for material application; false otherwise
 	 */
-	static bool IsActorValidForMaterialApplication( AActor* TargetActor );
+	static UNREALED_API bool IsActorValidForMaterialApplication( AActor* TargetActor );
 
 	/**
 	 * Attempts to apply the material to the specified actor.
@@ -318,7 +318,7 @@ public:
 	 *
 	 * @return	true if the material was successfully applied to the actor
 	 */
-	static bool ApplyMaterialToActor( AActor* TargetActor, UMaterialInterface* MaterialToApply, int32 OptionalMaterialSlot = -1 );
+	static UNREALED_API bool ApplyMaterialToActor( AActor* TargetActor, UMaterialInterface* MaterialToApply, int32 OptionalMaterialSlot = -1 );
 
 private:
 	/**
diff --git a/Engine/Source/Editor/UnrealEd/Public/AutoReimport/AssetSourceFilenameCache.h b/Engine/Source/Editor/UnrealEd/Public/AutoReimport/AssetSourceFilenameCache.h
--- a/Engine/Source/Editor/UnrealEd/Public/AutoReimport/AssetSourceFilenameCache.h
+++ b/Engine/Source/Editor/UnrealEd/Public/AutoReimport/AssetSourceFilenameCache.h
@@ -22,22 +22,22 @@ struct FAssetImportInfo;
 
 
 /** Class responsible for maintaing a cache of clean source file names (bla.txt) to asset data */
-class UNREALED_API FAssetSourceFilenameCache
+class FAssetSourceFilenameCache
 {
 public:
-	FAssetSourceFilenameCache();
+	UNREALED_API FAssetSourceFilenameCache();
 
 	/** Singleton access */
-	static FAssetSourceFilenameCache& Get();
+	static UNREALED_API FAssetSourceFilenameCache& Get();
 
 	/** Helper functions to extract asset import information from asset registry tags */
-	static TOptional<FAssetImportInfo> ExtractAssetImportInfo(const FAssetData& AssetData);
+	static UNREALED_API TOptional<FAssetImportInfo> ExtractAssetImportInfo(const FAssetData& AssetData);
 
 	/** Retrieve a list of assets that were imported from the specified filename */
-	TArray<FAssetData> GetAssetsPertainingToFile(const IAssetRegistry& Registry, const FString& AbsoluteFilename) const;
+	UNREALED_API TArray<FAssetData> GetAssetsPertainingToFile(const IAssetRegistry& Registry, const FString& AbsoluteFilename) const;
 
 	/** Shutdown this instance of the cache */
-	void Shutdown();
+	UNREALED_API void Shutdown();
 
 	/** Event for when an asset has been renamed, and has been updated in our source file cache */
 	DECLARE_EVENT_TwoParams( FAssetSourceFilenameCache, FAssetRenamedEvent, const FAssetData&, const FString& );
@@ -45,10 +45,10 @@ public:
 
 private:
 	/** Delegate bindings that keep the cache up-to-date */
-	void HandleOnAssetsAdded(TConstArrayView<FAssetData> Assets);
-	void HandleOnAssetsRemoved(TConstArrayView<FAssetData> AssetData);
-	void HandleOnAssetRenamed(const FAssetData& AssetData, const FString& OldPath);
-	void HandleOnAssetUpdated(const FAssetImportInfo& OldData, const UAssetImportData* ImportData);
+	UNREALED_API void HandleOnAssetsAdded(TConstArrayView<FAssetData> Assets);
+	UNREALED_API void HandleOnAssetsRemoved(TConstArrayView<FAssetData> AssetData);
+	UNREALED_API void HandleOnAssetRenamed(const FAssetData& AssetData, const FString& OldPath);
+	UNREALED_API void HandleOnAssetUpdated(const FAssetImportInfo& OldData, const UAssetImportData* ImportData);
 	
 	/** Event that is triggered when an asset has been renamed, and we've updated our cache */
 	FAssetRenamedEvent AssetRenamedEvent;
diff --git a/Engine/Source/Editor/UnrealEd/Public/AutoReimport/AutoReimportManager.h b/Engine/Source/Editor/UnrealEd/Public/AutoReimport/AutoReimportManager.h
--- a/Engine/Source/Editor/UnrealEd/Public/AutoReimport/AutoReimportManager.h
+++ b/Engine/Source/Editor/UnrealEd/Public/AutoReimport/AutoReimportManager.h
@@ -23,30 +23,30 @@ struct FPathAndMountPoint
 };
 	
 /* Deals with auto reimporting of objects when the objects file on disk is modified*/
-UCLASS(config=Editor, transient)
-class UNREALED_API UAutoReimportManager : public UObject
+UCLASS(config=Editor, transient, MinimalAPI)
+class UAutoReimportManager : public UObject
 {
 	GENERATED_UCLASS_BODY()
 public:
 
-	~UAutoReimportManager();
+	UNREALED_API ~UAutoReimportManager();
 
 	/** Initialize during engine startup */
-	void Initialize();
+	UNREALED_API void Initialize();
 
 	/** Get a list of absolute directories that we are monitoring */
-	TArray<FPathAndMountPoint> GetMonitoredDirectories() const;
+	UNREALED_API TArray<FPathAndMountPoint> GetMonitoredDirectories() const;
 
 	/** Report an external change to the manager, such that a subsequent equal change reported by the os be ignored */
-	void IgnoreNewFile(const FString& Filename);
-	void IgnoreFileModification(const FString& Filename);
-	void IgnoreMovedFile(const FString& SrcFilename, const FString& DstFilename);
-	void IgnoreDeletedFile(const FString& Filename);
+	UNREALED_API void IgnoreNewFile(const FString& Filename);
+	UNREALED_API void IgnoreFileModification(const FString& Filename);
+	UNREALED_API void IgnoreMovedFile(const FString& SrcFilename, const FString& DstFilename);
+	UNREALED_API void IgnoreDeletedFile(const FString& Filename);
 
 private:
 
 	/** UObject Interface */
-	virtual void BeginDestroy() override;
+	UNREALED_API virtual void BeginDestroy() override;
 
 	/** Private implementation of the reimport manager */
 	TSharedPtr<class FAutoReimportManager> Implementation;
diff --git a/Engine/Source/Editor/UnrealEd/Public/AutomatedAssetImportData.h b/Engine/Source/Editor/UnrealEd/Public/AutomatedAssetImportData.h
--- a/Engine/Source/Editor/UnrealEd/Public/AutomatedAssetImportData.h
+++ b/Engine/Source/Editor/UnrealEd/Public/AutomatedAssetImportData.h
@@ -22,22 +22,22 @@ DECLARE_LOG_CATEGORY_EXTERN(LogAutomatedImport, Log, All);
 /**
  * Contains data for a group of assets to import
  */ 
-UCLASS(Transient, BlueprintType)
-class UNREALED_API UAutomatedAssetImportData : public UObject
+UCLASS(Transient, BlueprintType, MinimalAPI)
+class UAutomatedAssetImportData : public UObject
 {
 	GENERATED_BODY()
 
 public:
-	UAutomatedAssetImportData();
+	UNREALED_API UAutomatedAssetImportData();
 
 	/** @return true if this group contains enough valid data to import*/
-	bool IsValid() const;
+	UNREALED_API bool IsValid() const;
 
 	/** Initalizes the group */
-	void Initialize(TSharedPtr<FJsonObject> InImportGroupJsonData);
+	UNREALED_API void Initialize(TSharedPtr<FJsonObject> InImportGroupJsonData);
 
 	/** @return the display name of the group */
-	FString GetDisplayName() const; 
+	UNREALED_API FString GetDisplayName() const; 
 public:
 	/** Display name of the group. This is for logging purposes only. */
 	UPROPERTY(BlueprintReadWrite, Category="Asset Import Data")
diff --git a/Engine/Source/Editor/UnrealEd/Public/BaseWidgetBlueprint.h b/Engine/Source/Editor/UnrealEd/Public/BaseWidgetBlueprint.h
--- a/Engine/Source/Editor/UnrealEd/Public/BaseWidgetBlueprint.h
+++ b/Engine/Source/Editor/UnrealEd/Public/BaseWidgetBlueprint.h
@@ -5,8 +5,8 @@
 #include "Blueprint/UserWidgetBlueprint.h"
 #include "BaseWidgetBlueprint.generated.h"
 
-UCLASS(Abstract)
-class UNREALED_API UBaseWidgetBlueprint : public UUserWidgetBlueprint
+UCLASS(Abstract, MinimalAPI)
+class UBaseWidgetBlueprint : public UUserWidgetBlueprint
 {
 	GENERATED_UCLASS_BODY()
 
@@ -17,19 +17,19 @@ public:
 	TObjectPtr<class UWidgetTree> WidgetTree;
 #endif
 
-	virtual void PostLoad() override;
+	UNREALED_API virtual void PostLoad() override;
 
 	/**
 	* Returns collection of widgets that represent the 'source' (user edited) widgets for this
 	* blueprint - avoids calling virtual functions on instances and is therefore safe to use
 	* throughout compilation.
 	*/
-	TArray<class UWidget*> GetAllSourceWidgets();
-	TArray<const class UWidget*> GetAllSourceWidgets() const;
+	UNREALED_API TArray<class UWidget*> GetAllSourceWidgets();
+	UNREALED_API TArray<const class UWidget*> GetAllSourceWidgets() const;
 
 	/** Identical to GetAllSourceWidgets, but as an algorithm */
-	void ForEachSourceWidget(TFunctionRef<void(class UWidget*)> Fn);
-	void ForEachSourceWidget(TFunctionRef<void(class UWidget*)> Fn) const;
+	UNREALED_API void ForEachSourceWidget(TFunctionRef<void(class UWidget*)> Fn);
+	UNREALED_API void ForEachSourceWidget(TFunctionRef<void(class UWidget*)> Fn) const;
 
 private:
 	void ForEachSourceWidgetImpl(TFunctionRef<void(class UWidget*)> Fn) const;
diff --git a/Engine/Source/Editor/UnrealEd/Public/Bookmarks/BookmarkScoped.h b/Engine/Source/Editor/UnrealEd/Public/Bookmarks/BookmarkScoped.h
--- a/Engine/Source/Editor/UnrealEd/Public/Bookmarks/BookmarkScoped.h
+++ b/Engine/Source/Editor/UnrealEd/Public/Bookmarks/BookmarkScoped.h
@@ -12,15 +12,15 @@ class UBookMark;
  * Provides a way to temporarily bookmark and restore all viewports
  * This allows camera locations and rotations to be preserved across a map reload
  */
-class UNREALED_API FBookmarkScoped
+class FBookmarkScoped
 {
 public:
 
 	/** Constructor; bookmark all viewports */
-	FBookmarkScoped();
+	UNREALED_API FBookmarkScoped();
 
 	/** Destructor; restore all viewports */
-	~FBookmarkScoped();
+	UNREALED_API ~FBookmarkScoped();
 
 private:
 
diff --git a/Engine/Source/Editor/UnrealEd/Public/Bookmarks/IBookmarkTypeActions.h b/Engine/Source/Editor/UnrealEd/Public/Bookmarks/IBookmarkTypeActions.h
--- a/Engine/Source/Editor/UnrealEd/Public/Bookmarks/IBookmarkTypeActions.h
+++ b/Engine/Source/Editor/UnrealEd/Public/Bookmarks/IBookmarkTypeActions.h
@@ -1,4 +1,4 @@
-﻿// Copyright Epic Games, Inc. All Rights Reserved.
+// Copyright Epic Games, Inc. All Rights Reserved.
 
 #pragma once
 
@@ -9,7 +9,7 @@
  * Actions that can be applied to or completed with bookmarks.
  * Must be registered, see IBookmarkTypeTools.
  */
-class UNREALED_API IBookmarkTypeActions
+class IBookmarkTypeActions
 {
 public:
 
@@ -37,4 +37,4 @@ public:
 	 * @param InViewportClient	The viewport that will be updated.
 	 */
 	virtual void JumpToBookmark(class UBookmarkBase* Bookmark, const TSharedPtr<struct FBookmarkBaseJumpToSettings> InSettings, class FEditorViewportClient& InViewportClient) = 0;
-};
\ No newline at end of file
+};
diff --git a/Engine/Source/Editor/UnrealEd/Public/Bookmarks/IBookmarkTypeTools.h b/Engine/Source/Editor/UnrealEd/Public/Bookmarks/IBookmarkTypeTools.h
--- a/Engine/Source/Editor/UnrealEd/Public/Bookmarks/IBookmarkTypeTools.h
+++ b/Engine/Source/Editor/UnrealEd/Public/Bookmarks/IBookmarkTypeTools.h
@@ -1,4 +1,4 @@
-﻿// Copyright Epic Games, Inc. All Rights Reserved.
+// Copyright Epic Games, Inc. All Rights Reserved.
 
 #pragma once
 
@@ -8,7 +8,7 @@
 /**
  * Provides convenience methods for interacting with Bookmarks.
  */
-class UNREALED_API IBookmarkTypeTools : public FNoncopyable
+class IBookmarkTypeTools : public FNoncopyable
 {
 protected:
 
@@ -16,7 +16,7 @@ protected:
 
 public:
 
-	static IBookmarkTypeTools& Get();
+	static UNREALED_API IBookmarkTypeTools& Get();
 
 	/**
 	 * Gets the current maximum number of bookmarks allowed.
diff --git a/Engine/Source/Editor/UnrealEd/Public/BusyCursor.h b/Engine/Source/Editor/UnrealEd/Public/BusyCursor.h
--- a/Engine/Source/Editor/UnrealEd/Public/BusyCursor.h
+++ b/Engine/Source/Editor/UnrealEd/Public/BusyCursor.h
@@ -7,9 +7,9 @@
 /**
  * While in scope, sets the cursor to the busy (hourglass) cursor for all windows.
  */
-class UNREALED_API FScopedBusyCursor
+class FScopedBusyCursor
 {
 public:
-	FScopedBusyCursor();
-	~FScopedBusyCursor();
+	UNREALED_API FScopedBusyCursor();
+	UNREALED_API ~FScopedBusyCursor();
 };
diff --git a/Engine/Source/Editor/UnrealEd/Public/CameraController.h b/Engine/Source/Editor/UnrealEd/Public/CameraController.h
--- a/Engine/Source/Editor/UnrealEd/Public/CameraController.h
+++ b/Engine/Source/Editor/UnrealEd/Public/CameraController.h
@@ -221,7 +221,7 @@ public:
 	/** Constructor */
 	UNREALED_API FEditorCameraController();
 
-	UNREALED_API virtual ~FEditorCameraController() {}
+	virtual ~FEditorCameraController() {}
 
 	/** Sets the configuration for this camera controller */
 	void SetConfig( const FCameraControllerConfig& InConfig )
diff --git a/Engine/Source/Editor/UnrealEd/Public/ClassIconFinder.h b/Engine/Source/Editor/UnrealEd/Public/ClassIconFinder.h
--- a/Engine/Source/Editor/UnrealEd/Public/ClassIconFinder.h
+++ b/Engine/Source/Editor/UnrealEd/Public/ClassIconFinder.h
@@ -37,7 +37,7 @@ public:
 	UNREALED_API static const UClass* GetIconClassForAssetData(const FAssetData& InAssetData, bool* bOutIsClassType = nullptr);
 
 	/** Find the large thumbnail name to use for the supplied class */
-	UNREALED_API static const FSlateBrush* FindThumbnailForClass(const UClass* InClass, const FName& InDefaultName = FName() )
+	static const FSlateBrush* FindThumbnailForClass(const UClass* InClass, const FName& InDefaultName = FName() )
 	{
 		return FSlateIconFinder::FindCustomIconBrushForClass(InClass, TEXT("ClassThumbnail"), InDefaultName);
 	}
diff --git a/Engine/Source/Editor/UnrealEd/Public/Commandlets/ChunkDependencyInfo.h b/Engine/Source/Editor/UnrealEd/Public/Commandlets/ChunkDependencyInfo.h
--- a/Engine/Source/Editor/UnrealEd/Public/Commandlets/ChunkDependencyInfo.h
+++ b/Engine/Source/Editor/UnrealEd/Public/Commandlets/ChunkDependencyInfo.h
@@ -47,19 +47,19 @@ struct FChunkDependencyTreeNode
 };
 
 /** This is read out of config and defines a tree of chunk dependencies */
-UCLASS(config=Engine, defaultconfig)
-class UNREALED_API UChunkDependencyInfo : public UObject
+UCLASS(config=Engine, defaultconfig, MinimalAPI)
+class UChunkDependencyInfo : public UObject
 {
 	GENERATED_UCLASS_BODY()
 
 	/** Will return an existing dependency graph, or build if it necessary */
-	const FChunkDependencyTreeNode* GetOrBuildChunkDependencyGraph(int32 HighestChunk = 0);
+	UNREALED_API const FChunkDependencyTreeNode* GetOrBuildChunkDependencyGraph(int32 HighestChunk = 0);
 
 	/** Will create a dependency tree starting with RootTreeNode. If HighestChunk is == 0 it will only add dependencies on 0 for chunks already in the dependencies list */
-	const FChunkDependencyTreeNode* BuildChunkDependencyGraph(int32 HighestChunk);
+	UNREALED_API const FChunkDependencyTreeNode* BuildChunkDependencyGraph(int32 HighestChunk);
 
 	/** Removes redundant chunks from a chunk list */
-	void RemoveRedundantChunks(TArray<int32>& ChunkIDs) const;
+	UNREALED_API void RemoveRedundantChunks(TArray<int32>& ChunkIDs) const;
 
 	/** List of dependencies used to remove redundant chunks */
 	UPROPERTY(config)
diff --git a/Engine/Source/Editor/UnrealEd/Public/ComponentAssetBroker.h b/Engine/Source/Editor/UnrealEd/Public/ComponentAssetBroker.h
--- a/Engine/Source/Editor/UnrealEd/Public/ComponentAssetBroker.h
+++ b/Engine/Source/Editor/UnrealEd/Public/ComponentAssetBroker.h
@@ -42,48 +42,48 @@ public:
 // FComponentAssetBrokerageage
 
 /** Utility class that associates assets with component classes */
-class UNREALED_API FComponentAssetBrokerage
+class FComponentAssetBrokerage
 {
 public:
 	/** Find set of components that support this asset */
-	static TArray< TSubclassOf<UActorComponent> > GetComponentsForAsset(const UObject* InAsset);
+	static UNREALED_API TArray< TSubclassOf<UActorComponent> > GetComponentsForAsset(const UObject* InAsset);
 
 	/** Get the primary component for this asset*/
-	static TSubclassOf<UActorComponent> GetPrimaryComponentForAsset(UClass* InAssetClass);
+	static UNREALED_API TSubclassOf<UActorComponent> GetPrimaryComponentForAsset(UClass* InAssetClass);
 
 	/** Assign the assigned asset to the supplied component */
-	static bool AssignAssetToComponent(UActorComponent* InComponent, UObject* InAsset);
+	static UNREALED_API bool AssignAssetToComponent(UActorComponent* InComponent, UObject* InAsset);
 
 	/** Get the currently assigned asset from the component */
-	static UObject* GetAssetFromComponent(UActorComponent* InComponent);
+	static UNREALED_API UObject* GetAssetFromComponent(UActorComponent* InComponent);
 
 	/** See if this component supports assets of any type */
-	static bool SupportsAssets(UActorComponent* InComponent);
+	static UNREALED_API bool SupportsAssets(UActorComponent* InComponent);
 
 	/** Register a component class for a specified asset class */
-	static void RegisterAssetToComponentMapping(UClass* InAssetClass, TSubclassOf<UActorComponent> InComponentClass, bool bSetAsPrimary);
+	static UNREALED_API void RegisterAssetToComponentMapping(UClass* InAssetClass, TSubclassOf<UActorComponent> InComponentClass, bool bSetAsPrimary);
 
 	/** Unregister a component type for a specified asset class */
-	static void UnregisterAssetToComponentMapping(UClass* InAssetClass, TSubclassOf<UActorComponent> InComponentClass);
+	static UNREALED_API void UnregisterAssetToComponentMapping(UClass* InAssetClass, TSubclassOf<UActorComponent> InComponentClass);
 
 	/** Try to find the broker for the specified component type */
-	static TSharedPtr<IComponentAssetBroker> FindBrokerByComponentType(TSubclassOf<UActorComponent> InComponentClass);
+	static UNREALED_API TSharedPtr<IComponentAssetBroker> FindBrokerByComponentType(TSubclassOf<UActorComponent> InComponentClass);
 
 	/** Try to find the *primary* broker for the specified asset type */
-	static TSharedPtr<IComponentAssetBroker> FindBrokerByAssetType(UClass* InAssetClass);
+	static UNREALED_API TSharedPtr<IComponentAssetBroker> FindBrokerByAssetType(UClass* InAssetClass);
 
 	/** Get the currently supported assets, optionally filtered by the supplied component class */
-	static TArray<UClass*> GetSupportedAssets(UClass* InFilterComponentClass = NULL);
+	static UNREALED_API TArray<UClass*> GetSupportedAssets(UClass* InFilterComponentClass = NULL);
 
 
 	/** Register a component class for a specified asset class */
-	static void RegisterBroker(TSharedPtr<IComponentAssetBroker> Broker, TSubclassOf<UActorComponent> InComponentClass, bool bSetAsPrimary, bool bMapComponentForAssets);
+	static UNREALED_API void RegisterBroker(TSharedPtr<IComponentAssetBroker> Broker, TSubclassOf<UActorComponent> InComponentClass, bool bSetAsPrimary, bool bMapComponentForAssets);
 
 	/** Unregister a component type for a specified asset class */
-	static void UnregisterBroker(TSharedPtr<IComponentAssetBroker> Broker);
+	static UNREALED_API void UnregisterBroker(TSharedPtr<IComponentAssetBroker> Broker);
 
 	/** Shut down the brokerage; should only be called by the editor during shutdown */
-	static void PRIVATE_ShutdownBrokerage();
+	static UNREALED_API void PRIVATE_ShutdownBrokerage();
 
 private:
 	FComponentAssetBrokerage() {}
diff --git a/Engine/Source/Editor/UnrealEd/Public/ComponentTypeRegistry.h b/Engine/Source/Editor/UnrealEd/Public/ComponentTypeRegistry.h
--- a/Engine/Source/Editor/UnrealEd/Public/ComponentTypeRegistry.h
+++ b/Engine/Source/Editor/UnrealEd/Public/ComponentTypeRegistry.h
@@ -24,9 +24,9 @@ struct FComponentTypeEntry
 	TObjectPtr<class UClass> ComponentClass;
 };
 
-struct UNREALED_API FComponentTypeRegistry
+struct FComponentTypeRegistry
 {
-	static FComponentTypeRegistry& Get();
+	static UNREALED_API FComponentTypeRegistry& Get();
 
 	/**
 	 * Called when the user changes the text in the search box.
@@ -34,22 +34,22 @@ struct UNREALED_API FComponentTypeRegistry
 	 * @return Deleate that can be used to handle change notifications. change notifications are raised when entries are 
 	 *	added or removed from the component type list
 	 */
-	FOnComponentTypeListChanged& SubscribeToComponentList(TArray<FComponentClassComboEntryPtr>*& OutComponentList);
-	FOnComponentTypeListChanged& SubscribeToComponentList(const TArray<FComponentTypeEntry>*& OutComponentList);
-	FOnComponentTypeListChanged& GetOnComponentTypeListChanged();
+	UNREALED_API FOnComponentTypeListChanged& SubscribeToComponentList(TArray<FComponentClassComboEntryPtr>*& OutComponentList);
+	UNREALED_API FOnComponentTypeListChanged& SubscribeToComponentList(const TArray<FComponentTypeEntry>*& OutComponentList);
+	UNREALED_API FOnComponentTypeListChanged& GetOnComponentTypeListChanged();
 
 	/**
 	 * Called when a specific class has been updated and should force the component type registry to update as well
 	 */
-	void InvalidateClass(TSubclassOf<UActorComponent> ClassToUpdate);
+	UNREALED_API void InvalidateClass(TSubclassOf<UActorComponent> ClassToUpdate);
 
 	/** Schedules a full update of the component type registry on the next frame */
-	void Invalidate();
+	UNREALED_API void Invalidate();
 
 	/**
 	 * Attempts to locate the class entry corresponding to the given object path.
 	 */
-	FComponentClassComboEntryPtr FindClassEntryForObjectPath(FTopLevelAssetPath InObjectPath) const;
+	UNREALED_API FComponentClassComboEntryPtr FindClassEntryForObjectPath(FTopLevelAssetPath InObjectPath) const;
 
 private:
 	void OnReloadComplete(EReloadCompleteReason Reason);
diff --git a/Engine/Source/Editor/UnrealEd/Public/ComponentVisualizer.h b/Engine/Source/Editor/UnrealEd/Public/ComponentVisualizer.h
--- a/Engine/Source/Editor/UnrealEd/Public/ComponentVisualizer.h
+++ b/Engine/Source/Editor/UnrealEd/Public/ComponentVisualizer.h
@@ -93,7 +93,7 @@ public:
  * Describes a chain of properties from the parent actor of a given component, to the component itself.
  */
 USTRUCT()
-struct UNREALED_API FComponentPropertyPath
+struct FComponentPropertyPath
 {
 public:
 	 GENERATED_USTRUCT_BODY()
@@ -113,10 +113,10 @@ public:
 	AActor* GetParentOwningActor() const { return ParentOwningActor.Get(); }
 
 	/** Gets a pointer to the component, or nullptr if it is not valid */
-	UActorComponent* GetComponent() const;
+	UNREALED_API UActorComponent* GetComponent() const;
 
 	/** Determines whether the property path is valid or not */
-	bool IsValid() const;
+	UNREALED_API bool IsValid() const;
 
 	bool operator ==(const FComponentPropertyPath& InRHS) const
 	{
@@ -131,7 +131,7 @@ public:
 private:
 
 	/** Sets the component referred to by the object */
-	void Set(const UActorComponent* Component);
+	UNREALED_API void Set(const UActorComponent* Component);
 
 	UPROPERTY()
 	TWeakObjectPtr<AActor> ParentOwningActor;
@@ -146,7 +146,7 @@ private:
 
 
 /** Base class for a component visualizer, that draw editor information for a particular component class */
-class UNREALED_API FComponentVisualizer : public TSharedFromThis<FComponentVisualizer>
+class FComponentVisualizer : public TSharedFromThis<FComponentVisualizer>
 {
 public:
 	FComponentVisualizer() {}
@@ -201,17 +201,17 @@ public:
 
 	/** Find the name of the property that points to this component */
 	UE_DEPRECATED(4.24, "Please use the FComponentPropertyPath class to build property name paths for components.")
-	static FPropertyNameAndIndex GetComponentPropertyName(const UActorComponent* Component);
+	static UNREALED_API FPropertyNameAndIndex GetComponentPropertyName(const UActorComponent* Component);
 
 	/** Get a component pointer from the property name */
 	UE_DEPRECATED(4.24, "Please use the FComponentPropertyPath::GetComponent() to retrieve a component pointer from a property name path.")
-	static UActorComponent* GetComponentFromPropertyName(const AActor* CompOwner, const FPropertyNameAndIndex& Property);
+	static UNREALED_API UActorComponent* GetComponentFromPropertyName(const AActor* CompOwner, const FPropertyNameAndIndex& Property);
 
 	/** Notify that a component property has been modified */
-	static void NotifyPropertyModified(UActorComponent* Component, FProperty* Property, EPropertyChangeType::Type PropertyChangeType = EPropertyChangeType::Unspecified);
+	static UNREALED_API void NotifyPropertyModified(UActorComponent* Component, FProperty* Property, EPropertyChangeType::Type PropertyChangeType = EPropertyChangeType::Unspecified);
 
 	/** Notify that many component properties have been modified */
-	static void NotifyPropertiesModified(UActorComponent* Component, const TArray<FProperty*>& Properties, EPropertyChangeType::Type PropertyChangeType = EPropertyChangeType::Unspecified);
+	static UNREALED_API void NotifyPropertiesModified(UActorComponent* Component, const TArray<FProperty*>& Properties, EPropertyChangeType::Type PropertyChangeType = EPropertyChangeType::Unspecified);
 };
 
 struct FCachedComponentVisualizer
