[UnrealEd]

diff --git a/Engine/Source/Editor/UnrealEd/Public/ComponentVisualizerManager.h b/Engine/Source/Editor/UnrealEd/Public/ComponentVisualizerManager.h
--- a/Engine/Source/Editor/UnrealEd/Public/ComponentVisualizerManager.h
+++ b/Engine/Source/Editor/UnrealEd/Public/ComponentVisualizerManager.h
@@ -12,63 +12,63 @@ class SWidget;
 struct FViewportClick;
 
 /** Class that managed active component visualizer and routes input to it */
-class UNREALED_API FComponentVisualizerManager
+class FComponentVisualizerManager
 {
 public:
-	FComponentVisualizerManager();
+	UNREALED_API FComponentVisualizerManager();
 	virtual ~FComponentVisualizerManager() {}
 
 
 	/** Activate a component visualizer given a clicked proxy */
-	bool HandleProxyForComponentVis(FEditorViewportClient* InViewportClient, HHitProxy *HitProxy, const FViewportClick &Click);
+	UNREALED_API bool HandleProxyForComponentVis(FEditorViewportClient* InViewportClient, HHitProxy *HitProxy, const FViewportClick &Click);
 
 	/** Activate a component visualizer given the specific visualizer */
-	bool SetActiveComponentVis(FEditorViewportClient* InViewportClient, TSharedPtr<FComponentVisualizer>& InVisualizer);
+	UNREALED_API bool SetActiveComponentVis(FEditorViewportClient* InViewportClient, TSharedPtr<FComponentVisualizer>& InVisualizer);
 
 	/** Clear active component visualizer */
-	void ClearActiveComponentVis();
+	UNREALED_API void ClearActiveComponentVis();
 
 	/** Handle a click on the specified level editor viewport client */
-	bool HandleClick(FEditorViewportClient* InViewportClient, HHitProxy *HitProxy, const FViewportClick &Click);
+	UNREALED_API bool HandleClick(FEditorViewportClient* InViewportClient, HHitProxy *HitProxy, const FViewportClick &Click);
 
 	/** Pass key input to active visualizer */
-	bool HandleInputKey(FEditorViewportClient* InViewportClient, FViewport* Viewport, FKey Key, EInputEvent Event) const;
+	UNREALED_API bool HandleInputKey(FEditorViewportClient* InViewportClient, FViewport* Viewport, FKey Key, EInputEvent Event) const;
 
 	/** Pass delta input to active visualizer */
-	bool HandleInputDelta(FEditorViewportClient* InViewportClient, FViewport* InViewport, FVector& InDrag, FRotator& InRot, FVector& InScale) const;
+	UNREALED_API bool HandleInputDelta(FEditorViewportClient* InViewportClient, FViewport* InViewport, FVector& InDrag, FRotator& InRot, FVector& InScale) const;
 
 	/** Pass box select input to active visualizer */
-	bool HandleBoxSelect(const FBox& InBox, FEditorViewportClient* InViewportClient, FViewport* InViewport) const;
+	UNREALED_API bool HandleBoxSelect(const FBox& InBox, FEditorViewportClient* InViewportClient, FViewport* InViewport) const;
 
 	/** Pass frustum select input to active visualizer */
-	bool HandleFrustumSelect(const FConvexVolume &InFrustum, FEditorViewportClient* InViewportClient, FViewport* InViewport) const;
+	UNREALED_API bool HandleFrustumSelect(const FConvexVolume &InFrustum, FEditorViewportClient* InViewportClient, FViewport* InViewport) const;
 
 	/** Return whether focus on selection should focus on bounding box defined by active visualizer */
-	bool HasFocusOnSelectionBoundingBox(FBox& OutBoundingBox) const;
+	UNREALED_API bool HasFocusOnSelectionBoundingBox(FBox& OutBoundingBox) const;
 
 	/** Pass snap input to active visualizer */
-	bool HandleSnapTo(const bool bInAlign, const bool bInUseLineTrace, const bool bInUseBounds, const bool bInUsePivot, AActor* InDestination);
+	UNREALED_API bool HandleSnapTo(const bool bInAlign, const bool bInUseLineTrace, const bool bInUseBounds, const bool bInUsePivot, AActor* InDestination);
 
 	/** Get widget location from active visualizer */
-	bool GetWidgetLocation(const FEditorViewportClient* InViewportClient, FVector& OutLocation) const;
+	UNREALED_API bool GetWidgetLocation(const FEditorViewportClient* InViewportClient, FVector& OutLocation) const;
 
 	/** Get custom widget coordinate system from active visualizer */
-	bool GetCustomInputCoordinateSystem(const FEditorViewportClient* InViewportClient, FMatrix& OutMatrix) const;
+	UNREALED_API bool GetCustomInputCoordinateSystem(const FEditorViewportClient* InViewportClient, FMatrix& OutMatrix) const;
 
 	/** Gets called when the mouse tracking has started (dragging behavior) */
-	void TrackingStarted(FEditorViewportClient* InViewportClient);
+	UNREALED_API void TrackingStarted(FEditorViewportClient* InViewportClient);
 
 	/** Gets called when the mouse tracking has stopped (dragging behavior) */
-	void TrackingStopped(FEditorViewportClient* InViewportClient, bool bInDidMove);
+	UNREALED_API void TrackingStopped(FEditorViewportClient* InViewportClient, bool bInDidMove);
 
 	/** Generate context menu for the component visualizer */
-	TSharedPtr<SWidget> GenerateContextMenuForComponentVis() const;
+	UNREALED_API TSharedPtr<SWidget> GenerateContextMenuForComponentVis() const;
 
 	/** Returns whether there is currently an active visualizer */
-	bool IsActive() const;
+	UNREALED_API bool IsActive() const;
 
 	/** Returns whether the component being visualized is an archetype or not */
-	bool IsVisualizingArchetype() const;
+	UNREALED_API bool IsVisualizingArchetype() const;
 
 private:
 	/** Currently 'active' visualizer that we should pass input to etc */
@@ -76,4 +76,4 @@ private:
 
 	/** The viewport client for the currently active visualizer */
 	FEditorViewportClient* EditedVisualizerViewportClient;
-};
\ No newline at end of file
+};
diff --git a/Engine/Source/Editor/UnrealEd/Public/CookPackageSplitter.h b/Engine/Source/Editor/UnrealEd/Public/CookPackageSplitter.h
--- a/Engine/Source/Editor/UnrealEd/Public/CookPackageSplitter.h
+++ b/Engine/Source/Editor/UnrealEd/Public/CookPackageSplitter.h
@@ -226,22 +226,22 @@ namespace Private
 {
 
 /** Interface for internal use only (used by REGISTER_COOKPACKAGE_SPLITTER to register an ICookPackageSplitter for a class) */
-class UNREALED_API FRegisteredCookPackageSplitter
+class FRegisteredCookPackageSplitter
 {
 public:
-	FRegisteredCookPackageSplitter();
-	virtual ~FRegisteredCookPackageSplitter();
+	UNREALED_API FRegisteredCookPackageSplitter();
+	UNREALED_API virtual ~FRegisteredCookPackageSplitter();
 
 	virtual UClass* GetSplitDataClass() const = 0;
 	virtual bool ShouldSplitPackage(UObject* Object) const = 0;
 	virtual ICookPackageSplitter* CreateInstance(UObject* Object) const = 0;
 	virtual FString GetSplitterDebugName() const = 0;
 
-	static void ForEach(TFunctionRef<void(FRegisteredCookPackageSplitter*)> Func);
+	static UNREALED_API void ForEach(TFunctionRef<void(FRegisteredCookPackageSplitter*)> Func);
 
 private:
 	
-	static TLinkedList<FRegisteredCookPackageSplitter*>*& GetRegisteredList();
+	static UNREALED_API TLinkedList<FRegisteredCookPackageSplitter*>*& GetRegisteredList();
 	TLinkedList<FRegisteredCookPackageSplitter*> GlobalListLink;
 };
 
@@ -270,4 +270,4 @@ namespace PREPROCESSOR_JOIN(SplitterClass, SplitDataClass) \
 	static PREPROCESSOR_JOIN(PREPROCESSOR_JOIN(SplitterClass, SplitDataClass), _Register) DefaultObject; \
 }
 
-#endif
\ No newline at end of file
+#endif
diff --git a/Engine/Source/Editor/UnrealEd/Public/CrashReporterSettings.h b/Engine/Source/Editor/UnrealEd/Public/CrashReporterSettings.h
--- a/Engine/Source/Editor/UnrealEd/Public/CrashReporterSettings.h
+++ b/Engine/Source/Editor/UnrealEd/Public/CrashReporterSettings.h
@@ -14,8 +14,8 @@
 /**
  * Implements per-project cooker settings exposed to the editor
  */
-UCLASS(config = EditorPerProjectUserSettings)
-class UNREALED_API UCrashReporterSettings
+UCLASS(config = EditorPerProjectUserSettings, MinimalAPI)
+class UCrashReporterSettings
 	: public UObject
 {
 	GENERATED_UCLASS_BODY()
diff --git a/Engine/Source/Editor/UnrealEd/Public/CurveTableEditorUtils.h b/Engine/Source/Editor/UnrealEd/Public/CurveTableEditorUtils.h
--- a/Engine/Source/Editor/UnrealEd/Public/CurveTableEditorUtils.h
+++ b/Engine/Source/Editor/UnrealEd/Public/CurveTableEditorUtils.h
@@ -8,7 +8,7 @@
 
 class UCurveTable;
 
-struct UNREALED_API FCurveTableEditorUtils
+struct FCurveTableEditorUtils
 {
 	enum class ECurveTableChangeInfo
 	{
@@ -30,7 +30,7 @@ struct UNREALED_API FCurveTableEditorUtils
 	public:
 		UNREALED_API static FCurveTableEditorManager& Get();
 
-		class UNREALED_API ListenerType : public InnerListenerType<FCurveTableEditorManager>
+		class ListenerType : public InnerListenerType<FCurveTableEditorManager>
 		{
 		public:
 			virtual void SelectionChange(const UCurveTable* CurveTable, FName RowName) { }
@@ -39,6 +39,6 @@ struct UNREALED_API FCurveTableEditorUtils
 
 	typedef FCurveTableEditorManager::ListenerType INotifyOnCurveTableChanged;
 
-	static void BroadcastPreChange(UCurveTable* DataTable, ECurveTableChangeInfo Info);
-	static void BroadcastPostChange(UCurveTable* DataTable, ECurveTableChangeInfo Info);
+	static UNREALED_API void BroadcastPreChange(UCurveTable* DataTable, ECurveTableChangeInfo Info);
+	static UNREALED_API void BroadcastPostChange(UCurveTable* DataTable, ECurveTableChangeInfo Info);
 };
diff --git a/Engine/Source/Editor/UnrealEd/Public/DataTableEditorUtils.h b/Engine/Source/Editor/UnrealEd/Public/DataTableEditorUtils.h
--- a/Engine/Source/Editor/UnrealEd/Public/DataTableEditorUtils.h
+++ b/Engine/Source/Editor/UnrealEd/Public/DataTableEditorUtils.h
@@ -52,7 +52,7 @@ enum class ERowInsertionPosition
 	Bottom,
 };
 
-struct UNREALED_API FDataTableEditorUtils
+struct FDataTableEditorUtils
 {
 	enum class EDataTableChangeInfo
 	{
@@ -74,7 +74,7 @@ struct UNREALED_API FDataTableEditorUtils
 	public:
 		UNREALED_API static FDataTableEditorManager& Get();
 
-		class UNREALED_API ListenerType : public InnerListenerType<FDataTableEditorManager>
+		class ListenerType : public InnerListenerType<FDataTableEditorManager>
 		{
 		public:
 			virtual void SelectionChange(const UDataTable* DataTable, FName RowName) { }
@@ -83,53 +83,53 @@ struct UNREALED_API FDataTableEditorUtils
 
 	typedef FDataTableEditorManager::ListenerType INotifyOnDataTableChanged;
 
-	static bool RemoveRow(UDataTable* DataTable, FName Name);
-	static uint8* AddRow(UDataTable* DataTable, FName RowName);
-	static uint8* DuplicateRow(UDataTable* DataTable, FName SourceRowName, FName RowName);
-	static bool RenameRow(UDataTable* DataTable, FName OldName, FName NewName);
-	static bool MoveRow(UDataTable* DataTable, FName RowName, ERowMoveDirection Direction, int32 NumRowsToMoveBy = 1);
-	static bool SelectRow(const UDataTable* DataTable, FName RowName);
-	static bool DiffersFromDefault(UDataTable* DataTable, FName RowName);
-	static bool ResetToDefault(UDataTable* DataTable, FName RowName);
+	static UNREALED_API bool RemoveRow(UDataTable* DataTable, FName Name);
+	static UNREALED_API uint8* AddRow(UDataTable* DataTable, FName RowName);
+	static UNREALED_API uint8* DuplicateRow(UDataTable* DataTable, FName SourceRowName, FName RowName);
+	static UNREALED_API bool RenameRow(UDataTable* DataTable, FName OldName, FName NewName);
+	static UNREALED_API bool MoveRow(UDataTable* DataTable, FName RowName, ERowMoveDirection Direction, int32 NumRowsToMoveBy = 1);
+	static UNREALED_API bool SelectRow(const UDataTable* DataTable, FName RowName);
+	static UNREALED_API bool DiffersFromDefault(UDataTable* DataTable, FName RowName);
+	static UNREALED_API bool ResetToDefault(UDataTable* DataTable, FName RowName);
 
-	static uint8* AddRowAboveOrBelowSelection(UDataTable* DataTable, const FName& RowName, const FName& NewRowName, ERowInsertionPosition InsertPosition);
+	static UNREALED_API uint8* AddRowAboveOrBelowSelection(UDataTable* DataTable, const FName& RowName, const FName& NewRowName, ERowInsertionPosition InsertPosition);
 
-	static void BroadcastPreChange(UDataTable* DataTable, EDataTableChangeInfo Info);
-	static void BroadcastPostChange(UDataTable* DataTable, EDataTableChangeInfo Info);
+	static UNREALED_API void BroadcastPreChange(UDataTable* DataTable, EDataTableChangeInfo Info);
+	static UNREALED_API void BroadcastPostChange(UDataTable* DataTable, EDataTableChangeInfo Info);
 
 	/** Reads a data table and parses out editable copies of rows and columns */
-	static void CacheDataTableForEditing(const UDataTable* DataTable, TArray<FDataTableEditorColumnHeaderDataPtr>& OutAvailableColumns, TArray<FDataTableEditorRowListViewDataPtr>& OutAvailableRows);
+	static UNREALED_API void CacheDataTableForEditing(const UDataTable* DataTable, TArray<FDataTableEditorColumnHeaderDataPtr>& OutAvailableColumns, TArray<FDataTableEditorRowListViewDataPtr>& OutAvailableRows);
 
 	/** Generic version that works with any datatable-like structure */
-	static void CacheDataForEditing(const UScriptStruct* RowStruct, const TMap<FName, uint8*>& RowMap, TArray<FDataTableEditorColumnHeaderDataPtr>& OutAvailableColumns, TArray<FDataTableEditorRowListViewDataPtr>& OutAvailableRows);
+	static UNREALED_API void CacheDataForEditing(const UScriptStruct* RowStruct, const TMap<FName, uint8*>& RowMap, TArray<FDataTableEditorColumnHeaderDataPtr>& OutAvailableColumns, TArray<FDataTableEditorRowListViewDataPtr>& OutAvailableRows);
 
 	/** Returns all script structs that can be used as a data table row. This only includes loaded ones */
-	static TArray<UScriptStruct*> GetPossibleStructs();
+	static UNREALED_API TArray<UScriptStruct*> GetPossibleStructs();
 
 	/** Fills in an array with all possible DataTable structs, unloaded and loaded */
-	static void GetPossibleStructAssetData(TArray<FAssetData>& StructAssets);
+	static UNREALED_API void GetPossibleStructAssetData(TArray<FAssetData>& StructAssets);
 	
 	/** Utility function which verifies that the specified struct type is viable for data tables */
-	static bool IsValidTableStruct(const UScriptStruct* Struct);
+	static UNREALED_API bool IsValidTableStruct(const UScriptStruct* Struct);
 
 	/** Add a UI action for search for references, useful for customizations */
-	static void AddSearchForReferencesContextMenu(class FDetailWidgetRow& RowNameDetailWidget, FExecuteAction SearchForReferencesAction);
+	static UNREALED_API void AddSearchForReferencesContextMenu(class FDetailWidgetRow& RowNameDetailWidget, FExecuteAction SearchForReferencesAction);
 
 	/** Short description for a data or curve handle */
-	static FText GetHandleShortDescription(const UObject* TableAsset, FName RowName);
+	static UNREALED_API FText GetHandleShortDescription(const UObject* TableAsset, FName RowName);
 
 	/** Tooltip text for the data table row type */
-	static FText GetRowTypeInfoTooltipText(FDataTableEditorColumnHeaderDataPtr ColumnHeaderDataPtr);
+	static UNREALED_API FText GetRowTypeInfoTooltipText(FDataTableEditorColumnHeaderDataPtr ColumnHeaderDataPtr);
 
 	/** Doc excerpt name for the data table row type */
-	static FString GetRowTypeTooltipDocExcerptName(FDataTableEditorColumnHeaderDataPtr ColumnHeaderDataPtr);
+	static UNREALED_API FString GetRowTypeTooltipDocExcerptName(FDataTableEditorColumnHeaderDataPtr ColumnHeaderDataPtr);
 
 	/** Link to variable type doc  */
-	static const FString VariableTypesTooltipDocLink;
+	static UNREALED_API const FString VariableTypesTooltipDocLink;
 
 	/** Delegate called when a data table struct is selected */
 	DECLARE_DELEGATE_OneParam(FOnDataTableStructSelected, UScriptStruct*);
 
 	/** Creates a combo box that allows selecting from the list of possible row structures */
-	static TSharedRef<SWidget> MakeRowStructureComboBox(FOnDataTableStructSelected OnSelected);
+	static UNREALED_API TSharedRef<SWidget> MakeRowStructureComboBox(FOnDataTableStructSelected OnSelected);
 };
diff --git a/Engine/Source/Editor/UnrealEd/Public/Dialogs/Dialogs.h b/Engine/Source/Editor/UnrealEd/Public/Dialogs/Dialogs.h
--- a/Engine/Source/Editor/UnrealEd/Public/Dialogs/Dialogs.h
+++ b/Engine/Source/Editor/UnrealEd/Public/Dialogs/Dialogs.h
@@ -73,7 +73,7 @@ public:
 /**
  * A Dialog that displays a warning message to the user and provides the option to not display it in the future
  */
-class UNREALED_API FSuppressableWarningDialog 
+class FSuppressableWarningDialog 
 {
 public:
 
@@ -146,10 +146,10 @@ public:
 	 * 
 	 * @param FSetupInfo Info - struct used to initialize the dialog. 
 	 */
-	FSuppressableWarningDialog ( const FSuppressableWarningDialog::FSetupInfo& Info );
+	UNREALED_API FSuppressableWarningDialog ( const FSuppressableWarningDialog::FSetupInfo& Info );
 	
 	/** Launches warning window, returns user response or suppressed */
-	EResult ShowModal() const;
+	UNREALED_API EResult ShowModal() const;
 
 private:
 
diff --git a/Engine/Source/Editor/UnrealEd/Public/Dialogs/DlgPickAssetPath.h b/Engine/Source/Editor/UnrealEd/Public/Dialogs/DlgPickAssetPath.h
--- a/Engine/Source/Editor/UnrealEd/Public/Dialogs/DlgPickAssetPath.h
+++ b/Engine/Source/Editor/UnrealEd/Public/Dialogs/DlgPickAssetPath.h
@@ -24,7 +24,7 @@ public:
 
 	SLATE_END_ARGS()
 
-	UNREALED_API SDlgPickAssetPath()
+	SDlgPickAssetPath()
 	:	UserResponse(EAppReturnType::Cancel)
 	{
 	}
diff --git a/Engine/Source/Editor/UnrealEd/Public/Dialogs/DlgPickPath.h b/Engine/Source/Editor/UnrealEd/Public/Dialogs/DlgPickPath.h
--- a/Engine/Source/Editor/UnrealEd/Public/Dialogs/DlgPickPath.h
+++ b/Engine/Source/Editor/UnrealEd/Public/Dialogs/DlgPickPath.h
@@ -24,7 +24,7 @@ public:
 
 	SLATE_END_ARGS()
 
-	UNREALED_API SDlgPickPath()
+	SDlgPickPath()
 	:	UserResponse(EAppReturnType::Cancel)
 	{
 	}
diff --git a/Engine/Source/Editor/UnrealEd/Public/Dialogs/DlgReferenceTree.h b/Engine/Source/Editor/UnrealEd/Public/Dialogs/DlgReferenceTree.h
--- a/Engine/Source/Editor/UnrealEd/Public/Dialogs/DlgReferenceTree.h
+++ b/Engine/Source/Editor/UnrealEd/Public/Dialogs/DlgReferenceTree.h
@@ -114,7 +114,7 @@ typedef TSharedPtr< struct FReferenceTreeDataContainer > FReferenceTreeItemPtr;
 
 typedef STreeView< FReferenceTreeItemPtr > SReferenceTreeView;
 
-class UNREALED_API SReferenceTree : public SCompoundWidget
+class SReferenceTree : public SCompoundWidget
 {
 public:
 	SLATE_BEGIN_ARGS( SReferenceTree )
@@ -123,12 +123,12 @@ public:
 		SLATE_ARGUMENT(UObject*, Object)
 	SLATE_END_ARGS()
 
-	void Construct( const FArguments& InArgs );
+	UNREALED_API void Construct( const FArguments& InArgs );
 
-	~SReferenceTree();
+	UNREALED_API ~SReferenceTree();
 
-	static TWeakPtr< SWindow > SingletonInstance;
-	static void OpenDialog(UObject* InObject);
+	static UNREALED_API TWeakPtr< SWindow > SingletonInstance;
+	static UNREALED_API void OpenDialog(UObject* InObject);
 
 private:
 	/** Generates a row for the tree. */
diff --git a/Engine/Source/Editor/UnrealEd/Public/DiffUtils.h b/Engine/Source/Editor/UnrealEd/Public/DiffUtils.h
--- a/Engine/Source/Editor/UnrealEd/Public/DiffUtils.h
+++ b/Engine/Source/Editor/UnrealEd/Public/DiffUtils.h
@@ -268,7 +268,7 @@ namespace DiffUtils
 DECLARE_DELEGATE(FOnDiffEntryFocused);
 DECLARE_DELEGATE_RetVal(TSharedRef<SWidget>, FGenerateDiffEntryWidget);
 
-class UNREALED_API FBlueprintDifferenceTreeEntry
+class FBlueprintDifferenceTreeEntry
 {
 public:
 	FBlueprintDifferenceTreeEntry(FOnDiffEntryFocused InOnFocus, FGenerateDiffEntryWidget InGenerateWidget, TArray< TSharedPtr<FBlueprintDifferenceTreeEntry> > InChildren = TArray< TSharedPtr<FBlueprintDifferenceTreeEntry> >())
@@ -282,16 +282,16 @@ public:
 	virtual ~FBlueprintDifferenceTreeEntry() = default;
 
 	/** Displays message to user saying there are no differences */
-	static TSharedPtr<FBlueprintDifferenceTreeEntry> NoDifferencesEntry();
+	static UNREALED_API TSharedPtr<FBlueprintDifferenceTreeEntry> NoDifferencesEntry();
 
 	/** Displays message to user warning that there may be undetected differences */
-	static TSharedPtr<FBlueprintDifferenceTreeEntry> UnknownDifferencesEntry();
+	static UNREALED_API TSharedPtr<FBlueprintDifferenceTreeEntry> UnknownDifferencesEntry();
 
 	/** Create category message for the diff UI */
-	static TSharedPtr<FBlueprintDifferenceTreeEntry> CreateCategoryEntry(const FText& LabelText, const FText& ToolTipText, FOnDiffEntryFocused FocusCallback, const TArray< TSharedPtr<FBlueprintDifferenceTreeEntry> >& Children, bool bHasDifferences);
+	static UNREALED_API TSharedPtr<FBlueprintDifferenceTreeEntry> CreateCategoryEntry(const FText& LabelText, const FText& ToolTipText, FOnDiffEntryFocused FocusCallback, const TArray< TSharedPtr<FBlueprintDifferenceTreeEntry> >& Children, bool bHasDifferences);
 
 	/** Create category message for the merge UI */
-	static TSharedPtr<FBlueprintDifferenceTreeEntry> CreateCategoryEntryForMerge(const FText& LabelText, const FText& ToolTipText, FOnDiffEntryFocused FocusCallback, const TArray< TSharedPtr<FBlueprintDifferenceTreeEntry> >& Children, bool bHasRemoteDifferences, bool bHasLocalDifferences, bool bHasConflicts);
+	static UNREALED_API TSharedPtr<FBlueprintDifferenceTreeEntry> CreateCategoryEntryForMerge(const FText& LabelText, const FText& ToolTipText, FOnDiffEntryFocused FocusCallback, const TArray< TSharedPtr<FBlueprintDifferenceTreeEntry> >& Children, bool bHasRemoteDifferences, bool bHasLocalDifferences, bool bHasConflicts);
 	
 	FOnDiffEntryFocused OnFocus;
 	FGenerateDiffEntryWidget GenerateWidget;
diff --git a/Engine/Source/Editor/UnrealEd/Public/DragAndDrop/AssetDragDropOp.h b/Engine/Source/Editor/UnrealEd/Public/DragAndDrop/AssetDragDropOp.h
--- a/Engine/Source/Editor/UnrealEd/Public/DragAndDrop/AssetDragDropOp.h
+++ b/Engine/Source/Editor/UnrealEd/Public/DragAndDrop/AssetDragDropOp.h
@@ -11,20 +11,20 @@
 class FAssetThumbnail;
 class UActorFactory;
 
-class UNREALED_API FAssetDragDropOp : public FDecoratedDragDropOp
+class FAssetDragDropOp : public FDecoratedDragDropOp
 {
 public:
 	DRAG_DROP_OPERATOR_TYPE(FAssetDragDropOp, FDecoratedDragDropOp)
 
-	static TSharedRef<FAssetDragDropOp> New(const FAssetData& InAssetData, UActorFactory* ActorFactory = nullptr);
+	static UNREALED_API TSharedRef<FAssetDragDropOp> New(const FAssetData& InAssetData, UActorFactory* ActorFactory = nullptr);
 
-	static TSharedRef<FAssetDragDropOp> New(TArray<FAssetData> InAssetData, UActorFactory* ActorFactory = nullptr);
+	static UNREALED_API TSharedRef<FAssetDragDropOp> New(TArray<FAssetData> InAssetData, UActorFactory* ActorFactory = nullptr);
 
-	static TSharedRef<FAssetDragDropOp> New(FString InAssetPath);
+	static UNREALED_API TSharedRef<FAssetDragDropOp> New(FString InAssetPath);
 
-	static TSharedRef<FAssetDragDropOp> New(TArray<FString> InAssetPaths);
+	static UNREALED_API TSharedRef<FAssetDragDropOp> New(TArray<FString> InAssetPaths);
 
-	static TSharedRef<FAssetDragDropOp> New(TArray<FAssetData> InAssetData, TArray<FString> InAssetPaths, UActorFactory* ActorFactory = nullptr);
+	static UNREALED_API TSharedRef<FAssetDragDropOp> New(TArray<FAssetData> InAssetData, TArray<FString> InAssetPaths, UActorFactory* ActorFactory = nullptr);
 
 	/** @return true if this drag operation contains assets */
 	bool HasAssets() const
@@ -57,21 +57,21 @@ public:
 	}
 
 public:
-	virtual ~FAssetDragDropOp();
+	UNREALED_API virtual ~FAssetDragDropOp();
 
-	virtual TSharedPtr<SWidget> GetDefaultDecorator() const override;
+	UNREALED_API virtual TSharedPtr<SWidget> GetDefaultDecorator() const override;
 
-	FText GetDecoratorText() const;
+	UNREALED_API FText GetDecoratorText() const;
 
 protected:
-	void Init(TArray<FAssetData> InAssetData, TArray<FString> InAssetPaths, UActorFactory* InActorFactory);
-	virtual void InitThumbnail();
+	UNREALED_API void Init(TArray<FAssetData> InAssetData, TArray<FString> InAssetPaths, UActorFactory* InActorFactory);
+	UNREALED_API virtual void InitThumbnail();
 
-	virtual bool HasFiles() const;
-	virtual bool HasFolders() const;
+	UNREALED_API virtual bool HasFiles() const;
+	UNREALED_API virtual bool HasFolders() const;
 
-	virtual int32 GetTotalCount() const;
-	virtual FText GetFirstItemText() const;
+	UNREALED_API virtual int32 GetTotalCount() const;
+	UNREALED_API virtual FText GetFirstItemText() const;
 
 private:
 	/** Data for the assets this item represents */
diff --git a/Engine/Source/Editor/UnrealEd/Public/DragAndDrop/CollectionDragDropOp.h b/Engine/Source/Editor/UnrealEd/Public/DragAndDrop/CollectionDragDropOp.h
--- a/Engine/Source/Editor/UnrealEd/Public/DragAndDrop/CollectionDragDropOp.h
+++ b/Engine/Source/Editor/UnrealEd/Public/DragAndDrop/CollectionDragDropOp.h
@@ -20,7 +20,7 @@
 class SWidget;
 struct FAssetData;
 
-class UNREALED_API FCollectionDragDropOp : public FDecoratedDragDropOp
+class FCollectionDragDropOp : public FDecoratedDragDropOp
 {
 public:
 	DRAG_DROP_OPERATOR_TYPE(FCollectionDragDropOp, FDecoratedDragDropOp)
@@ -42,12 +42,12 @@ public:
 	
 public:
 	/** @return The assets from this drag operation */
-	TArray<FAssetData> GetAssets() const;
+	UNREALED_API TArray<FAssetData> GetAssets() const;
 
-	virtual TSharedPtr<SWidget> GetDefaultDecorator() const override;
+	UNREALED_API virtual TSharedPtr<SWidget> GetDefaultDecorator() const override;
 
 private:
-	FText GetDecoratorText() const;
+	UNREALED_API FText GetDecoratorText() const;
 
 	EAssetTagItemViewMode AssetTagViewMode = EAssetTagItemViewMode::Standard;
 };
diff --git a/Engine/Source/Editor/UnrealEd/Public/DragTool_FrustumSelect.h b/Engine/Source/Editor/UnrealEd/Public/DragTool_FrustumSelect.h
--- a/Engine/Source/Editor/UnrealEd/Public/DragTool_FrustumSelect.h
+++ b/Engine/Source/Editor/UnrealEd/Public/DragTool_FrustumSelect.h
@@ -14,7 +14,7 @@ class UModel;
  * Draws a box in the current viewport and when the mouse button is released,
  * it selects/unselects everything inside of it.
  */
-class UNREALED_API FDragTool_ActorFrustumSelect : public FDragTool
+class FDragTool_ActorFrustumSelect : public FDragTool
 {
 public:
 	explicit FDragTool_ActorFrustumSelect(FLevelEditorViewportClient* InLevelViewportClient)
@@ -27,7 +27,7 @@ public:
 	 *
 	 * @param	InDelta		A delta of mouse movement.
 	 */
-	virtual void AddDelta( const FVector& InDelta ) override;
+	UNREALED_API virtual void AddDelta( const FVector& InDelta ) override;
 
 	/**
 	 * Starts a mouse drag behavior.  The start location is snapped to the editor constraints if bUseSnapping is true.
@@ -35,13 +35,13 @@ public:
 	 * @param	InViewportClient	The viewport client in which the drag event occurred.
 	 * @param	InStart				Where the mouse was when the drag started.
 	 */
-	virtual void StartDrag(FEditorViewportClient* InViewportClient, const FVector& InStart, const FVector2D& InStartScreen) override;
+	UNREALED_API virtual void StartDrag(FEditorViewportClient* InViewportClient, const FVector& InStart, const FVector2D& InStartScreen) override;
 	
 	/**
 	 * Ends a mouse drag behavior (the user has let go of the mouse button).
 	 */
-	virtual void EndDrag() override;
-	virtual void Render(const FSceneView* View,FCanvas* Canvas) override;
+	UNREALED_API virtual void EndDrag() override;
+	UNREALED_API virtual void Render(const FSceneView* View,FCanvas* Canvas) override;
 
 private:
 	/** 
@@ -50,7 +50,7 @@ private:
 	 * @param OutFrustum		The created frustum
 	 * @param bUseBoxFrustum	If true a frustum out of the current dragged box will be created.  false will use the view frustum.
 	 */
-	void CalculateFrustum( FSceneView* View, FConvexVolume& OutFrustum, bool bUseBoxFrustum );
+	UNREALED_API void CalculateFrustum( FSceneView* View, FConvexVolume& OutFrustum, bool bUseBoxFrustum );
 
 	/** 
 	 * Returns true if the provided BSP node intersects with the provided frustum 
@@ -60,19 +60,19 @@ private:
 	 * @param InFrustum				The frustum to check against.
 	 * @param bUseStrictSelection	true if the node must be entirely within the frustum
 	 */
-	bool IntersectsFrustum( const UModel& InModel, int32 NodeIndex, const FConvexVolume& InFrustum, bool bUseStrictSelection ) const;
+	UNREALED_API bool IntersectsFrustum( const UModel& InModel, int32 NodeIndex, const FConvexVolume& InFrustum, bool bUseStrictSelection ) const;
 	
 	/** Adds a hover effect to the passed in actor */
-	void AddHoverEffect( AActor& InActor );
+	UNREALED_API void AddHoverEffect( AActor& InActor );
 	
 	/** Adds a hover effect to the passed in bsp surface */
-	void AddHoverEffect( UModel& InModel, int32 SurfIndex );
+	UNREALED_API void AddHoverEffect( UModel& InModel, int32 SurfIndex );
 
 	/** Removes a hover effect from the passed in actor */
-	void RemoveHoverEffect( AActor& InActor );
+	UNREALED_API void RemoveHoverEffect( AActor& InActor );
 
 	/** Removes a hover effect from the passed in bsp surface */
-	void RemoveHoverEffect( UModel& InModel, int32 SurfIndex );
+	UNREALED_API void RemoveHoverEffect( UModel& InModel, int32 SurfIndex );
 
 	/** List of actors to repeatedly check when determining hover cues */
 	TArray<AActor*> ActorsToCheck;
diff --git a/Engine/Source/Editor/UnrealEd/Public/EdGraphNode_Comment.h b/Engine/Source/Editor/UnrealEd/Public/EdGraphNode_Comment.h
--- a/Engine/Source/Editor/UnrealEd/Public/EdGraphNode_Comment.h
+++ b/Engine/Source/Editor/UnrealEd/Public/EdGraphNode_Comment.h
@@ -111,7 +111,7 @@ public:
 	UNREALED_API const FCommentNodeSet& GetNodesUnderComment() const;
 
 	/** Return the font size of the comment */
-	UNREALED_API virtual int32 GetFontSize() const { return FontSize; }
+	virtual int32 GetFontSize() const { return FontSize; }
 
 	/** Override the default selection state of this graph node */
 	enum class ESelectionState : uint8 { Inherited, Selected, Deselected };
diff --git a/Engine/Source/Editor/UnrealEd/Public/EdGraphUtilities.h b/Engine/Source/Editor/UnrealEd/Public/EdGraphUtilities.h
--- a/Engine/Source/Editor/UnrealEd/Public/EdGraphUtilities.h
+++ b/Engine/Source/Editor/UnrealEd/Public/EdGraphUtilities.h
@@ -25,7 +25,7 @@ class UObject;
 
 //////////////////////////////////////////////////////////////////////////
 
-struct UNREALED_API FGraphPanelNodeFactory : public TSharedFromThis<FGraphPanelNodeFactory>
+struct FGraphPanelNodeFactory : public TSharedFromThis<FGraphPanelNodeFactory>
 {
 public:
 	virtual ~FGraphPanelNodeFactory() {}
@@ -34,7 +34,7 @@ public:
 
 //////////////////////////////////////////////////////////////////////////
 
-struct UNREALED_API FGraphPanelPinFactory : public TSharedFromThis<FGraphPanelPinFactory>
+struct FGraphPanelPinFactory : public TSharedFromThis<FGraphPanelPinFactory>
 {
 public:
 	virtual ~FGraphPanelPinFactory() {}
@@ -43,7 +43,7 @@ public:
 
 //////////////////////////////////////////////////////////////////////////
 
-struct UNREALED_API FGraphPanelPinConnectionFactory : public TSharedFromThis<FGraphPanelPinConnectionFactory>
+struct FGraphPanelPinConnectionFactory : public TSharedFromThis<FGraphPanelPinConnectionFactory>
 {
 public:
     virtual ~FGraphPanelPinConnectionFactory() {}
@@ -52,29 +52,29 @@ public:
 
 //////////////////////////////////////////////////////////////////////////
 
-class UNREALED_API FEdGraphUtilities
+class FEdGraphUtilities
 {
 public:
-	static void RegisterVisualNodeFactory(TSharedPtr<FGraphPanelNodeFactory> NewFactory);
-	static void UnregisterVisualNodeFactory(TSharedPtr<FGraphPanelNodeFactory> OldFactory);
+	static UNREALED_API void RegisterVisualNodeFactory(TSharedPtr<FGraphPanelNodeFactory> NewFactory);
+	static UNREALED_API void UnregisterVisualNodeFactory(TSharedPtr<FGraphPanelNodeFactory> OldFactory);
 
-	static void RegisterVisualPinFactory(TSharedPtr<FGraphPanelPinFactory> NewFactory);
-	static void UnregisterVisualPinFactory(TSharedPtr<FGraphPanelPinFactory> OldFactory);
+	static UNREALED_API void RegisterVisualPinFactory(TSharedPtr<FGraphPanelPinFactory> NewFactory);
+	static UNREALED_API void UnregisterVisualPinFactory(TSharedPtr<FGraphPanelPinFactory> OldFactory);
 
-    static void RegisterVisualPinConnectionFactory(TSharedPtr<FGraphPanelPinConnectionFactory> NewFactory);
-    static void UnregisterVisualPinConnectionFactory(TSharedPtr<FGraphPanelPinConnectionFactory> OldFactory);
+    static UNREALED_API void RegisterVisualPinConnectionFactory(TSharedPtr<FGraphPanelPinConnectionFactory> NewFactory);
+    static UNREALED_API void UnregisterVisualPinConnectionFactory(TSharedPtr<FGraphPanelPinConnectionFactory> OldFactory);
 
 	/** After pasting nodes, need to perform some fixup for pins etc. */
-	static void PostProcessPastedNodes(TSet<class UEdGraphNode*>& SpawnedNodes);
+	static UNREALED_API void PostProcessPastedNodes(TSet<class UEdGraphNode*>& SpawnedNodes);
 
 	/** Finds the net associated with a particular pin */
-	static UEdGraphPin* GetNetFromPin(UEdGraphPin* Pin);
+	static UNREALED_API UEdGraphPin* GetNetFromPin(UEdGraphPin* Pin);
 
 	/** 
 	 *	Clones (deep copies) a UEdGraph, including all of it's nodes and pins and their links, maintaining a mapping from the clone to the source nodes (even across multiple clonings) 
 	 *	A NewOuter of NULL indicates the transient package
 	 */
-	static UEdGraph* CloneGraph(UEdGraph* Source, UObject* NewOuter, FCompilerResultsLog* MessageLog = NULL, bool bCloningForCompile = false);
+	static UNREALED_API UEdGraph* CloneGraph(UEdGraph* Source, UObject* NewOuter, FCompilerResultsLog* MessageLog = NULL, bool bCloningForCompile = false);
 
 	/**
 	  * Clones the content from SourceGraph and merges it into MergeTarget;
@@ -89,17 +89,17 @@ public:
 	  * @param OutClonedNodes			Will populate with a full list of cloned nodes if provided
 	  * @param bInIsCompiling			TRUE if the function is being called during compilation, this will eliminate some nodes that will not be compiled
 	  */
-	static void CloneAndMergeGraphIn(UEdGraph* MergeTarget, UEdGraph* SourceGraph, FCompilerResultsLog& MessageLog, bool bRequireSchemaMatch, bool bInIsCompiling = false, TArray<UEdGraphNode*>* OutClonedNodes = NULL);
+	static UNREALED_API void CloneAndMergeGraphIn(UEdGraph* MergeTarget, UEdGraph* SourceGraph, FCompilerResultsLog& MessageLog, bool bRequireSchemaMatch, bool bInIsCompiling = false, TArray<UEdGraphNode*>* OutClonedNodes = NULL);
 
 	/**
 	 * Moves the contents of all of the children graphs of ParentGraph (recursively) into the MergeTarget graph.
 	 * This does not clone, it's destructive to the ParentGraph
 	 */
-	static void MergeChildrenGraphsIn(UEdGraph* MergeTarget, UEdGraph* ParentGraph, bool bRequireSchemaMatch, bool bInIsCompiling = false, FCompilerResultsLog* MessageLog = nullptr);
+	static UNREALED_API void MergeChildrenGraphsIn(UEdGraph* MergeTarget, UEdGraph* ParentGraph, bool bRequireSchemaMatch, bool bInIsCompiling = false, FCompilerResultsLog* MessageLog = nullptr);
 
 	/** Tries to rename the graph to have a name similar to BaseName */
-	static void RenameGraphCloseToName(UEdGraph* Graph, const FString& BaseName, int32 StartIndex = 1);
-	static void RenameGraphToNameOrCloseToName(UEdGraph* Graph, const FString& DesiredName);
+	static UNREALED_API void RenameGraphCloseToName(UEdGraph* Graph, const FString& BaseName, int32 StartIndex = 1);
+	static UNREALED_API void RenameGraphToNameOrCloseToName(UEdGraph* Graph, const FString& DesiredName);
 
 	/**
 	 * Exports a set of nodes to text
@@ -107,7 +107,7 @@ public:
 	 * @param	NodesToExport		Set of nodes to export, all nodes must be in the same graph
 	 * @param	[out] ExportedText	The text representation will be stored here
 	 */
-	static void ExportNodesToText(TSet<UObject*> NodesToExport, /*out*/ FString& ExportedText);
+	static UNREALED_API void ExportNodesToText(TSet<UObject*> NodesToExport, /*out*/ FString& ExportedText);
 
 	/**
 	 * Imports a set of previously exported nodes into a graph
@@ -116,7 +116,7 @@ public:
 	 * @param	TextToImport			The text representation of the nodes to import
 	 * @param	[out] ImportedNodeSet	The set of nodes that were imported
 	 */
-	static void ImportNodesFromText(UEdGraph* DestinationGraph, const FString& TextToImport, /*out*/ TSet<UEdGraphNode*>& ImportedNodeSet);
+	static UNREALED_API void ImportNodesFromText(UEdGraph* DestinationGraph, const FString& TextToImport, /*out*/ TSet<UEdGraphNode*>& ImportedNodeSet);
 
 	/**
 	 * Can we import nodes from this text ?
@@ -124,15 +124,15 @@ public:
 	 * @param	DestinationGraph		The target graph to place the imported nodes in
 	 * @param	TextToImport			The text representation of the nodes to import
 	 */
-	static bool CanImportNodesFromText(const UEdGraph* DestinationGraph, const FString& TextToImport);
+	static UNREALED_API bool CanImportNodesFromText(const UEdGraph* DestinationGraph, const FString& TextToImport);
 
 	// Calculates an approximate bounding box for the specified nodes
-	static FIntRect CalculateApproximateNodeBoundaries(const TArray<UEdGraphNode*>& Nodes);
+	static UNREALED_API FIntRect CalculateApproximateNodeBoundaries(const TArray<UEdGraphNode*>& Nodes);
 
 	/** Copy Common State of data from OldNode to NewNode **/
-	static void CopyCommonState(UEdGraphNode* OldNode, UEdGraphNode* NewNode);
+	static UNREALED_API void CopyCommonState(UEdGraphNode* OldNode, UEdGraphNode* NewNode);
 
-	struct UNREALED_API FNodeVisitor
+	struct FNodeVisitor
 	{
 		TSet<UEdGraphNode*> VisitedNodes;
 
@@ -140,32 +140,32 @@ public:
 		{
 		}
 
-		void TraverseNodes(UEdGraphNode* Node);
+		UNREALED_API void TraverseNodes(UEdGraphNode* Node);
 	};
 
 	/** returns true if the ParameterName is marked as a SetParam in the relevant Function */
-	static bool IsSetParam(const UFunction* Function, const FName ParameterName);
+	static UNREALED_API bool IsSetParam(const UFunction* Function, const FName ParameterName);
 
 	/** returns true if the ParameterName is marked as a MapParam, MapKeyParam, or MapValueParam in the relevant Function */
-	static bool IsMapParam(const UFunction* Function, const FName ParameterName);
+	static UNREALED_API bool IsMapParam(const UFunction* Function, const FName ParameterName);
 
 	/** returns true if the ParameterName is marked as ArrayTypeDependentParams */
-	static bool IsArrayDependentParam(const UFunction* Function, const FName ParameterName );
+	static UNREALED_API bool IsArrayDependentParam(const UFunction* Function, const FName ParameterName );
 
 	/** returns true if the ParamaterName is marked as a Set, Map, or array param in the relevant function */
-	static bool IsDynamicContainerParam(const UFunction* Function, const FName ParameterName);
+	static UNREALED_API bool IsDynamicContainerParam(const UFunction* Function, const FName ParameterName);
 
 	/** returns the first pin marked as an ArrayParam, usually returning nullptr */
-	static UEdGraphPin* FindArrayParamPin(const UFunction* Function, const UEdGraphNode* Node);
+	static UNREALED_API UEdGraphPin* FindArrayParamPin(const UFunction* Function, const UEdGraphNode* Node);
 
 	/** returns the first pin marked as an SetParam, usually returning nullptr */
-	static UEdGraphPin* FindSetParamPin(const UFunction* Function, const UEdGraphNode* Node);
+	static UNREALED_API UEdGraphPin* FindSetParamPin(const UFunction* Function, const UEdGraphNode* Node);
 	
 	/** returns the first pin marked as an MapParam, usually returning nullptr */
-	static UEdGraphPin* FindMapParamPin(const UFunction* Function, const UEdGraphNode* Node);
+	static UNREALED_API UEdGraphPin* FindMapParamPin(const UFunction* Function, const UEdGraphNode* Node);
 
 	/** returns the first pin referred to in a , and | delineated list of pin names */
-	static UEdGraphPin* FindPinFromMetaData(const UFunction* Function, const UEdGraphNode* Node, FName MetaData );
+	static UNREALED_API UEdGraphPin* FindPinFromMetaData(const UFunction* Function, const UEdGraphNode* Node, FName MetaData );
 
 	/**
 	* Output a map of pin names to a set of connections given a function entry pin
@@ -173,7 +173,7 @@ public:
 	* @param Node				The entry pin to gather connections from
 	* @param OutPinConnections	Output map of connection data
 	*/
-	static void GetPinConnectionMap(const UEdGraphNode* Node, TMap<FString, TSet<UEdGraphPin*>>& OutPinConnections);
+	static UNREALED_API void GetPinConnectionMap(const UEdGraphNode* Node, TMap<FString, TSet<UEdGraphPin*>>& OutPinConnections);
 
 	/**
 	* Reconnect the pin map to the given node
@@ -181,15 +181,15 @@ public:
 	* @param Node				The node to connect pins to
 	* @param PinConnections		Map of pin connections to set
 	*/
-	static void ReconnectPinMap(UEdGraphNode* Node, const TMap<FString, TSet<UEdGraphPin*>>& PinConnections);
+	static UNREALED_API void ReconnectPinMap(UEdGraphNode* Node, const TMap<FString, TSet<UEdGraphPin*>>& PinConnections);
 
 	/** Copy pin default values between nodes for any pins with matching names */
-	static void CopyPinDefaults(const UEdGraphNode* NodeFrom, UEdGraphNode* NodeTo);
+	static UNREALED_API void CopyPinDefaults(const UEdGraphNode* NodeFrom, UEdGraphNode* NodeTo);
 
 private:
-	static TArray< TSharedPtr<FGraphPanelNodeFactory> > VisualNodeFactories;
-	static TArray< TSharedPtr<FGraphPanelPinFactory> > VisualPinFactories;
-    static TArray< TSharedPtr<FGraphPanelPinConnectionFactory> > VisualPinConnectionFactories;
+	static UNREALED_API TArray< TSharedPtr<FGraphPanelNodeFactory> > VisualNodeFactories;
+	static UNREALED_API TArray< TSharedPtr<FGraphPanelPinFactory> > VisualPinFactories;
+    static UNREALED_API TArray< TSharedPtr<FGraphPanelPinConnectionFactory> > VisualPinConnectionFactories;
 	friend class FNodeFactory;
 
 	// Should never create an instance of this class
diff --git a/Engine/Source/Editor/UnrealEd/Public/EdMode.h b/Engine/Source/Editor/UnrealEd/Public/EdMode.h
--- a/Engine/Source/Editor/UnrealEd/Public/EdMode.h
+++ b/Engine/Source/Editor/UnrealEd/Public/EdMode.h
@@ -36,26 +36,26 @@ struct FViewportClick;
 /**
  * Base class for all editor modes.
  */
-class UNREALED_API FEdMode : public TSharedFromThis<FEdMode>, public FGCObject, public FEditorCommonDrawHelper, public FLegacyEdModeWidgetHelper
+class FEdMode : public TSharedFromThis<FEdMode>, public FGCObject, public FEditorCommonDrawHelper, public FLegacyEdModeWidgetHelper
 {
 public:
 	/** Friends so it can access mode's internals on construction */
 	friend class FEditorModeRegistry;
 
-	FEdMode();
-	virtual ~FEdMode();
+	UNREALED_API FEdMode();
+	UNREALED_API virtual ~FEdMode();
 
 	virtual void Initialize() {}
 
-	virtual bool MouseEnter( FEditorViewportClient* ViewportClient,FViewport* Viewport,int32 x, int32 y );
+	UNREALED_API virtual bool MouseEnter( FEditorViewportClient* ViewportClient,FViewport* Viewport,int32 x, int32 y );
 
-	virtual bool MouseLeave( FEditorViewportClient* ViewportClient,FViewport* Viewport );
+	UNREALED_API virtual bool MouseLeave( FEditorViewportClient* ViewportClient,FViewport* Viewport );
 
-	virtual bool MouseMove(FEditorViewportClient* ViewportClient,FViewport* Viewport,int32 x, int32 y);
+	UNREALED_API virtual bool MouseMove(FEditorViewportClient* ViewportClient,FViewport* Viewport,int32 x, int32 y);
 
-	virtual bool ReceivedFocus(FEditorViewportClient* ViewportClient,FViewport* Viewport);
+	UNREALED_API virtual bool ReceivedFocus(FEditorViewportClient* ViewportClient,FViewport* Viewport);
 
-	virtual bool LostFocus(FEditorViewportClient* ViewportClient,FViewport* Viewport);
+	UNREALED_API virtual bool LostFocus(FEditorViewportClient* ViewportClient,FViewport* Viewport);
 
 	/**
 	 * Called when the mouse is moved while a window input capture is in effect
@@ -67,16 +67,16 @@ public:
 	 *
 	 * @return	true if input was handled
 	 */
-	virtual bool CapturedMouseMove( FEditorViewportClient* InViewportClient, FViewport* InViewport, int32 InMouseX, int32 InMouseY );
+	UNREALED_API virtual bool CapturedMouseMove( FEditorViewportClient* InViewportClient, FViewport* InViewport, int32 InMouseX, int32 InMouseY );
 
 	/** Process all captured mouse moves that occurred during the current frame */
 	virtual bool ProcessCapturedMouseMoves( FEditorViewportClient* InViewportClient, FViewport* InViewport, const TArrayView<FIntPoint>& CapturedMouseMoves ) { return false; }
 
-	virtual bool InputKey(FEditorViewportClient* ViewportClient,FViewport* Viewport,FKey Key,EInputEvent Event);
-	virtual bool InputAxis(FEditorViewportClient* InViewportClient,FViewport* Viewport,int32 ControllerId,FKey Key,float Delta,float DeltaTime);
-	virtual bool InputDelta(FEditorViewportClient* InViewportClient,FViewport* InViewport,FVector& InDrag,FRotator& InRot,FVector& InScale);
-	virtual bool StartTracking(FEditorViewportClient* InViewportClient, FViewport* InViewport);
-	virtual bool EndTracking(FEditorViewportClient* InViewportClient, FViewport* InViewport);
+	UNREALED_API virtual bool InputKey(FEditorViewportClient* ViewportClient,FViewport* Viewport,FKey Key,EInputEvent Event);
+	UNREALED_API virtual bool InputAxis(FEditorViewportClient* InViewportClient,FViewport* Viewport,int32 ControllerId,FKey Key,float Delta,float DeltaTime);
+	UNREALED_API virtual bool InputDelta(FEditorViewportClient* InViewportClient,FViewport* InViewport,FVector& InDrag,FRotator& InRot,FVector& InScale);
+	UNREALED_API virtual bool StartTracking(FEditorViewportClient* InViewportClient, FViewport* InViewport);
+	UNREALED_API virtual bool EndTracking(FEditorViewportClient* InViewportClient, FViewport* InViewport);
 	// Added for handling EDIT Command...
 	virtual EEditAction::Type GetActionEditDuplicate() { return EEditAction::Skip; }
 	virtual EEditAction::Type GetActionEditDelete() { return EEditAction::Skip; }
@@ -89,7 +89,7 @@ public:
 	virtual bool ProcessEditCopy() { return false; }
 	virtual bool ProcessEditPaste() { return false; }
 
-	virtual void Tick(FEditorViewportClient* ViewportClient,float DeltaTime);
+	UNREALED_API virtual void Tick(FEditorViewportClient* ViewportClient,float DeltaTime);
 
 	virtual bool IsCompatibleWith(FEditorModeID OtherModeID) const { return false; }
 	virtual void ActorMoveNotify() {}
@@ -125,7 +125,7 @@ public:
 	virtual bool ShouldDrawBrushWireframe( AActor* InActor ) const { return true; }
 
 	/** If Rotation Snap should be enabled for this mode*/ 
-	virtual bool IsSnapRotationEnabled();
+	UNREALED_API virtual bool IsSnapRotationEnabled();
 
 	/** If this mode should override the snap rotation
 	* @param	Rotation		The Rotation Override
@@ -141,17 +141,17 @@ public:
 		return TEXT("FEdMode");
 	}
 
-	virtual void Enter();
-	virtual void Exit();
-	virtual UTexture2D* GetVertexTexture();
+	UNREALED_API virtual void Enter();
+	UNREALED_API virtual void Exit();
+	UNREALED_API virtual UTexture2D* GetVertexTexture();
 
 	/**
 	 * Lets each tool determine if it wants to use the editor widget or not.  If the tool doesn't want to use it,
 	 * it will be fed raw mouse delta information (not snapped or altered in any way).
 	 */
-	virtual bool UsesTransformWidget() const override;
-	virtual bool UsesTransformWidget(UE::Widget::EWidgetMode CheckMode) const override;
-	virtual bool UsesPropertyWidgets() const override;
+	UNREALED_API virtual bool UsesTransformWidget() const override;
+	UNREALED_API virtual bool UsesTransformWidget(UE::Widget::EWidgetMode CheckMode) const override;
+	UNREALED_API virtual bool UsesPropertyWidgets() const override;
 
 	virtual void PostUndo() {}
 
@@ -167,7 +167,7 @@ public:
 	 * @param	InBox	The selection box to use, in worldspace coordinates.
 	 * @return		true if something was selected/deselected, false otherwise.
 	 */
-	virtual bool BoxSelect( FBox& InBox, bool InSelect = true );
+	UNREALED_API virtual bool BoxSelect( FBox& InBox, bool InSelect = true );
 
 	/**
 	 * Lets each mode/tool handle frustum selection in its own way.
@@ -175,12 +175,12 @@ public:
 	 * @param	InFrustum	The selection box to use, in worldspace coordinates.
 	 * @return	true if something was selected/deselected, false otherwise.
 	 */
-	virtual bool FrustumSelect( const FConvexVolume& InFrustum, FEditorViewportClient* InViewportClient, bool InSelect = true );
+	UNREALED_API virtual bool FrustumSelect( const FConvexVolume& InFrustum, FEditorViewportClient* InViewportClient, bool InSelect = true );
 
-	virtual void SelectNone();
+	UNREALED_API virtual void SelectNone();
 	virtual void SelectionChanged() {}
 
-	virtual bool HandleClick(FEditorViewportClient* InViewportClient, HHitProxy* HitProxy, const FViewportClick& Click);
+	UNREALED_API virtual bool HandleClick(FEditorViewportClient* InViewportClient, HHitProxy* HitProxy, const FViewportClick& Click);
 
 	/**
 	 * Allows an editor mode to override the bounding box used to focus the viewport on a selection
@@ -206,9 +206,9 @@ public:
 
 	// Tools
 
-	void SetCurrentTool( EModeTools InID );
-	void SetCurrentTool( FModeTool* InModeTool );
-	FModeTool* FindTool( EModeTools InID );
+	UNREALED_API void SetCurrentTool( EModeTools InID );
+	UNREALED_API void SetCurrentTool( FModeTool* InModeTool );
+	UNREALED_API FModeTool* FindTool( EModeTools InID );
 
 	const TArray<FModeTool*>& GetTools() const		{ return Tools; }
 
@@ -223,11 +223,11 @@ public:
 	/** @name Rendering */
 	//@{
 	/** Draws translucent polygons on brushes and volumes. */
-	virtual void Render(const FSceneView* View,FViewport* Viewport,FPrimitiveDrawInterface* PDI);
+	UNREALED_API virtual void Render(const FSceneView* View,FViewport* Viewport,FPrimitiveDrawInterface* PDI);
 	//void DrawGridSection(int32 ViewportLocX,int32 ViewportGridY,FVector* A,FVector* B,float* AX,float* BX,int32 Axis,int32 AlphaCase,FSceneView* View,FPrimitiveDrawInterface* PDI);
 
 	/** Overlays the editor hud (brushes, drag tools, static mesh vertices, etc*. */
-	virtual void DrawHUD(FEditorViewportClient* ViewportClient,FViewport* Viewport,const FSceneView* View,FCanvas* Canvas);
+	UNREALED_API virtual void DrawHUD(FEditorViewportClient* ViewportClient,FViewport* Viewport,const FSceneView* View,FCanvas* Canvas);
 	//@}
 
 	/**
@@ -237,16 +237,16 @@ public:
 	virtual bool HandleDragDuplicate() { return false; }
 
 	/** True if this mode uses a toolkit mode (eventually they all should) */
-	virtual bool UsesToolkits() const;
+	UNREALED_API virtual bool UsesToolkits() const;
 
 	/** Gets the toolkit created by this mode */
 	TSharedPtr<FModeToolkit> GetToolkit() { return Toolkit; }
 
 	/** Returns the world this toolkit is editing */
-	UWorld* GetWorld() const;
+	UNREALED_API UWorld* GetWorld() const;
 
 	/** Returns the owning mode manager for this mode */
-	class FEditorModeTools* GetModeManager() const;
+	UNREALED_API class FEditorModeTools* GetModeManager() const;
 
 	/** 
 	 * Called when the editor mode should rebuild its toolbar 
@@ -265,11 +265,11 @@ public:
 public:
 
 	/** Request that this mode be deleted at the next convenient opportunity (FEditorModeTools::Tick) */
-	void RequestDeletion();
+	UNREALED_API void RequestDeletion();
 
 private:
 	/** Called whenever a mode type is unregistered */
-	void OnModeUnregistered(FEditorModeID ModeID);
+	UNREALED_API void OnModeUnregistered(FEditorModeID ModeID);
 
 protected:
 	/** Optional array of tools for this mode. */
diff --git a/Engine/Source/Editor/UnrealEd/Public/Editor.h b/Engine/Source/Editor/UnrealEd/Public/Editor.h
--- a/Engine/Source/Editor/UnrealEd/Public/Editor.h
+++ b/Engine/Source/Editor/UnrealEd/Public/Editor.h
@@ -75,7 +75,7 @@ private:
  * FEditorDelegates
  * Delegates used by the editor.
  **/
-struct UNREALED_API FEditorDelegates
+struct FEditorDelegates
 {
 	/** delegate type for map change events ( Params: uint32 MapChangeFlags (MapChangeEventFlags) ) */
 	DECLARE_MULTICAST_DELEGATE_OneParam(FOnMapChanged, uint32);
@@ -176,187 +176,187 @@ struct UNREALED_API FEditorDelegates
 
 	/** Called when the CurrentLevel is switched to a new level.  Note that this event won't be fired for temporary
 		changes to the current level, such as when copying/pasting actors. */
-	static FSimpleMulticastDelegate NewCurrentLevel;
+	static UNREALED_API FSimpleMulticastDelegate NewCurrentLevel;
 	/** Called when the map has changed */
-	static FOnMapChanged MapChange;
+	static UNREALED_API FOnMapChanged MapChange;
 	/** Called when an actor is added to a layer */
-	static FSimpleMulticastDelegate LayerChange;
+	static UNREALED_API FSimpleMulticastDelegate LayerChange;
 	/** Called after an undo/redo */
-	static FSimpleMulticastDelegate PostUndoRedo;
+	static UNREALED_API FSimpleMulticastDelegate PostUndoRedo;
 	/** surfprops changed */
-	static FSimpleMulticastDelegate SurfProps;
+	static UNREALED_API FSimpleMulticastDelegate SurfProps;
 	/** Sent when requesting to display the properties of selected actors or BSP surfaces */
-	static FSimpleMulticastDelegate SelectedProps;
+	static UNREALED_API FSimpleMulticastDelegate SelectedProps;
 	/** Fits the currently assigned texture to the selected surfaces */
-	static FOnFitTextureToSurface FitTextureToSurface;
+	static UNREALED_API FOnFitTextureToSurface FitTextureToSurface;
 	/** Called when the editor mode is changed */
-	static FOnModeChanged ChangeEditorMode;
+	static UNREALED_API FOnModeChanged ChangeEditorMode;
 	/** Called when properties of an actor have changed */
-	static FSimpleMulticastDelegate ActorPropertiesChange;
+	static UNREALED_API FSimpleMulticastDelegate ActorPropertiesChange;
 	/** Called when the editor needs to be refreshed */
-	static FSimpleMulticastDelegate RefreshEditor;
+	static UNREALED_API FSimpleMulticastDelegate RefreshEditor;
 	/** called when all browsers need to be refreshed */
-	static FSimpleMulticastDelegate RefreshAllBrowsers;
+	static UNREALED_API FSimpleMulticastDelegate RefreshAllBrowsers;
 	/** called when the level browser need to be refreshed */
-	static FSimpleMulticastDelegate RefreshLevelBrowser;
+	static UNREALED_API FSimpleMulticastDelegate RefreshLevelBrowser;
 	/** called when the layer browser need to be refreshed */
-	static FSimpleMulticastDelegate RefreshLayerBrowser;
+	static UNREALED_API FSimpleMulticastDelegate RefreshLayerBrowser;
 	/** called when the primitive stats browser need to be refreshed */
-	static FSimpleMulticastDelegate RefreshPrimitiveStatsBrowser;
+	static UNREALED_API FSimpleMulticastDelegate RefreshPrimitiveStatsBrowser;
 	/** Called when an action is performed which interacts with the content browser; 
 	 *  load any selected assets which aren't already loaded */
-	static FSimpleMulticastDelegate LoadSelectedAssetsIfNeeded;
+	static UNREALED_API FSimpleMulticastDelegate LoadSelectedAssetsIfNeeded;
 	/** Called when load errors are about to be displayed */
-	static FSimpleMulticastDelegate DisplayLoadErrors;
+	static UNREALED_API FSimpleMulticastDelegate DisplayLoadErrors;
 	/** Called when an editor mode is being entered */
 	UE_DEPRECATED(4.24, "Use EditorModeIDEnter instead")
-	static FOnEditorModeTransitioned EditorModeEnter;
+	static UNREALED_API FOnEditorModeTransitioned EditorModeEnter;
 	/** Called when an editor mode is being exited */
 	UE_DEPRECATED(4.24, "Use EditorModeIDExit instead")
-	static FOnEditorModeTransitioned EditorModeExit;
+	static UNREALED_API FOnEditorModeTransitioned EditorModeExit;
 	/** Called when an editor mode ID is being entered */
 	UE_DEPRECATED(5.0, "Use the asset editor's mode manager to scope mode enter notifications.")
-	static FOnEditorModeIDTransitioned EditorModeIDEnter;
+	static UNREALED_API FOnEditorModeIDTransitioned EditorModeIDEnter;
 	/** Called when an editor mode ID is being exited */
 	UE_DEPRECATED(5.0, "Use the asset editor's mode manager to scope mode exit notifications.")
-	static FOnEditorModeIDTransitioned EditorModeIDExit;
+	static UNREALED_API FOnEditorModeIDTransitioned EditorModeIDExit;
 	/** Sent when a PIE session is beginning (before we decide if PIE can run - allows clients to avoid blocking PIE) */
-	static FOnPIEEvent PreBeginPIE;
+	static UNREALED_API FOnPIEEvent PreBeginPIE;
 	/** Sent when a PIE session is beginning (but hasn't actually started yet) */
-	static FOnPIEEvent BeginPIE;
+	static UNREALED_API FOnPIEEvent BeginPIE;
 	/** Sent when a PIE session has fully started and after BeginPlay() has been called */
-	static FOnPIEEvent PostPIEStarted;
+	static UNREALED_API FOnPIEEvent PostPIEStarted;
 	/** Sent when a PIE session is ending, before anything else happens */
-	static FOnPIEEvent PrePIEEnded;
+	static UNREALED_API FOnPIEEvent PrePIEEnded;
 	/** Sent when a PIE session is ending */
-	static FOnPIEEvent EndPIE;
+	static UNREALED_API FOnPIEEvent EndPIE;
 	/** Sent when a PIE session is paused */
-	static FOnPIEEvent PausePIE;
+	static UNREALED_API FOnPIEEvent PausePIE;
 	/** Sent when a PIE session is resumed */
-	static FOnPIEEvent ResumePIE;
+	static UNREALED_API FOnPIEEvent ResumePIE;
 	/** Sent when a PIE session is single-stepped */
-	static FOnPIEEvent SingleStepPIE;
+	static UNREALED_API FOnPIEEvent SingleStepPIE;
 	/** Sent just before the user switches between from PIE to SIE, or vice-versa.  Passes in whether we are currently in SIE */
-	static FOnPIEEvent OnPreSwitchBeginPIEAndSIE;
+	static UNREALED_API FOnPIEEvent OnPreSwitchBeginPIEAndSIE;
 	/** Sent after the user switches between from PIE to SIE, or vice-versa.  Passes in whether we are currently in SIE */
-	static FOnPIEEvent OnSwitchBeginPIEAndSIE;
+	static UNREALED_API FOnPIEEvent OnSwitchBeginPIEAndSIE;
 	/** Sent when a PIE session is cancelled */
-	static FSimpleMulticastDelegate CancelPIE;
+	static UNREALED_API FSimpleMulticastDelegate CancelPIE;
 	/** Sent when PC local play session is starting */
-	static FOnStandaloneLocalPlayEvent BeginStandaloneLocalPlay;
+	static UNREALED_API FOnStandaloneLocalPlayEvent BeginStandaloneLocalPlay;
 	/** Within a property window, the currently selected item was changed.*/
-	static FSimpleMulticastDelegate PropertySelectionChange;
+	static UNREALED_API FSimpleMulticastDelegate PropertySelectionChange;
 	/** Called after Landscape layer infomap update have completed */
-	static FSimpleMulticastDelegate PostLandscapeLayerUpdated;
+	static UNREALED_API FSimpleMulticastDelegate PostLandscapeLayerUpdated;
 	/** Called before SaveWorld is processed */
-	static FOnPreSaveWorldWithContext PreSaveWorldWithContext;
+	static UNREALED_API FOnPreSaveWorldWithContext PreSaveWorldWithContext;
 	/** Called after SaveWorld is processed */
-	static FOnPostSaveWorldWithContext PostSaveWorldWithContext;
+	static UNREALED_API FOnPostSaveWorldWithContext PostSaveWorldWithContext;
 	/** Called before any number of external actors will be saved */
-	static FOnPreSaveExternalActors PreSaveExternalActors;
+	static UNREALED_API FOnPreSaveExternalActors PreSaveExternalActors;
 	/** Called after any number of external actors has been saved */
-	static FOnPostSaveExternalActors PostSaveExternalActors;
+	static UNREALED_API FOnPostSaveExternalActors PostSaveExternalActors;
 	/** Called before any asset validation happens via the Asset Validation subsystem. */
-	static FSimpleMulticastDelegate OnPreAssetValidation;
+	static UNREALED_API FSimpleMulticastDelegate OnPreAssetValidation;
 	/** Called after asset validation happens by the Asset Validation subsystem. */
-	static FSimpleMulticastDelegate OnPostAssetValidation;
+	static UNREALED_API FSimpleMulticastDelegate OnPostAssetValidation;
 	/** Called when finishing picking a new blueprint class during construction */
-	static FOnFinishPickingBlueprintClass OnFinishPickingBlueprintClass;
+	static UNREALED_API FOnFinishPickingBlueprintClass OnFinishPickingBlueprintClass;
 	/** Called when beginning configuration of a new asset */
-	static FOnNewAssetCreation OnConfigureNewAssetProperties;
+	static UNREALED_API FOnNewAssetCreation OnConfigureNewAssetProperties;
 	/** Called when an asset is about to undergo a destructive action caused by the Editor UI (Delete, Move, Rename, Privatize, etc.) */
-	static FOnPreDestructiveAssetAction OnPreDestructiveAssetAction;
+	static UNREALED_API FOnPreDestructiveAssetAction OnPreDestructiveAssetAction;
 	/** Called when finishing configuration of a new asset */
-	static FOnNewAssetCreation OnNewAssetCreated;
+	static UNREALED_API FOnNewAssetCreation OnNewAssetCreated;
 	/** Called when new assets are being (re-)imported. */
 	UE_DEPRECATED(4.22, "Use the ImportSubsystem instead. GEditor->GetEditorSubsystem<UImportSubsystem>()")
-	static FOnAssetPreImport OnAssetPreImport;
+	static UNREALED_API FOnAssetPreImport OnAssetPreImport;
 	/** Called when new assets have been (re-)imported. */
 	UE_DEPRECATED(4.22, "Use the ImportSubsystem instead. GEditor->GetEditorSubsystem<UImportSubsystem>()")
-	static FOnAssetPostImport OnAssetPostImport;
+	static UNREALED_API FOnAssetPostImport OnAssetPostImport;
 	/** Called after an asset has been reimported */
 	UE_DEPRECATED(4.22, "Use the ImportSubsystem instead. GEditor->GetEditorSubsystem<UImportSubsystem>()")
-	static FOnAssetReimport OnAssetReimport;
+	static UNREALED_API FOnAssetReimport OnAssetReimport;
 	/** Called when new actors are dropped on to the viewport */
-	static FOnNewActorsDropped OnNewActorsDropped;
+	static UNREALED_API FOnNewActorsDropped OnNewActorsDropped;
 	/** Called when new actors are placed in the viewport */
-	static FOnNewActorsPlaced OnNewActorsPlaced;
+	static UNREALED_API FOnNewActorsPlaced OnNewActorsPlaced;
 	/** Called when grid snapping is changed */
-	static FOnGridSnappingChanged OnGridSnappingChanged;
+	static UNREALED_API FOnGridSnappingChanged OnGridSnappingChanged;
 	/** Called when a lighting build has started */
-	static FSimpleMulticastDelegate OnLightingBuildStarted;
+	static UNREALED_API FSimpleMulticastDelegate OnLightingBuildStarted;
 	/** Called when a lighting build has been kept */
-	static FSimpleMulticastDelegate OnLightingBuildKept;
+	static UNREALED_API FSimpleMulticastDelegate OnLightingBuildKept;
 	/** Called when a lighting build has failed (maybe called twice if cancelled) */
-	static FSimpleMulticastDelegate OnLightingBuildFailed;
+	static UNREALED_API FSimpleMulticastDelegate OnLightingBuildFailed;
 	/** Called when a lighting build has succeeded */
-	static FSimpleMulticastDelegate OnLightingBuildSucceeded;
+	static UNREALED_API FSimpleMulticastDelegate OnLightingBuildSucceeded;
 	/** Called when when attempting to apply an object to an actor (via drag drop) */
-	static FOnApplyObjectToActor OnApplyObjectToActor;
+	static UNREALED_API FOnApplyObjectToActor OnApplyObjectToActor;
 	/** Called when focusing viewport on a set of actors */
-	static FOnFocusViewportOnActors OnFocusViewportOnActors;
+	static UNREALED_API FOnFocusViewportOnActors OnFocusViewportOnActors;
 	/** Called before LoadMap is processed */
-	static FOnMapLoad OnMapLoad;
+	static UNREALED_API FOnMapLoad OnMapLoad;
 	/** Called when a map is opened, giving map name, and whether it was a template */
-	static FOnMapOpened OnMapOpened;
+	static UNREALED_API FOnMapOpened OnMapOpened;
 	/** Called when the editor camera is moved */
-	static FOnEditorCameraMoved OnEditorCameraMoved;
+	static UNREALED_API FOnEditorCameraMoved OnEditorCameraMoved;
 	/** Called when the editor camera is moved */
-	static FOnDollyPerspectiveCamera OnDollyPerspectiveCamera;
+	static UNREALED_API FOnDollyPerspectiveCamera OnDollyPerspectiveCamera;
 	/** Called on editor shutdown after packages have been successfully saved */
-	static FSimpleMulticastDelegate OnShutdownPostPackagesSaved;
+	static UNREALED_API FSimpleMulticastDelegate OnShutdownPostPackagesSaved;
 	/** Called when one or more packages have been deleted through save */
-	static FOnPackageDeleted OnPackageDeleted;
+	static UNREALED_API FOnPackageDeleted OnPackageDeleted;
 	/** Called when the user requests assets to be deleted to determine if the operation is available.  */
-	static FOnAssetsCanDelete OnAssetsCanDelete;
+	static UNREALED_API FOnAssetsCanDelete OnAssetsCanDelete;
 	/** Called when the user requests certain assets be deletedand  allows the addition of secondary assets that should also be deleted */
-	static FOnAssetsAddExtraObjectsToDelete OnAssetsAddExtraObjectsToDelete;
+	static UNREALED_API FOnAssetsAddExtraObjectsToDelete OnAssetsAddExtraObjectsToDelete;
 	/** Called when the user requests certain assets be deleted (DOES NOT imply that the asset will be deleted... the user could cancel) */
-	static FOnAssetsPreDelete OnAssetsPreDelete;
+	static UNREALED_API FOnAssetsPreDelete OnAssetsPreDelete;
 	/** Called when one or more assets have been deleted */
-	static FOnAssetsDeleted OnAssetsDeleted;
+	static UNREALED_API FOnAssetsDeleted OnAssetsDeleted;
 	/** Called when a user starts dragging something out of content browser (can be multiple assets) */
-	static FOnAssetDragStarted OnAssetDragStarted;
+	static UNREALED_API FOnAssetDragStarted OnAssetDragStarted;
 	/** Called when a user changes the UInputSettings::bEnableGestureRecognizer setting to refresh the available actions. */
-	static FSimpleMulticastDelegate OnEnableGestureRecognizerChanged;
+	static UNREALED_API FSimpleMulticastDelegate OnEnableGestureRecognizerChanged;
 	/** Called when Action or Axis mappings have been changed */
-	static FSimpleMulticastDelegate OnActionAxisMappingsChanged;
+	static UNREALED_API FSimpleMulticastDelegate OnActionAxisMappingsChanged;
 	/** Called from FEditorUtils::AddLevelToWorld after the level is added successfully to the world. */
-	static FOnAddLevelToWorld OnAddLevelToWorld;
+	static UNREALED_API FOnAddLevelToWorld OnAddLevelToWorld;
 	/** Sent before edit cut is handled */
-	static FOnEditCutActorsBegin OnEditCutActorsBegin;
+	static UNREALED_API FOnEditCutActorsBegin OnEditCutActorsBegin;
 	/** Sent after edit cut is handled */
-	static FOnEditCutActorsEnd OnEditCutActorsEnd;
+	static UNREALED_API FOnEditCutActorsEnd OnEditCutActorsEnd;
 	/** Sent before edit copy is handled */
-	static FOnEditCopyActorsBegin OnEditCopyActorsBegin;
+	static UNREALED_API FOnEditCopyActorsBegin OnEditCopyActorsBegin;
 	/** Sent after edit copy is handled */
-	static FOnEditCopyActorsEnd OnEditCopyActorsEnd;
+	static UNREALED_API FOnEditCopyActorsEnd OnEditCopyActorsEnd;
 	/** Sent before edit paste is handled */
-	static FOnEditPasteActorsBegin OnEditPasteActorsBegin;
+	static UNREALED_API FOnEditPasteActorsBegin OnEditPasteActorsBegin;
 	/** Sent after edit paste is handled */
-	static FOnEditPasteActorsEnd OnEditPasteActorsEnd;
+	static UNREALED_API FOnEditPasteActorsEnd OnEditPasteActorsEnd;
 	/** Sent before duplicate is handled */
-	static FOnDuplicateActorsBegin OnDuplicateActorsBegin;
+	static UNREALED_API FOnDuplicateActorsBegin OnDuplicateActorsBegin;
 	/** Sent after duplicate is handled */
-	static FOnDuplicateActorsEnd OnDuplicateActorsEnd;
+	static UNREALED_API FOnDuplicateActorsEnd OnDuplicateActorsEnd;
 	/** Sent when delete begin called */
-	static FOnDeleteActorsBegin OnDeleteActorsBegin;
+	static UNREALED_API FOnDeleteActorsBegin OnDeleteActorsBegin;
 	/** Sent when delete end called */
-	static FOnDeleteActorsEnd OnDeleteActorsEnd;
+	static UNREALED_API FOnDeleteActorsEnd OnDeleteActorsEnd;
 	/** Called when you want to view things in the reference viewer, these are bound to by asset manager editor plugins */
-	static FOnOpenReferenceViewer OnOpenReferenceViewer;
+	static UNREALED_API FOnOpenReferenceViewer OnOpenReferenceViewer;
 	/** Called when you want to view things in the size map */
-	static FOnViewAssetIdentifiers OnOpenSizeMap;
+	static UNREALED_API FOnViewAssetIdentifiers OnOpenSizeMap;
 	/** Called when you want to view things in the asset audit window */
-	static FOnViewAssetIdentifiers OnOpenAssetAudit;
+	static UNREALED_API FOnViewAssetIdentifiers OnOpenAssetAudit;
 	/** Called to try and edit an asset identifier, which could be a package or searchable name */
-	static FOnViewAssetIdentifiers OnEditAssetIdentifiers;
+	static UNREALED_API FOnViewAssetIdentifiers OnEditAssetIdentifiers;
 	/** Called when the editor requests a restart */
-	static FOnRestartRequested OnRestartRequested;
+	static UNREALED_API FOnRestartRequested OnRestartRequested;
 	/** Called when the editor has booted */
-	static FOnEditorBoot OnEditorBoot;
+	static UNREALED_API FOnEditorBoot OnEditorBoot;
 	/** Called when the editor has initialized */
-	static FOnEditorInitialized OnEditorInitialized;
+	static UNREALED_API FOnEditorInitialized OnEditorInitialized;
 };
 
 /**
diff --git a/Engine/Source/Editor/UnrealEd/Public/EditorActorFolders.h b/Engine/Source/Editor/UnrealEd/Public/EditorActorFolders.h
--- a/Engine/Source/Editor/UnrealEd/Public/EditorActorFolders.h
+++ b/Engine/Source/Editor/UnrealEd/Public/EditorActorFolders.h
@@ -29,13 +29,13 @@ DECLARE_MULTICAST_DELEGATE_ThreeParams(FOnActorFolderMoved, UWorld&, const FFold
 
 
 /** Class responsible for managing an in-memory representation of actor folders in the editor */
-struct UNREALED_API FActorFolders : public FGCObject, public IActorEditorContextClient
+struct FActorFolders : public FGCObject, public IActorEditorContextClient
 {
-	FActorFolders();
-	~FActorFolders();
+	UNREALED_API FActorFolders();
+	UNREALED_API ~FActorFolders();
 
 	// FGCObject Interface
-	virtual void AddReferencedObjects(FReferenceCollector& Collector) override;
+	UNREALED_API virtual void AddReferencedObjects(FReferenceCollector& Collector) override;
 	virtual FString GetReferencerName() const override
 	{
 		return "FActorFolders";
@@ -43,185 +43,185 @@ struct UNREALED_API FActorFolders : public FGCObject, public IActorEditorContext
 	// End FGCObject Interface
 
 	/** Singleton access */
-	static FActorFolders& Get();
+	static UNREALED_API FActorFolders& Get();
 
 	/** Folder creation and deletion events. Called whenever a folder is created or deleted in a world. */
-	static FOnActorFolderCreated OnFolderCreated;
-	static FOnActorFolderMoved 	OnFolderMoved;
-	static FOnActorFolderDeleted OnFolderDeleted;
+	static UNREALED_API FOnActorFolderCreated OnFolderCreated;
+	static UNREALED_API FOnActorFolderMoved 	OnFolderMoved;
+	static UNREALED_API FOnActorFolderDeleted OnFolderDeleted;
 
 	//~ Begin Deprecated
 
 	UE_DEPRECATED(5.0, "OnFolderCreate has been deprecated. Please use OnFolderCreated.")
-	static FOnActorFolderCreate OnFolderCreate;
+	static UNREALED_API FOnActorFolderCreate OnFolderCreate;
 
 	UE_DEPRECATED(5.0, "OnFolderMove has been deprecated. Please use OnFolderMoved.")
-	static FOnActorFolderMove 	OnFolderMove;
+	static UNREALED_API FOnActorFolderMove 	OnFolderMove;
 
 	UE_DEPRECATED(5.0, "OnFolderDelete has been deprecated. Please use OnFolderDeleted.")
-	static FOnActorFolderDelete OnFolderDelete;
+	static UNREALED_API FOnActorFolderDelete OnFolderDelete;
 
 	UE_DEPRECATED(5.0, "GetFolderProperties using FName has been deprecated. Please use new interface using FFolder.")
-	FActorFolderProps* GetFolderProperties(UWorld& InWorld, FName InPath);
+	UNREALED_API FActorFolderProps* GetFolderProperties(UWorld& InWorld, FName InPath);
 
 	UE_DEPRECATED(5.0, "GetDefaultFolderName using FName  has been deprecated. Please use new interface using FFolder.")
-	FName GetDefaultFolderName(UWorld& InWorld, FName ParentPath = FName());
+	UNREALED_API FName GetDefaultFolderName(UWorld& InWorld, FName ParentPath = FName());
 	
 	UE_DEPRECATED(5.0, "GetDefaultFolderNameForSelection using FName  has been deprecated. Please use new interface using FFolder.")
-	FName GetDefaultFolderNameForSelection(UWorld& InWorld);
+	UNREALED_API FName GetDefaultFolderNameForSelection(UWorld& InWorld);
 
 	UE_DEPRECATED(5.0, "GetFolderName using FName  has been deprecated. Please use new interface using FFolder.")
-	FName GetFolderName(UWorld& InWorld, FName ParentPath, FName FolderName);
+	UNREALED_API FName GetFolderName(UWorld& InWorld, FName ParentPath, FName FolderName);
 
 	UE_DEPRECATED(5.0, "CreateFolder using FName  has been deprecated. Please use new interface using FFolder.")
-	void CreateFolder(UWorld& InWorld, FName Path);
+	UNREALED_API void CreateFolder(UWorld& InWorld, FName Path);
 
 	UE_DEPRECATED(5.0, "CreateFolderContainingSelection using FName  has been deprecated. Please use new interface using FFolder.")
-	void CreateFolderContainingSelection(UWorld& InWorld, FName Path);
+	UNREALED_API void CreateFolderContainingSelection(UWorld& InWorld, FName Path);
 
 	UE_DEPRECATED(5.0, "SetSelectedFolderPath using FName  has been deprecated. Please use new interface using FFolder.")
-	void SetSelectedFolderPath(FName Path) const;
+	UNREALED_API void SetSelectedFolderPath(FName Path) const;
 
 	UE_DEPRECATED(5.0, "DeleteFolder using FName  has been deprecated. Please use new interface using FFolder.")
-	void DeleteFolder(UWorld& InWorld, FName FolderToDelete);
+	UNREALED_API void DeleteFolder(UWorld& InWorld, FName FolderToDelete);
 
 	UE_DEPRECATED(5.0, "RenameFolderInWorld using FName  has been deprecated. Please use new interface using FFolder.")
-	bool RenameFolderInWorld(UWorld& World, FName OldPath, FName NewPath);
+	UNREALED_API bool RenameFolderInWorld(UWorld& World, FName OldPath, FName NewPath);
 
 	//~ End Deprecated
-	static FFolder GetActorDescFolder(UWorld& InWorld, const FWorldPartitionActorDesc* InActorDesc);
+	static UNREALED_API FFolder GetActorDescFolder(UWorld& InWorld, const FWorldPartitionActorDesc* InActorDesc);
 
 	/** Apply an operation to each actor desc in the given list of folders. */
-	static void ForEachActorDescInFolders(UWorld& InWorld, const TSet<FName>& InPaths, TFunctionRef<bool(const FWorldPartitionActorDesc*)> Operation, const FFolder::FRootObject& InFolderRootObject = FFolder::GetInvalidRootObject());
+	static UNREALED_API void ForEachActorDescInFolders(UWorld& InWorld, const TSet<FName>& InPaths, TFunctionRef<bool(const FWorldPartitionActorDesc*)> Operation, const FFolder::FRootObject& InFolderRootObject = FFolder::GetInvalidRootObject());
 		
 	/** Apply an operation to each actor in the given list of folders. Will stop when operation returns false. */
-	static void ForEachActorInFolders(UWorld& InWorld, const TArray<FName>& InPaths, TFunctionRef<bool(AActor*)> Operation, const FFolder::FRootObject& InFolderRootObject = FFolder::GetInvalidRootObject());
-	static void ForEachActorInFolders(UWorld& InWorld, const TSet<FName>& InPaths, TFunctionRef<bool(AActor*)> Operation, const FFolder::FRootObject& InFolderRootObject = FFolder::GetInvalidRootObject());
+	static UNREALED_API void ForEachActorInFolders(UWorld& InWorld, const TArray<FName>& InPaths, TFunctionRef<bool(AActor*)> Operation, const FFolder::FRootObject& InFolderRootObject = FFolder::GetInvalidRootObject());
+	static UNREALED_API void ForEachActorInFolders(UWorld& InWorld, const TSet<FName>& InPaths, TFunctionRef<bool(AActor*)> Operation, const FFolder::FRootObject& InFolderRootObject = FFolder::GetInvalidRootObject());
 
 	/** Get an array of actors from a list of folders */
-	static void GetActorsFromFolders(UWorld& InWorld, const TArray<FName>& InPaths, TArray<AActor*>& OutActors, const FFolder::FRootObject& InFolderRootObject = FFolder::GetInvalidRootObject());
-	static void GetActorsFromFolders(UWorld& InWorld, const TSet<FName>& InPaths, TArray<AActor*>& OutActors, const FFolder::FRootObject& InFolderRootObject = FFolder::GetInvalidRootObject());
+	static UNREALED_API void GetActorsFromFolders(UWorld& InWorld, const TArray<FName>& InPaths, TArray<AActor*>& OutActors, const FFolder::FRootObject& InFolderRootObject = FFolder::GetInvalidRootObject());
+	static UNREALED_API void GetActorsFromFolders(UWorld& InWorld, const TSet<FName>& InPaths, TArray<AActor*>& OutActors, const FFolder::FRootObject& InFolderRootObject = FFolder::GetInvalidRootObject());
 
 	/** Get an array of weak actor pointers from a list of folders */
-	static void GetWeakActorsFromFolders(UWorld& InWorld, const TArray<FName>& InPaths, TArray<TWeakObjectPtr<AActor>>& OutActors, const FFolder::FRootObject& InFolderRootObject = FFolder::GetInvalidRootObject());
-	static void GetWeakActorsFromFolders(UWorld& InWorld, const TSet<FName>& InPaths, TArray<TWeakObjectPtr<AActor>>& OutActors, const FFolder::FRootObject& InFolderRootObject = FFolder::GetInvalidRootObject());
+	static UNREALED_API void GetWeakActorsFromFolders(UWorld& InWorld, const TArray<FName>& InPaths, TArray<TWeakObjectPtr<AActor>>& OutActors, const FFolder::FRootObject& InFolderRootObject = FFolder::GetInvalidRootObject());
+	static UNREALED_API void GetWeakActorsFromFolders(UWorld& InWorld, const TSet<FName>& InPaths, TArray<TWeakObjectPtr<AActor>>& OutActors, const FFolder::FRootObject& InFolderRootObject = FFolder::GetInvalidRootObject());
 
 	/** Tests whether a folder container exists for the specified world */
-	bool IsInitializedForWorld(UWorld& InWorld) const;
+	UNREALED_API bool IsInitializedForWorld(UWorld& InWorld) const;
 
 	/** Get a default folder name under the specified parent path */
-	FFolder GetDefaultFolderName(UWorld& InWorld, const FFolder& InParentFolder);
+	UNREALED_API FFolder GetDefaultFolderName(UWorld& InWorld, const FFolder& InParentFolder);
 	
 	/** Get a new default folder name that would apply to the current selection */
-	FFolder GetDefaultFolderForSelection(UWorld& InWorld, TArray<FFolder>* InSelectedFolders = nullptr);
+	UNREALED_API FFolder GetDefaultFolderForSelection(UWorld& InWorld, TArray<FFolder>* InSelectedFolders = nullptr);
 
 	/** Get folder name that is unique under specified parent path */
-	FFolder GetFolderName(UWorld& InWorld, const FFolder& InParentFolder, const FName& InLeafName);
+	UNREALED_API FFolder GetFolderName(UWorld& InWorld, const FFolder& InParentFolder, const FName& InLeafName);
 
 	/** Create a new folder in the specified world, of the specified path */
-	bool CreateFolder(UWorld& InWorld, const FFolder& InFolder);
+	UNREALED_API bool CreateFolder(UWorld& InWorld, const FFolder& InFolder);
 
 	/** Same as CreateFolder, but moves the current actor selection into the new folder as well */
-	void CreateFolderContainingSelection(UWorld& InWorld, const FFolder& InFolder);
+	UNREALED_API void CreateFolderContainingSelection(UWorld& InWorld, const FFolder& InFolder);
 
 	/** Sets the folder path for all the selected actors */
-	void SetSelectedFolderPath(const FFolder& InFolder) const;
+	UNREALED_API void SetSelectedFolderPath(const FFolder& InFolder) const;
 
 	/** Delete the specified folder in the world */
-	void DeleteFolder(UWorld& InWorld, const FFolder& InFolderToDelete);
+	UNREALED_API void DeleteFolder(UWorld& InWorld, const FFolder& InFolderToDelete);
 
 	/** Rename the specified path to a new name */
-	bool RenameFolderInWorld(UWorld& InWorld, const FFolder& OldPath, const FFolder& NewPath);
+	UNREALED_API bool RenameFolderInWorld(UWorld& InWorld, const FFolder& OldPath, const FFolder& NewPath);
 
 	/** Notify that a root object has been removed. Cleanup of existing folders with with root object. */
-	void OnFolderRootObjectRemoved(UWorld& InWorld, const FFolder::FRootObject& InFolderRootObject);
+	UNREALED_API void OnFolderRootObjectRemoved(UWorld& InWorld, const FFolder::FRootObject& InFolderRootObject);
 
 	/** Return if folder exists */
-	bool ContainsFolder(UWorld& InWorld, const FFolder& InFolder);
+	UNREALED_API bool ContainsFolder(UWorld& InWorld, const FFolder& InFolder);
 
 	/** Return the folder expansion state */
-	bool IsFolderExpanded(UWorld& InWorld, const FFolder& InFolder);
+	UNREALED_API bool IsFolderExpanded(UWorld& InWorld, const FFolder& InFolder);
 
 	/** Set the folder expansion state */
-	void SetIsFolderExpanded(UWorld& InWorld, const FFolder& InFolder, bool bIsExpanded);
+	UNREALED_API void SetIsFolderExpanded(UWorld& InWorld, const FFolder& InFolder, bool bIsExpanded);
 
 	/** Iterate on all folders of a world and pass it to the provided operation. */
-	void ForEachFolder(UWorld& InWorld, TFunctionRef<bool(const FFolder&)> Operation);
+	UNREALED_API void ForEachFolder(UWorld& InWorld, TFunctionRef<bool(const FFolder&)> Operation);
 
 	/** Iterate on all world's folders with the given root object and pass it to the provided operation. */
-	void ForEachFolderWithRootObject(UWorld& InWorld, const FFolder::FRootObject& InFolderRootObject, TFunctionRef<bool(const FFolder&)> Operation);
+	UNREALED_API void ForEachFolderWithRootObject(UWorld& InWorld, const FFolder::FRootObject& InFolderRootObject, TFunctionRef<bool(const FFolder&)> Operation);
 
 	/** Get the folder properties for the specified path. Returns nullptr if no properties exist */
-	FActorFolderProps* GetFolderProperties(UWorld& InWorld, const FFolder& InFolder);
+	UNREALED_API FActorFolderProps* GetFolderProperties(UWorld& InWorld, const FFolder& InFolder);
 
 	//~ Begin IActorEditorContextClient interface
-	virtual void OnExecuteActorEditorContextAction(UWorld* InWorld, const EActorEditorContextAction& InType, class AActor* InActor = nullptr) override;
-	virtual bool GetActorEditorContextDisplayInfo(UWorld* InWorld, FActorEditorContextClientDisplayInfo& OutDiplayInfo) const override;
+	UNREALED_API virtual void OnExecuteActorEditorContextAction(UWorld* InWorld, const EActorEditorContextAction& InType, class AActor* InActor = nullptr) override;
+	UNREALED_API virtual bool GetActorEditorContextDisplayInfo(UWorld* InWorld, FActorEditorContextClientDisplayInfo& OutDiplayInfo) const override;
 	virtual bool CanResetContext(UWorld* InWorld) const override { return true; };
-	virtual TSharedRef<SWidget> GetActorEditorContextWidget(UWorld* InWorld) const override;
+	UNREALED_API virtual TSharedRef<SWidget> GetActorEditorContextWidget(UWorld* InWorld) const override;
 	virtual FOnActorEditorContextClientChanged& GetOnActorEditorContextClientChanged() override { return ActorEditorContextClientChanged; }
 	//~ End IActorEditorContextClient interface
-	FFolder GetActorEditorContextFolder(UWorld& InWorld, bool bMustMatchCurrentLevel = true) const;
-	void SetActorEditorContextFolder(UWorld& InWorld, const FFolder& InFolder);
+	UNREALED_API FFolder GetActorEditorContextFolder(UWorld& InWorld, bool bMustMatchCurrentLevel = true) const;
+	UNREALED_API void SetActorEditorContextFolder(UWorld& InWorld, const FFolder& InFolder);
 
 private:
 	
-	static FFolder::FRootObject GetWorldFolderRootObject(UWorld& InWorld);
+	static UNREALED_API FFolder::FRootObject GetWorldFolderRootObject(UWorld& InWorld);
 
 	/** Broadcast when actor folder is created. */
-	void BroadcastOnActorFolderCreated(UWorld& InWorld, const FFolder& InFolder);
+	UNREALED_API void BroadcastOnActorFolderCreated(UWorld& InWorld, const FFolder& InFolder);
 
 	/** Broadcast when actor folder is deleted. */
-	void BroadcastOnActorFolderDeleted(UWorld& InWorld, const FFolder& InFolder);
+	UNREALED_API void BroadcastOnActorFolderDeleted(UWorld& InWorld, const FFolder& InFolder);
 
 	/** Broadcast when actor folder has moved. */
-	void BroadcastOnActorFolderMoved(UWorld& InWorld, const FFolder& InSrcFolder, const FFolder& InDstFolder);
+	UNREALED_API void BroadcastOnActorFolderMoved(UWorld& InWorld, const FFolder& InSrcFolder, const FFolder& InDstFolder);
 
 	/** Get or create a folder container for the specified world */
-	UWorldFolders& GetOrCreateWorldFolders(UWorld& InWorld);
+	UNREALED_API UWorldFolders& GetOrCreateWorldFolders(UWorld& InWorld);
 
 	/** Create and update a folder container for the specified world */
-	UWorldFolders& CreateWorldFolders(UWorld& InWorld);
+	UNREALED_API UWorldFolders& CreateWorldFolders(UWorld& InWorld);
 
 	/** Rebuild the folder list for the specified world. This can be very slow as it
 		iterates all actors in memory to rebuild the array of actors for this world */
-	void RebuildFolderListForWorld(UWorld& InWorld);
+	UNREALED_API void RebuildFolderListForWorld(UWorld& InWorld);
 
 	/** Called when an actor's folder has changed */
-	void OnActorFolderChanged(const AActor* InActor, FName OldPath);
+	UNREALED_API void OnActorFolderChanged(const AActor* InActor, FName OldPath);
 
 	/** Called when the actor list of the current world has changed */
-	void OnLevelActorListChanged();
+	UNREALED_API void OnLevelActorListChanged();
 
 	/** Called when all the levels have been changed */
-	void OnAllLevelsChanged();
+	UNREALED_API void OnAllLevelsChanged();
 
 	/** Called when an actor folder is added */
-	void OnActorFolderAdded(UActorFolder* InActorFolder);
+	UNREALED_API void OnActorFolderAdded(UActorFolder* InActorFolder);
 
 	/** Called when the global map in the editor has changed */
-	void OnMapChange(uint32 MapChangeFlags);
+	UNREALED_API void OnMapChange(uint32 MapChangeFlags);
 
 	/** Called after a world has been saved */
-	void OnWorldSaved(UWorld* World, FObjectPostSaveContext ObjectSaveContext);
+	UNREALED_API void OnWorldSaved(UWorld* World, FObjectPostSaveContext ObjectSaveContext);
 
 	/** Attempt to save the folders state */
-	void SaveWorldFoldersState(UWorld* World);
+	UNREALED_API void SaveWorldFoldersState(UWorld* World);
 
 	/** Remove any references to folder arrays for dead worlds */
-	void Housekeeping();
+	UNREALED_API void Housekeeping();
 
 	/** Add a folder to the folder map for the specified world. Does not trigger any events. */
-	bool AddFolderToWorld(UWorld& InWorld, const FFolder& InFolder);
+	UNREALED_API bool AddFolderToWorld(UWorld& InWorld, const FFolder& InFolder);
 
 	/** Removed folders from specified world. Can optionally trigger delete events. */
-	void RemoveFoldersFromWorld(UWorld& InWorld, const TArray<FFolder>& InFolders, bool bBroadcastDelete);
+	UNREALED_API void RemoveFoldersFromWorld(UWorld& InWorld, const TArray<FFolder>& InFolders, bool bBroadcastDelete);
 
 	/** Transient map of folders, keyed on world pointer */
 	TMap<TWeakObjectPtr<UWorld>, TObjectPtr<UWorldFolders>> WorldFolders;
 
 	/** Called when ActorEditorContextClient changed. */
-	void BroadcastOnActorEditorContextClientChanged(UWorld& InWorld);
+	UNREALED_API void BroadcastOnActorEditorContextClientChanged(UWorld& InWorld);
 
 	/** Delegate used to notify changes to ActorEditorContextSubsystem */
 	FOnActorEditorContextClientChanged ActorEditorContextClientChanged;
diff --git a/Engine/Source/Editor/UnrealEd/Public/EditorAnimUtils.h b/Engine/Source/Editor/UnrealEd/Public/EditorAnimUtils.h
--- a/Engine/Source/Editor/UnrealEd/Public/EditorAnimUtils.h
+++ b/Engine/Source/Editor/UnrealEd/Public/EditorAnimUtils.h
@@ -36,34 +36,34 @@ namespace EditorAnimUtils
 		FString Rename(const UObject* Asset) const;
 	};
 
-	class UNREALED_API FAnimationRetargetContext
+	class FAnimationRetargetContext
 	{
 	public:
-		FAnimationRetargetContext(const TArray<FAssetData>& AssetsToRetarget, bool bRetargetReferredAssets, bool bInConvertAnimationDataInComponentSpaces, const FNameDuplicationRule& NameRule=FNameDuplicationRule());
-		FAnimationRetargetContext(TArray<TWeakObjectPtr<UObject>> AssetsToRetarget, bool bRetargetReferredAssets, bool bInConvertAnimationDataInComponentSpaces, const FNameDuplicationRule& NameRule=FNameDuplicationRule());
+		UNREALED_API FAnimationRetargetContext(const TArray<FAssetData>& AssetsToRetarget, bool bRetargetReferredAssets, bool bInConvertAnimationDataInComponentSpaces, const FNameDuplicationRule& NameRule=FNameDuplicationRule());
+		UNREALED_API FAnimationRetargetContext(TArray<TWeakObjectPtr<UObject>> AssetsToRetarget, bool bRetargetReferredAssets, bool bInConvertAnimationDataInComponentSpaces, const FNameDuplicationRule& NameRule=FNameDuplicationRule());
 
 		/** Were we supplied anything that we can retarget */
-		bool HasAssetsToRetarget() const;
+		UNREALED_API bool HasAssetsToRetarget() const;
 
 		/** Did we duplicate any assets */
-		bool HasDuplicates() const;
+		UNREALED_API bool HasDuplicates() const;
 
 		/** Returns the UObject that was chosen to retarget if there was only one in the first place */
-		UObject* GetSingleTargetObject() const;
+		UNREALED_API UObject* GetSingleTargetObject() const;
 
 		/** Returns the duplicate of the supplied object if there is one, otherwise NULL */
-		UObject* GetDuplicate(const UObject* OriginalObject) const;
+		UNREALED_API UObject* GetDuplicate(const UObject* OriginalObject) const;
 
 		/** return all duplicates */
-		TArray<UObject*> GetAllDuplicates() const;
+		UNREALED_API TArray<UObject*> GetAllDuplicates() const;
 
 		/** Duplicates the assets stored for retargetting, populating maps of original assets to new asset */
-		void DuplicateAssetsToRetarget(UPackage* DestinationPackage, const FNameDuplicationRule* NameRule);
+		UNREALED_API void DuplicateAssetsToRetarget(UPackage* DestinationPackage, const FNameDuplicationRule* NameRule);
 
 		/** Retarget the contained assets */
-		void RetargetAnimations(USkeleton* OldSkeleton, USkeleton* NewSkeleton);
+		UNREALED_API void RetargetAnimations(USkeleton* OldSkeleton, USkeleton* NewSkeleton);
 	
-		void AddRemappedAsset(UAnimationAsset* OriginalAsset, UAnimationAsset* NewAsset);
+		UNREALED_API void AddRemappedAsset(UAnimationAsset* OriginalAsset, UAnimationAsset* NewAsset);
 
 	private:
 		/** Lists of assets to retarget. Populated from FAssetData supplied to constructor */
diff --git a/Engine/Source/Editor/UnrealEd/Public/EditorBuildUtils.h b/Engine/Source/Editor/UnrealEd/Public/EditorBuildUtils.h
--- a/Engine/Source/Editor/UnrealEd/Public/EditorBuildUtils.h
+++ b/Engine/Source/Editor/UnrealEd/Public/EditorBuildUtils.h
@@ -82,10 +82,10 @@ public:
 	};
 
 	/** Helper struct to specify settings for an automated editor build */
-	struct UNREALED_API FEditorAutomatedBuildSettings
+	struct FEditorAutomatedBuildSettings
 	{
 		/** Constructor */
-		FEditorAutomatedBuildSettings();
+		UNREALED_API FEditorAutomatedBuildSettings();
 
 		/** Behavior to take when a map build results in map check errors */
 		EAutomatedBuildBehavior BuildErrorBehavior;
@@ -195,12 +195,12 @@ public:
 	*
 	* @return	true if the build was triggered as user request; false if it did not 
 	*/
-	static UNREALED_API bool IsBuildingNavigationFromUserRequest() { return bBuildingNavigationFromUserRequest; }
+	static bool IsBuildingNavigationFromUserRequest() { return bBuildingNavigationFromUserRequest; }
 
 	/** 
 	* call it to notify that navigation builder finished building 
 	*/
-	static UNREALED_API void PathBuildingFinished() { bBuildingNavigationFromUserRequest = false; }
+	static void PathBuildingFinished() { bBuildingNavigationFromUserRequest = false; }
 
 	/**
 	 * Call this when an async custom build step has completed (successfully or not).
diff --git a/Engine/Source/Editor/UnrealEd/Public/EditorDirectories.h b/Engine/Source/Editor/UnrealEd/Public/EditorDirectories.h
--- a/Engine/Source/Editor/UnrealEd/Public/EditorDirectories.h
+++ b/Engine/Source/Editor/UnrealEd/Public/EditorDirectories.h
@@ -27,18 +27,18 @@ namespace ELastDirectory
 	};
 };
 
-class UNREALED_API FEditorDirectories
+class FEditorDirectories
 {
 public:
-	FEditorDirectories();
+	UNREALED_API FEditorDirectories();
 
-	static FEditorDirectories& Get();
+	static UNREALED_API FEditorDirectories& Get();
 
 	/** Initializes the "LastDir" array with default directories for loading/saving files */
-	void LoadLastDirectories();
+	UNREALED_API void LoadLastDirectories();
 
 	/** Writes the current "LastDir" array back out to the config files */
-	void SaveLastDirectories();
+	UNREALED_API void SaveLastDirectories();
 
 	/**
 	 *	Fetches the last directory used for the specified type
@@ -46,7 +46,7 @@ public:
 	 *	@param	InLastDir	the identifier for the directory type
 	 *	@return	FString		the string that was last set
 	 */
-	FString GetLastDirectory( ELastDirectory::Type InLastDir ) const;
+	UNREALED_API FString GetLastDirectory( ELastDirectory::Type InLastDir ) const;
 
 	
 	/**
@@ -55,7 +55,7 @@ public:
 	 *	@param	InLastDir	the identifier for the directory type
 	 *	@param	InLastStr	the string to set it as
 	 */
-	void SetLastDirectory( ELastDirectory::Type InLastDir, const FString& InLastStr );
+	UNREALED_API void SetLastDirectory( ELastDirectory::Type InLastDir, const FString& InLastStr );
 
 	/**
 	 *	Overrides how configuration is saved and the default directory
@@ -64,10 +64,10 @@ public:
 	 *	@param	InConfigSectionName	section name to use inside config file
 	 *	@param	InDefaultDir	default directory to use if config file did not have saved directory
 	 */
-	void SetOverride(const FString& InConfigFile, const FString& InConfigSectionName, const FString& InDefaultDir);
+	UNREALED_API void SetOverride(const FString& InConfigFile, const FString& InConfigSectionName, const FString& InDefaultDir);
 
 	/** Removes override of config file, config section and default directory */
-	void ResetOverride();
+	UNREALED_API void ResetOverride();
 
 private:
 	/** Array of the last directories used for various editor windows */
diff --git a/Engine/Source/Editor/UnrealEd/Public/EditorDragTools.h b/Engine/Source/Editor/UnrealEd/Public/EditorDragTools.h
--- a/Engine/Source/Editor/UnrealEd/Public/EditorDragTools.h
+++ b/Engine/Source/Editor/UnrealEd/Public/EditorDragTools.h
@@ -15,10 +15,10 @@ class FSceneView;
  * The base class that all drag tools inherit from.
  * The drag tools implement special behaviors for the user clicking and dragging in a viewport.
  */
-class UNREALED_API FDragTool
+class FDragTool
 {
 public:
-	FDragTool(FEditorModeTools* InModeTools);
+	UNREALED_API FDragTool(FEditorModeTools* InModeTools);
 	virtual ~FDragTool() {}
 
 	/**
@@ -27,7 +27,7 @@ public:
 	 *
 	 * @param	InDelta		A delta of mouse movement.
 	 */
-	virtual void AddDelta( const FVector& InDelta );
+	UNREALED_API virtual void AddDelta( const FVector& InDelta );
 
 	/**
 	 * Starts a mouse drag behavior.  The start location is snapped to the editor constraints if bUseSnapping is true.
@@ -35,12 +35,12 @@ public:
 	 * @param	InViewportClient	The viewport client in which the drag event occurred.
 	 * @param	InStart				Where the mouse was when the drag started.
 	 */
-	virtual void StartDrag(FEditorViewportClient* InViewportClient, const FVector& InStartWorld, const FVector2D& InStartScreen);
+	UNREALED_API virtual void StartDrag(FEditorViewportClient* InViewportClient, const FVector& InStartWorld, const FVector2D& InStartScreen);
 
 	/**
 	 * Ends a mouse drag behavior (the user has let go of the mouse button).
 	 */
-	virtual void EndDrag();
+	UNREALED_API virtual void EndDrag();
 	virtual void Render3D(const FSceneView* View,FPrimitiveDrawInterface* PDI) {}
 	virtual void Render(const FSceneView* View,FCanvas* Canvas) {}
 
diff --git a/Engine/Source/Editor/UnrealEd/Public/EditorFolderUtils.h b/Engine/Source/Editor/UnrealEd/Public/EditorFolderUtils.h
--- a/Engine/Source/Editor/UnrealEd/Public/EditorFolderUtils.h
+++ b/Engine/Source/Editor/UnrealEd/Public/EditorFolderUtils.h
@@ -8,11 +8,11 @@
 #include "Misc/Paths.h"
 #include "UObject/NameTypes.h"
 
-class UNREALED_API FEditorFolderUtils
+class FEditorFolderUtils
 {
 public:
 	/** Get the leaf name of a specified folder path */
-	static FName GetLeafName(const FName& InPath);
+	static UNREALED_API FName GetLeafName(const FName& InPath);
 
 	/** Get the parent path for the specified folder path */
 	FORCEINLINE static FName GetParentPath(const FName& InPath)
@@ -21,6 +21,6 @@ public:
 	}
 
 	/** Check if the specified path is a child of the specified parent */
-	static bool PathIsChildOf(const FString& PotentialChild, const FString& Parent);
-	static bool PathIsChildOf(const FName& PotentialChild, const FName& Parent);
-};
\ No newline at end of file
+	static UNREALED_API bool PathIsChildOf(const FString& PotentialChild, const FString& Parent);
+	static UNREALED_API bool PathIsChildOf(const FName& PotentialChild, const FName& Parent);
+};
