[UnrealEd]

diff --git a/Engine/Source/Editor/UnrealEd/Public/EditorModeManager.h b/Engine/Source/Editor/UnrealEd/Public/EditorModeManager.h
--- a/Engine/Source/Editor/UnrealEd/Public/EditorModeManager.h
+++ b/Engine/Source/Editor/UnrealEd/Public/EditorModeManager.h
@@ -33,32 +33,32 @@ class UTypedElementSelectionSet;
 /**
  * A helper class to store the state of the various editor modes.
  */
-class UNREALED_API FEditorModeTools : public FGCObject, public FEditorUndoClient, public TSharedFromThis<FEditorModeTools>
+class FEditorModeTools : public FGCObject, public FEditorUndoClient, public TSharedFromThis<FEditorModeTools>
 {
 public:
-	FEditorModeTools();
-	virtual ~FEditorModeTools();
+	UNREALED_API FEditorModeTools();
+	UNREALED_API virtual ~FEditorModeTools();
 
 	/**
 	 * Set the default editor mode for these tools
 	 * 
 	 * @param	DefaultModeID		The mode ID for the new default mode
 	 */
-	void SetDefaultMode( const FEditorModeID DefaultModeID );
+	UNREALED_API void SetDefaultMode( const FEditorModeID DefaultModeID );
 
 	/**
 	 * Adds a new default mode to this tool's list of default modes.  You can have multiple default modes, but they all must be compatible with each other.
 	 * 
 	 * @param	DefaultModeID		The mode ID for the new default mode
 	 */
-	void AddDefaultMode( const FEditorModeID DefaultModeID );
+	UNREALED_API void AddDefaultMode( const FEditorModeID DefaultModeID );
 
 	/**
 	 * Removes a default mode
 	 * 
 	 * @param	DefaultModeID		The mode ID for the default mode to remove
 	 */
-	void RemoveDefaultMode( const FEditorModeID DefaultModeID );
+	UNREALED_API void RemoveDefaultMode( const FEditorModeID DefaultModeID );
 
 	/**
 	 * Returns whether or not the provided mode ID is a default mode
@@ -68,12 +68,12 @@ public:
 	/**
 	 * Activates the default modes defined by this class.  Note that there can be more than one default mode, and this call will activate them all in sequence.
 	 */
-	void ActivateDefaultMode();
+	UNREALED_API void ActivateDefaultMode();
 
 	/** 
 	 * Returns true if the default modes are active.  Note that there can be more than one default mode, and this will only return true if all default modes are active.
 	 */
-	bool IsDefaultModeActive() const;
+	UNREALED_API bool IsDefaultModeActive() const;
 
 	/**
 	 * Activates an editor mode. Shuts down all other active modes which cannot run with the passed in mode.
@@ -81,19 +81,19 @@ public:
 	 * @param InID		The ID of the editor mode to activate.
 	 * @param bToggle	true if the passed in editor mode should be toggled off if it is already active.
 	 */
-	void ActivateMode( FEditorModeID InID, bool bToggle = false );
+	UNREALED_API void ActivateMode( FEditorModeID InID, bool bToggle = false );
 
 	/**
 	 * Deactivates an editor mode. 
 	 * 
 	 * @param InID		The ID of the editor mode to deactivate.
 	 */
-	void DeactivateMode(FEditorModeID InID);
+	UNREALED_API void DeactivateMode(FEditorModeID InID);
 
 	/**
 	 * Deactivate the mode and entirely purge it from memory. Used when a mode type is unregistered
 	 */
-	void DestroyMode(FEditorModeID InID);
+	UNREALED_API void DestroyMode(FEditorModeID InID);
 
 
 
@@ -101,25 +101,25 @@ public:
 	 * Whether or not the mode toolbox (where mode details panels and some tools are) should be shown.
 	 */
 	UE_DEPRECATED(4.26, "Individual toolkit hosts, such as the level editor, should handle determining if they show a mode toolbox for hosted toolkits.")
-	bool ShouldShowModeToolbox() const;
+	UNREALED_API bool ShouldShowModeToolbox() const;
 protected:
 	/** Exits the given editor mode */
-	void ExitMode(UEdMode* InMode);
+	UNREALED_API void ExitMode(UEdMode* InMode);
 
 	/** Removes the mode ID from the tools manager when a mode is unregistered */
-	void OnModeUnregistered(FEditorModeID ModeID);
+	UNREALED_API void OnModeUnregistered(FEditorModeID ModeID);
 		
-	void DeactivateModeAtIndex(int32 Index);
+	UNREALED_API void DeactivateModeAtIndex(int32 Index);
 public:
 
 	/**
 	 * Deactivates all modes, note some modes can never be deactivated.
 	 */
-	void DeactivateAllModes();
+	UNREALED_API void DeactivateAllModes();
 
-	UEdMode* GetActiveScriptableMode(FEditorModeID InID) const;
+	UNREALED_API UEdMode* GetActiveScriptableMode(FEditorModeID InID) const;
 
-	virtual UTexture2D* GetVertexTexture() const;
+	UNREALED_API virtual UTexture2D* GetVertexTexture() const;
 
 	/**
 	 * Returns true if the current mode is not the specified ModeID.  Also optionally warns the user.
@@ -129,23 +129,23 @@ public:
 	 * @param	bNotifyUser		If true, display the error as a notification, instead of a dialog
 	 * @return					true if the current mode is not the specified mode.
 	 */
-	bool EnsureNotInMode(FEditorModeID ModeID, const FText& ErrorMsg = FText::GetEmpty(), bool bNotifyUser = false) const;
+	UNREALED_API bool EnsureNotInMode(FEditorModeID ModeID, const FText& ErrorMsg = FText::GetEmpty(), bool bNotifyUser = false) const;
 
-	FMatrix GetCustomDrawingCoordinateSystem();
-	FMatrix GetCustomInputCoordinateSystem();
-	FMatrix GetLocalCoordinateSystem();
+	UNREALED_API FMatrix GetCustomDrawingCoordinateSystem();
+	UNREALED_API FMatrix GetCustomInputCoordinateSystem();
+	UNREALED_API FMatrix GetLocalCoordinateSystem();
 	
 	/** 
 	 * Returns true if the passed in editor mode is active 
 	 */
-	bool IsModeActive( FEditorModeID InID ) const;
+	UNREALED_API bool IsModeActive( FEditorModeID InID ) const;
 
 	/**
 	 * Returns a pointer to an active mode specified by the passed in ID
 	 * If the editor mode is not active, NULL is returned
 	 */
-	FEdMode* GetActiveMode( FEditorModeID InID );
-	const FEdMode* GetActiveMode( FEditorModeID InID ) const;
+	UNREALED_API FEdMode* GetActiveMode( FEditorModeID InID );
+	UNREALED_API const FEdMode* GetActiveMode( FEditorModeID InID ) const;
 
 	template <typename SpecificModeType>
 	SpecificModeType* GetActiveModeTyped( FEditorModeID InID )
@@ -163,55 +163,55 @@ public:
 	 * Returns the active tool of the passed in editor mode.
 	 * If the passed in editor mode is not active or the mode has no active tool, NULL is returned
 	 */
-	const FModeTool* GetActiveTool( FEditorModeID InID ) const;
+	UNREALED_API const FModeTool* GetActiveTool( FEditorModeID InID ) const;
 
 	void SetShowWidget( bool InShowWidget )	{ bShowWidget = InShowWidget; }
-	bool GetShowWidget() const;
+	UNREALED_API bool GetShowWidget() const;
 
 	/** Cycle the widget mode, forwarding queries to modes */
-	void CycleWidgetMode (void);
+	UNREALED_API void CycleWidgetMode (void);
 
 	/** Check with modes to see if the widget mode can be cycled */
-	bool CanCycleWidgetMode() const;
+	UNREALED_API bool CanCycleWidgetMode() const;
 
 	/**Save Widget Settings to Ini file*/
-	void SaveWidgetSettings();
+	UNREALED_API void SaveWidgetSettings();
 	/**Load Widget Settings from Ini file*/
-	void LoadWidgetSettings();
+	UNREALED_API void LoadWidgetSettings();
 
 	/** Gets the widget axis to be drawn */
-	EAxisList::Type GetWidgetAxisToDraw( UE::Widget::EWidgetMode InWidgetMode ) const;
+	UNREALED_API EAxisList::Type GetWidgetAxisToDraw( UE::Widget::EWidgetMode InWidgetMode ) const;
 
 	/** Mouse tracking interface.  Passes tracking messages to all active modes */
-	bool StartTracking(FEditorViewportClient* InViewportClient, FViewport* InViewport);
-	bool EndTracking(FEditorViewportClient* InViewportClient, FViewport* InViewport);
+	UNREALED_API bool StartTracking(FEditorViewportClient* InViewportClient, FViewport* InViewport);
+	UNREALED_API bool EndTracking(FEditorViewportClient* InViewportClient, FViewport* InViewport);
 	bool IsTracking() const { return bIsTracking; }
 
-	bool AllowsViewportDragTool() const;
+	UNREALED_API bool AllowsViewportDragTool() const;
 
 	/** Notifies all active modes that a map change has occured */
-	void MapChangeNotify();
+	UNREALED_API void MapChangeNotify();
 
 	/** Notifies all active modes to empty their selections */
-	void SelectNone();
+	UNREALED_API void SelectNone();
 
 	/** Notifies all active modes of box selection attempts */
-	bool BoxSelect( FBox& InBox, bool InSelect );
+	UNREALED_API bool BoxSelect( FBox& InBox, bool InSelect );
 
 	/** Notifies all active modes of frustum selection attempts */
-	bool FrustumSelect( const FConvexVolume& InFrustum, FEditorViewportClient* InViewportClient, bool InSelect );
+	UNREALED_API bool FrustumSelect( const FConvexVolume& InFrustum, FEditorViewportClient* InViewportClient, bool InSelect );
 
 	/** true if any active mode uses a transform widget */
-	bool UsesTransformWidget() const;
+	UNREALED_API bool UsesTransformWidget() const;
 
 	/** true if any active mode uses the passed in transform widget */
-	bool UsesTransformWidget( UE::Widget::EWidgetMode CheckMode ) const;
+	UNREALED_API bool UsesTransformWidget( UE::Widget::EWidgetMode CheckMode ) const;
 
 	/** Sets the current widget axis */
-	void SetCurrentWidgetAxis( EAxisList::Type NewAxis );
+	UNREALED_API void SetCurrentWidgetAxis( EAxisList::Type NewAxis );
 
 	/** Notifies all active modes of mouse click messages. */
-	bool HandleClick(FEditorViewportClient* InViewportClient,  HHitProxy *HitProxy, const FViewportClick &Click );
+	UNREALED_API bool HandleClick(FEditorViewportClient* InViewportClient,  HHitProxy *HitProxy, const FViewportClick &Click );
 
 	/**
 	 * Allows editor modes to override the bounding box used to focus the viewport on a selection
@@ -221,25 +221,25 @@ public:
 	 * @param InOutBox		The box that should be computed for the actor and component
 	 * @return bool			true if a mode overrides the box and populated InOutBox, false if it did not populate InOutBox
 	 */
-	bool ComputeBoundingBoxForViewportFocus(AActor* Actor, UPrimitiveComponent* PrimitiveComponent, FBox& InOutBox);
+	UNREALED_API bool ComputeBoundingBoxForViewportFocus(AActor* Actor, UPrimitiveComponent* PrimitiveComponent, FBox& InOutBox);
 
 	/** true if the passed in brush actor should be drawn in wireframe */	
-	bool ShouldDrawBrushWireframe( AActor* InActor ) const;
+	UNREALED_API bool ShouldDrawBrushWireframe( AActor* InActor ) const;
 
 	/** true if brush vertices should be drawn */
-	bool ShouldDrawBrushVertices() const;
+	UNREALED_API bool ShouldDrawBrushVertices() const;
 
 	/** Ticks all active modes */
-	void Tick( FEditorViewportClient* ViewportClient, float DeltaTime );
+	UNREALED_API void Tick( FEditorViewportClient* ViewportClient, float DeltaTime );
 
 	/** Notifies all active modes of any change in mouse movement */
-	bool InputDelta( FEditorViewportClient* InViewportClient,FViewport* InViewport,FVector& InDrag,FRotator& InRot,FVector& InScale );
+	UNREALED_API bool InputDelta( FEditorViewportClient* InViewportClient,FViewport* InViewport,FVector& InDrag,FRotator& InRot,FVector& InScale );
 
 	/** Notifies all active modes of captured mouse movement */	
-	bool CapturedMouseMove( FEditorViewportClient* InViewportClient, FViewport* InViewport, int32 InMouseX, int32 InMouseY );
+	UNREALED_API bool CapturedMouseMove( FEditorViewportClient* InViewportClient, FViewport* InViewport, int32 InMouseX, int32 InMouseY );
 
 	/** Notifies all active modes of all captured mouse movement */	
-	bool ProcessCapturedMouseMoves( FEditorViewportClient* InViewportClient, FViewport* InViewport, const TArrayView<FIntPoint>& CapturedMouseMoves );
+	UNREALED_API bool ProcessCapturedMouseMoves( FEditorViewportClient* InViewportClient, FViewport* InViewport, const TArrayView<FIntPoint>& CapturedMouseMoves );
 
 	/** 
 	 * Notifies all active modes of keyboard input 
@@ -248,98 +248,98 @@ public:
 	 *  set this to false due to some behaviors being routed in different conditions to legacy modes compared
 	 *  to the input router (see its use in EditorViewportClient.cpp).
 	 */
-	bool InputKey( FEditorViewportClient* InViewportClient, FViewport* Viewport, FKey Key, EInputEvent Event, bool bRouteToToolsContext = true);
+	UNREALED_API bool InputKey( FEditorViewportClient* InViewportClient, FViewport* Viewport, FKey Key, EInputEvent Event, bool bRouteToToolsContext = true);
 
 	/** Notifies all active modes of axis movement */
-	bool InputAxis( FEditorViewportClient* InViewportClient, FViewport* Viewport, int32 ControllerId, FKey Key, float Delta, float DeltaTime);
+	UNREALED_API bool InputAxis( FEditorViewportClient* InViewportClient, FViewport* Viewport, int32 ControllerId, FKey Key, float Delta, float DeltaTime);
 
-	bool MouseEnter( FEditorViewportClient* InViewportClient, FViewport* Viewport, int32 X, int32 Y );
+	UNREALED_API bool MouseEnter( FEditorViewportClient* InViewportClient, FViewport* Viewport, int32 X, int32 Y );
 	
-	bool MouseLeave( FEditorViewportClient* InViewportClient, FViewport* Viewport );
+	UNREALED_API bool MouseLeave( FEditorViewportClient* InViewportClient, FViewport* Viewport );
 
 	/** Notifies all active modes that the mouse has moved */
-	bool MouseMove( FEditorViewportClient* InViewportClient, FViewport* Viewport, int32 X, int32 Y );
+	UNREALED_API bool MouseMove( FEditorViewportClient* InViewportClient, FViewport* Viewport, int32 X, int32 Y );
 
 	/** Notifies all active modes that a viewport has received focus */
-	bool ReceivedFocus( FEditorViewportClient* InViewportClient, FViewport* Viewport );
+	UNREALED_API bool ReceivedFocus( FEditorViewportClient* InViewportClient, FViewport* Viewport );
 
 	/** Notifies all active modes that a viewport has lost focus */
-	bool LostFocus( FEditorViewportClient* InViewportClient, FViewport* Viewport );
+	UNREALED_API bool LostFocus( FEditorViewportClient* InViewportClient, FViewport* Viewport );
 
 	/** Draws all active modes */	
-	void DrawActiveModes( const FSceneView* InView, FPrimitiveDrawInterface* PDI );
+	UNREALED_API void DrawActiveModes( const FSceneView* InView, FPrimitiveDrawInterface* PDI );
 
 	/** Renders all active modes */
-	void Render( const FSceneView* InView, FViewport* Viewport, FPrimitiveDrawInterface* PDI );
+	UNREALED_API void Render( const FSceneView* InView, FViewport* Viewport, FPrimitiveDrawInterface* PDI );
 
 	/** Draws the HUD for all active modes */
-	void DrawHUD( FEditorViewportClient* InViewportClient,FViewport* Viewport,const FSceneView* View,FCanvas* Canvas );
+	UNREALED_API void DrawHUD( FEditorViewportClient* InViewportClient,FViewport* Viewport,const FSceneView* View,FCanvas* Canvas );
 
 	/** 
 	 * Get a pivot point specified by any active modes around which the camera should orbit
 	 * @param	OutPivot	The custom pivot point returned by the mode/tool
 	 * @return	true if a custom pivot point was specified, false otherwise.
 	 */
-	bool GetPivotForOrbit( FVector& OutPivot ) const;
+	UNREALED_API bool GetPivotForOrbit( FVector& OutPivot ) const;
 
 	/** Calls PostUndo on all active modes */
 	// Begin FEditorUndoClient
-	virtual void PostUndo(bool bSuccess) override;
-	virtual void PostRedo(bool bSuccess) override;
+	UNREALED_API virtual void PostUndo(bool bSuccess) override;
+	UNREALED_API virtual void PostRedo(bool bSuccess) override;
 	// End of FEditorUndoClient
 
 	/** True if we should allow widget move */
-	bool AllowWidgetMove() const;
+	UNREALED_API bool AllowWidgetMove() const;
 
 	/** True if we should disallow mouse delta tracking. */
-	bool DisallowMouseDeltaTracking() const;
+	UNREALED_API bool DisallowMouseDeltaTracking() const;
 
 	/** Get a cursor to override the default with, if any */
-	bool GetCursor(EMouseCursor::Type& OutCursor) const;
+	UNREALED_API bool GetCursor(EMouseCursor::Type& OutCursor) const;
 
 	/** Get override cursor visibility settings */
-	bool GetOverrideCursorVisibility(bool& bWantsOverride, bool& bHardwareCursorVisible, bool bSoftwareCursorVisible) const;
+	UNREALED_API bool GetOverrideCursorVisibility(bool& bWantsOverride, bool& bHardwareCursorVisible, bool bSoftwareCursorVisible) const;
 
 	/** Called before converting mouse movement to drag/rot */
-	bool PreConvertMouseMovement(FEditorViewportClient* InViewportClient);
+	UNREALED_API bool PreConvertMouseMovement(FEditorViewportClient* InViewportClient);
 
 	/** Called after converting mouse movement to drag/rot */
-	bool PostConvertMouseMovement(FEditorViewportClient* InViewportClient);
+	UNREALED_API bool PostConvertMouseMovement(FEditorViewportClient* InViewportClient);
 
 	/**
 	 * Returns a good location to draw the widget at.
 	 */
-	FVector GetWidgetLocation() const;
+	UNREALED_API FVector GetWidgetLocation() const;
 
 	/**
 	 * Changes the current widget mode.
 	 */
-	void SetWidgetMode( UE::Widget::EWidgetMode InWidgetMode );
+	UNREALED_API void SetWidgetMode( UE::Widget::EWidgetMode InWidgetMode );
 
 	/**
 	 * Allows you to temporarily override the widget mode.  Call this function again
 	 * with WM_None to turn off the override.
 	 */
-	void SetWidgetModeOverride( UE::Widget::EWidgetMode InWidgetMode );
+	UNREALED_API void SetWidgetModeOverride( UE::Widget::EWidgetMode InWidgetMode );
 
 	/**
 	 * Retrieves the current widget mode, taking overrides into account.
 	 */
-	UE::Widget::EWidgetMode GetWidgetMode() const;
+	UNREALED_API UE::Widget::EWidgetMode GetWidgetMode() const;
 
 
 	/**
 	* Set Scale On The Widget
 	*/
-	void SetWidgetScale(float InScale);
+	UNREALED_API void SetWidgetScale(float InScale);
 
 	/**
 	*  Get Widget Scale
 	*/
-	float GetWidgetScale() const;
+	UNREALED_API float GetWidgetScale() const;
 
 	// FGCObject interface
-	virtual void AddReferencedObjects( FReferenceCollector& Collector ) override;
+	UNREALED_API virtual void AddReferencedObjects( FReferenceCollector& Collector ) override;
 	virtual FString GetReferencerName() const override
 	{
 		return TEXT("FEditorModeTools");
@@ -349,12 +349,12 @@ public:
 	/**
 	 * Loads the state that was saved in the INI file
 	 */
-	void LoadConfig(void);
+	UNREALED_API void LoadConfig(void);
 
 	/**
 	 * Saves the current state to the INI file
 	 */
-	void SaveConfig(void);
+	UNREALED_API void SaveConfig(void);
 
 	/** 
 	 * Sets the pivot locations
@@ -362,7 +362,7 @@ public:
 	 * @param Location 		The location to set
 	 * @param bIncGridBase	Whether or not to also set the GridBase
 	 */
-	void SetPivotLocation( const FVector& Location, const bool bIncGridBase );
+	UNREALED_API void SetPivotLocation( const FVector& Location, const bool bIncGridBase );
 
 	/**
 	 * Multicast delegate for OnModeEntered and OnModeExited callbacks.
@@ -395,10 +395,10 @@ public:
 	 * 
 	 * @param bGetRawValue true when you want the actual value of CoordSystem, not the value modified by the state.
 	 */
-	ECoordSystem GetCoordSystem(bool bGetRawValue = false);
+	UNREALED_API ECoordSystem GetCoordSystem(bool bGetRawValue = false);
 
 	/** Sets the current CoordSystem */
-	void SetCoordSystem(ECoordSystem NewCoordSystem);
+	UNREALED_API void SetCoordSystem(ECoordSystem NewCoordSystem);
 
 	/** Sets the hide viewport UI state */
 	void SetHideViewportUI( bool bInHideViewportUI ) { bHideViewportUI = bInHideViewportUI; }
@@ -407,7 +407,7 @@ public:
 	bool IsViewportUIHidden() const { return bHideViewportUI; }
 
 	/** Called by Editors when they are about to close */
-	bool OnRequestClose();
+	UNREALED_API bool OnRequestClose();
 
 	bool PivotShown;
 	bool Snapping;
@@ -428,91 +428,91 @@ public:
 	FString InfoString;
 
 	/** Sets the host for toolkits created via modes from this mode manager (can only be called once) */
-	void SetToolkitHost(TSharedRef<IToolkitHost> Host);
+	UNREALED_API void SetToolkitHost(TSharedRef<IToolkitHost> Host);
 
 	/** Returns the host for toolkits created via modes from this mode manager */
-	TSharedPtr<IToolkitHost> GetToolkitHost() const;
+	UNREALED_API TSharedPtr<IToolkitHost> GetToolkitHost() const;
 
 	/** Check if toolkit host exists */
-	bool HasToolkitHost() const;
+	UNREALED_API bool HasToolkitHost() const;
 
 	/**
 	 * Returns the set of selected actors.
 	 */
-	virtual USelection* GetSelectedActors() const;
+	UNREALED_API virtual USelection* GetSelectedActors() const;
 
 	/**
 	 * @return the set of selected non-actor objects.
 	 */
-	virtual USelection* GetSelectedObjects() const;
+	UNREALED_API virtual USelection* GetSelectedObjects() const;
 
 	/**
 	 * Returns the set of selected components.
 	 */
-	virtual USelection* GetSelectedComponents() const;
+	UNREALED_API virtual USelection* GetSelectedComponents() const;
 
 	/**
 	 * Returns the selection set for the toolkit host.
 	 * (i.e. the selection set for the level editor)
 	 */
-	UTypedElementSelectionSet* GetEditorSelectionSet() const;
+	UNREALED_API UTypedElementSelectionSet* GetEditorSelectionSet() const;
 
 	/**
 	 * Stores the current selection under the given key, and clears the current selection state if requested.
 	 */
-	void StoreSelection(FName SelectionStoreKey, bool bClearSelection = true);
+	UNREALED_API void StoreSelection(FName SelectionStoreKey, bool bClearSelection = true);
 
 	/**
 	 * Restores the selection to the state that was stored using the given key.
 	 */
-	void RestoreSelection(FName SelectionStoreKey);
+	UNREALED_API void RestoreSelection(FName SelectionStoreKey);
 
 	/**
 	 * Returns the world that is being edited by this mode manager
 	 */ 
-	virtual UWorld* GetWorld() const;
+	UNREALED_API virtual UWorld* GetWorld() const;
 
 	/**
 	 * Returns the currently hovered viewport client
 	 */
-	FEditorViewportClient* GetHoveredViewportClient() const;
+	UNREALED_API FEditorViewportClient* GetHoveredViewportClient() const;
 
 	/**
 	 * Returns the currently focused viewport client
 	 */
-	FEditorViewportClient* GetFocusedViewportClient() const;
+	UNREALED_API FEditorViewportClient* GetFocusedViewportClient() const;
 
 	/**
 	 * Whether or not the current selection has a scene component selected
  	 */
-	bool SelectionHasSceneComponent() const;
+	UNREALED_API bool SelectionHasSceneComponent() const;
 
-	bool IsSelectionAllowed(AActor* InActor, const bool bInSelected) const;
+	UNREALED_API bool IsSelectionAllowed(AActor* InActor, const bool bInSelected) const;
 
-	bool IsSelectionHandled(AActor* InActor, const bool bInSelected) const;
+	UNREALED_API bool IsSelectionHandled(AActor* InActor, const bool bInSelected) const;
 
-	bool ProcessEditDuplicate();
-	bool ProcessEditDelete();
-	bool ProcessEditCut();
-	bool ProcessEditCopy();
-	bool ProcessEditPaste();
-	EEditAction::Type  GetActionEditDuplicate();
-	EEditAction::Type  GetActionEditDelete();
-	EEditAction::Type  GetActionEditCut();
-	EEditAction::Type  GetActionEditCopy();
-	EEditAction::Type GetActionEditPaste();
+	UNREALED_API bool ProcessEditDuplicate();
+	UNREALED_API bool ProcessEditDelete();
+	UNREALED_API bool ProcessEditCut();
+	UNREALED_API bool ProcessEditCopy();
+	UNREALED_API bool ProcessEditPaste();
+	UNREALED_API EEditAction::Type  GetActionEditDuplicate();
+	UNREALED_API EEditAction::Type  GetActionEditDelete();
+	UNREALED_API EEditAction::Type  GetActionEditCut();
+	UNREALED_API EEditAction::Type  GetActionEditCopy();
+	UNREALED_API EEditAction::Type GetActionEditPaste();
 
 	UE_DEPRECATED(5.0, "This function is redundant, and is handled as part of a call to ActivateMode.")
-	void DeactivateOtherVisibleModes(FEditorModeID InMode);
-	bool IsSnapRotationEnabled() const;
-	bool SnapRotatorToGridOverride(FRotator& InRotation) const;
-	void ActorsDuplicatedNotify(TArray<AActor*>& InPreDuplicateSelection, TArray<AActor*>& InPostDuplicateSelection, const bool bOffsetLocations);
-	void ActorMoveNotify();
-	void ActorSelectionChangeNotify();
-	void ActorPropChangeNotify();
-	void UpdateInternalData();
-	bool IsOnlyVisibleActiveMode(FEditorModeID InMode) const;
-	bool IsOnlyActiveMode(FEditorModeID InMode) const;
+	UNREALED_API void DeactivateOtherVisibleModes(FEditorModeID InMode);
+	UNREALED_API bool IsSnapRotationEnabled() const;
+	UNREALED_API bool SnapRotatorToGridOverride(FRotator& InRotation) const;
+	UNREALED_API void ActorsDuplicatedNotify(TArray<AActor*>& InPreDuplicateSelection, TArray<AActor*>& InPostDuplicateSelection, const bool bOffsetLocations);
+	UNREALED_API void ActorMoveNotify();
+	UNREALED_API void ActorSelectionChangeNotify();
+	UNREALED_API void ActorPropChangeNotify();
+	UNREALED_API void UpdateInternalData();
+	UNREALED_API bool IsOnlyVisibleActiveMode(FEditorModeID InMode) const;
+	UNREALED_API bool IsOnlyActiveMode(FEditorModeID InMode) const;
 
 	/*
 	* Sets the active Modes ToolBar Palette Tab to the named Palette
@@ -520,32 +520,32 @@ public:
 	//void  InvokeToolPaletteTab(FEditorModeID InMode, FName InPaletteName);
 
 	/** returns true if all active EdModes are OK with an AutoSave happening now  */
-	bool CanAutoSave() const;
+	UNREALED_API bool CanAutoSave() const;
 	
 	/** returns true if all active EdModes are OK support operation on current asset */
-	bool IsOperationSupportedForCurrentAsset(EAssetOperation InOperation) const;
+	UNREALED_API bool IsOperationSupportedForCurrentAsset(EAssetOperation InOperation) const;
 
-	void RemoveAllDelegateHandlers();
+	UNREALED_API void RemoveAllDelegateHandlers();
 
 	/** @return ToolsContext for this Mode Manager */
-	UModeManagerInteractiveToolsContext* GetInteractiveToolsContext() const;
+	UNREALED_API UModeManagerInteractiveToolsContext* GetInteractiveToolsContext() const;
 
 
 protected:
 	/** 
 	 * Delegate handlers
 	 **/
-	void OnEditorSelectionChanged(UObject* NewSelection);
-	void OnEditorSelectNone();
+	UNREALED_API void OnEditorSelectionChanged(UObject* NewSelection);
+	UNREALED_API void OnEditorSelectNone();
 
 
 	/** Handles the notification when a world is going through GC to clean up any modes pending deactivation. */
-	void OnWorldCleanup(UWorld* InWorld, bool bSessionEnded, bool bCleanupResources);
+	UNREALED_API void OnWorldCleanup(UWorld* InWorld, bool bSessionEnded, bool bCleanupResources);
 
-	virtual void DrawBrackets(FEditorViewportClient* ViewportClient, FViewport* Viewport, const FSceneView* View, FCanvas* Canvas);
+	UNREALED_API virtual void DrawBrackets(FEditorViewportClient* ViewportClient, FViewport* Viewport, const FSceneView* View, FCanvas* Canvas);
 
-	void ForEachEdMode(TFunctionRef<bool(UEdMode*)> InCalllback) const;
-	bool TestAllModes(TFunctionRef<bool(UEdMode*)> InCalllback, bool bExpected) const;
+	UNREALED_API void ForEachEdMode(TFunctionRef<bool(UEdMode*)> InCalllback) const;
+	UNREALED_API bool TestAllModes(TFunctionRef<bool(UEdMode*)> InCalllback, bool bExpected) const;
 	
 	template <class InterfaceToCastTo>
 	void ForEachEdMode(TFunctionRef<bool(InterfaceToCastTo*)> InCallback) const
@@ -561,7 +561,7 @@ protected:
 		});
 	}
 
-	void ExitAllModesPendingDeactivate();
+	UNREALED_API void ExitAllModesPendingDeactivate();
 
 	/** List of default modes for this tool.  These must all be compatible with each other. */
 	TArray<FEditorModeID> DefaultModeIDs;
diff --git a/Engine/Source/Editor/UnrealEd/Public/EditorModeRegistry.h b/Engine/Source/Editor/UnrealEd/Public/EditorModeRegistry.h
--- a/Engine/Source/Editor/UnrealEd/Public/EditorModeRegistry.h
+++ b/Engine/Source/Editor/UnrealEd/Public/EditorModeRegistry.h
@@ -56,12 +56,12 @@ struct IEditorModeFactory : public TSharedFromThis<IEditorModeFactory>
 	virtual TSharedRef<FEdMode> CreateMode() const = 0;
 };
 
-struct UNREALED_API FEditorModeFactory : IEditorModeFactory
+struct FEditorModeFactory : IEditorModeFactory
 {
 public:
-	FEditorModeFactory(const FEditorModeInfo& InModeInfo);
-	FEditorModeFactory(FEditorModeInfo&& InModeInfo);
-	virtual ~FEditorModeFactory();
+	UNREALED_API FEditorModeFactory(const FEditorModeInfo& InModeInfo);
+	UNREALED_API FEditorModeFactory(FEditorModeInfo&& InModeInfo);
+	UNREALED_API virtual ~FEditorModeFactory();
 
 	/** Information pertaining to this factory's mode */
 	FEditorModeInfo ModeInfo;
@@ -72,9 +72,9 @@ public:
 	/**
 	 * Gets the information pertaining to the mode type that this factory creates
 	 */
-	virtual FEditorModeInfo GetModeInfo() const final;
+	UNREALED_API virtual FEditorModeInfo GetModeInfo() const final;
 
-	virtual TSharedRef<FEdMode> CreateMode() const final;
+	UNREALED_API virtual TSharedRef<FEdMode> CreateMode() const final;
 };
 
 /**
diff --git a/Engine/Source/Editor/UnrealEd/Public/EditorModeTools.h b/Engine/Source/Editor/UnrealEd/Public/EditorModeTools.h
--- a/Engine/Source/Editor/UnrealEd/Public/EditorModeTools.h
+++ b/Engine/Source/Editor/UnrealEd/Public/EditorModeTools.h
@@ -47,11 +47,11 @@ enum EModeTools : int8
 /**
  * Base class for all editor mode tools.
  */
-class UNREALED_API FModeTool
+class FModeTool
 {
 public:
-	FModeTool();
-	virtual ~FModeTool();
+	UNREALED_API FModeTool();
+	UNREALED_API virtual ~FModeTool();
 
 	/** Returns the name that gets reported to the editor. */
 	virtual FString GetName() const		{ return TEXT("Default"); }
@@ -105,8 +105,8 @@ public:
 		return false;
 	}
 
-	virtual void Render(const FSceneView* View,FViewport* Viewport,FPrimitiveDrawInterface* PDI);
-	virtual void DrawHUD(FEditorViewportClient* ViewportClient,FViewport* Viewport,const FSceneView* View,FCanvas* Canvas);
+	UNREALED_API virtual void Render(const FSceneView* View,FViewport* Viewport,FPrimitiveDrawInterface* PDI);
+	UNREALED_API virtual void DrawHUD(FEditorViewportClient* ViewportClient,FViewport* Viewport,const FSceneView* View,FCanvas* Canvas);
 
 	//@{
 	virtual bool StartModify()	{ return 0; }
diff --git a/Engine/Source/Editor/UnrealEd/Public/EditorUndoClient.h b/Engine/Source/Editor/UnrealEd/Public/EditorUndoClient.h
--- a/Engine/Source/Editor/UnrealEd/Public/EditorUndoClient.h
+++ b/Engine/Source/Editor/UnrealEd/Public/EditorUndoClient.h
@@ -14,11 +14,11 @@ class FTransactionObjectEvent;
 /**
  * Interface for tools wanting to handle undo/redo operations
  */
-class UNREALED_API FEditorUndoClient
+class FEditorUndoClient
 {
 public:
 	/** Always unregister for undo on destruction, just in case. */
-	virtual ~FEditorUndoClient();
+	UNREALED_API virtual ~FEditorUndoClient();
 
 	/**
 	 * Called to see if the context of the current undo/redo operation is a match for the client
@@ -48,11 +48,11 @@ public:
 };
 
 /** An undo client that registers itself in its constructor and unregisters itself in its destructor */
-class UNREALED_API FSelfRegisteringEditorUndoClient : public FEditorUndoClient
+class FSelfRegisteringEditorUndoClient : public FEditorUndoClient
 {
 public:
 	/** Register in constructor */
-	FSelfRegisteringEditorUndoClient();
+	UNREALED_API FSelfRegisteringEditorUndoClient();
 
 	/** FEditorUndoClient already unregisters in its destructor */
-};
\ No newline at end of file
+};
diff --git a/Engine/Source/Editor/UnrealEd/Public/EditorViewportClient.h b/Engine/Source/Editor/UnrealEd/Public/EditorViewportClient.h
--- a/Engine/Source/Editor/UnrealEd/Public/EditorViewportClient.h
+++ b/Engine/Source/Editor/UnrealEd/Public/EditorViewportClient.h
@@ -55,16 +55,16 @@ namespace EDragTool
 /**
 * Unreal level editor actions
 */
-class UNREALED_API FViewportNavigationCommands : public TCommands<FViewportNavigationCommands>
+class FViewportNavigationCommands : public TCommands<FViewportNavigationCommands>
 {
 
 public:
-	FViewportNavigationCommands();
+	UNREALED_API FViewportNavigationCommands();
 
 	/**
 	* Initialize commands
 	*/
-	virtual void RegisterCommands() override;
+	UNREALED_API virtual void RegisterCommands() override;
 
 
 	TSharedPtr< FUICommandInfo > Forward;
@@ -207,13 +207,13 @@ struct FDropQuery
 /** 
  * Stores the transformation data for the viewport camera
  */
-struct UNREALED_API FViewportCameraTransform
+struct FViewportCameraTransform
 {
 public:
-	FViewportCameraTransform();
+	UNREALED_API FViewportCameraTransform();
 
 	/** Sets the transform's location */
-	void SetLocation( const FVector& Position );
+	UNREALED_API void SetLocation( const FVector& Position );
 
 	/** Sets the transform's rotation */
 	void SetRotation( const FRotator& Rotation )
@@ -234,7 +234,7 @@ public:
 	}
 
 	/** Check if transition curve is playing. */
-	bool IsPlaying();
+	UNREALED_API bool IsPlaying();
 
 	/** @return The transform's location */
 	FORCEINLINE const FVector& GetLocation() const { return ViewLocation; }
@@ -254,19 +254,19 @@ public:
 	 * @param InDesiredLocation	The location to transition to
 	 * @param bInstant			If the desired location should be set instantly rather than transitioned to over time
 	 */
-	void TransitionToLocation( const FVector& InDesiredLocation, TWeakPtr<SWidget> EditorViewportWidget, bool bInstant );
+	UNREALED_API void TransitionToLocation( const FVector& InDesiredLocation, TWeakPtr<SWidget> EditorViewportWidget, bool bInstant );
 
 	/**
 	 * Updates any current location transitions
 	 *
 	 * @return true if there is currently a transition
 	 */
-	bool UpdateTransition();
+	UNREALED_API bool UpdateTransition();
 
 	/**
 	 * Computes a matrix to use for viewport location and rotation when orbiting
 	 */
-	FMatrix ComputeOrbitMatrix() const;
+	UNREALED_API FMatrix ComputeOrbitMatrix() const;
 
 private:
 	/** The time when a transition to the desired location began */
@@ -311,14 +311,14 @@ private:
 DECLARE_MULTICAST_DELEGATE_OneParam(FEditorViewportViewModifierDelegate, FEditorViewportViewModifierParams&);
 
 /** Viewport client for editor viewports. Contains common functionality for camera movement, rendering debug information, etc. */
-class UNREALED_API FEditorViewportClient : public FCommonViewportClient, public FViewElementDrawer, public FGCObject
+class FEditorViewportClient : public FCommonViewportClient, public FViewElementDrawer, public FGCObject
 {
 public:
 	friend class FMouseDeltaTracker;
 
-	FEditorViewportClient(FEditorModeTools* InModeTools, FPreviewScene* InPreviewScene = nullptr, const TWeakPtr<SEditorViewport>& InEditorViewportWidget = nullptr);
+	UNREALED_API FEditorViewportClient(FEditorModeTools* InModeTools, FPreviewScene* InPreviewScene = nullptr, const TWeakPtr<SEditorViewport>& InEditorViewportWidget = nullptr);
 
-	virtual ~FEditorViewportClient();
+	UNREALED_API virtual ~FEditorViewportClient();
 
 	/** Non-copyable */
 	FEditorViewportClient(const FEditorViewportClient&) = delete;
@@ -341,7 +341,7 @@ public:
 	 * @param bShouldBeRealtime	If true, this viewport will be realtime, if false this viewport will not be realtime
 	 * @param SystemDisplayName	This display name of whatever system is overriding realtime. This name is displayed to users in the viewport options menu	
 	 */
-	void AddRealtimeOverride(bool bShouldBeRealtime, FText SystemDisplayName);
+	UNREALED_API void AddRealtimeOverride(bool bShouldBeRealtime, FText SystemDisplayName);
 
 	UE_DEPRECATED(4.26, "SetRealtimeOverride is replaced with AddRealtimeOverride, as multiple overrides can now be added.")
 	void SetRealtimeOverride(bool bShouldBeRealtime, FText SystemDisplayName)
@@ -352,7 +352,7 @@ public:
 	/**
 	 * Returns whether there's a realtime override registered with the given system name.
 	 */
-	bool HasRealtimeOverride(FText SystemDisplayName) const;
+	UNREALED_API bool HasRealtimeOverride(FText SystemDisplayName) const;
 
 	/**
 	 * Removes the most recent realtime override registered with the given system name.
@@ -360,7 +360,7 @@ public:
 	 * @param Whether to assert if no matching realtime override was found.
 	 * @return Whether any matching realtime override was found.
 	 */
-	bool RemoveRealtimeOverride(FText SystemDisplayName, bool bCheckMissingOverride = true);
+	UNREALED_API bool RemoveRealtimeOverride(FText SystemDisplayName, bool bCheckMissingOverride = true);
 
 	UE_DEPRECATED(4.26, "RemoveRealtimeOverride now takes a system display name to help remove the correct override.")
 	void RemoveRealtimeOverride()
@@ -374,7 +374,7 @@ public:
 	 *
 	 * @return   Whether there was any realtime override to be popped.
 	 */
-	bool PopRealtimeOverride();
+	UNREALED_API bool PopRealtimeOverride();
 
 	/**
 	 * Toggles whether or not the viewport updates in realtime and returns the updated state.
@@ -382,13 +382,13 @@ public:
 	 *
 	 * @return		The current state of the realtime flag.
 	 */
-	bool ToggleRealtime();
+	UNREALED_API bool ToggleRealtime();
 
 	/** 
 	 * Sets whether or not the viewport updates in realtime.
 	 * Note: This value is saved between editor sessions so it should not be used for temporary states.  For that see SetRealtimeOverride
 	 */
-	void SetRealtime(bool bInRealtime);
+	UNREALED_API void SetRealtime(bool bInRealtime);
 
 	/** @return		True if viewport is in realtime mode, false otherwise. */
 	bool IsRealtime() const
@@ -410,7 +410,7 @@ public:
 	/**
 	 * Saves the realtime state to a config location.  Does not save any temp overrides
 	 */ 
-	void SaveRealtimeStateToConfig(bool& ConfigVar) const;
+	UNREALED_API void SaveRealtimeStateToConfig(bool& ConfigVar) const;
 
 	/**
 	 * @return true if there are any temp realtime overrides set
@@ -420,47 +420,47 @@ public:
 	/**
 	 * @return  If an override is set this returns the message indicating what set it
 	 */
-	FText GetRealtimeOverrideMessage() const;
+	UNREALED_API FText GetRealtimeOverrideMessage() const;
 
 	UE_DEPRECATED(4.25, "SetRealtime no longer takes in bStoreCurrentValue parameter. For temporary overrides use AddRealtimeOverride")
-	void SetRealtime(bool bInRealtime, bool bStoreCurrentValue);
+	UNREALED_API void SetRealtime(bool bInRealtime, bool bStoreCurrentValue);
 
 	/**
 	 * Restores realtime setting to stored value. This will only enable realtime and 
 	 * never disable it (unless bAllowDisable is true)
 	 */
 	UE_DEPRECATED(4.25, "To save and restore realtime state non-permanently use AddRealtimeOverride and RemoveRealtimeOverride")
-	void RestoreRealtime(const bool bAllowDisable = false);
+	UNREALED_API void RestoreRealtime(const bool bAllowDisable = false);
 
 	// this set ups camera for both orbit and non orbit control
-	void SetCameraSetup(const FVector& LocationForOrbiting, const FRotator& InOrbitRotation, const FVector& InOrbitZoom, const FVector& InOrbitLookAt, const FVector& InViewLocation, const FRotator &InViewRotation );
+	UNREALED_API void SetCameraSetup(const FVector& LocationForOrbiting, const FRotator& InOrbitRotation, const FVector& InOrbitZoom, const FVector& InOrbitLookAt, const FVector& InViewLocation, const FRotator &InViewRotation );
 
 	/** Callback for toggling the camera lock flag. */
-	virtual void SetCameraLock();
+	UNREALED_API virtual void SetCameraLock();
 
 	/** Callback for checking the camera lock flag. */
-	bool IsCameraLocked() const;
+	UNREALED_API bool IsCameraLocked() const;
 
 	/** Callback for toggling the grid show flag. */
-	void SetShowGrid();
+	UNREALED_API void SetShowGrid();
 
 	/** Callback for checking the grid show flag. */
-	bool IsSetShowGridChecked() const;
+	UNREALED_API bool IsSetShowGridChecked() const;
 
 	/** Sets the show bounds flag */
-	void SetShowBounds(bool bShow);
+	UNREALED_API void SetShowBounds(bool bShow);
 
 	/** Callback for toggling the bounds show flag. */
-	void ToggleShowBounds();
+	UNREALED_API void ToggleShowBounds();
 
 	/** Callback for checking the bounds show flag. */
-	bool IsSetShowBoundsChecked() const;
+	UNREALED_API bool IsSetShowBoundsChecked() const;
 
 	/** Callback for toggling the collision geometry show flag. */
-	void SetShowCollision();
+	UNREALED_API void SetShowCollision();
 
 	/** Callback for checking the collision geometry show flag. */
-	bool IsSetShowCollisionChecked() const;
+	UNREALED_API bool IsSetShowCollisionChecked() const;
 
 	/** Gets ViewportCameraTransform object for the current viewport type */
 	FViewportCameraTransform& GetViewTransform()
@@ -509,7 +509,7 @@ public:
 		}
 	}
 	/** Perform default camera movement that would normally happen* This can be used by systems to override how the camera is triggered*/
-	void PeformDefaultCameraMovement(FVector& Drag, FRotator& Rot, FVector& Scale);
+	UNREALED_API void PeformDefaultCameraMovement(FVector& Drag, FRotator& Rot, FVector& Scale);
 
 	/** Sets ortho zoom amount */
 	void SetOrthoZoom( float InOrthoZoom ) 
@@ -550,10 +550,10 @@ public:
 	}
 
 	/** @return The number of units per pixel displayed in this viewport */
-	float GetOrthoUnitsPerPixel(const FViewport* Viewport) const;
+	UNREALED_API float GetOrthoUnitsPerPixel(const FViewport* Viewport) const;
 
 	/** Get a prettified string representation of the specified unreal units */
-	static FString UnrealUnitsToSiUnits(float UnrealUnits);
+	static UNREALED_API FString UnrealUnitsToSiUnits(float UnrealUnits);
 
 	void RemoveCameraRoll()
 	{
@@ -562,60 +562,60 @@ public:
 		SetViewRotation( Rotation );
 	}
 
-	void SetInitialViewTransform(ELevelViewportType ViewportType, const FVector& ViewLocation, const FRotator& ViewRotation, float InOrthoZoom );
+	UNREALED_API void SetInitialViewTransform(ELevelViewportType ViewportType, const FVector& ViewLocation, const FRotator& ViewRotation, float InOrthoZoom );
 
-	void TakeHighResScreenShot();
+	UNREALED_API void TakeHighResScreenShot();
 
 	/** Called when an editor mode has been (de)activated */
-	void OnEditorModeIDChanged(const FEditorModeID& EditorModeID, bool bIsEntering);
+	UNREALED_API void OnEditorModeIDChanged(const FEditorModeID& EditorModeID, bool bIsEntering);
 
 	/** FViewElementDrawer interface */
-	virtual void Draw(const FSceneView* View,FPrimitiveDrawInterface* PDI) override;
-	virtual void Draw(FViewport* Viewport,FCanvas* Canvas) override;
+	UNREALED_API virtual void Draw(const FSceneView* View,FPrimitiveDrawInterface* PDI) override;
+	UNREALED_API virtual void Draw(FViewport* Viewport,FCanvas* Canvas) override;
 
 	/**
 	 * Gets the world space cursor info from the current mouse position
 	 *
 	 * @return					An FViewportCursorLocation containing information about the mouse position in world space.
 	 */
-	FViewportCursorLocation GetCursorWorldLocationFromMousePos();
+	UNREALED_API FViewportCursorLocation GetCursorWorldLocationFromMousePos();
 
 	/** FViewportClient interface */
-	virtual bool ProcessScreenShots(FViewport* Viewport) override;
-	virtual void RedrawRequested(FViewport* Viewport) override;
-	virtual void RequestInvalidateHitProxy(FViewport* Viewport) override;
+	UNREALED_API virtual bool ProcessScreenShots(FViewport* Viewport) override;
+	UNREALED_API virtual void RedrawRequested(FViewport* Viewport) override;
+	UNREALED_API virtual void RequestInvalidateHitProxy(FViewport* Viewport) override;
 	UE_DEPRECATED(5.1, "This version of InputKey is deprecated. Please use the version that takes FInputKeyEventArgs instead.")
-	virtual bool InputKey(FViewport* InViewport, int32 ControllerId, FKey Key, EInputEvent Event, float AmountDepressed = 1.f, bool bGamepad=false) override;
-	virtual bool InputKey(const FInputKeyEventArgs& EventArgs) override;
+	UNREALED_API virtual bool InputKey(FViewport* InViewport, int32 ControllerId, FKey Key, EInputEvent Event, float AmountDepressed = 1.f, bool bGamepad=false) override;
+	UNREALED_API virtual bool InputKey(const FInputKeyEventArgs& EventArgs) override;
 	UE_DEPRECATED(5.1, "This version of InputAxis is deprecated. Please use the version that takes a DeviceId instead.")
-	virtual bool InputAxis(FViewport* InViewport, int32 ControllerId, FKey Key, float Delta, float DeltaTime, int32 NumSamples=1, bool bGamepad=false) override;
-	virtual bool InputAxis(FViewport* Viewport, FInputDeviceId DeviceID, FKey Key, float Delta, float DeltaTime, int32 NumSamples=1, bool bGamepad=false) override;
-	virtual bool InputGesture(FViewport* Viewport, EGestureEvent GestureType, const FVector2D& GestureDelta, bool bIsDirectionInvertedFromDevice) override;
-	virtual void ReceivedFocus(FViewport* Viewport) override;
-	virtual void MouseEnter(FViewport* Viewport,int32 x, int32 y) override;
-	virtual void MouseMove(FViewport* Viewport,int32 x, int32 y) override;
-	virtual void MouseLeave( FViewport* Viewport ) override;
-	virtual EMouseCursor::Type GetCursor(FViewport* Viewport,int32 X,int32 Y) override;
-	virtual void CapturedMouseMove( FViewport* InViewport, int32 InMouseX, int32 InMouseY ) override;
-	virtual void ProcessAccumulatedPointerInput(FViewport* InViewport) override;
-	virtual bool IsOrtho() const override;
-	virtual void LostFocus(FViewport* Viewport) override;
-	virtual FStatUnitData* GetStatUnitData() const override;
-	virtual FStatHitchesData* GetStatHitchesData() const override;
-	virtual const TArray<FString>* GetEnabledStats() const override;
-	virtual void SetEnabledStats(const TArray<FString>& InEnabledStats) override;
-	virtual bool IsStatEnabled(const FString& InName) const override;
+	UNREALED_API virtual bool InputAxis(FViewport* InViewport, int32 ControllerId, FKey Key, float Delta, float DeltaTime, int32 NumSamples=1, bool bGamepad=false) override;
+	UNREALED_API virtual bool InputAxis(FViewport* Viewport, FInputDeviceId DeviceID, FKey Key, float Delta, float DeltaTime, int32 NumSamples=1, bool bGamepad=false) override;
+	UNREALED_API virtual bool InputGesture(FViewport* Viewport, EGestureEvent GestureType, const FVector2D& GestureDelta, bool bIsDirectionInvertedFromDevice) override;
+	UNREALED_API virtual void ReceivedFocus(FViewport* Viewport) override;
+	UNREALED_API virtual void MouseEnter(FViewport* Viewport,int32 x, int32 y) override;
+	UNREALED_API virtual void MouseMove(FViewport* Viewport,int32 x, int32 y) override;
+	UNREALED_API virtual void MouseLeave( FViewport* Viewport ) override;
+	UNREALED_API virtual EMouseCursor::Type GetCursor(FViewport* Viewport,int32 X,int32 Y) override;
+	UNREALED_API virtual void CapturedMouseMove( FViewport* InViewport, int32 InMouseX, int32 InMouseY ) override;
+	UNREALED_API virtual void ProcessAccumulatedPointerInput(FViewport* InViewport) override;
+	UNREALED_API virtual bool IsOrtho() const override;
+	UNREALED_API virtual void LostFocus(FViewport* Viewport) override;
+	UNREALED_API virtual FStatUnitData* GetStatUnitData() const override;
+	UNREALED_API virtual FStatHitchesData* GetStatHitchesData() const override;
+	UNREALED_API virtual const TArray<FString>* GetEnabledStats() const override;
+	UNREALED_API virtual void SetEnabledStats(const TArray<FString>& InEnabledStats) override;
+	UNREALED_API virtual bool IsStatEnabled(const FString& InName) const override;
 
 protected:
 
-	virtual bool Internal_InputKey(const FInputKeyEventArgs& EventArgs);
-	virtual bool Internal_InputAxis(FViewport* Viewport, FInputDeviceId DeviceID, FKey Key, float Delta, float DeltaTime, int32 NumSamples=1, bool bGamepad=false);
+	UNREALED_API virtual bool Internal_InputKey(const FInputKeyEventArgs& EventArgs);
+	UNREALED_API virtual bool Internal_InputAxis(FViewport* Viewport, FInputDeviceId DeviceID, FKey Key, float Delta, float DeltaTime, int32 NumSamples=1, bool bGamepad=false);
 
 public:
 	
 	/** FGCObject interface */
-	virtual void AddReferencedObjects( FReferenceCollector& Collector ) override;
-	virtual FString GetReferencerName() const override;
+	UNREALED_API virtual void AddReferencedObjects( FReferenceCollector& Collector ) override;
+	UNREALED_API virtual FString GetReferencerName() const override;
 
 	/**
 	 * Called when the user clicks in the viewport
@@ -627,7 +627,7 @@ public:
 	 * @param HitX			The X location of the mouse
 	 * @param HitY			The Y location of the mouse
 	 */
-	virtual void ProcessClick(FSceneView& View, HHitProxy* HitProxy, FKey Key, EInputEvent Event, uint32 HitX, uint32 HitY);
+	UNREALED_API virtual void ProcessClick(FSceneView& View, HHitProxy* HitProxy, FKey Key, EInputEvent Event, uint32 HitX, uint32 HitY);
 	
 	/**
 	 * Called when mouse movement tracking begins
@@ -652,17 +652,17 @@ public:
 	 * @param Rot			The amount the widget was rotated  (the value depends on the coordinate system of the widget.  See GetWidgetCoordSystem )
 	 * @param Scale			The amount the widget was scaled (the value depends on the coordinate system of the widget.  See GetWidgetCoordSystem )
 	 */
-	virtual bool InputWidgetDelta(FViewport* InViewport, EAxisList::Type CurrentAxis, FVector& Drag, FRotator& Rot, FVector& Scale);
+	UNREALED_API virtual bool InputWidgetDelta(FViewport* InViewport, EAxisList::Type CurrentAxis, FVector& Drag, FRotator& Rot, FVector& Scale);
 
 	/**
 	 * Sets the current widget mode
 	 */
-	virtual void SetWidgetMode(UE::Widget::EWidgetMode NewMode);
+	UNREALED_API virtual void SetWidgetMode(UE::Widget::EWidgetMode NewMode);
 
 	/**
 	 * Whether or not the new widget mode can be set in this viewport
 	 */
-	virtual bool CanSetWidgetMode(UE::Widget::EWidgetMode NewMode) const;
+	UNREALED_API virtual bool CanSetWidgetMode(UE::Widget::EWidgetMode NewMode) const;
 
 	/**
 	 * Whether or not the widget mode can be cycled
@@ -672,43 +672,43 @@ public:
 	/**
 	 * @return The current display mode for transform widget 
 	 */
-	virtual UE::Widget::EWidgetMode GetWidgetMode() const;
+	UNREALED_API virtual UE::Widget::EWidgetMode GetWidgetMode() const;
 
 	/**
 	 * @return The world space location of the transform widget
 	 */
-	virtual FVector GetWidgetLocation() const;
+	UNREALED_API virtual FVector GetWidgetLocation() const;
 
 	/**
 	 * @return The current coordinate system for drawing and input of the transform widget.  
 	 * For world coordiante system return the identity matrix
 	 */
-	virtual FMatrix GetWidgetCoordSystem() const;
+	UNREALED_API virtual FMatrix GetWidgetCoordSystem() const;
 
 	/**
 	* @return The local coordinate system for the transform widget.
 	* For world coordiante system return the identity matrix
 	*/
-	virtual FMatrix GetLocalCoordinateSystem() const;
+	UNREALED_API virtual FMatrix GetLocalCoordinateSystem() const;
 	/**
 	 * Sets the coordinate system space to use
 	 */
-	virtual void SetWidgetCoordSystemSpace( ECoordSystem NewCoordSystem );
+	UNREALED_API virtual void SetWidgetCoordSystemSpace( ECoordSystem NewCoordSystem );
 
 	/**
 	 * @return The coordinate system space (world or local) to display the widget in
 	 */
-	virtual ECoordSystem GetWidgetCoordSystemSpace() const;
+	UNREALED_API virtual ECoordSystem GetWidgetCoordSystemSpace() const;
 
 	/**
 	 * Sets the current axis being manipulated by the transform widget
 	 */
-	virtual void SetCurrentWidgetAxis( EAxisList::Type InAxis );
+	UNREALED_API virtual void SetCurrentWidgetAxis( EAxisList::Type InAxis );
 
 	/**
 	 * Adjusts the current transform widget size by the provided delta value
 	 */
-	void AdjustTransformWidgetSize(const int32 SizeDelta);
+	UNREALED_API void AdjustTransformWidgetSize(const int32 SizeDelta);
 
 	/**
 	 * Called to do any additional set up of the view for rendering
@@ -716,7 +716,7 @@ public:
 	 * @param ViewFamily	The view family being rendered
 	 * @param View			The view being rendered
 	 */
-	virtual void SetupViewForRendering( FSceneViewFamily& ViewFamily, FSceneView& View );
+	UNREALED_API virtual void SetupViewForRendering( FSceneViewFamily& ViewFamily, FSceneView& View );
 
 	/**
 	 * Called to draw onto the viewports 2D canvas
@@ -725,15 +725,15 @@ public:
 	 * @param View			The view of the scene to be rendered
 	 * @param Canvas		The canvas to draw on
 	 */
-	virtual void DrawCanvas(FViewport& InViewport, FSceneView& View, FCanvas& Canvas);
+	UNREALED_API virtual void DrawCanvas(FViewport& InViewport, FSceneView& View, FCanvas& Canvas);
 
 	// Draws a visualization of the preview light if it was recently moved
-	virtual void DrawPreviewLightVisualization(const FSceneView* View, FPrimitiveDrawInterface* PDI);
+	UNREALED_API virtual void DrawPreviewLightVisualization(const FSceneView* View, FPrimitiveDrawInterface* PDI);
 
 	/**
 	 * Render the drag tool in the viewport
 	 */
-	void RenderDragTool(const FSceneView* View, FCanvas* Canvas);
+	UNREALED_API void RenderDragTool(const FSceneView* View, FCanvas* Canvas);
 
 	/**
 	 * Configures the specified FSceneView object with the view and projection matrices for this viewport.
@@ -741,17 +741,17 @@ public:
 	 * @param	StereoPass	Which eye we're drawing this view for when in stereo mode
 	 * @return	A pointer to the view within the view family which represents the viewport's primary view.
 	 */
-	virtual FSceneView* CalcSceneView(FSceneViewFamily* ViewFamily, const int32 StereoViewIndex = INDEX_NONE);
+	UNREALED_API virtual FSceneView* CalcSceneView(FSceneViewFamily* ViewFamily, const int32 StereoViewIndex = INDEX_NONE);
 
 	/** 
 	 * @return The scene being rendered in this viewport
 	 */
-	virtual FSceneInterface* GetScene() const;
+	UNREALED_API virtual FSceneInterface* GetScene() const;
 
 	/**
 	 * @return The background color of the viewport 
 	 */
-	virtual FLinearColor GetBackgroundColor() const;
+	UNREALED_API virtual FLinearColor GetBackgroundColor() const;
 
 	/**
 	 * Called to override any post process settings for the view
@@ -763,17 +763,17 @@ public:
 	/**
 	 * Ticks this viewport client
 	 */
-	virtual void Tick(float DeltaSeconds);
+	UNREALED_API virtual void Tick(float DeltaSeconds);
 
 	/**
 	 * Called each frame to update the viewport based on delta mouse movements
 	 */
-	virtual void UpdateMouseDelta();
+	UNREALED_API virtual void UpdateMouseDelta();
 	
 	/**
 	 * Called each frame to update the viewport based on delta trackpad gestures
 	 */
-	virtual void UpdateGestureDelta();
+	UNREALED_API virtual void UpdateGestureDelta();
 
 
 	/** 
@@ -782,7 +782,7 @@ public:
 	*
 	* @return 		A valid pointer to the viewports world scene.
 	*/
-	virtual UWorld* GetWorld() const override;
+	UNREALED_API virtual UWorld* GetWorld() const override;
 
 	/** If true, this is a level editor viewport */
 	virtual bool IsLevelEditorClient() const { return false; }
@@ -793,29 +793,29 @@ public:
 	 * @param DragToolType	The type of drag tool to make
 	 * @return The new drag tool
 	 */
-	virtual TSharedPtr<FDragTool> MakeDragTool( EDragTool::Type DragToolType );
+	UNREALED_API virtual TSharedPtr<FDragTool> MakeDragTool( EDragTool::Type DragToolType );
 
 	/** @return true if a drag tool can be used */
-	virtual bool CanUseDragTool() const;
+	UNREALED_API virtual bool CanUseDragTool() const;
 
 	/** @return Whether or not to orbit the camera */
-	virtual bool ShouldOrbitCamera() const;
+	UNREALED_API virtual bool ShouldOrbitCamera() const;
 
-	bool IsMovingCamera() const;
+	UNREALED_API bool IsMovingCamera() const;
 
 	virtual void UpdateLinkedOrthoViewports( bool bInvalidate = false ) {}
 
 	/**
 	 * @return true to lock the pitch of the viewport camera
 	 */
-	virtual bool ShouldLockPitch() const;
+	UNREALED_API virtual bool ShouldLockPitch() const;
 
 	/**
 	 * Called when the mouse cursor is hovered over a hit proxy
 	 * 
 	 * @param HoveredHitProxy	The hit proxy currently hovered over
 	 */
-	virtual void CheckHoveredHitProxy( HHitProxy* HoveredHitProxy );
+	UNREALED_API virtual void CheckHoveredHitProxy( HHitProxy* HoveredHitProxy );
 
 	/** Returns true if a placement dragging actor exists */
 	virtual bool HasDropPreviewActors() const { return false; }
@@ -877,7 +877,7 @@ public:
 	virtual bool UseAppTime() const { return IsRealtime() && !IsSimulateInEditorViewport(); }
 public:
 	/** True if the window is maximized or floating */
-	bool IsVisible() const;
+	UNREALED_API bool IsVisible() const;
 
 	bool IsSimulateInEditorViewport() const 
 	{ 
@@ -899,14 +899,14 @@ public:
 	 *
 	 * @param	bWantStats	true if stats should be displayed
 	 */
-	virtual void SetShowStats( bool bWantStats ) override;
+	UNREALED_API virtual void SetShowStats( bool bWantStats ) override;
 	
 	/**
 	 * Sets how the viewport is displayed (lit, wireframe, etc) for the current viewport type
 	 * 
 	 * @param	InViewModeIndex				View mode to set for the current viewport type
 	 */
-	virtual void SetViewMode(EViewModeIndex InViewModeIndex);
+	UNREALED_API virtual void SetViewMode(EViewModeIndex InViewModeIndex);
 
 	/**
 	 * Sets how the viewport is displayed (lit, wireframe, etc)
@@ -914,15 +914,15 @@ public:
 	 * @param	InPerspViewModeIndex		View mode to set when this viewport is of type LVT_Perspective
 	 * @param	InOrthoViewModeIndex		View mode to set when this viewport is not of type LVT_Perspective
 	 */
-	void SetViewModes(const EViewModeIndex InPerspViewModeIndex, const EViewModeIndex InOrthoViewModeIndex);
+	UNREALED_API void SetViewModes(const EViewModeIndex InPerspViewModeIndex, const EViewModeIndex InOrthoViewModeIndex);
 
 	/** Set the viewmode param. */
-	void SetViewModeParam(int32 InViewModeParam);
+	UNREALED_API void SetViewModeParam(int32 InViewModeParam);
 
 	/**
 	 * @return The current view mode in this viewport, for the current viewport type
 	 */
-	EViewModeIndex GetViewMode() const;
+	UNREALED_API EViewModeIndex GetViewMode() const;
 
 	/**
 	 * @return The view mode to use when this viewport is of type LVT_Perspective
@@ -947,7 +947,7 @@ public:
 	}
 
 	/** @return True if InViewModeIndex is the current view mode param */
-	bool IsViewModeParam(int32 InViewModeParam) const;
+	UNREALED_API bool IsViewModeParam(int32 InViewModeParam) const;
 
 	/**
 	 * Invalidates this viewport and optionally child views.
@@ -955,12 +955,12 @@ public:
 	 * @param	bInvalidateChildViews		[opt] If true (the default), invalidate views that see this viewport as their parent.
 	 * @param	bInvalidateHitProxies		[opt] If true (the default), invalidates cached hit proxies too.
 	 */
-	void Invalidate(bool bInvalidateChildViews=true, bool bInvalidateHitProxies=true);
+	UNREALED_API void Invalidate(bool bInvalidateChildViews=true, bool bInvalidateHitProxies=true);
 	
 	/**
 	 * Gets the dimensions of the viewport
 	 */
-	void GetViewportDimensions( FIntPoint& OutOrigin, FIntPoint& OutSize );
+	UNREALED_API void GetViewportDimensions( FIntPoint& OutOrigin, FIntPoint& OutSize );
 	
 	/**
 	 * Determines which axis InKey and InDelta most refer to and returns a corresponding FVector.  This
@@ -968,40 +968,40 @@ public:
 	 *
 	 * @param	InNudge		If 1, this delta is coming from a keyboard nudge and not the mouse
 	 */
-	FVector TranslateDelta( FKey InKey, float InDelta, bool InNudge );
+	UNREALED_API FVector TranslateDelta( FKey InKey, float InDelta, bool InNudge );
 
 	/**
 	 * Returns the effective viewport type (taking into account any actor locking or camera possession)
 	 */
-	virtual ELevelViewportType GetViewportType() const;
+	UNREALED_API virtual ELevelViewportType GetViewportType() const;
 
 	/**
 	 * Set the viewport type of the client
 	 *
 	 * @param InViewportType	The viewport type to set the client to
 	 */
-	virtual void SetViewportType( ELevelViewportType InViewportType );
+	UNREALED_API virtual void SetViewportType( ELevelViewportType InViewportType );
 	
 	/**
 	 * Rotate through viewport view options
 	 */
-	virtual void RotateViewportType();
+	UNREALED_API virtual void RotateViewportType();
 
 	/**
 	* @return If the viewport option in the array is the active viewport type
 	*/
-	bool IsActiveViewportTypeInRotation() const;
+	UNREALED_API bool IsActiveViewportTypeInRotation() const;
 
 	/**
 	 * @return If InViewportType is the active viewport type 
 	 */
-	bool IsActiveViewportType(ELevelViewportType InViewportType) const;
+	UNREALED_API bool IsActiveViewportType(ELevelViewportType InViewportType) const;
 
 	/** Returns true if this viewport is perspective. */
-	bool IsPerspective() const;
+	UNREALED_API bool IsPerspective() const;
 
 	/** Is the aspect ratio currently constrained? */
-	virtual bool IsAspectRatioConstrained() const;
+	UNREALED_API virtual bool IsAspectRatioConstrained() const;
 
 	/** 
 	 * Focuses the viewport to the center of the bounding box ensuring that the entire box is in view 
@@ -1009,7 +1009,7 @@ public:
 	 * @param BoundingBox			The box to focus
 	 * @param bInstant			Whether or not to focus the viewport instantly or over time
 	 */
-	void FocusViewportOnBox( const FBox& BoundingBox, bool bInstant = false );
+	UNREALED_API void FocusViewportOnBox( const FBox& BoundingBox, bool bInstant = false );
 
 	/**
 	 * Translates the viewport so that the given LookAt point is at the center of viewport, while maintaining current Location/LookAt distance
@@ -1017,21 +1017,21 @@ public:
 	 * @param NewLookAt			The new NewLookAt point to focus on
 	 * @param bInstant			Whether or not to focus the viewport instantly or over time
 	 */
-	void CenterViewportAtPoint(const FVector& NewLookAt, bool bInstant = false);
+	UNREALED_API void CenterViewportAtPoint(const FVector& NewLookAt, bool bInstant = false);
 
 	FEditorCameraController* GetCameraController(void) { return CameraController; }
 
-	void InputAxisForOrbit(FViewport* Viewport, const FVector& DragDelta, FVector& Drag, FRotator& Rot);
+	UNREALED_API void InputAxisForOrbit(FViewport* Viewport, const FVector& DragDelta, FVector& Drag, FRotator& Rot);
 
 	/**
 	 * Implements screenshot capture for editor viewports.  Should be called by derived class' InputKey.
 	 */
-	bool InputTakeScreenshot(FViewport* Viewport, FKey Key, EInputEvent Event);
+	UNREALED_API bool InputTakeScreenshot(FViewport* Viewport, FKey Key, EInputEvent Event);
 
 	/**
 	 * Opens the screenshot in the uses default bitmap viewer (determined by OS)
 	 */
-	void OpenScreenshot( FString SourceFilePath );
+	UNREALED_API void OpenScreenshot( FString SourceFilePath );
 
 	/**
 	 * Takes the screenshot capture, this is called by keybinded events as well InputTakeScreenshot
@@ -1039,48 +1039,48 @@ public:
 	 * @param Viewport				The viewport to take a screenshot of.
 	 * @param bInvalidateViewport	Some methods already invalidate the viewport before calling this method.
 	 */
-	void TakeScreenshot(FViewport* Viewport, bool bInvalidateViewport);
+	UNREALED_API void TakeScreenshot(FViewport* Viewport, bool bInvalidateViewport);
 
 	/**
 	 * Converts a generic movement delta into drag/rotation deltas based on the viewport and keys held down.
 	 */
-	void ConvertMovementToDragRot( const FVector& InDelta, FVector& InDragDelta, FRotator& InRotDelta ) const;
-	void ConvertMovementToOrbitDragRot(const FVector& InDelta, FVector& InDragDelta, FRotator& InRotDelta) const;
+	UNREALED_API void ConvertMovementToDragRot( const FVector& InDelta, FVector& InDragDelta, FRotator& InRotDelta ) const;
+	UNREALED_API void ConvertMovementToOrbitDragRot(const FVector& InDelta, FVector& InDragDelta, FRotator& InRotDelta) const;
 
 	
 	/** Toggle between orbit camera and fly camera */
-	void ToggleOrbitCamera( bool bEnableOrbitCamera );
+	UNREALED_API void ToggleOrbitCamera( bool bEnableOrbitCamera );
 
 	/**
 	 * Sets the camera view location such that the LookAtPoint point is at the specified location.
 	 */
-	void SetViewLocationForOrbiting(const FVector& LookAtPoint, float DistanceToCamera = 256.f );
+	UNREALED_API void SetViewLocationForOrbiting(const FVector& LookAtPoint, float DistanceToCamera = 256.f );
 
 	/**
 	 * Moves the viewport Scamera according to the specified drag and rotation.
 	 */
-	void MoveViewportCamera( const FVector& InDrag, const FRotator& InRot, bool bDollyCamera = false );
+	UNREALED_API void MoveViewportCamera( const FVector& InDrag, const FRotator& InRot, bool bDollyCamera = false );
 
 	/** 
 	 * Get the custom pivot point around which the camera should orbit for this viewport
 	 * @param	OutPivot	The custom pivot point specified by the viewport
 	 * @return	true if a custom pivot point was specified, false otherwise.
 	 */
-	virtual bool GetPivotForOrbit(FVector& OutPivot) const;
+	UNREALED_API virtual bool GetPivotForOrbit(FVector& OutPivot) const;
 
 	// Utility functions to return the modifier key states
-	bool IsAltPressed() const;
-	bool IsCtrlPressed() const;
-	bool IsShiftPressed() const;
-	bool IsCmdPressed() const;
+	UNREALED_API bool IsAltPressed() const;
+	UNREALED_API bool IsCtrlPressed() const;
+	UNREALED_API bool IsShiftPressed() const;
+	UNREALED_API bool IsCmdPressed() const;
 	
 	// Functions for registering and evaluating prioritized input chords. 
 	// When a prioritized chord evaluates true its required keys are 
 	// ignored by all lower priority chords, providing a way to resolve 
 	// conflict between input key combinations.
-	void RegisterPrioritizedInputChord(const FPrioritizedInputChord& InInputCord);
-	void UnregisterPrioritizedInputChord(const FName InInputCordName);
-	bool IsPrioritizedInputChordPressed(const FName InInputCordName) const;
+	UNREALED_API void RegisterPrioritizedInputChord(const FPrioritizedInputChord& InInputCord);
+	UNREALED_API void UnregisterPrioritizedInputChord(const FName InInputCordName);
+	UNREALED_API bool IsPrioritizedInputChordPressed(const FName InInputCordName) const;
 
 	/**
 	 * Utility function to return whether the command accepts the key states
@@ -1088,19 +1088,19 @@ public:
 	 * @param InOptionalKey (Optional) input key being tested against. If not specified, the current viewport's key state's key will be used
 	 * @return True if one of the command's chords accepts the input :
 	 */
-	bool IsCommandChordPressed(const TSharedPtr<FUICommandInfo> InCommand, FKey InOptionalKey = FKey()) const;
+	UNREALED_API bool IsCommandChordPressed(const TSharedPtr<FUICommandInfo> InCommand, FKey InOptionalKey = FKey()) const;
 
 	/** @return True if the window is in an immersive viewport */
-	bool IsInImmersiveViewport() const;
+	UNREALED_API bool IsInImmersiveViewport() const;
 
 	void ClearAudioFocus() { bHasAudioFocus = false; }
 	
 	void SetAudioFocus() { bHasAudioFocus = true; }
 
-	void MarkMouseMovedSinceClick();
+	UNREALED_API void MarkMouseMovedSinceClick();
 
 	/** Determines whether this viewport is currently allowed to use Absolute Movement */
-	bool IsUsingAbsoluteTranslation(bool bAlsoCheckAbsoluteRotation = false) const;
+	UNREALED_API bool IsUsingAbsoluteTranslation(bool bAlsoCheckAbsoluteRotation = false) const;
 
 	bool IsForcedRealtimeAudio() const { return bForceAudioRealtime; }
 
@@ -1113,31 +1113,31 @@ public:
 	/** @return true if a mouse button is down and it's movement being tracked for operations inside the viewport */
 	bool IsTracking() const { return bIsTracking; }
 
-	EAxisList::Type GetCurrentWidgetAxis() const;
+	UNREALED_API EAxisList::Type GetCurrentWidgetAxis() const;
 
 	/** Overrides current cursor. */
-	void SetRequiredCursorOverride( bool WantOverride, EMouseCursor::Type RequiredCursor = EMouseCursor::Default ); 
+	UNREALED_API void SetRequiredCursorOverride( bool WantOverride, EMouseCursor::Type RequiredCursor = EMouseCursor::Default ); 
 
 	/** Overrides current widget mode */
-	void SetWidgetModeOverride(UE::Widget::EWidgetMode InWidgetMode);
+	UNREALED_API void SetWidgetModeOverride(UE::Widget::EWidgetMode InWidgetMode);
 
 	/** Get the camera speed for this viewport */
-	float GetCameraSpeed() const;
+	UNREALED_API float GetCameraSpeed() const;
 
 	/** Get the camera speed for this viewport based on the specified speed setting */
-	float GetCameraSpeed(int32 SpeedSetting) const;
+	UNREALED_API float GetCameraSpeed(int32 SpeedSetting) const;
 
 	/** Set the speed setting for the camera in this viewport */
-	virtual void SetCameraSpeedSetting(int32 SpeedSetting);
+	UNREALED_API virtual void SetCameraSpeedSetting(int32 SpeedSetting);
 
 	/** Get the camera speed setting for this viewport */
-	virtual int32 GetCameraSpeedSetting() const;
+	UNREALED_API virtual int32 GetCameraSpeedSetting() const;
 
 	/** Get the camera speed scalar for this viewport */
-	virtual float GetCameraSpeedScalar() const;
+	UNREALED_API virtual float GetCameraSpeedScalar() const;
 
 	/** Set the camera speed scalar for this viewport */
-	virtual void SetCameraSpeedScalar(float SpeedScalar);
+	UNREALED_API virtual void SetCameraSpeedScalar(float SpeedScalar);
 
 	/** Editor mode tool manager being used for this viewport client */
 	FEditorModeTools* GetModeTools() const
@@ -1147,7 +1147,7 @@ public:
 
 	/** Legacy adapter for a toolkit to take ownership of a mode manager that may have been created by this viewport client */
 	UE_DEPRECATED(5.0, "This function is meant for legacy edge cases (pre-UAssetEditor), where the toolkit may not be the original owner of the mode manager.")
-	void TakeOwnershipOfModeManager(TSharedPtr<FEditorModeTools>& ModeManagerPtr);
+	UNREALED_API void TakeOwnershipOfModeManager(TSharedPtr<FEditorModeTools>& ModeManagerPtr);
 
 	/** Get the editor viewport widget */
 	TSharedPtr<SEditorViewport> GetEditorViewportWidget() const { return EditorViewportWidget.Pin(); }
@@ -1155,7 +1155,7 @@ public:
 	/**
 	 * Computes a matrix to use for viewport location and rotation
 	 */
-	virtual FMatrix CalcViewRotationMatrix(const FRotator& InViewRotation) const;
+	UNREALED_API virtual FMatrix CalcViewRotationMatrix(const FRotator& InViewRotation) const;
 
 protected:
 
@@ -1170,7 +1170,7 @@ protected:
 
 public:
 
-	void DrawBoundingBox(FBox &Box, FCanvas* InCanvas, const FSceneView* InView, const FViewport* InViewport, const FLinearColor& InColor, const bool bInDrawBracket, const FString &InLabelText) ;
+	UNREALED_API void DrawBoundingBox(FBox &Box, FCanvas* InCanvas, const FSceneView* InView, const FViewport* InViewport, const FLinearColor& InColor, const bool bInDrawBracket, const FString &InLabelText) ;
 	
 	/**
 	 * Draws a screen space bounding box around the specified actor
@@ -1183,11 +1183,11 @@ public:
 	 * @param	bInDrawBracket	True to draw a bracket, otherwise a box will be rendered
 	 * @param	bInLabelText	Optional label text to draw
 	 */
-	void DrawActorScreenSpaceBoundingBox( FCanvas* InCanvas, const FSceneView* InView, FViewport* InViewport, AActor* InActor, const FLinearColor& InColor, const bool bInDrawBracket, const FString& InLabelText = TEXT( "" ) );
+	UNREALED_API void DrawActorScreenSpaceBoundingBox( FCanvas* InCanvas, const FSceneView* InView, FViewport* InViewport, AActor* InActor, const FLinearColor& InColor, const bool bInDrawBracket, const FString& InLabelText = TEXT( "" ) );
 
-	void SetGameView(bool bGameViewEnable);
+	UNREALED_API void SetGameView(bool bGameViewEnable);
 
-	virtual void SetVREditView(bool bGameViewEnable);
+	UNREALED_API virtual void SetVREditView(bool bGameViewEnable);
 	/**
 	 * Returns true if this viewport is excluding non-game elements from its display
 	 */
@@ -1219,50 +1219,50 @@ public:
 	}
 
 	/** Get the near clipping plane for this viewport. */
-	float GetNearClipPlane() const;
+	UNREALED_API float GetNearClipPlane() const;
 
 	/** Override the near clipping plane. Set to a negative value to disable the override. */
-	void OverrideNearClipPlane(float InNearPlane);
+	UNREALED_API void OverrideNearClipPlane(float InNearPlane);
 
 	/** Get the far clipping plane override for this viewport. */
-	float GetFarClipPlaneOverride() const;
+	UNREALED_API float GetFarClipPlaneOverride() const;
 	
 	/** Override the far clipping plane. Set to a negative value to disable the override. */
-	void OverrideFarClipPlane(const float InFarPlane);
+	UNREALED_API void OverrideFarClipPlane(const float InFarPlane);
 
 	/** When collision draw mode changes, this function allows hidden objects to be drawn, so hidden colliding objects can be seen */
-	void UpdateHiddenCollisionDrawing();
+	UNREALED_API void UpdateHiddenCollisionDrawing();
 	/** Returns the scene depth at the given viewport X,Y */
-	float GetSceneDepthAtLocation(int32 X, int32 Y);
+	UNREALED_API float GetSceneDepthAtLocation(int32 X, int32 Y);
 
 	/** Returns the location of the object at the given viewport X,Y */
-	FVector GetHitProxyObjectLocation(int32 X, int32 Y);
+	UNREALED_API FVector GetHitProxyObjectLocation(int32 X, int32 Y);
 
 	/** Returns the map allowing to convert from the viewmode param to a name. */
 	TMap<int32, FName>& GetViewModeParamNameMap() { return ViewModeParamNameMap; }
 
 	/** Show or hide the widget. */
-	void ShowWidget(const bool bShow);
+	UNREALED_API void ShowWidget(const bool bShow);
 
 	/**
 	 * Returns whether or not the flight camera is active
 	 *
 	 * @return true if the flight camera is active
 	 */
-	bool IsFlightCameraActive() const;
+	UNREALED_API bool IsFlightCameraActive() const;
 
 	/** Delegate handler fired when a show flag is toggled */
-	virtual void HandleToggleShowFlag(FEngineShowFlags::EShowFlag EngineShowFlagIndex);
+	UNREALED_API virtual void HandleToggleShowFlag(FEngineShowFlags::EShowFlag EngineShowFlagIndex);
 
 	/** Delegate handler fired to determine the state of a show flag */
-	virtual bool HandleIsShowFlagEnabled(FEngineShowFlags::EShowFlag EngineShowFlagIndex) const;
+	UNREALED_API virtual bool HandleIsShowFlagEnabled(FEngineShowFlags::EShowFlag EngineShowFlagIndex) const;
 
 	/**
 	 * Changes the buffer visualization mode for this viewport.
 	 *
 	 * @param InName	The ID of the required visualization mode
 	 */
-	void ChangeBufferVisualizationMode( FName InName );
+	UNREALED_API void ChangeBufferVisualizationMode( FName InName );
 
 	/**
 	 * Checks if a buffer visualization mode is selected.
@@ -1270,19 +1270,19 @@ public:
 	 * @param InName	The ID of the required visualization mode
 	 * @return	true if the supplied buffer visualization mode is checked
 	 */
-	bool IsBufferVisualizationModeSelected( FName InName ) const;
+	UNREALED_API bool IsBufferVisualizationModeSelected( FName InName ) const;
 
 	/**
 	 * Returns the FText display name associated with CurrentBufferVisualizationMode.
 	 */
-	FText GetCurrentBufferVisualizationModeDisplayName() const;
+	UNREALED_API FText GetCurrentBufferVisualizationModeDisplayName() const;
 
 	/**
 	 * Changes the Nanite visualization mode for this viewport.
 	 *
 	 * @param InName	The ID of the required visualization mode
 	 */
-	void ChangeNaniteVisualizationMode(FName InName);
+	UNREALED_API void ChangeNaniteVisualizationMode(FName InName);
 
 	/**
 	 * Checks if a Nanite visualization mode is selected.
@@ -1290,19 +1290,19 @@ public:
 	 * @param InName	The ID of the required visualization mode
 	 * @return	true if the supplied Nanite visualization mode is checked
 	 */
-	bool IsNaniteVisualizationModeSelected(FName InName) const;
+	UNREALED_API bool IsNaniteVisualizationModeSelected(FName InName) const;
 
 	/**
 	 * Returns the FText display name associated with CurrentNaniteVisualizationMode.
 	 */
-	FText GetCurrentNaniteVisualizationModeDisplayName() const;
+	UNREALED_API FText GetCurrentNaniteVisualizationModeDisplayName() const;
 
 	/**
 	 * Changes the Lumen visualization mode for this viewport.
 	 *
 	 * @param InName	The ID of the required visualization mode
 	 */
-	void ChangeLumenVisualizationMode(FName InName);
+	UNREALED_API void ChangeLumenVisualizationMode(FName InName);
 
 	/**
 	 * Checks if a Lumen visualization mode is selected.
@@ -1310,19 +1310,19 @@ public:
 	 * @param InName	The ID of the required visualization mode
 	 * @return	true if the supplied Lumen visualization mode is checked
 	 */
-	bool IsLumenVisualizationModeSelected(FName InName) const;
+	UNREALED_API bool IsLumenVisualizationModeSelected(FName InName) const;
 
 	/**
 	 * Returns the FText display name associated with CurrentLumenVisualizationMode.
 	 */
-	FText GetCurrentLumenVisualizationModeDisplayName() const;
+	UNREALED_API FText GetCurrentLumenVisualizationModeDisplayName() const;
 
 	/**
 	 * Changes the Strata visualization mode for this viewport.
 	 *
 	 * @param InName	The ID of the required visualization mode
 	 */
-	void ChangeStrataVisualizationMode(FName InName);
+	UNREALED_API void ChangeStrataVisualizationMode(FName InName);
 
 	/**
 	 * Checks if a Strata visualization mode is selected.
@@ -1330,19 +1330,19 @@ public:
 	 * @param InName	The ID of the required visualization mode
 	 * @return	true if the supplied Strata visualization mode is checked
 	 */
-	bool IsStrataVisualizationModeSelected(FName InName) const;
+	UNREALED_API bool IsStrataVisualizationModeSelected(FName InName) const;
 
 	/**
 	 * Returns the FText display name associated with CurrentStrataVisualizationMode.
 	 */
-	FText GetCurrentStrataVisualizationModeDisplayName() const;
+	UNREALED_API FText GetCurrentStrataVisualizationModeDisplayName() const;
 
 	/**
 	 * Changes the Groom visualization mode for this viewport.
 	 *
 	 * @param InName	The ID of the required visualization mode
 	 */
-	void ChangeGroomVisualizationMode(FName InName);
+	UNREALED_API void ChangeGroomVisualizationMode(FName InName);
 
 	/**
 	 * Checks if a Groom visualization mode is selected.
@@ -1350,19 +1350,19 @@ public:
 	 * @param InName	The ID of the required visualization mode
 	 * @return	true if the supplied Groom visualization mode is checked
 	 */
-	bool IsGroomVisualizationModeSelected(FName InName) const;
+	UNREALED_API bool IsGroomVisualizationModeSelected(FName InName) const;
 
 	/**
 	 * Returns the FText display name associated with CurrentGroomVisualizationMode.
 	 */
-	FText GetCurrentGroomVisualizationModeDisplayName() const;
+	UNREALED_API FText GetCurrentGroomVisualizationModeDisplayName() const;
 
 	/**
 	* Changes the virtual shadow map visualization mode for this viewport.
 	*
 	* @param InName	The ID of the required visualization mode
 	*/
-	void ChangeVirtualShadowMapVisualizationMode(FName InName);
+	UNREALED_API void ChangeVirtualShadowMapVisualizationMode(FName InName);
 
 	/**
 	* Checks if a virtual shadow map visualization mode is selected.
@@ -1370,24 +1370,24 @@ public:
 	* @param InName	The ID of the required visualization mode
 	* @return	true if the supplied virtual shadow map visualization mode is checked
 	*/
-	bool IsVirtualShadowMapVisualizationModeSelected(FName InName) const;
+	UNREALED_API bool IsVirtualShadowMapVisualizationModeSelected(FName InName) const;
 
 	/**
 	* Returns the FText display name associated with CurrentVirtualShadowMapVisualizationMode.
 	*/
-	FText GetCurrentVirtualShadowMapVisualizationModeDisplayName() const;
+	UNREALED_API FText GetCurrentVirtualShadowMapVisualizationModeDisplayName() const;
 
 	/**
 	* Returns whether visualize debug material is enabled.
 	*/
-	bool IsVisualizeCalibrationMaterialEnabled() const;
+	UNREALED_API bool IsVisualizeCalibrationMaterialEnabled() const;
 
 	/**
 	 * Changes the ray tracing debug visualization mode for this viewport
 	 *
 	 * @param InName	The ID of the required ray tracing debug visualization mode
 	 */
-	void ChangeRayTracingDebugVisualizationMode(FName InName);
+	UNREALED_API void ChangeRayTracingDebugVisualizationMode(FName InName);
 
 	/**
 	 * Checks if a ray tracing debug visualization mode is selected
@@ -1395,14 +1395,14 @@ public:
 	 * @param InName	The ID of the required ray tracing debug visualization mode
 	 * @return	true if the supplied ray tracing debug visualization mode is checked
 	 */
-	bool IsRayTracingDebugVisualizationModeSelected(FName InName) const;
+	UNREALED_API bool IsRayTracingDebugVisualizationModeSelected(FName InName) const;
 
 	/**
 	 * Changes the GPU Skin Cache visualization mode for this viewport
 	 *
 	 * @param InName	The ID of the required GPU Skin Cache visualization mode
 	 */
-	void ChangeGPUSkinCacheVisualizationMode(FName InName);
+	UNREALED_API void ChangeGPUSkinCacheVisualizationMode(FName InName);
 
 	/**
 	 * Checks if a GPU Skin Cache visualization mode is selected
@@ -1410,39 +1410,39 @@ public:
 	 * @param InName	The ID of the required GPU Skin Cache visualization mode
 	 * @return	true if the supplied GPU Skin Cache visualization mode is checked
 	 */
-	bool IsGPUSkinCacheVisualizationModeSelected(FName InName) const;
+	UNREALED_API bool IsGPUSkinCacheVisualizationModeSelected(FName InName) const;
 
 	/**
 	* Returns the FText display name associated with CurrentGPUSkinCacheVisualizationMode.
 	*/
-	FText GetCurrentGPUSkinCacheVisualizationModeDisplayName() const;
+	UNREALED_API FText GetCurrentGPUSkinCacheVisualizationModeDisplayName() const;
 
 	/** @return True if PreviewResolutionFraction is supported. */
-	bool SupportsPreviewResolutionFraction() const;
+	UNREALED_API bool SupportsPreviewResolutionFraction() const;
 
 	/** @return default resolution fraction for UI based on display resolution and user settings. */
-	float GetDefaultPrimaryResolutionFractionTarget() const;
+	UNREALED_API float GetDefaultPrimaryResolutionFractionTarget() const;
 
 	/** @return whether preview screen percentage for UI. */
-	bool IsPreviewingScreenPercentage() const;
+	UNREALED_API bool IsPreviewingScreenPercentage() const;
 
 	/** Sets whether preview screen percentage for UI. */
-	void SetPreviewingScreenPercentage(bool bIsPreviewing);
+	UNREALED_API void SetPreviewingScreenPercentage(bool bIsPreviewing);
 
 	/** @return preview screen percentage for UI. */
-	int32 GetPreviewScreenPercentage() const;
+	UNREALED_API int32 GetPreviewScreenPercentage() const;
 
 	/** Set preview screen percentage on UI behalf. */
-	void SetPreviewScreenPercentage(int32 PreviewScreenPercentage);
+	UNREALED_API void SetPreviewScreenPercentage(int32 PreviewScreenPercentage);
 
 	/** @return True if DPI preview is supported. */
-	bool SupportsLowDPIPreview() const;
+	UNREALED_API bool SupportsLowDPIPreview() const;
 
 	/** @return whether previewing for low DPI. */
-	bool IsLowDPIPreview() const;
+	UNREALED_API bool IsLowDPIPreview() const;
 
 	/** Set whether previewing for low DPI. */
-	void SetLowDPIPreview(bool LowDPIPreview);
+	UNREALED_API void SetLowDPIPreview(bool LowDPIPreview);
 	
 	/** Mouse info is usually transformed to gizmo space before FEdMode handles it, this allows raw delta X and Y access */
 	FMouseDeltaTracker* GetMouseDeltaTracker() const { return  MouseDeltaTracker; }
@@ -1451,7 +1451,7 @@ public:
 	virtual uint32 GetCachedMouseY() const { return CachedMouseY; }
 
 	/** @return True if the camera speed should be scaled by its view distance. */
-	virtual bool ShouldScaleCameraSpeedByDistance() const;
+	UNREALED_API virtual bool ShouldScaleCameraSpeedByDistance() const;
 
 	/**
 	 * Enable customization of the EngineShowFlags for rendering. After calling this function,
@@ -1459,10 +1459,10 @@ public:
 	 * Changes made to the ShowFlags will be used for that frame but .EngineShowFlags will not be modified.
 	 * @param OverrideFunc custom override function that will be called every frame until override is disabled.
 	 */
-	void EnableOverrideEngineShowFlags(TUniqueFunction<void(FEngineShowFlags&)> OverrideFunc);
+	UNREALED_API void EnableOverrideEngineShowFlags(TUniqueFunction<void(FEngineShowFlags&)> OverrideFunc);
 
 	/** Disable EngineShowFlags override if enabled */
-	void DisableOverrideEngineShowFlags();
+	UNREALED_API void DisableOverrideEngineShowFlags();
 
 	/** @return true if Override EngineShowFlags are currently enabled */
 	bool IsEngineShowFlagsOverrideEnabled() const { return !! OverrideShowFlagsFunc; }
@@ -1472,23 +1472,23 @@ public:
 
 protected:
 	/** Invalidates the viewport widget (if valid) to register its active timer */
-	void InvalidateViewportWidget();
+	UNREALED_API void InvalidateViewportWidget();
 
 	/** Constant for how much the camera safe zone rectangle is inset when being displayed in the editor */
-	static float const SafePadding;
+	static UNREALED_API float const SafePadding;
 	/**
 	 * Called when the perspective viewport camera moves
 	 */
 	virtual void PerspectiveCameraMoved() {}
 
 	/** Updates the rotate widget with the passed in delta rotation. */
-	void ApplyDeltaToRotateWidget(const FRotator& InRot);
+	UNREALED_API void ApplyDeltaToRotateWidget(const FRotator& InRot);
 
 	/** Invalidates this and other linked viewports (anything viewing the same scene) */
-	virtual void RedrawAllViewportsIntoThisScene();
+	UNREALED_API virtual void RedrawAllViewportsIntoThisScene();
 
 	/** FCommonViewportClient interface */
-	virtual float UpdateViewportClientWindowDPIScale() const override;
+	UNREALED_API virtual float UpdateViewportClientWindowDPIScale() const override;
 
 	/**
 	 * Used to store the required cursor visibility states and override cursor appearance
@@ -1514,21 +1514,21 @@ protected:
 	/**
 	 * Updates the visibility of the hardware and software cursors according to the viewport's state.
 	 */
-	void UpdateAndApplyCursorVisibility();
+	UNREALED_API void UpdateAndApplyCursorVisibility();
 
 	/** Setup the cursor visibility state we require and store in RequiredCursorVisibiltyAndAppearance struct */
-	void UpdateRequiredCursorVisibility();
+	UNREALED_API void UpdateRequiredCursorVisibility();
 	
 	/** Sets the required hardware and software cursor. */
-	void SetRequiredCursor(const bool bHardwareCursorVisible, const bool bSoftwareCursorVisible);
+	UNREALED_API void SetRequiredCursor(const bool bHardwareCursorVisible, const bool bSoftwareCursorVisible);
 
 	/** 
 	 * Apply the required cursor visibility states from the RequiredCursorVisibiltyAndAppearance struct 
 	 * @param	View				True - Set the position of the software cursor if being made visible. This defaults to FALSE.
 	*/
-	void ApplyRequiredCursorVisibility(  bool bUpdateSoftwareCursorPostion = false );
+	UNREALED_API void ApplyRequiredCursorVisibility(  bool bUpdateSoftwareCursorPostion = false );
 
-	bool ShouldUseMoveCanvasMovement() const;
+	UNREALED_API bool ShouldUseMoveCanvasMovement() const;
 
 
 	/**
@@ -1539,7 +1539,7 @@ protected:
 	 * @param	InRotation		Specifies the rotation to apply to the axes
 	 * @param	InAxis			Specifies which axes to draw
 	 */
-	void DrawAxes(FViewport* Viewport,FCanvas* Canvas, const FRotator* InRotation = NULL, EAxisList::Type InAxis = EAxisList::XYZ);
+	UNREALED_API void DrawAxes(FViewport* Viewport,FCanvas* Canvas, const FRotator* InRotation = NULL, EAxisList::Type InAxis = EAxisList::XYZ);
 
 	/**
 	 * Draws viewport scale units
@@ -1548,7 +1548,7 @@ protected:
 	 * @param	InCanvas		Canvas to draw on
 	 * @param	InView			Scene view used for rendering
 	 */
-	void DrawScaleUnits(FViewport* Viewport, FCanvas* Canvas, const FSceneView& InView);
+	UNREALED_API void DrawScaleUnits(FViewport* Viewport, FCanvas* Canvas, const FSceneView& InView);
 
 	/**
 	 * Starts tracking the mouse due to mouse input
@@ -1556,7 +1556,7 @@ protected:
 	 * @param InputState	The mouse and keyboard input state at the time the mouse input happened
 	 * @param View			The view of the scene in this viewport
 	 */
-	void StartTrackingDueToInput( const struct FInputEventState& InputState, FSceneView& View );
+	UNREALED_API void StartTrackingDueToInput( const struct FInputEventState& InputState, FSceneView& View );
 
 	/**
 	 * Handles clicking in the viewport
@@ -1564,7 +1564,7 @@ protected:
 	 * @param InputState	The mouse and keyboard input state at the time the click happened
 	 * @param View			The view of the scene in this viewport
 	 */
-	void ProcessClickInViewport( const FInputEventState& InputState, FSceneView& View );
+	UNREALED_API void ProcessClickInViewport( const FInputEventState& InputState, FSceneView& View );
 
 	/**
 	 * Handles double clicking in the viewport
@@ -1572,7 +1572,7 @@ protected:
 	 * @param InputState	The mouse and keyboard input state at the time the click happened
 	 * @param View			The view of the scene in this viewport
 	 */
-	void ProcessDoubleClickInViewport( const struct FInputEventState& InputState, FSceneView& View );
+	UNREALED_API void ProcessDoubleClickInViewport( const struct FInputEventState& InputState, FSceneView& View );
 
 	/**
 	 * Called when a user zooms the ortho viewport
@@ -1580,85 +1580,85 @@ protected:
 	 * @param InputState	The current state of mouse and keyboard buttons
 	 * @param Scale			Allows modifying the default amount of zoom
 	 */
-	void OnOrthoZoom( const struct FInputEventState& InputState, float Scale = 1.0f );
+	UNREALED_API void OnOrthoZoom( const struct FInputEventState& InputState, float Scale = 1.0f );
 
 	/**
 	 * Called when a user dollys the perspective camera
 	 *
 	 * @param InputState	The current state of mouse and keyboard buttons
 	 */
-	void OnDollyPerspectiveCamera( const struct FInputEventState& InputState );
+	UNREALED_API void OnDollyPerspectiveCamera( const struct FInputEventState& InputState );
 
 	/**
 	 * Called when a user changes the camera speed
 	 *
 	 * @param InputState	The current state of mouse and keyboard buttons
 	 */
-	void OnChangeCameraSpeed( const struct FInputEventState& InputState );
+	UNREALED_API void OnChangeCameraSpeed( const struct FInputEventState& InputState );
 	
 	/**
 	 * Stops any mouse tracking
 	 */
-	void StopTracking();
+	UNREALED_API void StopTracking();
 
 	/**
 	 * Aborts mouse tracking (stop and cancel)
 	 */
-	virtual void AbortTracking();
+	UNREALED_API virtual void AbortTracking();
 
 	/** Enables or disables camera lock **/
-	void EnableCameraLock(bool bEnable);
+	UNREALED_API void EnableCameraLock(bool bEnable);
 
 	/**
 	 * Gets a joystick state cache for the specified controller ID
 	 */
-	FCachedJoystickState* GetJoystickState(const uint32 InControllerID);
+	UNREALED_API FCachedJoystickState* GetJoystickState(const uint32 InControllerID);
 
 	/** Helper used by DrawSafeFrames to get the current safe frame aspect ratio. */
 	virtual bool GetActiveSafeFrame(float& OutAspectRatio) const { return false; }
 
 	/** Helper function to calculate the safe frame rectangle on the current viewport */
-	bool CalculateEditorConstrainedViewRect(FSlateRect& OutSafeFrameRect, FViewport* InViewport, const float DPIScale);
+	UNREALED_API bool CalculateEditorConstrainedViewRect(FSlateRect& OutSafeFrameRect, FViewport* InViewport, const float DPIScale);
 
 	virtual void NudgeSelectedObjects(const struct FInputEventState& InputState) {}
 
 private:
 	/** @return Whether or not the camera should be panned or dollied */
-	bool ShouldPanOrDollyCamera() const;
+	UNREALED_API bool ShouldPanOrDollyCamera() const;
 
-	void ConditionalCheckHoveredHitProxy();
+	UNREALED_API void ConditionalCheckHoveredHitProxy();
 
 	/** Returns true if perspective flight camera input mode is currently active in this viewport */
-	bool IsFlightCameraInputModeActive() const;
+	UNREALED_API bool IsFlightCameraInputModeActive() const;
 
 
 	/** Moves a perspective camera */
-	void MoveViewportPerspectiveCamera( const FVector& InDrag, const FRotator& InRot, bool bDollyCamera = false );
+	UNREALED_API void MoveViewportPerspectiveCamera( const FVector& InDrag, const FRotator& InRot, bool bDollyCamera = false );
 
 	
 	/**Applies Joystick axis control to camera movement*/
-	void UpdateCameraMovementFromJoystick(const bool bRelativeMovement, FCameraControllerConfig& InConfig);
+	UNREALED_API void UpdateCameraMovementFromJoystick(const bool bRelativeMovement, FCameraControllerConfig& InConfig);
 
 	/**
 	 * Updates real-time camera movement.  Should be called every viewport tick!
 	 *
 	 * @param	DeltaTime	Time interval in seconds since last update
 	 */
-	void UpdateCameraMovement( float DeltaTime );
+	UNREALED_API void UpdateCameraMovement( float DeltaTime );
 
 	/**
 	 * Forcibly disables lighting show flags if there are no lights in the scene, or restores lighting show
 	 * flags if lights are added to the scene.
 	 */
-	void UpdateLightingShowFlags( FEngineShowFlags& InOutShowFlags );
+	UNREALED_API void UpdateLightingShowFlags( FEngineShowFlags& InOutShowFlags );
 
 	/** InOut might get adjusted depending on viewmode or viewport type */
-	void ApplyEditorViewModeAdjustments(FEngineShowFlags& InOut) const;
+	UNREALED_API void ApplyEditorViewModeAdjustments(FEngineShowFlags& InOut) const;
 
 	/** Renders the safe frame lines. */
-	void DrawSafeFrames(FViewport& Viewport, FSceneView& View, FCanvas& Canvas);
+	UNREALED_API void DrawSafeFrames(FViewport& Viewport, FSceneView& View, FCanvas& Canvas);
 
-	void DrawSafeFrameQuad( FCanvas &Canvas, FVector2D V1, FVector2D V2 );
+	UNREALED_API void DrawSafeFrameQuad( FCanvas &Canvas, FVector2D V1, FVector2D V2 );
 	
 	virtual FEngineShowFlags* GetEngineShowFlags() override
 	{ 
@@ -1668,31 +1668,31 @@ private:
 	/**
 	 * Set a specific stat to either enabled or disabled (returns true if there are any stats enabled)
 	 */
-	int32 SetStatEnabled(const TCHAR* InName, const bool bEnable, const bool bAll = false);
+	UNREALED_API int32 SetStatEnabled(const TCHAR* InName, const bool bEnable, const bool bAll = false);
 
 	/** Delegate handler to see if a stat is enabled on this viewport */
-	void HandleViewportStatCheckEnabled(const TCHAR* InName, bool& bOutCurrentEnabled, bool& bOutOthersEnabled);
+	UNREALED_API void HandleViewportStatCheckEnabled(const TCHAR* InName, bool& bOutCurrentEnabled, bool& bOutOthersEnabled);
 
 	/** Delegate handler for when stats are enabled in a viewport */
-	void HandleViewportStatEnabled(const TCHAR* InName);
+	UNREALED_API void HandleViewportStatEnabled(const TCHAR* InName);
 
 	/** Delegate handler for when stats are disabled in a viewport */
-	void HandleViewportStatDisabled(const TCHAR* InName);
+	UNREALED_API void HandleViewportStatDisabled(const TCHAR* InName);
 
 	/** Delegate handler for when all stats are disabled in a viewport */
-	void HandleViewportStatDisableAll(const bool bInAnyViewport);
+	UNREALED_API void HandleViewportStatDisableAll(const bool bInAnyViewport);
 
 	/** Delegate handler for when a window DPI changes and we might need to adjust the scenes resolution */
-	void HandleWindowDPIScaleChanged(TSharedRef<SWindow> InWindow);
+	UNREALED_API void HandleWindowDPIScaleChanged(TSharedRef<SWindow> InWindow);
 
 	/** Handle the camera about to be moved or stopped **/
 	virtual void BeginCameraMovement(bool bHasMovement) {}
 	virtual void EndCameraMovement() {}
 
-	float GetMinimumOrthoZoom() const;
+	UNREALED_API float GetMinimumOrthoZoom() const;
 
 public:
-	static const uint32 MaxCameraSpeeds;
+	static UNREALED_API const uint32 MaxCameraSpeeds;
 
 	/** Delegate used to get whether or not this client is in an immersive viewport */
 	FViewportStateGetter ImmersiveDelegate;
@@ -1946,10 +1946,10 @@ protected:
 
 public:
 	/* Default view mode for perspective viewports */
-	static const EViewModeIndex DefaultPerspectiveViewMode;
+	static UNREALED_API const EViewModeIndex DefaultPerspectiveViewMode;
 
 	/* Default view mode for orthographic viewports */
-	static const EViewModeIndex DefaultOrthoViewMode;
+	static UNREALED_API const EViewModeIndex DefaultOrthoViewMode;
 	
 	/** Flag to lock the viewport fly camera */
 	bool bLockFlightCamera;
@@ -1966,7 +1966,7 @@ protected:
 
 	// Used by the custom temporal upscaler plugin
 public:
-	class UNREALED_API ICustomTemporalUpscalerData
+	class ICustomTemporalUpscalerData
 	{
 	public:
 		virtual ~ICustomTemporalUpscalerData()
@@ -2060,7 +2060,7 @@ private:
 
 
 
-class UNREALED_API FEditorViewportStats
+class FEditorViewportStats
 {
 public:
 	enum Category
@@ -2095,33 +2095,33 @@ public:
 	/**
 	 * Used commits a single usage record for whichever category is sent to it.
 	 */
-	static void Used(Category InCategory);
+	static UNREALED_API void Used(Category InCategory);
 
 	/**
 	 * Begins the frame for capturing using statements.  If nothing is logged between the begin and
 	 * end frame we reset the last using tracking variable in EndFrame.
 	 */
-	static void BeginFrame();
+	static UNREALED_API void BeginFrame();
 
 	/**
 	 * Using commits a single usage record for whichever category is sent to it only if it's different
 	 * from the last category that was sent to Using.  This should be used to capture usage data for modes
 	 * where it's difficult to tell when they ended.
 	 */
-	static void Using(Category InCategory);
+	static UNREALED_API void Using(Category InCategory);
 
 	/**
 	 * Doesn't use anything, but ensures that the last using item is not reset.
 	 */
-	static void NoOpUsing();
+	static UNREALED_API void NoOpUsing();
 
 	/**
 	 * Use EndEndUsing to manually reset the Using state so that the next call to Using will commit a new record.
 	 * Useful when you know a transition has occurred like the mouse button has been released.
 	 */
-	static void EndFrame();
+	static UNREALED_API void EndFrame();
 
-	static void SendUsageData();
+	static UNREALED_API void SendUsageData();
 
 private:
 	static void Initialize();
diff --git a/Engine/Source/Editor/UnrealEd/Public/EditorViewportCommands.h b/Engine/Source/Editor/UnrealEd/Public/EditorViewportCommands.h
--- a/Engine/Source/Editor/UnrealEd/Public/EditorViewportCommands.h
+++ b/Engine/Source/Editor/UnrealEd/Public/EditorViewportCommands.h
@@ -15,10 +15,10 @@ class FUICommandList;
 /**
  * Class containing commands for editor viewport actions common to all viewports
  */
-class UNREALED_API FEditorViewportCommands : public TCommands<FEditorViewportCommands>
+class FEditorViewportCommands : public TCommands<FEditorViewportCommands>
 {
 public:
-	FEditorViewportCommands();
+	UNREALED_API FEditorViewportCommands();
 	
 	/** Viewport pane configurations */
 	TSharedPtr< FUICommandInfo > ViewportConfig_OnePane;
@@ -248,7 +248,7 @@ public:
 	TSharedPtr< FUICommandInfo > SurfaceSnapping;
 public:
 	/** Registers our commands with the binding system */
-	virtual void RegisterCommands() override;
+	UNREALED_API virtual void RegisterCommands() override;
 };
 
 /**
diff --git a/Engine/Source/Editor/UnrealEd/Public/EditorViewportTabContent.h b/Engine/Source/Editor/UnrealEd/Public/EditorViewportTabContent.h
--- a/Engine/Source/Editor/UnrealEd/Public/EditorViewportTabContent.h
+++ b/Engine/Source/Editor/UnrealEd/Public/EditorViewportTabContent.h
@@ -14,7 +14,7 @@ class SDockTab;
  * Represents the content in a viewport tab in an editor.
  * Each SDockTab holding viewports in an editor contains and owns one of these.
  */
-class UNREALED_API FEditorViewportTabContent : public FViewportTabContent, public TSharedFromThis<FEditorViewportTabContent>
+class FEditorViewportTabContent : public FViewportTabContent, public TSharedFromThis<FEditorViewportTabContent>
 {
 public:
 	/** @return The string used to identify the layout of this viewport tab */
@@ -28,31 +28,31 @@ public:
 		return ActiveViewportLayout;
 	}
 	   
-	TSharedPtr<FEditorViewportLayout> ConstructViewportLayoutByTypeName(const FName& TypeName, bool bSwitchingLayouts);
+	UNREALED_API TSharedPtr<FEditorViewportLayout> ConstructViewportLayoutByTypeName(const FName& TypeName, bool bSwitchingLayouts);
 
-	virtual void Initialize(AssetEditorViewportFactoryFunction Func, TSharedPtr<SDockTab> InParentTab, const FString& InLayoutString);
+	UNREALED_API virtual void Initialize(AssetEditorViewportFactoryFunction Func, TSharedPtr<SDockTab> InParentTab, const FString& InLayoutString);
 
-	TSharedPtr<SAssetEditorViewport> CreateSlateViewport(FName InTypeName, const FAssetEditorViewportConstructionArgs& ConstructionArgs) const;
+	UNREALED_API TSharedPtr<SAssetEditorViewport> CreateSlateViewport(FName InTypeName, const FAssetEditorViewportConstructionArgs& ConstructionArgs) const;
 
 	/**
 	* Sets the current layout by changing the contained layout object
 	* 
 	* @param ConfigurationName		The name of the layout (for the names in namespace EditorViewportConfigurationNames)
 	*/
-	void SetViewportConfiguration(const FName& ConfigurationName) override;
+	UNREALED_API void SetViewportConfiguration(const FName& ConfigurationName) override;
 
-	TSharedPtr<SEditorViewport> GetFirstViewport();
+	UNREALED_API TSharedPtr<SEditorViewport> GetFirstViewport();
 
-	void UpdateViewportTabWidget();
+	UNREALED_API void UpdateViewportTabWidget();
 
 	/** Save any configuration required to persist state for this viewport layout */
-	void SaveConfig() const;
+	UNREALED_API void SaveConfig() const;
 
 	/**
 	 * Refresh the current layout.
 	 * @note Does not save config
 	 */
-	void RefreshViewportConfiguration();
+	UNREALED_API void RefreshViewportConfiguration();
 
 	/**
 	 * Attempt to find a viewport creation factory function in the set of registered factory functions.
@@ -60,11 +60,11 @@ public:
 	 * @param InTypeName The name of the viewport creation factory to find.
 	 * @returns the creation function if one is found with the given typename, otherwise nullptr.
 	 */
-	const AssetEditorViewportFactoryFunction* FindViewportCreationFactory(FName InTypeName) const;
+	UNREALED_API const AssetEditorViewportFactoryFunction* FindViewportCreationFactory(FName InTypeName) const;
 
 protected:
-	virtual TSharedPtr<FEditorViewportLayout> FactoryViewportLayout(bool bIsSwitchingLayouts);
-	virtual FName GetLayoutTypeNameFromLayoutString() const;
+	UNREALED_API virtual TSharedPtr<FEditorViewportLayout> FactoryViewportLayout(bool bIsSwitchingLayouts);
+	UNREALED_API virtual FName GetLayoutTypeNameFromLayoutString() const;
 
 private:
 	using AssetEditorViewportCreationFactories = TMap<FName, AssetEditorViewportFactoryFunction>;
