[UnrealEd]

diff --git a/Engine/Source/Editor/UnrealEd/Public/EditorWorldExtension.h b/Engine/Source/Editor/UnrealEd/Public/EditorWorldExtension.h
--- a/Engine/Source/Editor/UnrealEd/Public/EditorWorldExtension.h
+++ b/Engine/Source/Editor/UnrealEd/Public/EditorWorldExtension.h
@@ -38,8 +38,8 @@ struct FEditorWorldExtensionActorData
 	bool bValidForPIE = false;
 };
 
-UCLASS()
-class UNREALED_API UEditorWorldExtension : public UObject
+UCLASS(MinimalAPI)
+class UEditorWorldExtension : public UObject
 {
 	GENERATED_BODY()
 
@@ -47,10 +47,10 @@ class UNREALED_API UEditorWorldExtension : public UObject
 public:
 	
 	/** Default constructor */
-	UEditorWorldExtension();
+	UNREALED_API UEditorWorldExtension();
 
 	/** Default destructor */
-	virtual ~UEditorWorldExtension();
+	UNREALED_API virtual ~UEditorWorldExtension();
 
 	/** Initialize extension */
 	virtual void Init() {};
@@ -61,14 +61,14 @@ public:
 	/** Give base class the chance to tick */
 	virtual void Tick( float DeltaSeconds ) {};
 
-	virtual bool InputKey( FEditorViewportClient* InViewportClient, FViewport* Viewport, FKey Key, EInputEvent Event);
-	virtual bool InputAxis( FEditorViewportClient* InViewportClient, FViewport* Viewport, int32 ControllerId, FKey Key, float Delta, float DeltaTime);
+	UNREALED_API virtual bool InputKey( FEditorViewportClient* InViewportClient, FViewport* Viewport, FKey Key, EInputEvent Event);
+	UNREALED_API virtual bool InputAxis( FEditorViewportClient* InViewportClient, FViewport* Viewport, int32 ControllerId, FKey Key, float Delta, float DeltaTime);
 
 	/** Gets the world owning this extension */
-	virtual UWorld* GetWorld() const override;
+	UNREALED_API virtual UWorld* GetWorld() const override;
 
 	/** Gets the world owning this extension's non-PIE valid actors when current world is a play world */
-	virtual UWorld* GetLastEditorWorld() const;
+	UNREALED_API virtual UWorld* GetLastEditorWorld() const;
 
 	/**  Spawns a transient actor that we can use in the current world of this extension (templated for convenience) */
 	template<class T>
@@ -78,27 +78,27 @@ public:
 	}
 
 	/** Spawns a transient actor that we can use in the current world of this extension */
-	AActor* SpawnTransientSceneActor(TSubclassOf<AActor> ActorClass, const FString& ActorName, const bool bWithSceneComponent = false, const EObjectFlags InObjectFlags = EObjectFlags::RF_Transient | EObjectFlags::RF_DuplicateTransient, const bool bValidForPIE = false);
+	UNREALED_API AActor* SpawnTransientSceneActor(TSubclassOf<AActor> ActorClass, const FString& ActorName, const bool bWithSceneComponent = false, const EObjectFlags InObjectFlags = EObjectFlags::RF_Transient | EObjectFlags::RF_DuplicateTransient, const bool bValidForPIE = false);
 
 	/** Destroys a transient actor we created earlier */
-	void DestroyTransientActor(AActor* Actor);
+	UNREALED_API void DestroyTransientActor(AActor* Actor);
 
 	/** Sets if this extension should be ticked. */
-	void SetActive(const bool bInActive);
+	UNREALED_API void SetActive(const bool bInActive);
 
 	/** If this extension is currently being ticked. */
-	bool IsActive() const;
+	UNREALED_API bool IsActive() const;
 
 	/** Get the owning collection of extensions */
-	UEditorWorldExtensionCollection* GetOwningCollection();
+	UNREALED_API UEditorWorldExtensionCollection* GetOwningCollection();
 
 	/** Executes command */
-	bool ExecCommand(const FString& InCommand);
+	UNREALED_API bool ExecCommand(const FString& InCommand);
 
 protected:
 	
 	/** Reparent actors to a new world */
-	virtual void TransitionWorld(UWorld* NewWorld, EEditorWorldExtensionTransitionState TransitionState);
+	UNREALED_API virtual void TransitionWorld(UWorld* NewWorld, EEditorWorldExtensionTransitionState TransitionState);
 
 	/** Give child class a chance to act on entering simulate mode */
 	virtual void EnteredSimulateInEditor() {};
@@ -112,10 +112,10 @@ protected:
 private:
 
 	/** Reparent the actors to a new world. */
-	void ReparentActor(AActor* Actor, UWorld* NewWorld);
+	UNREALED_API void ReparentActor(AActor* Actor, UWorld* NewWorld);
 
 	/** Let the FEditorWorldExtensionCollection set the world of this extension before init */
-	void InitInternal(UEditorWorldExtensionCollection* InOwningExtensionsCollection);
+	UNREALED_API void InitInternal(UEditorWorldExtensionCollection* InOwningExtensionsCollection);
 
 	UPROPERTY()
 	TArray<FEditorWorldExtensionActorData> ExtensionActors;
@@ -127,8 +127,8 @@ private:
 /**
  * Holds a collection of UEditorExtension
  */
-UCLASS()
-class UNREALED_API UEditorWorldExtensionCollection : public UObject
+UCLASS(MinimalAPI)
+class UEditorWorldExtensionCollection : public UObject
 {
 	GENERATED_BODY()
 
@@ -136,73 +136,73 @@ class UNREALED_API UEditorWorldExtensionCollection : public UObject
 public:
 
 	/** Default constructor */
-	UEditorWorldExtensionCollection();
+	UNREALED_API UEditorWorldExtensionCollection();
 
 	/** Default destructor */
-	virtual ~UEditorWorldExtensionCollection();
+	UNREALED_API virtual ~UEditorWorldExtensionCollection();
 
 	/** Gets the world from the world context */
-	virtual UWorld* GetWorld() const override;
+	UNREALED_API virtual UWorld* GetWorld() const override;
 
 	/** Gets the last editor world, will only be non-null when current world is a play world. */
-	UWorld* GetLastEditorWorld() const;
+	UNREALED_API UWorld* GetLastEditorWorld() const;
 
 	/**
 	 * Checks if the passed extension already exists and creates one if it doesn't.
 	 * @param EditorExtensionClass the subclass of an extension to create if necessary and add.
 	 */
-	UEditorWorldExtension* AddExtension(TSubclassOf<UEditorWorldExtension> EditorExtensionClass);
+	UNREALED_API UEditorWorldExtension* AddExtension(TSubclassOf<UEditorWorldExtension> EditorExtensionClass);
 
 	/** 
 	 * Adds an extension to the collection
 	 * @param	EditorExtension			The UEditorExtension that will be created, initialized and added to the collection.
 	 */
-	void AddExtension( UEditorWorldExtension* EditorExtension );
+	UNREALED_API void AddExtension( UEditorWorldExtension* EditorExtension );
 	
 	/** 
 	 * Removes an extension from the collection and calls Shutdown() on the extension
 	 * @param	EditorExtension			The UEditorExtension to remove.  It must already have been added.
 	 */
-	void RemoveExtension( UEditorWorldExtension* EditorExtension );
+	UNREALED_API void RemoveExtension( UEditorWorldExtension* EditorExtension );
 
 	/**
 	 * Find an extension based on the class
 	 * @param	EditorExtensionClass	The class to find an extension with
 	 * @return							The first extension that is found based on class
 	 */
-	UEditorWorldExtension* FindExtension( TSubclassOf<UEditorWorldExtension> EditorExtensionClass );
+	UNREALED_API UEditorWorldExtension* FindExtension( TSubclassOf<UEditorWorldExtension> EditorExtensionClass );
 
 	/** Ticks all extensions */
-	void Tick( float DeltaSeconds );
+	UNREALED_API void Tick( float DeltaSeconds );
 
 	/** Notifies all extensions of keyboard input */
-	bool InputKey( FEditorViewportClient* InViewportClient, FViewport* Viewport, FKey Key, EInputEvent Event);
+	UNREALED_API bool InputKey( FEditorViewportClient* InViewportClient, FViewport* Viewport, FKey Key, EInputEvent Event);
 
 	/** Notifies all extensions of axis movement */
-	bool InputAxis( FEditorViewportClient* InViewportClient, FViewport* Viewport, int32 ControllerId, FKey Key, float Delta, float DeltaTime);
+	UNREALED_API bool InputAxis( FEditorViewportClient* InViewportClient, FViewport* Viewport, int32 ControllerId, FKey Key, float Delta, float DeltaTime);
 	
 	/** Show or hide all the actors of extensions that belong to this collection. */
-	void ShowAllActors(const bool bShow);
+	UNREALED_API void ShowAllActors(const bool bShow);
 
 private:
 
 	/** Sets the world for this collection and gives every extension an opportunity to transition */
-	void SetWorld(UWorld* NewWorld, EEditorWorldExtensionTransitionState TransitionState /* = EEditorWorldExtensionTransitionState::TransitionAll */);
+	UNREALED_API void SetWorld(UWorld* NewWorld, EEditorWorldExtensionTransitionState TransitionState /* = EEditorWorldExtensionTransitionState::TransitionAll */);
 
 	/** Transitions actors in every extension to the specified world */
-	void TransitionWorld(UWorld* NewWorld, EEditorWorldExtensionTransitionState TransitionState);
+	UNREALED_API void TransitionWorld(UWorld* NewWorld, EEditorWorldExtensionTransitionState TransitionState);
 
 	/** Called by the editor after PIE or Simulate is started */
-	void PostPIEStarted( bool bIsSimulatingInEditor );
+	UNREALED_API void PostPIEStarted( bool bIsSimulatingInEditor );
 
 	/** Called just before PIE or Simulate ends */
-	void OnPreEndPIE(bool bWasSimulatingInEditor);
+	UNREALED_API void OnPreEndPIE(bool bWasSimulatingInEditor);
 
 	/** Called when PIE or Simulate ends */
-	void OnEndPIE( bool bWasSimulatingInEditor );
+	UNREALED_API void OnEndPIE( bool bWasSimulatingInEditor );
 
 	/** Called when switching between play and simulate */
-	void SwitchPIEAndSIE(bool bIsSimulatingInEditor);
+	UNREALED_API void SwitchPIEAndSIE(bool bIsSimulatingInEditor);
 
 	/** World context */
 	TWeakObjectPtr<UWorld> Currentworld;
@@ -220,24 +220,24 @@ private:
 /**
  * Holds a map of extension collections paired with worlds
  */
-UCLASS()
-class UNREALED_API UEditorWorldExtensionManager	: public UObject
+UCLASS(MinimalAPI)
+class UEditorWorldExtensionManager	: public UObject
 {
 	GENERATED_BODY()
 
 public:
 	/** Default constructor */
-	UEditorWorldExtensionManager();
+	UNREALED_API UEditorWorldExtensionManager();
 
 	/** Default destructor */
-	virtual ~UEditorWorldExtensionManager();
+	UNREALED_API virtual ~UEditorWorldExtensionManager();
 
 	/** Gets the editor world wrapper that is found with the world passed.
 	 * Adds one for this world if there was non found. */
-	UEditorWorldExtensionCollection* GetEditorWorldExtensions(UWorld* InWorld, const bool bCreateIfNeeded = true);
+	UNREALED_API UEditorWorldExtensionCollection* GetEditorWorldExtensions(UWorld* InWorld, const bool bCreateIfNeeded = true);
 
 	/** Ticks all the collections */
-	void Tick( float DeltaSeconds );
+	UNREALED_API void Tick( float DeltaSeconds );
 
 private:
 
diff --git a/Engine/Source/Editor/UnrealEd/Public/EditorWorldUtils.h b/Engine/Source/Editor/UnrealEd/Public/EditorWorldUtils.h
--- a/Engine/Source/Editor/UnrealEd/Public/EditorWorldUtils.h
+++ b/Engine/Source/Editor/UnrealEd/Public/EditorWorldUtils.h
@@ -15,7 +15,7 @@
  * On destruction, will destroy the world and unroot it.
  * This class will also set GWorld & the EditorWorldContext to this world.
  */
-class UNREALED_API FScopedEditorWorld
+class FScopedEditorWorld
 {
 public:
 	/**
@@ -23,31 +23,31 @@ public:
 	 * @param InWorld					The world to manage.
 	 * @param InInitializationValues	The initialization values to use for the world.
 	 */
-	FScopedEditorWorld(UWorld* InWorld, const UWorld::InitializationValues& InInitializationValues, EWorldType::Type InWorldType = EWorldType::Editor);
+	UNREALED_API FScopedEditorWorld(UWorld* InWorld, const UWorld::InitializationValues& InInitializationValues, EWorldType::Type InWorldType = EWorldType::Editor);
 
 	/**
 	 * FStringView Constructor - Load the specified package & initialize the world as an editor world.
 	 * @param InLongPackageName			Path to a package containing a world.
 	 * @param InInitializationValues	The initialization values to use for the world.
 	 */
-	FScopedEditorWorld(const FStringView InLongPackageName, const UWorld::InitializationValues& InInitializationValues, EWorldType::Type InWorldType = EWorldType::Editor);
+	UNREALED_API FScopedEditorWorld(const FStringView InLongPackageName, const UWorld::InitializationValues& InInitializationValues, EWorldType::Type InWorldType = EWorldType::Editor);
 
 	/**
 	 * SoftObjectPtr Constructor - Initialize the provided world as an editor world.
 	 * @param InWorld					World soft object pointer.
 	 * @param InInitializationValues	The initialization values to use for the world.
 	 */
-	FScopedEditorWorld(const TSoftObjectPtr<UWorld>& InSoftWorld, const UWorld::InitializationValues& InInitializationValues, EWorldType::Type InWorldType = EWorldType::Editor);
+	UNREALED_API FScopedEditorWorld(const TSoftObjectPtr<UWorld>& InSoftWorld, const UWorld::InitializationValues& InInitializationValues, EWorldType::Type InWorldType = EWorldType::Editor);
 
 	/**
 	 * Destructor - Destroy the provided world.
 	 */
-	~FScopedEditorWorld();
+	UNREALED_API ~FScopedEditorWorld();
 
 	/**
 	 * Obtain the world managed by this scope, or null if the initialization failed.
 	 */
-	UWorld* GetWorld() const;
+	UNREALED_API UWorld* GetWorld() const;
 
 private:
 	FScopedEditorWorld();
@@ -60,4 +60,4 @@ private:
 /**
  * Load a world package, managing the WorldTypePreLoadMap to ensure the correct world type is specified in UWorld::PostLoad()
  */
-UNREALED_API UPackage* LoadWorldPackageForEditor(const FStringView InLongPackageName, EWorldType::Type InWorldType = EWorldType::Editor, uint32 InLoadFlags = LOAD_None);
\ No newline at end of file
+UNREALED_API UPackage* LoadWorldPackageForEditor(const FStringView InLongPackageName, EWorldType::Type InWorldType = EWorldType::Editor, uint32 InLoadFlags = LOAD_None);
diff --git a/Engine/Source/Editor/UnrealEd/Public/Elements/Actor/ActorElementDetailsInterface.h b/Engine/Source/Editor/UnrealEd/Public/Elements/Actor/ActorElementDetailsInterface.h
--- a/Engine/Source/Editor/UnrealEd/Public/Elements/Actor/ActorElementDetailsInterface.h
+++ b/Engine/Source/Editor/UnrealEd/Public/Elements/Actor/ActorElementDetailsInterface.h
@@ -12,11 +12,11 @@
 
 struct FTypedElementHandle;
 
-UCLASS()
-class UNREALED_API UActorElementDetailsInterface : public UObject, public ITypedElementDetailsInterface
+UCLASS(MinimalAPI)
+class UActorElementDetailsInterface : public UObject, public ITypedElementDetailsInterface
 {
 	GENERATED_BODY()
 
 public:
-	virtual TUniquePtr<ITypedElementDetailsObject> GetDetailsObject(const FTypedElementHandle& InElementHandle) override;
+	UNREALED_API virtual TUniquePtr<ITypedElementDetailsObject> GetDetailsObject(const FTypedElementHandle& InElementHandle) override;
 };
diff --git a/Engine/Source/Editor/UnrealEd/Public/Elements/Actor/ActorElementEditorAssetDataInterface.h b/Engine/Source/Editor/UnrealEd/Public/Elements/Actor/ActorElementEditorAssetDataInterface.h
--- a/Engine/Source/Editor/UnrealEd/Public/Elements/Actor/ActorElementEditorAssetDataInterface.h
+++ b/Engine/Source/Editor/UnrealEd/Public/Elements/Actor/ActorElementEditorAssetDataInterface.h
@@ -13,11 +13,11 @@ class UObject;
 struct FAssetData;
 struct FTypedElementHandle;
 
-UCLASS()
-class UNREALED_API UActorElementEditorAssetDataInterface : public UActorElementAssetDataInterface
+UCLASS(MinimalAPI)
+class UActorElementEditorAssetDataInterface : public UActorElementAssetDataInterface
 {
 	GENERATED_BODY()
 
 public:
-	virtual TArray<FAssetData> GetAllReferencedAssetDatas(const FTypedElementHandle& InElementHandle) override;
+	UNREALED_API virtual TArray<FAssetData> GetAllReferencedAssetDatas(const FTypedElementHandle& InElementHandle) override;
 };
diff --git a/Engine/Source/Editor/UnrealEd/Public/Elements/Actor/ActorElementEditorCopyAndPaste.h b/Engine/Source/Editor/UnrealEd/Public/Elements/Actor/ActorElementEditorCopyAndPaste.h
--- a/Engine/Source/Editor/UnrealEd/Public/Elements/Actor/ActorElementEditorCopyAndPaste.h
+++ b/Engine/Source/Editor/UnrealEd/Public/Elements/Actor/ActorElementEditorCopyAndPaste.h
@@ -16,8 +16,8 @@
 class AActor;
 class AGroupActor;
 
-UCLASS(BlueprintType)
-class UNREALED_API UActorElementsCopy : public UObject
+UCLASS(BlueprintType, MinimalAPI)
+class UActorElementsCopy : public UObject
 {
 	GENERATED_BODY()
 
@@ -26,25 +26,25 @@ public:
 	TArray<TObjectPtr<AActor>> ActorsToCopy;
 };
 
-UCLASS()
-class UNREALED_API UActorElementsExporterT3D : public UExporter
+UCLASS(MinimalAPI)
+class UActorElementsExporterT3D : public UExporter
 {
 public:
 	GENERATED_BODY()
 
 public:
-	UActorElementsExporterT3D(const FObjectInitializer& ObjectInitializer = FObjectInitializer());
+	UNREALED_API UActorElementsExporterT3D(const FObjectInitializer& ObjectInitializer = FObjectInitializer());
 
 	//~ Begin UExporter Interface
-	virtual bool ExportText(const FExportObjectInnerContext* Context, UObject* Object, const TCHAR* Type, FOutputDevice& Out, FFeedbackContext* Warn, uint32 PortFlags = 0) override;
+	UNREALED_API virtual bool ExportText(const FExportObjectInnerContext* Context, UObject* Object, const TCHAR* Type, FOutputDevice& Out, FFeedbackContext* Warn, uint32 PortFlags = 0) override;
 	//~ End UExporter Interface
 };
 
-struct UNREALED_API FActorElementEditorPasteImporter : public FWorldElementPasteImporter
+struct FActorElementEditorPasteImporter : public FWorldElementPasteImporter
 {
-	virtual void Import(FContext& Context) override;
+	UNREALED_API virtual void Import(FContext& Context) override;
 	
-	virtual TArray<FTypedElementHandle> GetImportedElements() override;
+	UNREALED_API virtual TArray<FTypedElementHandle> GetImportedElements() override;
 	
 private:
 	void PostImportProcess();
diff --git a/Engine/Source/Editor/UnrealEd/Public/Elements/Actor/ActorElementEditorSelectionInterface.h b/Engine/Source/Editor/UnrealEd/Public/Elements/Actor/ActorElementEditorSelectionInterface.h
--- a/Engine/Source/Editor/UnrealEd/Public/Elements/Actor/ActorElementEditorSelectionInterface.h
+++ b/Engine/Source/Editor/UnrealEd/Public/Elements/Actor/ActorElementEditorSelectionInterface.h
@@ -16,17 +16,17 @@ class UObject;
 struct FTypedElementHandle;
 struct FTypedElementIsSelectedOptions;
 
-UCLASS()
-class UNREALED_API UActorElementEditorSelectionInterface : public UActorElementSelectionInterface
+UCLASS(MinimalAPI)
+class UActorElementEditorSelectionInterface : public UActorElementSelectionInterface
 {
 	GENERATED_BODY()
 
 public:
-	virtual bool SelectElement(const FTypedElementHandle& InElementHandle, const FTypedElementListPtr& InSelectionSet, const FTypedElementSelectionOptions& InSelectionOptions) override;
-	virtual bool DeselectElement(const FTypedElementHandle& InElementHandle, const FTypedElementListPtr& InSelectionSet, const FTypedElementSelectionOptions& InSelectionOptions) override;
-	virtual bool IsElementSelected(const FTypedElementHandle& InElementHandle, const FTypedElementListConstPtr& SelectionSetPtr, const FTypedElementIsSelectedOptions& InSelectionOptions) override;
-	virtual bool ShouldPreventTransactions(const FTypedElementHandle& InElementHandle) override;
-	virtual TUniquePtr<ITypedElementTransactedElement> CreateTransactedElementImpl() override;
+	UNREALED_API virtual bool SelectElement(const FTypedElementHandle& InElementHandle, const FTypedElementListPtr& InSelectionSet, const FTypedElementSelectionOptions& InSelectionOptions) override;
+	UNREALED_API virtual bool DeselectElement(const FTypedElementHandle& InElementHandle, const FTypedElementListPtr& InSelectionSet, const FTypedElementSelectionOptions& InSelectionOptions) override;
+	UNREALED_API virtual bool IsElementSelected(const FTypedElementHandle& InElementHandle, const FTypedElementListConstPtr& SelectionSetPtr, const FTypedElementIsSelectedOptions& InSelectionOptions) override;
+	UNREALED_API virtual bool ShouldPreventTransactions(const FTypedElementHandle& InElementHandle) override;
+	UNREALED_API virtual TUniquePtr<ITypedElementTransactedElement> CreateTransactedElementImpl() override;
 
-	static bool IsActorSelected(const AActor* InActor, FTypedElementListConstRef InSelectionSet, const FTypedElementIsSelectedOptions& InSelectionOptions);
+	static UNREALED_API bool IsActorSelected(const AActor* InActor, FTypedElementListConstRef InSelectionSet, const FTypedElementIsSelectedOptions& InSelectionOptions);
 };
diff --git a/Engine/Source/Editor/UnrealEd/Public/Elements/Actor/ActorElementEditorViewportInteractionCustomization.h b/Engine/Source/Editor/UnrealEd/Public/Elements/Actor/ActorElementEditorViewportInteractionCustomization.h
--- a/Engine/Source/Editor/UnrealEd/Public/Elements/Actor/ActorElementEditorViewportInteractionCustomization.h
+++ b/Engine/Source/Editor/UnrealEd/Public/Elements/Actor/ActorElementEditorViewportInteractionCustomization.h
@@ -7,12 +7,12 @@
 
 class AActor;
 
-class UNREALED_API FActorElementEditorViewportInteractionCustomization : public FTypedElementViewportInteractionCustomization
+class FActorElementEditorViewportInteractionCustomization : public FTypedElementViewportInteractionCustomization
 {
 public:
-	virtual void GizmoManipulationDeltaUpdate(const TTypedElement<ITypedElementWorldInterface>& InElementWorldHandle, const UE::Widget::EWidgetMode InWidgetMode, const EAxisList::Type InDragAxis, const FInputDeviceState& InInputState, const FTransform& InDeltaTransform, const FVector& InPivotLocation, ETypedElementViewportInteractionDragMovementType GizmoDeltaType) override;
-	virtual void MirrorElement(const TTypedElement<ITypedElementWorldInterface>& InElementWorldHandle, const FVector& InMirrorScale, const FVector& InPivotLocation) override;
-	virtual bool GetFocusBounds(const TTypedElement<ITypedElementWorldInterface>& InElementWorldHandle, FBoxSphereBounds& OutBounds) override;
+	UNREALED_API virtual void GizmoManipulationDeltaUpdate(const TTypedElement<ITypedElementWorldInterface>& InElementWorldHandle, const UE::Widget::EWidgetMode InWidgetMode, const EAxisList::Type InDragAxis, const FInputDeviceState& InInputState, const FTransform& InDeltaTransform, const FVector& InPivotLocation, ETypedElementViewportInteractionDragMovementType GizmoDeltaType) override;
+	UNREALED_API virtual void MirrorElement(const TTypedElement<ITypedElementWorldInterface>& InElementWorldHandle, const FVector& InMirrorScale, const FVector& InPivotLocation) override;
+	UNREALED_API virtual bool GetFocusBounds(const TTypedElement<ITypedElementWorldInterface>& InElementWorldHandle, FBoxSphereBounds& OutBounds) override;
 	
-	static void ApplyDeltaToActor(AActor* InActor, const bool InIsDelta, const FVector* InDeltaTranslationPtr, const FRotator* InDeltaRotationPtr, const FVector* InDeltaScalePtr, const FVector& InPivotLocation, const FInputDeviceState& InInputState, ETypedElementViewportInteractionDragMovementType GizmoDeltaType);
+	static UNREALED_API void ApplyDeltaToActor(AActor* InActor, const bool InIsDelta, const FVector* InDeltaTranslationPtr, const FRotator* InDeltaRotationPtr, const FVector* InDeltaScalePtr, const FVector& InPivotLocation, const FInputDeviceState& InInputState, ETypedElementViewportInteractionDragMovementType GizmoDeltaType);
 };
diff --git a/Engine/Source/Editor/UnrealEd/Public/Elements/Actor/ActorElementEditorWorldInterface.h b/Engine/Source/Editor/UnrealEd/Public/Elements/Actor/ActorElementEditorWorldInterface.h
--- a/Engine/Source/Editor/UnrealEd/Public/Elements/Actor/ActorElementEditorWorldInterface.h
+++ b/Engine/Source/Editor/UnrealEd/Public/Elements/Actor/ActorElementEditorWorldInterface.h
@@ -5,25 +5,25 @@
 #include "Elements/Actor/ActorElementWorldInterface.h"
 #include "ActorElementEditorWorldInterface.generated.h"
 
-UCLASS()
-class UNREALED_API UActorElementEditorWorldInterface : public UActorElementWorldInterface
+UCLASS(MinimalAPI)
+class UActorElementEditorWorldInterface : public UActorElementWorldInterface
 {
 	GENERATED_BODY()
 
 public:
-	virtual bool GetPivotOffset(const FTypedElementHandle& InElementHandle, FVector& OutPivotOffset) override;
-	virtual bool SetPivotOffset(const FTypedElementHandle& InElementHandle, const FVector& InPivotOffset) override;
-	virtual void NotifyMovementStarted(const FTypedElementHandle& InElementHandle) override;
-	virtual void NotifyMovementOngoing(const FTypedElementHandle& InElementHandle) override;
-	virtual void NotifyMovementEnded(const FTypedElementHandle& InElementHandle) override;
-	virtual bool CanDeleteElement(const FTypedElementHandle& InElementHandle) override;
-	virtual bool DeleteElements(TArrayView<const FTypedElementHandle> InElementHandles, UWorld* InWorld, UTypedElementSelectionSet* InSelectionSet, const FTypedElementDeletionOptions& InDeletionOptions) override;
-	virtual bool CanDuplicateElement(const FTypedElementHandle& InElementHandle) override;
-	virtual void DuplicateElements(TArrayView<const FTypedElementHandle> InElementHandles, UWorld* InWorld, const FVector& InLocationOffset, TArray<FTypedElementHandle>& OutNewElements) override;
-	virtual bool CanCopyElement(const FTypedElementHandle& InElementHandle) override;
-	virtual void CopyElements(TArrayView<const FTypedElementHandle> InElementHandles, FOutputDevice& Out) override;
-	virtual TSharedPtr<FWorldElementPasteImporter> GetPasteImporter() override;
-	virtual bool IsElementInConvexVolume(const FTypedElementHandle& Handle, const FConvexVolume& InVolume, bool bMustEncompassEntireElement = false);
-	virtual bool IsElementInBox(const FTypedElementHandle& Handle, const FBox& InBox, bool bMustEncompassEntireElement = false);
-	virtual TArray<FTypedElementHandle> GetSelectionElementsFromSelectionFunction(const FTypedElementHandle& InElementHandle, const FWorldSelectionElementArgs& SelectionArgs, const TFunction<bool(const FTypedElementHandle&, const FWorldSelectionElementArgs&)>& SelectionFunction) override;
+	UNREALED_API virtual bool GetPivotOffset(const FTypedElementHandle& InElementHandle, FVector& OutPivotOffset) override;
+	UNREALED_API virtual bool SetPivotOffset(const FTypedElementHandle& InElementHandle, const FVector& InPivotOffset) override;
+	UNREALED_API virtual void NotifyMovementStarted(const FTypedElementHandle& InElementHandle) override;
+	UNREALED_API virtual void NotifyMovementOngoing(const FTypedElementHandle& InElementHandle) override;
+	UNREALED_API virtual void NotifyMovementEnded(const FTypedElementHandle& InElementHandle) override;
+	UNREALED_API virtual bool CanDeleteElement(const FTypedElementHandle& InElementHandle) override;
+	UNREALED_API virtual bool DeleteElements(TArrayView<const FTypedElementHandle> InElementHandles, UWorld* InWorld, UTypedElementSelectionSet* InSelectionSet, const FTypedElementDeletionOptions& InDeletionOptions) override;
+	UNREALED_API virtual bool CanDuplicateElement(const FTypedElementHandle& InElementHandle) override;
+	UNREALED_API virtual void DuplicateElements(TArrayView<const FTypedElementHandle> InElementHandles, UWorld* InWorld, const FVector& InLocationOffset, TArray<FTypedElementHandle>& OutNewElements) override;
+	UNREALED_API virtual bool CanCopyElement(const FTypedElementHandle& InElementHandle) override;
+	UNREALED_API virtual void CopyElements(TArrayView<const FTypedElementHandle> InElementHandles, FOutputDevice& Out) override;
+	UNREALED_API virtual TSharedPtr<FWorldElementPasteImporter> GetPasteImporter() override;
+	UNREALED_API virtual bool IsElementInConvexVolume(const FTypedElementHandle& Handle, const FConvexVolume& InVolume, bool bMustEncompassEntireElement = false);
+	UNREALED_API virtual bool IsElementInBox(const FTypedElementHandle& Handle, const FBox& InBox, bool bMustEncompassEntireElement = false);
+	UNREALED_API virtual TArray<FTypedElementHandle> GetSelectionElementsFromSelectionFunction(const FTypedElementHandle& InElementHandle, const FWorldSelectionElementArgs& SelectionArgs, const TFunction<bool(const FTypedElementHandle&, const FWorldSelectionElementArgs&)>& SelectionFunction) override;
 };
diff --git a/Engine/Source/Editor/UnrealEd/Public/Elements/Component/ComponentElementDetailsInterface.h b/Engine/Source/Editor/UnrealEd/Public/Elements/Component/ComponentElementDetailsInterface.h
--- a/Engine/Source/Editor/UnrealEd/Public/Elements/Component/ComponentElementDetailsInterface.h
+++ b/Engine/Source/Editor/UnrealEd/Public/Elements/Component/ComponentElementDetailsInterface.h
@@ -12,12 +12,12 @@
 
 struct FTypedElementHandle;
 
-UCLASS()
-class UNREALED_API UComponentElementDetailsInterface : public UObject, public ITypedElementDetailsInterface
+UCLASS(MinimalAPI)
+class UComponentElementDetailsInterface : public UObject, public ITypedElementDetailsInterface
 {
 	GENERATED_BODY()
 
 public:
 	virtual bool IsTopLevelElement(const FTypedElementHandle& InElementHandle) override { return false; }
-	virtual TUniquePtr<ITypedElementDetailsObject> GetDetailsObject(const FTypedElementHandle& InElementHandle) override;
+	UNREALED_API virtual TUniquePtr<ITypedElementDetailsObject> GetDetailsObject(const FTypedElementHandle& InElementHandle) override;
 };
diff --git a/Engine/Source/Editor/UnrealEd/Public/Elements/Component/ComponentElementEditorCopyAndPaste.h b/Engine/Source/Editor/UnrealEd/Public/Elements/Component/ComponentElementEditorCopyAndPaste.h
--- a/Engine/Source/Editor/UnrealEd/Public/Elements/Component/ComponentElementEditorCopyAndPaste.h
+++ b/Engine/Source/Editor/UnrealEd/Public/Elements/Component/ComponentElementEditorCopyAndPaste.h
@@ -1,4 +1,4 @@
-﻿// Copyright Epic Games, Inc. All Rights Reserved.
+// Copyright Epic Games, Inc. All Rights Reserved.
 
 #pragma once
 
@@ -10,8 +10,8 @@
 
 #include "ComponentElementEditorCopyAndPaste.generated.h"
 
-UCLASS(BlueprintType)
-class UNREALED_API UComponentElementsCopy : public UObject
+UCLASS(BlueprintType, MinimalAPI)
+class UComponentElementsCopy : public UObject
 {
 	GENERATED_BODY()
 
@@ -20,26 +20,26 @@ public:
 	TArray<TObjectPtr<UActorComponent>> ComponentsToCopy;
 };
 
-UCLASS()
-class UNREALED_API UComponentElementsExporterT3D : public UExporter
+UCLASS(MinimalAPI)
+class UComponentElementsExporterT3D : public UExporter
 {
 public:
 	GENERATED_BODY()
 
 public:
-	explicit UComponentElementsExporterT3D(const FObjectInitializer& ObjectInitializer = FObjectInitializer());
+	UNREALED_API explicit UComponentElementsExporterT3D(const FObjectInitializer& ObjectInitializer = FObjectInitializer());
 
 	//~ Begin UExporter Interface
-	virtual bool ExportText(const FExportObjectInnerContext* Context, UObject* Object, const TCHAR* Type, FOutputDevice& Out, FFeedbackContext* Warn, uint32 PortFlags = 0) override;
+	UNREALED_API virtual bool ExportText(const FExportObjectInnerContext* Context, UObject* Object, const TCHAR* Type, FOutputDevice& Out, FFeedbackContext* Warn, uint32 PortFlags = 0) override;
 	//~ End UExporter Interface
 };
 
-struct UNREALED_API FComponentElementEditorPasteImporter : public FWorldElementPasteImporter
+struct FComponentElementEditorPasteImporter : public FWorldElementPasteImporter
 {
 public:
-	virtual void Import(FContext& Context) override;
+	UNREALED_API virtual void Import(FContext& Context) override;
 	
-	virtual TArray<FTypedElementHandle> GetImportedElements() override;
+	UNREALED_API virtual TArray<FTypedElementHandle> GetImportedElements() override;
 private:
 	TArray<TObjectPtr<UActorComponent>> ImportedComponents;
-};
\ No newline at end of file
+};
diff --git a/Engine/Source/Editor/UnrealEd/Public/Elements/Component/ComponentElementEditorSelectionInterface.h b/Engine/Source/Editor/UnrealEd/Public/Elements/Component/ComponentElementEditorSelectionInterface.h
--- a/Engine/Source/Editor/UnrealEd/Public/Elements/Component/ComponentElementEditorSelectionInterface.h
+++ b/Engine/Source/Editor/UnrealEd/Public/Elements/Component/ComponentElementEditorSelectionInterface.h
@@ -16,15 +16,15 @@ class UObject;
 struct FTypedElementHandle;
 struct FTypedElementIsSelectedOptions;
 
-UCLASS()
-class UNREALED_API UComponentElementEditorSelectionInterface : public UComponentElementSelectionInterface
+UCLASS(MinimalAPI)
+class UComponentElementEditorSelectionInterface : public UComponentElementSelectionInterface
 {
 	GENERATED_BODY()
 
 public:
-	virtual bool IsElementSelected(const FTypedElementHandle& InElementHandle, const FTypedElementListConstPtr& InSelectionSet, const FTypedElementIsSelectedOptions& InSelectionOptions) override;
-	virtual bool ShouldPreventTransactions(const FTypedElementHandle& InElementHandle) override;
-	virtual TUniquePtr<ITypedElementTransactedElement> CreateTransactedElementImpl() override;
+	UNREALED_API virtual bool IsElementSelected(const FTypedElementHandle& InElementHandle, const FTypedElementListConstPtr& InSelectionSet, const FTypedElementIsSelectedOptions& InSelectionOptions) override;
+	UNREALED_API virtual bool ShouldPreventTransactions(const FTypedElementHandle& InElementHandle) override;
+	UNREALED_API virtual TUniquePtr<ITypedElementTransactedElement> CreateTransactedElementImpl() override;
 
-	static bool IsComponentSelected(const UActorComponent* InComponent, FTypedElementListConstRef InSelectionSet, const FTypedElementIsSelectedOptions& InSelectionOptions);
+	static UNREALED_API bool IsComponentSelected(const UActorComponent* InComponent, FTypedElementListConstRef InSelectionSet, const FTypedElementIsSelectedOptions& InSelectionOptions);
 };
diff --git a/Engine/Source/Editor/UnrealEd/Public/Elements/Component/ComponentElementEditorViewportInteractionCustomization.h b/Engine/Source/Editor/UnrealEd/Public/Elements/Component/ComponentElementEditorViewportInteractionCustomization.h
--- a/Engine/Source/Editor/UnrealEd/Public/Elements/Component/ComponentElementEditorViewportInteractionCustomization.h
+++ b/Engine/Source/Editor/UnrealEd/Public/Elements/Component/ComponentElementEditorViewportInteractionCustomization.h
@@ -8,12 +8,12 @@
 class FProperty;
 class USceneComponent;
 
-class UNREALED_API FComponentElementEditorViewportInteractionCustomization : public FTypedElementViewportInteractionCustomization
+class FComponentElementEditorViewportInteractionCustomization : public FTypedElementViewportInteractionCustomization
 {
 public:
-	virtual bool GetGizmoPivotLocation(const TTypedElement<ITypedElementWorldInterface>& InElementWorldHandle, const UE::Widget::EWidgetMode InWidgetMode, FVector& OutPivotLocation) override;
-	virtual void GizmoManipulationDeltaUpdate(const TTypedElement<ITypedElementWorldInterface>& InElementWorldHandle, const UE::Widget::EWidgetMode InWidgetMode, const EAxisList::Type InDragAxis, const FInputDeviceState& InInputState, const FTransform& InDeltaTransform, const FVector& InPivotLocation, ETypedElementViewportInteractionDragMovementType GizmoDeltaType) override;
-	virtual bool GetFocusBounds(const TTypedElement<ITypedElementWorldInterface>& InElementWorldHandle, FBoxSphereBounds& OutBounds) override;
+	UNREALED_API virtual bool GetGizmoPivotLocation(const TTypedElement<ITypedElementWorldInterface>& InElementWorldHandle, const UE::Widget::EWidgetMode InWidgetMode, FVector& OutPivotLocation) override;
+	UNREALED_API virtual void GizmoManipulationDeltaUpdate(const TTypedElement<ITypedElementWorldInterface>& InElementWorldHandle, const UE::Widget::EWidgetMode InWidgetMode, const EAxisList::Type InDragAxis, const FInputDeviceState& InInputState, const FTransform& InDeltaTransform, const FVector& InPivotLocation, ETypedElementViewportInteractionDragMovementType GizmoDeltaType) override;
+	UNREALED_API virtual bool GetFocusBounds(const TTypedElement<ITypedElementWorldInterface>& InElementWorldHandle, FBoxSphereBounds& OutBounds) override;
 	
-	static void ApplyDeltaToComponent(USceneComponent* InComponent, const bool InIsDelta, const FVector* InDeltaTranslationPtr, const FRotator* InDeltaRotationPtr, const FVector* InDeltaScalePtr, const FVector& InPivotLocation, const FInputDeviceState& InInputState);
+	static UNREALED_API void ApplyDeltaToComponent(USceneComponent* InComponent, const bool InIsDelta, const FVector* InDeltaTranslationPtr, const FRotator* InDeltaRotationPtr, const FVector* InDeltaScalePtr, const FVector& InPivotLocation, const FInputDeviceState& InInputState);
 };
diff --git a/Engine/Source/Editor/UnrealEd/Public/Elements/Component/ComponentElementEditorWorldInterface.h b/Engine/Source/Editor/UnrealEd/Public/Elements/Component/ComponentElementEditorWorldInterface.h
--- a/Engine/Source/Editor/UnrealEd/Public/Elements/Component/ComponentElementEditorWorldInterface.h
+++ b/Engine/Source/Editor/UnrealEd/Public/Elements/Component/ComponentElementEditorWorldInterface.h
@@ -5,23 +5,23 @@
 #include "Elements/Component/ComponentElementWorldInterface.h"
 #include "ComponentElementEditorWorldInterface.generated.h"
 
-UCLASS()
-class UNREALED_API UComponentElementEditorWorldInterface : public UComponentElementWorldInterface
+UCLASS(MinimalAPI)
+class UComponentElementEditorWorldInterface : public UComponentElementWorldInterface
 {
 	GENERATED_BODY()
 
 public:
-	virtual void NotifyMovementStarted(const FTypedElementHandle& InElementHandle) override;
-	virtual void NotifyMovementOngoing(const FTypedElementHandle& InElementHandle) override;
-	virtual void NotifyMovementEnded(const FTypedElementHandle& InElementHandle) override;
-	virtual bool CanDeleteElement(const FTypedElementHandle& InElementHandle) override;
-	virtual bool DeleteElements(TArrayView<const FTypedElementHandle> InElementHandles, UWorld* InWorld, UTypedElementSelectionSet* InSelectionSet, const FTypedElementDeletionOptions& InDeletionOptions) override;
-	virtual bool CanDuplicateElement(const FTypedElementHandle& InElementHandle) override;
-	virtual void DuplicateElements(TArrayView<const FTypedElementHandle> InElementHandles, UWorld* InWorld, const FVector& InLocationOffset, TArray<FTypedElementHandle>& OutNewElements) override;
-	virtual bool CanCopyElement(const FTypedElementHandle& InElementHandle) override;
-	virtual void CopyElements(TArrayView<const FTypedElementHandle> InElementHandles, FOutputDevice& Out) override;
-	virtual TSharedPtr<FWorldElementPasteImporter> GetPasteImporter() override;
-	virtual bool IsElementInConvexVolume(const FTypedElementHandle& Handle, const FConvexVolume& InVolume, bool bMustEncompassEntireElement = false);
-	virtual bool IsElementInBox(const FTypedElementHandle& Handle, const FBox& InBox, bool bMustEncompassEntireElement = false);
-	virtual TArray<FTypedElementHandle> GetSelectionElementsFromSelectionFunction(const FTypedElementHandle& InElementHandle, const FWorldSelectionElementArgs& SelectionArgs, const TFunction<bool(const FTypedElementHandle&, const FWorldSelectionElementArgs&)>& SelectionFunction) override;
+	UNREALED_API virtual void NotifyMovementStarted(const FTypedElementHandle& InElementHandle) override;
+	UNREALED_API virtual void NotifyMovementOngoing(const FTypedElementHandle& InElementHandle) override;
+	UNREALED_API virtual void NotifyMovementEnded(const FTypedElementHandle& InElementHandle) override;
+	UNREALED_API virtual bool CanDeleteElement(const FTypedElementHandle& InElementHandle) override;
+	UNREALED_API virtual bool DeleteElements(TArrayView<const FTypedElementHandle> InElementHandles, UWorld* InWorld, UTypedElementSelectionSet* InSelectionSet, const FTypedElementDeletionOptions& InDeletionOptions) override;
+	UNREALED_API virtual bool CanDuplicateElement(const FTypedElementHandle& InElementHandle) override;
+	UNREALED_API virtual void DuplicateElements(TArrayView<const FTypedElementHandle> InElementHandles, UWorld* InWorld, const FVector& InLocationOffset, TArray<FTypedElementHandle>& OutNewElements) override;
+	UNREALED_API virtual bool CanCopyElement(const FTypedElementHandle& InElementHandle) override;
+	UNREALED_API virtual void CopyElements(TArrayView<const FTypedElementHandle> InElementHandles, FOutputDevice& Out) override;
+	UNREALED_API virtual TSharedPtr<FWorldElementPasteImporter> GetPasteImporter() override;
+	UNREALED_API virtual bool IsElementInConvexVolume(const FTypedElementHandle& Handle, const FConvexVolume& InVolume, bool bMustEncompassEntireElement = false);
+	UNREALED_API virtual bool IsElementInBox(const FTypedElementHandle& Handle, const FBox& InBox, bool bMustEncompassEntireElement = false);
+	UNREALED_API virtual TArray<FTypedElementHandle> GetSelectionElementsFromSelectionFunction(const FTypedElementHandle& InElementHandle, const FWorldSelectionElementArgs& SelectionArgs, const TFunction<bool(const FTypedElementHandle&, const FWorldSelectionElementArgs&)>& SelectionFunction) override;
 };
diff --git a/Engine/Source/Editor/UnrealEd/Public/Elements/Object/ObjectElementDetailsInterface.h b/Engine/Source/Editor/UnrealEd/Public/Elements/Object/ObjectElementDetailsInterface.h
--- a/Engine/Source/Editor/UnrealEd/Public/Elements/Object/ObjectElementDetailsInterface.h
+++ b/Engine/Source/Editor/UnrealEd/Public/Elements/Object/ObjectElementDetailsInterface.h
@@ -12,11 +12,11 @@
 
 struct FTypedElementHandle;
 
-UCLASS()
-class UNREALED_API UObjectElementDetailsInterface : public UObject, public ITypedElementDetailsInterface
+UCLASS(MinimalAPI)
+class UObjectElementDetailsInterface : public UObject, public ITypedElementDetailsInterface
 {
 	GENERATED_BODY()
 
 public:
-	virtual TUniquePtr<ITypedElementDetailsObject> GetDetailsObject(const FTypedElementHandle& InElementHandle) override;
+	UNREALED_API virtual TUniquePtr<ITypedElementDetailsObject> GetDetailsObject(const FTypedElementHandle& InElementHandle) override;
 };
diff --git a/Engine/Source/Editor/UnrealEd/Public/Elements/Object/ObjectElementEditorSelectionInterface.h b/Engine/Source/Editor/UnrealEd/Public/Elements/Object/ObjectElementEditorSelectionInterface.h
--- a/Engine/Source/Editor/UnrealEd/Public/Elements/Object/ObjectElementEditorSelectionInterface.h
+++ b/Engine/Source/Editor/UnrealEd/Public/Elements/Object/ObjectElementEditorSelectionInterface.h
@@ -13,14 +13,14 @@ class ITypedElementTransactedElement;
 class UObject;
 struct FTypedElementHandle;
 
-UCLASS()
-class UNREALED_API UObjectElementEditorSelectionInterface : public UObjectElementSelectionInterface
+UCLASS(MinimalAPI)
+class UObjectElementEditorSelectionInterface : public UObjectElementSelectionInterface
 {
 	GENERATED_BODY()
 
 public:
-	virtual bool ShouldPreventTransactions(const FTypedElementHandle& InElementHandle) override;
-	virtual TUniquePtr<ITypedElementTransactedElement> CreateTransactedElementImpl() override;
+	UNREALED_API virtual bool ShouldPreventTransactions(const FTypedElementHandle& InElementHandle) override;
+	UNREALED_API virtual TUniquePtr<ITypedElementTransactedElement> CreateTransactedElementImpl() override;
 
-	static bool ShouldObjectPreventTransactions(const UObject* InObject);
+	static UNREALED_API bool ShouldObjectPreventTransactions(const UObject* InObject);
 };
diff --git a/Engine/Source/Editor/UnrealEd/Public/Elements/SMInstance/SMInstanceElementDetailsInterface.h b/Engine/Source/Editor/UnrealEd/Public/Elements/SMInstance/SMInstanceElementDetailsInterface.h
--- a/Engine/Source/Editor/UnrealEd/Public/Elements/SMInstance/SMInstanceElementDetailsInterface.h
+++ b/Engine/Source/Editor/UnrealEd/Public/Elements/SMInstance/SMInstanceElementDetailsInterface.h
@@ -5,11 +5,11 @@
 #include "Elements/Interfaces/TypedElementDetailsInterface.h"
 #include "SMInstanceElementDetailsInterface.generated.h"
 
-UCLASS()
-class UNREALED_API USMInstanceElementDetailsInterface : public UObject, public ITypedElementDetailsInterface
+UCLASS(MinimalAPI)
+class USMInstanceElementDetailsInterface : public UObject, public ITypedElementDetailsInterface
 {
 	GENERATED_BODY()
 
 public:
-	virtual TUniquePtr<ITypedElementDetailsObject> GetDetailsObject(const FTypedElementHandle& InElementHandle) override;
+	UNREALED_API virtual TUniquePtr<ITypedElementDetailsObject> GetDetailsObject(const FTypedElementHandle& InElementHandle) override;
 };
diff --git a/Engine/Source/Editor/UnrealEd/Public/Elements/SMInstance/SMInstanceElementEditorSelectionInterface.h b/Engine/Source/Editor/UnrealEd/Public/Elements/SMInstance/SMInstanceElementEditorSelectionInterface.h
--- a/Engine/Source/Editor/UnrealEd/Public/Elements/SMInstance/SMInstanceElementEditorSelectionInterface.h
+++ b/Engine/Source/Editor/UnrealEd/Public/Elements/SMInstance/SMInstanceElementEditorSelectionInterface.h
@@ -5,15 +5,15 @@
 #include "Elements/SMInstance/SMInstanceElementSelectionInterface.h"
 #include "SMInstanceElementEditorSelectionInterface.generated.h"
 
-UCLASS()
-class UNREALED_API USMInstanceElementEditorSelectionInterface : public USMInstanceElementSelectionInterface
+UCLASS(MinimalAPI)
+class USMInstanceElementEditorSelectionInterface : public USMInstanceElementSelectionInterface
 {
 	GENERATED_BODY()
 
 public:
-	virtual bool IsElementSelected(const FTypedElementHandle& InElementHandle, const FTypedElementListConstPtr& InSelectionSet, const FTypedElementIsSelectedOptions& InSelectionOptions) override;
-	virtual bool SelectElement(const FTypedElementHandle& InElementHandle, const FTypedElementListPtr& InSelectionSet, const FTypedElementSelectionOptions& InSelectionOptions) override;
-	virtual bool DeselectElement(const FTypedElementHandle& InElementHandle, const FTypedElementListPtr& InSelectionSet, const FTypedElementSelectionOptions& InSelectionOptions) override;
-	virtual bool ShouldPreventTransactions(const FTypedElementHandle& InElementHandle) override;
-	virtual TUniquePtr<ITypedElementTransactedElement> CreateTransactedElementImpl() override;
+	UNREALED_API virtual bool IsElementSelected(const FTypedElementHandle& InElementHandle, const FTypedElementListConstPtr& InSelectionSet, const FTypedElementIsSelectedOptions& InSelectionOptions) override;
+	UNREALED_API virtual bool SelectElement(const FTypedElementHandle& InElementHandle, const FTypedElementListPtr& InSelectionSet, const FTypedElementSelectionOptions& InSelectionOptions) override;
+	UNREALED_API virtual bool DeselectElement(const FTypedElementHandle& InElementHandle, const FTypedElementListPtr& InSelectionSet, const FTypedElementSelectionOptions& InSelectionOptions) override;
+	UNREALED_API virtual bool ShouldPreventTransactions(const FTypedElementHandle& InElementHandle) override;
+	UNREALED_API virtual TUniquePtr<ITypedElementTransactedElement> CreateTransactedElementImpl() override;
 };
diff --git a/Engine/Source/Editor/UnrealEd/Public/Elements/SMInstance/SMInstanceElementEditorWorldInterface.h b/Engine/Source/Editor/UnrealEd/Public/Elements/SMInstance/SMInstanceElementEditorWorldInterface.h
--- a/Engine/Source/Editor/UnrealEd/Public/Elements/SMInstance/SMInstanceElementEditorWorldInterface.h
+++ b/Engine/Source/Editor/UnrealEd/Public/Elements/SMInstance/SMInstanceElementEditorWorldInterface.h
@@ -5,20 +5,20 @@
 #include "Elements/SMInstance/SMInstanceElementWorldInterface.h"
 #include "SMInstanceElementEditorWorldInterface.generated.h"
 
-UCLASS()
-class UNREALED_API USMInstanceElementEditorWorldInterface : public USMInstanceElementWorldInterface
+UCLASS(MinimalAPI)
+class USMInstanceElementEditorWorldInterface : public USMInstanceElementWorldInterface
 {
 	GENERATED_BODY()
 
 public:
-	virtual bool CanDeleteElement(const FTypedElementHandle& InElementHandle) override;
-	virtual bool DeleteElement(const FTypedElementHandle& InElementHandle, UWorld* InWorld, UTypedElementSelectionSet* InSelectionSet, const FTypedElementDeletionOptions& InDeletionOptions) override;
-	virtual bool DeleteElements(TArrayView<const FTypedElementHandle> InElementHandles, UWorld* InWorld, UTypedElementSelectionSet* InSelectionSet, const FTypedElementDeletionOptions& InDeletionOptions) override;
-	virtual bool CanDuplicateElement(const FTypedElementHandle& InElementHandle) override;
-	virtual FTypedElementHandle DuplicateElement(const FTypedElementHandle& InElementHandle, UWorld* InWorld, const FVector& InLocationOffset) override;
-	virtual void DuplicateElements(TArrayView<const FTypedElementHandle> InElementHandles, UWorld* InWorld, const FVector& InLocationOffset, TArray<FTypedElementHandle>& OutNewElements) override;
-	virtual bool CanPromoteElement(const FTypedElementHandle& InElementHandle) override;
-	virtual FTypedElementHandle PromoteElement(const FTypedElementHandle& InElementHandle, UWorld* OverrideWorld /* = nullptr */) override;
-	virtual bool IsElementInConvexVolume(const FTypedElementHandle& Handle, const FConvexVolume& InVolume, bool bMustEncompassEntireElement = false);
-	virtual bool IsElementInBox(const FTypedElementHandle& Handle, const FBox& InBox, bool bMustEncompassEntireElement = false);
+	UNREALED_API virtual bool CanDeleteElement(const FTypedElementHandle& InElementHandle) override;
+	UNREALED_API virtual bool DeleteElement(const FTypedElementHandle& InElementHandle, UWorld* InWorld, UTypedElementSelectionSet* InSelectionSet, const FTypedElementDeletionOptions& InDeletionOptions) override;
+	UNREALED_API virtual bool DeleteElements(TArrayView<const FTypedElementHandle> InElementHandles, UWorld* InWorld, UTypedElementSelectionSet* InSelectionSet, const FTypedElementDeletionOptions& InDeletionOptions) override;
+	UNREALED_API virtual bool CanDuplicateElement(const FTypedElementHandle& InElementHandle) override;
+	UNREALED_API virtual FTypedElementHandle DuplicateElement(const FTypedElementHandle& InElementHandle, UWorld* InWorld, const FVector& InLocationOffset) override;
+	UNREALED_API virtual void DuplicateElements(TArrayView<const FTypedElementHandle> InElementHandles, UWorld* InWorld, const FVector& InLocationOffset, TArray<FTypedElementHandle>& OutNewElements) override;
+	UNREALED_API virtual bool CanPromoteElement(const FTypedElementHandle& InElementHandle) override;
+	UNREALED_API virtual FTypedElementHandle PromoteElement(const FTypedElementHandle& InElementHandle, UWorld* OverrideWorld /* = nullptr */) override;
+	UNREALED_API virtual bool IsElementInConvexVolume(const FTypedElementHandle& Handle, const FConvexVolume& InVolume, bool bMustEncompassEntireElement = false);
+	UNREALED_API virtual bool IsElementInBox(const FTypedElementHandle& Handle, const FBox& InBox, bool bMustEncompassEntireElement = false);
 };
diff --git a/Engine/Source/Editor/UnrealEd/Public/Factories.h b/Engine/Source/Editor/UnrealEd/Public/Factories.h
--- a/Engine/Source/Editor/UnrealEd/Public/Factories.h
+++ b/Engine/Source/Editor/UnrealEd/Public/Factories.h
@@ -20,14 +20,14 @@ class UStaticMesh;
  * This class is a simple customizable object factory driven from a text buffer.  
  * Subclasses need to implement CanCreateClass and ProcessConstructedObject.
  */
-class UNREALED_API FCustomizableTextObjectFactory
+class FCustomizableTextObjectFactory
 {
 protected:
 	FFeedbackContext* WarningContext;
 	FObjectInstancingGraph InstanceGraph;
 public:
 	/** Constructor for the factory; takes a context for emitting warnings such as GWarn */
-	FCustomizableTextObjectFactory(FFeedbackContext* InWarningContext);
+	UNREALED_API FCustomizableTextObjectFactory(FFeedbackContext* InWarningContext);
 	virtual ~FCustomizableTextObjectFactory() {}
 
 	/**
@@ -37,24 +37,24 @@ public:
 	 *	@param	Flags		Flags used when creating AnimObjects
 	 *	@param	TextBuffer	Text buffer with descriptions of nodes
 	 */
-	void ProcessBuffer(UObject* InParent, EObjectFlags Flags, const FString& TextBuffer);
+	UNREALED_API void ProcessBuffer(UObject* InParent, EObjectFlags Flags, const FString& TextBuffer);
 
-	void ProcessBuffer(UObject* InParent, EObjectFlags Flags, const TCHAR* TextBuffer);
+	UNREALED_API void ProcessBuffer(UObject* InParent, EObjectFlags Flags, const TCHAR* TextBuffer);
 
 	/**
 	 *	Determine if it is possible to create objects from the specified TextBuffer
 	 */
-	bool CanCreateObjectsFromText(const FString& TextBuffer) const;
+	UNREALED_API bool CanCreateObjectsFromText(const FString& TextBuffer) const;
 
 protected:
 	/** Return true if the an object of type ObjectClass is allowed to be created; If false is returned, the object and subobjects will be ignored. */
-	virtual bool CanCreateClass(UClass* ObjectClass, bool& bOmitSubObjs) const;
+	UNREALED_API virtual bool CanCreateClass(UClass* ObjectClass, bool& bOmitSubObjs) const;
 
 	/** Return true if the an object of type ObjectClass is allowed to be created with a given name under the provided parent; If false is returned, the object and subobjects will be ignored. */
-	virtual bool CanCreateObject(UObject* InParent, UClass* ObjectClass, const FName& InDesiredName) const;
+	UNREALED_API virtual bool CanCreateObject(UObject* InParent, UClass* ObjectClass, const FName& InDesiredName) const;
 
 	/** This is called on each created object after the property text is imported */
-	virtual void ProcessConstructedObject(UObject* CreatedObject);
+	UNREALED_API virtual void ProcessConstructedObject(UObject* CreatedObject);
 
 	/** Post handling of constructed objects by the factory */
 	virtual void PostProcessConstructedObjects() {};
@@ -63,7 +63,7 @@ protected:
 	virtual void UpdateObjectName(UClass* ObjectClass, UObject* InParent, FName& InOutObjName) {};
 
 	/** Util to ensure that InName is a valid name for a new object within InParent. Will rename any existing object within InParent if it is called InName. */
-	static void ClearObjectNameUsage(UObject* InParent, FName InName);
+	static UNREALED_API void ClearObjectNameUsage(UObject* InParent, FName InName);
 
 	/** If we cant do anything with the line ourselves hand off to child class */
 	virtual void ProcessUnidentifiedLine(const FString& StrLine) {}
diff --git a/Engine/Source/Editor/UnrealEd/Public/FbxImporter.h b/Engine/Source/Editor/UnrealEd/Public/FbxImporter.h
--- a/Engine/Source/Editor/UnrealEd/Public/FbxImporter.h
+++ b/Engine/Source/Editor/UnrealEd/Public/FbxImporter.h
@@ -471,7 +471,7 @@ struct FbxSceneInfo
 /**
 * FBX basic data conversion class.
 */
-class UNREALED_API FFbxDataConverter
+class FFbxDataConverter
 {
 public:
 	static void SetJointPostConversionMatrix(FbxAMatrix ConversionMatrix) { JointPostConversionMatrix = ConversionMatrix; }
@@ -481,35 +481,35 @@ public:
 	static const FbxAMatrix &GetAxisConversionMatrix() { return AxisConversionMatrix; }
 	static const FbxAMatrix &GetAxisConversionMatrixInv() { return AxisConversionMatrixInv; }
 
-	static FVector ConvertPos(FbxVector4 Vector);
-	static FVector ConvertDir(FbxVector4 Vector);
-	static FRotator ConvertEuler(FbxDouble3 Euler);
-	static FVector ConvertScale(FbxDouble3 Vector);
-	static FVector ConvertScale(FbxVector4 Vector);
-	static FRotator ConvertRotation(FbxQuaternion Quaternion);
-	static FVector ConvertRotationToFVect(FbxQuaternion Quaternion, bool bInvertRot);
-	static FQuat ConvertRotToQuat(FbxQuaternion Quaternion);
-	static float ConvertDist(FbxDouble Distance);
-	static bool ConvertPropertyValue(FbxProperty& FbxProperty, FProperty& UnrealProperty, union UPropertyValue& OutUnrealPropertyValue);
-	static FTransform ConvertTransform(FbxAMatrix Matrix);
-	static FMatrix ConvertMatrix(const FbxAMatrix& Matrix);
-	static FbxAMatrix ConvertMatrix(const FMatrix& Matrix);
+	static UNREALED_API FVector ConvertPos(FbxVector4 Vector);
+	static UNREALED_API FVector ConvertDir(FbxVector4 Vector);
+	static UNREALED_API FRotator ConvertEuler(FbxDouble3 Euler);
+	static UNREALED_API FVector ConvertScale(FbxDouble3 Vector);
+	static UNREALED_API FVector ConvertScale(FbxVector4 Vector);
+	static UNREALED_API FRotator ConvertRotation(FbxQuaternion Quaternion);
+	static UNREALED_API FVector ConvertRotationToFVect(FbxQuaternion Quaternion, bool bInvertRot);
+	static UNREALED_API FQuat ConvertRotToQuat(FbxQuaternion Quaternion);
+	static UNREALED_API float ConvertDist(FbxDouble Distance);
+	static UNREALED_API bool ConvertPropertyValue(FbxProperty& FbxProperty, FProperty& UnrealProperty, union UPropertyValue& OutUnrealPropertyValue);
+	static UNREALED_API FTransform ConvertTransform(FbxAMatrix Matrix);
+	static UNREALED_API FMatrix ConvertMatrix(const FbxAMatrix& Matrix);
+	static UNREALED_API FbxAMatrix ConvertMatrix(const FMatrix& Matrix);
 
 	/*
 	 * Convert fbx linear space color to sRGB FColor
 	 */
-	static FColor ConvertColor(FbxDouble3 Color);
+	static UNREALED_API FColor ConvertColor(FbxDouble3 Color);
 
-	static FbxVector4 ConvertToFbxPos(FVector Vector);
-	static FbxVector4 ConvertToFbxRot(FVector Vector);
-	static FbxVector4 ConvertToFbxScale(FVector Vector);
+	static UNREALED_API FbxVector4 ConvertToFbxPos(FVector Vector);
+	static UNREALED_API FbxVector4 ConvertToFbxRot(FVector Vector);
+	static UNREALED_API FbxVector4 ConvertToFbxScale(FVector Vector);
 	
 	/*
 	* Convert sRGB FColor to fbx linear space color
 	*/
-	static FbxDouble3   ConvertToFbxColor(FColor Color);
-	static FbxString	ConvertToFbxString(FName Name);
-	static FbxString	ConvertToFbxString(const FString& String);
+	static UNREALED_API FbxDouble3   ConvertToFbxColor(FColor Color);
+	static UNREALED_API FbxString	ConvertToFbxString(FName Name);
+	static UNREALED_API FbxString	ConvertToFbxString(const FString& String);
 
 	// FbxCamera with no rotation faces X with Y-up while ours faces X with Z-up so add a -90 degrees roll to compensate
 	static FRotator GetCameraRotation() { return FRotator(0.f, 0.f, -90.f); }
@@ -518,9 +518,9 @@ public:
 	static FRotator GetLightRotation() { return FRotator(0.f, 90.f, 0.f); }
 
 private:
-	static FbxAMatrix JointPostConversionMatrix;
-	static FbxAMatrix AxisConversionMatrix;
-	static FbxAMatrix AxisConversionMatrixInv;
+	static UNREALED_API FbxAMatrix JointPostConversionMatrix;
+	static UNREALED_API FbxAMatrix AxisConversionMatrix;
+	static UNREALED_API FbxAMatrix AxisConversionMatrixInv;
 };
 
 FBXImportOptions* GetImportOptions( class FFbxImporter* FbxImporter, UFbxImportUI* ImportUI, bool bShowOptionDialog, bool bIsAutomated, const FString& FullPath, bool& OutOperationCanceled, bool& OutImportAll, bool bIsObjFormat, const FString& InFilename, bool bForceImportType = false, EFBXImportType ImportType = FBXIT_StaticMesh);
@@ -1861,10 +1861,10 @@ private:
 
 
 /** message Logger for FBX. Saves all the messages and prints when it's destroyed */
-class UNREALED_API FFbxLogger
+class FFbxLogger
 {
-	FFbxLogger();
-	~FFbxLogger();
+	UNREALED_API FFbxLogger();
+	UNREALED_API ~FFbxLogger();
 
 	/** Error messages **/
 	TArray<TSharedRef<FTokenizedMessage>> TokenizedErrorMessages;
@@ -1881,7 +1881,7 @@ class UNREALED_API FFbxLogger
 * We add this only top level of functions where it needs to be handled
 * if the importer already has logger set, it won't set anymore
 */
-class UNREALED_API FFbxLoggerSetter
+class FFbxLoggerSetter
 {
 	class FFbxLogger Logger;
 	FFbxImporter * Importer;
diff --git a/Engine/Source/Editor/UnrealEd/Public/FeedbackContextEditor.h b/Engine/Source/Editor/UnrealEd/Public/FeedbackContextEditor.h
--- a/Engine/Source/Editor/UnrealEd/Public/FeedbackContextEditor.h
+++ b/Engine/Source/Editor/UnrealEd/Public/FeedbackContextEditor.h
@@ -18,7 +18,7 @@ class SBuildProgressWidget;
 /**
  * A FFeedbackContext implementation for use in UnrealEd.
  */
-class UNREALED_API FFeedbackContextEditor : public FFeedbackContext
+class FFeedbackContextEditor : public FFeedbackContext
 {
 	/** Slate slow task widget */
 	TWeakPtr<class SWindow> SlowTaskWindow;
@@ -33,12 +33,12 @@ class UNREALED_API FFeedbackContextEditor : public FFeedbackContext
 
 public:
 
-	FFeedbackContextEditor();
+	UNREALED_API FFeedbackContextEditor();
 
-	virtual void StartSlowTask( const FText& Task, bool bShowCancelButton=false ) override;
-	virtual void FinalizeSlowTask( ) override;
-	virtual void ProgressReported( const float TotalProgressInterp, FText DisplayMessage ) override;
-	virtual bool IsPlayingInEditor() const override;
+	UNREALED_API virtual void StartSlowTask( const FText& Task, bool bShowCancelButton=false ) override;
+	UNREALED_API virtual void FinalizeSlowTask( ) override;
+	UNREALED_API virtual void ProgressReported( const float TotalProgressInterp, FText DisplayMessage ) override;
+	UNREALED_API virtual bool IsPlayingInEditor() const override;
 
 	void SetContext( FContextSupplier* InSupplier ) override {}
 
@@ -47,9 +47,9 @@ public:
 	 * (i.e. the ongoing slow task or the last one that ran).
 	 * The user cancel flag is reset when starting a new root slow task.
 	 */
-	virtual bool ReceivedUserCancel() override;
+	UNREALED_API virtual bool ReceivedUserCancel() override;
 
-	void OnUserCancel();
+	UNREALED_API void OnUserCancel();
 
 	virtual bool YesNof( const FText& Question ) override
 	{
@@ -60,8 +60,8 @@ public:
 	 * Show the Build Progress Window 
 	 * @return Handle to the Build Progress Widget created
 	 */
-	TWeakPtr<class SBuildProgressWidget> ShowBuildProgressWindow() override;
+	UNREALED_API TWeakPtr<class SBuildProgressWidget> ShowBuildProgressWindow() override;
 	
 	/** Close the Build Progress Window */
-	void CloseBuildProgressWindow() override;
+	UNREALED_API void CloseBuildProgressWindow() override;
 };
diff --git a/Engine/Source/Editor/UnrealEd/Public/FractureToolDelegates.h b/Engine/Source/Editor/UnrealEd/Public/FractureToolDelegates.h
--- a/Engine/Source/Editor/UnrealEd/Public/FractureToolDelegates.h
+++ b/Engine/Source/Editor/UnrealEd/Public/FractureToolDelegates.h
@@ -6,12 +6,12 @@
 #include "Delegates/Delegate.h"
 #include "HAL/Platform.h"
 
-class UNREALED_API FFractureToolDelegates
+class FFractureToolDelegates
 {
 public:
 
 	/** Return a single FFractureToolDelegates object */
-	static FFractureToolDelegates& Get();
+	static UNREALED_API FFractureToolDelegates& Get();
 
 	FSimpleMulticastDelegate OnFractureExpansionBegin;
 	FSimpleMulticastDelegate OnFractureExpansionUpdate;
@@ -28,4 +28,4 @@ public:
 
 	DECLARE_MULTICAST_DELEGATE_OneParam(FOnUpdateVisualizationSettings, bool);
 	FOnUpdateVisualizationSettings OnUpdateVisualizationSettings;
-};
\ No newline at end of file
+};
diff --git a/Engine/Source/Editor/UnrealEd/Public/GPUSkinCacheVisualizationMenuCommands.h b/Engine/Source/Editor/UnrealEd/Public/GPUSkinCacheVisualizationMenuCommands.h
--- a/Engine/Source/Editor/UnrealEd/Public/GPUSkinCacheVisualizationMenuCommands.h
+++ b/Engine/Source/Editor/UnrealEd/Public/GPUSkinCacheVisualizationMenuCommands.h
@@ -11,7 +11,7 @@
 
 class FEditorViewportClient;
 
-class UNREALED_API FGPUSkinCacheVisualizationMenuCommands : public TCommands<FGPUSkinCacheVisualizationMenuCommands>
+class FGPUSkinCacheVisualizationMenuCommands : public TCommands<FGPUSkinCacheVisualizationMenuCommands>
 {
 public:
 	enum class FGPUSkinCacheVisualizationType : uint8
@@ -38,15 +38,15 @@ public:
 	typedef TMultiMap<FName, FGPUSkinCacheVisualizationRecord> TGPUSkinCacheVisualizationModeCommandMap;
 	typedef TGPUSkinCacheVisualizationModeCommandMap::TConstIterator TCommandConstIterator;
 
-	FGPUSkinCacheVisualizationMenuCommands();
+	UNREALED_API FGPUSkinCacheVisualizationMenuCommands();
 
-	TCommandConstIterator CreateCommandConstIterator() const;
+	UNREALED_API TCommandConstIterator CreateCommandConstIterator() const;
 
-	static void BuildVisualisationSubMenu(FMenuBuilder& Menu);
+	static UNREALED_API void BuildVisualisationSubMenu(FMenuBuilder& Menu);
 
-	virtual void RegisterCommands() override;
+	UNREALED_API virtual void RegisterCommands() override;
 
-	void BindCommands(FUICommandList& CommandList, const TSharedPtr<FEditorViewportClient>& Client) const;
+	UNREALED_API void BindCommands(FUICommandList& CommandList, const TSharedPtr<FEditorViewportClient>& Client) const;
 
 	inline bool IsPopulated() const
 	{
@@ -54,11 +54,11 @@ public:
 	}
 
 private:
-	void BuildCommandMap();
-	bool AddCommandTypeToMenu(FMenuBuilder& Menu, const FGPUSkinCacheVisualizationType Type, bool bSeparatorBefore = false) const;
+	UNREALED_API void BuildCommandMap();
+	UNREALED_API bool AddCommandTypeToMenu(FMenuBuilder& Menu, const FGPUSkinCacheVisualizationType Type, bool bSeparatorBefore = false) const;
 
-	static void ChangeGPUSkinCacheVisualizationMode(TWeakPtr<FEditorViewportClient> WeakClient, FName InName);
-	static bool IsGPUSkinCacheVisualizationModeSelected(TWeakPtr<FEditorViewportClient> WeakClient, FName InName);
+	static UNREALED_API void ChangeGPUSkinCacheVisualizationMode(TWeakPtr<FEditorViewportClient> WeakClient, FName InName);
+	static UNREALED_API bool IsGPUSkinCacheVisualizationModeSelected(TWeakPtr<FEditorViewportClient> WeakClient, FName InName);
 
 	TGPUSkinCacheVisualizationModeCommandMap CommandMap;
 };
diff --git a/Engine/Source/Editor/UnrealEd/Public/GraphEditor.h b/Engine/Source/Editor/UnrealEd/Public/GraphEditor.h
--- a/Engine/Source/Editor/UnrealEd/Public/GraphEditor.h
+++ b/Engine/Source/Editor/UnrealEd/Public/GraphEditor.h
@@ -635,7 +635,7 @@ public:
 
 
 	/** Returns the graph panel used for this graph editor */
-	UNREALED_API virtual SGraphPanel* GetGraphPanel() const
+	virtual SGraphPanel* GetGraphPanel() const
 	{
 		if (Implementation.IsValid())
 		{
diff --git a/Engine/Source/Editor/UnrealEd/Public/HierarchicalLOD.h b/Engine/Source/Editor/UnrealEd/Public/HierarchicalLOD.h
--- a/Engine/Source/Editor/UnrealEd/Public/HierarchicalLOD.h
+++ b/Engine/Source/Editor/UnrealEd/Public/HierarchicalLOD.h
@@ -24,8 +24,8 @@ class UHierarchicalLODSetup;
 class ULevel;
 class UWorld;
 
-UCLASS(config = Engine, meta = (DisplayName = "Hierarchical LOD"), defaultconfig)
-class UNREALED_API UHierarchicalLODSettings : public UDeveloperSettings
+UCLASS(config = Engine, meta = (DisplayName = "Hierarchical LOD"), defaultconfig, MinimalAPI)
+class UHierarchicalLODSettings : public UDeveloperSettings
 {
 	GENERATED_UCLASS_BODY()
 
@@ -53,9 +53,9 @@ public:
 	TSoftObjectPtr<class UMaterialInterface> BaseMaterial;
 
 #if WITH_EDITOR
-	static bool IsValidFlattenMaterial(const UMaterialInterface* InBaseMaterial, bool bShowToaster);
+	static UNREALED_API bool IsValidFlattenMaterial(const UMaterialInterface* InBaseMaterial, bool bShowToaster);
 
-	virtual void PostEditChangeProperty(struct FPropertyChangedEvent& PropertyChangedEvent) override;
+	UNREALED_API virtual void PostEditChangeProperty(struct FPropertyChangedEvent& PropertyChangedEvent) override;
 #endif // WITH_EDITOR
 };
 
@@ -65,50 +65,50 @@ public:
  *
  * This builds list of clusters and make sure it's sorted in the order of lower cost to high and merge clusters
  **/
-struct UNREALED_API FHierarchicalLODBuilder
+struct FHierarchicalLODBuilder
 {
-	FHierarchicalLODBuilder(UWorld* InWorld, bool bInPersistentLevelOnly = false);
+	UNREALED_API FHierarchicalLODBuilder(UWorld* InWorld, bool bInPersistentLevelOnly = false);
 
 	/** DO NOT USE. This constructor is for internal usage only for hot-reload purposes. */
-	FHierarchicalLODBuilder();
+	UNREALED_API FHierarchicalLODBuilder();
 
 	/**
 	* Build, Builds the clusters and spawn LODActors with their merged Static Meshes
 	*/
-	void Build();
+	UNREALED_API void Build();
 	
 	/**
 	* PreviewBuild, Builds the clusters and spawns LODActors but without actually creating/merging new StaticMeshes
 	*/
-	void PreviewBuild();
+	UNREALED_API void PreviewBuild();
 
 	/**
 	* Clear all the HLODs and the ALODActors that were created for them
 	*/
-	void ClearHLODs();
+	UNREALED_API void ClearHLODs();
 
 	/**
 	* Clear only the ALODActorsPreview 
 	*/
-	void ClearPreviewBuild();
+	UNREALED_API void ClearPreviewBuild();
 
 	/** Builds the LOD meshes for all LODActors inside of the World's Levels */
-	void BuildMeshesForLODActors(bool bForceAll);
+	UNREALED_API void BuildMeshesForLODActors(bool bForceAll);
 
 	/** Saves HLOD meshes for actors in all the World's levels */
-	void SaveMeshesForActors();
+	UNREALED_API void SaveMeshesForActors();
 
 	/** Get the list of mesh packages to save for a given level */
-	void GetMeshesPackagesToSave(ULevel* InLevel, TSet<UPackage*>& InHLODPackagesToSave, const FString& PreviousLevelName = "");
+	UNREALED_API void GetMeshesPackagesToSave(ULevel* InLevel, TSet<UPackage*>& InHLODPackagesToSave, const FString& PreviousLevelName = "");
 
 	/** Delete HLOD packages that are empty. */
-	void DeleteEmptyHLODPackages(ULevel* InLevel);
+	UNREALED_API void DeleteEmptyHLODPackages(ULevel* InLevel);
 
 	/** 
 	 * @param	bInForce	Whether to force the recalculation of this actor's build flag. If this is false then the cached flag is used an only recalculated every so often.
 	 * @return whether a build is needed (i.e. any LOD actors are dirty) 
 	 */
-	bool NeedsBuild(bool bInForce = false) const;
+	UNREALED_API bool NeedsBuild(bool bInForce = false) const;
 
 	/**
 	* Build a single LOD Actor's mesh
@@ -116,7 +116,7 @@ struct UNREALED_API FHierarchicalLODBuilder
 	* @param LODActor - LODActor to build mesh for
 	* @param LODLevel - LODLevel to build the mesh for
 	*/
-	void BuildMeshForLODActor(ALODActor* LODActor, const uint32 LODLevel);
+	UNREALED_API void BuildMeshForLODActor(ALODActor* LODActor, const uint32 LODLevel);
 
 private:
 	/**
diff --git a/Engine/Source/Editor/UnrealEd/Public/IPIEAuthorizer.h b/Engine/Source/Editor/UnrealEd/Public/IPIEAuthorizer.h
--- a/Engine/Source/Editor/UnrealEd/Public/IPIEAuthorizer.h
+++ b/Engine/Source/Editor/UnrealEd/Public/IPIEAuthorizer.h
@@ -9,7 +9,7 @@
 
 class FString;
 
-class UNREALED_API IPIEAuthorizer : public IModularFeature
+class IPIEAuthorizer : public IModularFeature
 {
 public:
 	virtual ~IPIEAuthorizer() {}
@@ -28,4 +28,4 @@ public:
 	 * @return False to block play-in-editor from continuing, true to allow.
 	 */
 	virtual bool RequestPIEPermission(bool bIsSimulateInEditor, FString& OutReason) const = 0;
-};
\ No newline at end of file
+};
diff --git a/Engine/Source/Editor/UnrealEd/Public/IPackageAutoSaver.h b/Engine/Source/Editor/UnrealEd/Public/IPackageAutoSaver.h
--- a/Engine/Source/Editor/UnrealEd/Public/IPackageAutoSaver.h
+++ b/Engine/Source/Editor/UnrealEd/Public/IPackageAutoSaver.h
@@ -5,7 +5,7 @@
 #include "CoreMinimal.h"
 
 /** An interface to handle the creation, destruction, and restoration of auto-saved packages */
-class UNREALED_API IPackageAutoSaver
+class IPackageAutoSaver
 {
 public:
 	IPackageAutoSaver()
