[UnrealEd]

diff --git a/Engine/Source/Editor/UnrealEd/Public/Kismet2/BlueprintEditorUtils.h b/Engine/Source/Editor/UnrealEd/Public/Kismet2/BlueprintEditorUtils.h
--- a/Engine/Source/Editor/UnrealEd/Public/Kismet2/BlueprintEditorUtils.h
+++ b/Engine/Source/Editor/UnrealEd/Public/Kismet2/BlueprintEditorUtils.h
@@ -57,17 +57,17 @@ namespace EGraphRemoveFlags
 	};
 };
 
-struct UNREALED_API FFunctionFromNodeHelper
+struct FFunctionFromNodeHelper
 {
 	UFunction* const Function;
 	const UK2Node* const Node;
 
-	static UFunction* FunctionFromNode(const UK2Node* Node);
+	static UNREALED_API UFunction* FunctionFromNode(const UK2Node* Node);
 
-	FFunctionFromNodeHelper(const UObject* Obj);
+	UNREALED_API FFunctionFromNodeHelper(const UObject* Obj);
 };
 
-class UNREALED_API FBasePinChangeHelper
+class FBasePinChangeHelper
 {
 public:
 	static bool NodeIsNotTransient(const UK2Node* Node)
@@ -89,31 +89,31 @@ public:
 
 	virtual void EditCreateDelegates(class UK2Node_CreateDelegate* CallSite) {}
 
-	void Broadcast(UBlueprint* InBlueprint, class UK2Node_EditablePinBase* InTargetNode, UEdGraph* Graph);
+	UNREALED_API void Broadcast(UBlueprint* InBlueprint, class UK2Node_EditablePinBase* InTargetNode, UEdGraph* Graph);
 };
 
-class UNREALED_API FParamsChangedHelper : public FBasePinChangeHelper
+class FParamsChangedHelper : public FBasePinChangeHelper
 {
 public:
 	TSet<UBlueprint*> ModifiedBlueprints;
 	TSet<UEdGraph*> ModifiedGraphs;
 
-	virtual void EditCompositeTunnelNode(class UK2Node_Tunnel* TunnelNode) override;
+	UNREALED_API virtual void EditCompositeTunnelNode(class UK2Node_Tunnel* TunnelNode) override;
 
-	virtual void EditMacroInstance(class UK2Node_MacroInstance* MacroInstance, UBlueprint* Blueprint) override;
+	UNREALED_API virtual void EditMacroInstance(class UK2Node_MacroInstance* MacroInstance, UBlueprint* Blueprint) override;
 
-	virtual void EditCallSite(class UK2Node_CallFunction* CallSite, UBlueprint* Blueprint) override;
+	UNREALED_API virtual void EditCallSite(class UK2Node_CallFunction* CallSite, UBlueprint* Blueprint) override;
 
-	virtual void EditDelegates(class UK2Node_BaseMCDelegate* CallSite, UBlueprint* Blueprint) override;
+	UNREALED_API virtual void EditDelegates(class UK2Node_BaseMCDelegate* CallSite, UBlueprint* Blueprint) override;
 
-	virtual void EditCreateDelegates(class UK2Node_CreateDelegate* CallSite) override;
+	UNREALED_API virtual void EditCreateDelegates(class UK2Node_CreateDelegate* CallSite) override;
 
 };
 
-struct UNREALED_API FUCSComponentId
+struct FUCSComponentId
 {
 public:
-	FUCSComponentId(const class UK2Node_AddComponent* UCSNode);
+	UNREALED_API FUCSComponentId(const class UK2Node_AddComponent* UCSNode);
 	FGuid GetAssociatedGuid() const { return GraphNodeGuid; }
 
 private:
@@ -122,7 +122,7 @@ private:
 
 DECLARE_CYCLE_STAT_EXTERN(TEXT("Notify Blueprint Changed"), EKismetCompilerStats_NotifyBlueprintChanged, STATGROUP_KismetCompiler, );
 
-struct UNREALED_API FCompilerRelevantNodeLink
+struct FCompilerRelevantNodeLink
 {
 	UK2Node* Node;
 	UEdGraphPin* LinkedPin;
@@ -137,107 +137,107 @@ struct UNREALED_API FCompilerRelevantNodeLink
 /** Array type for GetCompilerRelevantNodeLinks() */
 typedef TArray<FCompilerRelevantNodeLink, TInlineAllocator<4> > FCompilerRelevantNodeLinkArray;
 
-class UNREALED_API FBlueprintEditorUtils
+class FBlueprintEditorUtils
 {
 public:
 
 	/**
 	 * Schedules and refreshes all nodes in the blueprint, making sure that nodes that affect function signatures get regenerated first
 	 */
-	static void RefreshAllNodes(UBlueprint* Blueprint);
+	static UNREALED_API void RefreshAllNodes(UBlueprint* Blueprint);
 
 	/** Event fired after RefreshAllNodes is called */
 	DECLARE_MULTICAST_DELEGATE_OneParam(FOnRefreshAllNodes, UBlueprint* /*Blueprint*/);
-	static FOnRefreshAllNodes OnRefreshAllNodesEvent;
+	static UNREALED_API FOnRefreshAllNodes OnRefreshAllNodesEvent;
 
 	/**
 	 * Reconstructs all nodes in the blueprint, node reconstruction order determined by FCompareNodePriority.
 	 */
-	static void ReconstructAllNodes(UBlueprint* Blueprint);
+	static UNREALED_API void ReconstructAllNodes(UBlueprint* Blueprint);
 
 	/** Event fired after ReconstructAllNodes is called */
 	DECLARE_MULTICAST_DELEGATE_OneParam(FOnReconstructAllNodes, UBlueprint* /*Blueprint*/);
-	static FOnReconstructAllNodes OnReconstructAllNodesEvent;
+	static UNREALED_API FOnReconstructAllNodes OnReconstructAllNodesEvent;
 
 	/**
 	 * Optimized refresh of nodes that depend on external blueprints.  Refreshes the nodes, but does not recompile the skeleton class
 	 */
-	static void RefreshExternalBlueprintDependencyNodes(UBlueprint* Blueprint, UStruct* RefreshOnlyChild = NULL);
+	static UNREALED_API void RefreshExternalBlueprintDependencyNodes(UBlueprint* Blueprint, UStruct* RefreshOnlyChild = NULL);
 
 	/**
 	 * Refresh the nodes of an individual graph.
 	 * 
 	 * @param	Graph	The graph to refresh.
 	 */
-	static void RefreshGraphNodes(const UEdGraph* Graph);
+	static UNREALED_API void RefreshGraphNodes(const UEdGraph* Graph);
 
 	/**
 	 * Replaces any deprecated nodes with new ones
 	 */
-	static void ReplaceDeprecatedNodes(UBlueprint* Blueprint);
+	static UNREALED_API void ReplaceDeprecatedNodes(UBlueprint* Blueprint);
 
 	/**
 	 * Preloads the object and all the members it owns (nodes, pins, etc)
 	 */
-	static void PreloadMembers(UObject* InObject);
+	static UNREALED_API void PreloadMembers(UObject* InObject);
 
 	/**
 	 * Preloads the construction script, and all templates therein, for the given Blueprint object
 	 */
-	static void PreloadConstructionScript(UBlueprint* Blueprint);
+	static UNREALED_API void PreloadConstructionScript(UBlueprint* Blueprint);
 
 	/**
 	 * Preloads the given construction script, and all templates therein
 	 */
-	static void PreloadConstructionScript(USimpleConstructionScript* SimpleConstructionScript);
+	static UNREALED_API void PreloadConstructionScript(USimpleConstructionScript* SimpleConstructionScript);
 
 	/** 
 	 * Helper function to patch the new CDO into the linker where the old one existed 
 	 */
-	static void PatchNewCDOIntoLinker(UObject* CDO, FLinkerLoad* Linker, int32 ExportIndex, FUObjectSerializeContext* InLoadContext);
+	static UNREALED_API void PatchNewCDOIntoLinker(UObject* CDO, FLinkerLoad* Linker, int32 ExportIndex, FUObjectSerializeContext* InLoadContext);
 
 	/** 
 	 * Procedure used to remove old function implementations and child properties from data only blueprints.
 	 */
-	static void RemoveStaleFunctions(UBlueprintGeneratedClass* Class, UBlueprint* Blueprint);
+	static UNREALED_API void RemoveStaleFunctions(UBlueprintGeneratedClass* Class, UBlueprint* Blueprint);
 
 	/**
 	 *  Synchronizes Blueprint's GeneratedClass's properties with the NewVariable declarations in the blueprint
 	 */
-	static void RefreshVariables(UBlueprint* Blueprint);
+	static UNREALED_API void RefreshVariables(UBlueprint* Blueprint);
 
 	/** Helper function to punch through and honor UAnimGraphNode_Base::PreloadRequiredAssets, which formerly relied on loading assets during compile */
-	static void PreloadBlueprintSpecificData(UBlueprint* Blueprint);
+	static UNREALED_API void PreloadBlueprintSpecificData(UBlueprint* Blueprint);
 
 	/**
 	 * Regenerates the class at class load time, and refreshes the blueprint
 	 */
 	UE_DEPRECATED(5.2, "Replace with FKismetEditorUtilities::CompileBlueprint")
-	static UClass* RegenerateBlueprintClass(UBlueprint* Blueprint, UClass* ClassToRegenerate, UObject* PreviousCDO);
+	static UNREALED_API UClass* RegenerateBlueprintClass(UBlueprint* Blueprint, UClass* ClassToRegenerate, UObject* PreviousCDO);
 	
 	/**
 	 * Links external dependencies
 	 */
-	static void LinkExternalDependencies(UBlueprint* Blueprint);
+	static UNREALED_API void LinkExternalDependencies(UBlueprint* Blueprint);
 
 	/**
 	 * Replace subobjects of CDO in linker
 	 */
-	static void PatchCDOSubobjectsIntoExport(UObject* PreviousCDO, UObject* NewCDO);
+	static UNREALED_API void PatchCDOSubobjectsIntoExport(UObject* PreviousCDO, UObject* NewCDO);
 
 	/** Recreates class meta data */
-	static void RecreateClassMetaData(UBlueprint* Blueprint, UClass* Class, bool bRemoveExistingMetaData);
+	static UNREALED_API void RecreateClassMetaData(UBlueprint* Blueprint, UClass* Class, bool bRemoveExistingMetaData);
 
 	/**
 	 * Copies the default properties of all parent blueprint classes in the chain to the specified blueprint's skeleton CDO
 	 */
-	static void PropagateParentBlueprintDefaults(UClass* ClassToPropagate);
+	static UNREALED_API void PropagateParentBlueprintDefaults(UClass* ClassToPropagate);
 
 	/** Called on a Blueprint after it has been duplicated */
-	static void PostDuplicateBlueprint(UBlueprint* Blueprint, bool bDuplicateForPIE);
+	static UNREALED_API void PostDuplicateBlueprint(UBlueprint* Blueprint, bool bDuplicateForPIE);
 
 	/** Consigns the blueprint's generated classes to oblivion */
-	static void RemoveGeneratedClasses(UBlueprint* Blueprint);
+	static UNREALED_API void RemoveGeneratedClasses(UBlueprint* Blueprint);
 
 	/**
 	 * Helper function to get the blueprint that ultimately owns a node.
@@ -245,7 +245,7 @@ public:
 	 * @param	InNode	Node to find the blueprint for.
 	 * @return	The corresponding blueprint or NULL.
 	 */
-	static UBlueprint* FindBlueprintForNode(const UEdGraphNode* Node);
+	static UNREALED_API UBlueprint* FindBlueprintForNode(const UEdGraphNode* Node);
 
 	/**
 	 * Helper function to get the blueprint that ultimately owns a node.  Cannot fail.
@@ -253,7 +253,7 @@ public:
 	 * @param	InNode	Node to find the blueprint for.
 	 * @return	The corresponding blueprint or NULL.
 	 */
-	static UBlueprint* FindBlueprintForNodeChecked(const UEdGraphNode* Node);
+	static UNREALED_API UBlueprint* FindBlueprintForNodeChecked(const UEdGraphNode* Node);
 
 	/**
 	 * Helper function to get the blueprint that ultimately owns a graph.
@@ -261,7 +261,7 @@ public:
 	 * @param	InGraph	Graph to find the blueprint for.
 	 * @return	The corresponding blueprint or NULL.
 	 */
-	static UBlueprint* FindBlueprintForGraph(const UEdGraph* Graph);
+	static UNREALED_API UBlueprint* FindBlueprintForGraph(const UEdGraph* Graph);
 
 	/**
 	 * Helper function to get the blueprint that ultimately owns a graph.  Cannot fail.
@@ -269,38 +269,38 @@ public:
 	 * @param	InGraph	Graph to find the blueprint for.
 	 * @return	The corresponding blueprint or NULL.
 	 */
-	static UBlueprint* FindBlueprintForGraphChecked(const UEdGraph* Graph);
+	static UNREALED_API UBlueprint* FindBlueprintForGraphChecked(const UEdGraph* Graph);
 
 	/** Helper function to get the SkeletonClass, returns nullptr for UClasses that are not generated by a UBlueprint */
-	static UClass* GetSkeletonClass(UClass* FromClass);
-	static const UClass* GetSkeletonClass(const UClass* FromClass);
+	static UNREALED_API UClass* GetSkeletonClass(UClass* FromClass);
+	static UNREALED_API const UClass* GetSkeletonClass(const UClass* FromClass);
 
 	/** Helper function to get the most up to date class , returns FromClass for native types, SkeletonClass for UBlueprint generated classes */
-	static UClass* GetMostUpToDateClass(UClass* FromClass);
-	static const UClass* GetMostUpToDateClass(const UClass* FromClass);
+	static UNREALED_API UClass* GetMostUpToDateClass(UClass* FromClass);
+	static UNREALED_API const UClass* GetMostUpToDateClass(const UClass* FromClass);
 	
 	/** Looks at the most up to data class and returns whether the given property exists in it as well */
-	static bool PropertyStillExists(FProperty* Property);
+	static UNREALED_API bool PropertyStillExists(FProperty* Property);
 
 	/** Returns the skeleton version of the property, skeleton classes are often more up to date than the authoritative GeneratedClass */
-	static FProperty* GetMostUpToDateProperty(FProperty* Property);
-	static const FProperty* GetMostUpToDateProperty(const FProperty* Property);
+	static UNREALED_API FProperty* GetMostUpToDateProperty(FProperty* Property);
+	static UNREALED_API const FProperty* GetMostUpToDateProperty(const FProperty* Property);
 
-	static UFunction* GetMostUpToDateFunction(UFunction* Function);
-	static const UFunction* GetMostUpToDateFunction(const UFunction* Function);
+	static UNREALED_API UFunction* GetMostUpToDateFunction(UFunction* Function);
+	static UNREALED_API const UFunction* GetMostUpToDateFunction(const UFunction* Function);
 
 	/**
 	 * Updates sources of delegates.
 	 */
-	static void UpdateDelegatesInBlueprint(UBlueprint* Blueprint);
+	static UNREALED_API void UpdateDelegatesInBlueprint(UBlueprint* Blueprint);
 
 	/**
 	 * Whether or not the blueprint should regenerate its class on load or not.  This prevents macros and other BP types not marked for reconstruction from being recompiled all the time
 	 */
-	static bool ShouldRegenerateBlueprint(UBlueprint* Blueprint);
+	static UNREALED_API bool ShouldRegenerateBlueprint(UBlueprint* Blueprint);
 
 	/** Returns true if compilation for the given blueprint has been disabled */
-	static bool IsCompileOnLoadDisabled(UBlueprint* Blueprint);
+	static UNREALED_API bool IsCompileOnLoadDisabled(UBlueprint* Blueprint);
 
 	/**
 	 * Blueprint has structurally changed (added/removed functions, graphs, etc...). Performs the following actions:
@@ -308,7 +308,7 @@ public:
 	 *  - Notifies any observers.
 	 *  - Marks the package as dirty.
 	 */
-	static void MarkBlueprintAsStructurallyModified(UBlueprint* Blueprint);
+	static UNREALED_API void MarkBlueprintAsStructurallyModified(UBlueprint* Blueprint);
 
 	/**
 	 * Blueprint has changed in some manner that invalidates the compiled data (link made/broken, default value changed, etc...)
@@ -318,10 +318,10 @@ public:
 	 * @param	Blueprint				The Blueprint to mark as modified
 	 * @param	PropertyChangedEvent	Used when marking the blueprint as modified due to a changed property (optional)
 	 */
-	static void MarkBlueprintAsModified(UBlueprint* Blueprint, FPropertyChangedEvent PropertyChangedEvent = FPropertyChangedEvent(nullptr));
+	static UNREALED_API void MarkBlueprintAsModified(UBlueprint* Blueprint, FPropertyChangedEvent PropertyChangedEvent = FPropertyChangedEvent(nullptr));
 
 	/** See whether or not the specified graph name / entry point name is unique */
-	static bool IsGraphNameUnique(UObject* InOuter, const FName& InName);
+	static UNREALED_API bool IsGraphNameUnique(UObject* InOuter, const FName& InName);
 
 	/**
 	 * Creates a new empty graph.
@@ -332,7 +332,7 @@ public:
 	 *
 	 * @return	null if it fails, else.
 	 */
-	static class UEdGraph* CreateNewGraph(UObject* ParentScope, const FName& GraphName, TSubclassOf<class UEdGraph> GraphClass, TSubclassOf<class UEdGraphSchema> SchemaClass);
+	static UNREALED_API class UEdGraph* CreateNewGraph(UObject* ParentScope, const FName& GraphName, TSubclassOf<class UEdGraph> GraphClass, TSubclassOf<class UEdGraphSchema> SchemaClass);
 
 	/**
 	 * Creates a new function graph with a signature that matches InNode
@@ -340,7 +340,7 @@ public:
 	 * @param InNode        Node to copy signature from
 	 * @param InSchemaClass The schema for the new graph
 	 */
-	static void CreateMatchingFunction(UK2Node_CallFunction* InNode, TSubclassOf<class UEdGraphSchema> InSchemaClass);
+	static UNREALED_API void CreateMatchingFunction(UK2Node_CallFunction* InNode, TSubclassOf<class UEdGraphSchema> InSchemaClass);
 
 	/**
 	 * Creates a function graph, but does not add it to the blueprint.  If bIsUserCreated is true, the entry/exit nodes will be editable. 
@@ -410,7 +410,7 @@ public:
 	*
 	* @return	True if this function can be converted to a custom event
 	*/
-	static bool IsFunctionConvertableToEvent(UBlueprint* const BlueprintObj, UFunction* const Function);
+	static UNREALED_API bool IsFunctionConvertableToEvent(UBlueprint* const BlueprintObj, UFunction* const Function);
 
 	/**
 	* Get the override class of a given function from its name
@@ -421,29 +421,29 @@ public:
 	* 
 	* @return The override class of a given function
 	*/
-	static UClass* const GetOverrideFunctionClass(UBlueprint* Blueprint, const FName FuncName, UFunction** OutFunction = nullptr);
+	static UNREALED_API UClass* const GetOverrideFunctionClass(UBlueprint* Blueprint, const FName FuncName, UFunction** OutFunction = nullptr);
 
 	/** Adds a macro graph to this blueprint.  If bIsUserCreated is true, the entry/exit nodes will be editable. SignatureFromClass is used to find signature for entry/exit nodes if using an existing signature. */
-	static void AddMacroGraph(UBlueprint* Blueprint, class UEdGraph* Graph,  bool bIsUserCreated, UClass* SignatureFromClass);
+	static UNREALED_API void AddMacroGraph(UBlueprint* Blueprint, class UEdGraph* Graph,  bool bIsUserCreated, UClass* SignatureFromClass);
 
 	/** Adds an interface graph to this blueprint */
-	static void AddInterfaceGraph(UBlueprint* Blueprint, class UEdGraph* Graph, UClass* InterfaceClass);
+	static UNREALED_API void AddInterfaceGraph(UBlueprint* Blueprint, class UEdGraph* Graph, UClass* InterfaceClass);
 
 	/** Adds an ubergraph page to this blueprint */
-	static void AddUbergraphPage(UBlueprint* Blueprint, class UEdGraph* Graph);
+	static UNREALED_API void AddUbergraphPage(UBlueprint* Blueprint, class UEdGraph* Graph);
 
 	/** Returns the name of the Ubergraph Function that the provided blueprint uses */
-	static FName GetUbergraphFunctionName(const UBlueprint* ForBlueprint);
+	static UNREALED_API FName GetUbergraphFunctionName(const UBlueprint* ForBlueprint);
 
 	/** Adds a domain-specific graph to this blueprint */
-	static void AddDomainSpecificGraph(UBlueprint* Blueprint, class UEdGraph* Graph);
+	static UNREALED_API void AddDomainSpecificGraph(UBlueprint* Blueprint, class UEdGraph* Graph);
 
 	/**
 	 * Remove the supplied set of graphs from the Blueprint.
 	 *
 	 * @param	GraphsToRemove	The graphs to remove.
 	 */
-	static void RemoveGraphs( UBlueprint* Blueprint, const TArray<class UEdGraph*>& GraphsToRemove );
+	static UNREALED_API void RemoveGraphs( UBlueprint* Blueprint, const TArray<class UEdGraph*>& GraphsToRemove );
 
 	/**
 	 * Removes the supplied graph from the Blueprint.
@@ -452,7 +452,7 @@ public:
 	 * @param GraphToRemove		The graph to remove.
 	 * @param Flags				Options to control the removal process
 	 */
-	static void RemoveGraph( UBlueprint* Blueprint, class UEdGraph* GraphToRemove, EGraphRemoveFlags::Type Flags = EGraphRemoveFlags::Default );
+	static UNREALED_API void RemoveGraph( UBlueprint* Blueprint, class UEdGraph* GraphToRemove, EGraphRemoveFlags::Type Flags = EGraphRemoveFlags::Default );
 
 	/**
 	 * Tries to rename the supplied graph.
@@ -461,7 +461,7 @@ public:
 	 * @param Graph				The graph to rename.
 	 * @param NewName			The new name for the graph
 	 */
-	static void RenameGraph(class UEdGraph* Graph, const FString& NewName );
+	static UNREALED_API void RenameGraph(class UEdGraph* Graph, const FString& NewName );
 
 	/**
 	 * Renames the graph of the supplied node with a valid name based off of the suggestion.
@@ -469,7 +469,7 @@ public:
 	 * @param GraphNode			The node of the graph to rename.
 	 * @param DesiredName		The initial form of the name to try
 	 */
-	static void RenameGraphWithSuggestion(class UEdGraph* Graph, TSharedPtr<class INameValidatorInterface> NameValidator, const FString& DesiredName );
+	static UNREALED_API void RenameGraphWithSuggestion(class UEdGraph* Graph, TSharedPtr<class INameValidatorInterface> NameValidator, const FString& DesiredName );
 
 	/**
 	 * Removes the supplied node from the Blueprint.
@@ -477,7 +477,7 @@ public:
 	 * @param Node				The node to remove.
 	 * @param bDontRecompile	If true, the blueprint will not be marked as modified, and will not be recompiled.  Useful for if you are removing several node at once, and don't want to recompile each time
 	 */
-	static void RemoveNode (UBlueprint* Blueprint, UEdGraphNode* Node, bool bDontRecompile=false);
+	static UNREALED_API void RemoveNode (UBlueprint* Blueprint, UEdGraphNode* Node, bool bDontRecompile=false);
 
 	/**
 	 * Returns the graph's top level graph (climbing up the hierarchy until there are no more graphs)
@@ -486,16 +486,16 @@ public:
 	 *
 	 * @return				The top level graph
 	 */
-	static UEdGraph* GetTopLevelGraph(const UEdGraph* InGraph);
+	static UNREALED_API UEdGraph* GetTopLevelGraph(const UEdGraph* InGraph);
 
 	/** Determines if the graph is ReadOnly, this differs from editable in that it is never expected to be edited and is in a read-only state */
-	static bool IsGraphReadOnly(UEdGraph* InGraph);
+	static UNREALED_API bool IsGraphReadOnly(UEdGraph* InGraph);
 
 	/** Look to see if an event already exists to override a particular function */
-	static UK2Node_Event* FindOverrideForFunction(const UBlueprint* Blueprint, const UClass* SignatureClass, FName SignatureName);
+	static UNREALED_API UK2Node_Event* FindOverrideForFunction(const UBlueprint* Blueprint, const UClass* SignatureClass, FName SignatureName);
 
 	/** Find the Custom Event if it already exists in the Blueprint */
-	static UK2Node_Event* FindCustomEventNode(const UBlueprint* Blueprint, FName const CustomName);
+	static UNREALED_API UK2Node_Event* FindCustomEventNode(const UBlueprint* Blueprint, FName const CustomName);
 
 	/** Returns all nodes in all graphs of the specified class */
 	template< class T > 
@@ -549,155 +549,155 @@ public:
 	}
 
 	/** Gather all bps that Blueprint depends on */
-	static void GatherDependencies(const UBlueprint* Blueprint, TSet<TWeakObjectPtr<UBlueprint>>& OutDependencies, TSet<TWeakObjectPtr<UStruct>>& OutUDSDependencies);
+	static UNREALED_API void GatherDependencies(const UBlueprint* Blueprint, TSet<TWeakObjectPtr<UBlueprint>>& OutDependencies, TSet<TWeakObjectPtr<UStruct>>& OutUDSDependencies);
 
 	/** Returns cached a list of loaded Blueprints that are dependent on the given Blueprint. */
-	static void GetDependentBlueprints(UBlueprint* Blueprint, TArray<UBlueprint*>& DependentBlueprints);
+	static UNREALED_API void GetDependentBlueprints(UBlueprint* Blueprint, TArray<UBlueprint*>& DependentBlueprints);
 
 	/** Searches the reference graph to find blueprints that are dependent on this BP */
-	static void FindDependentBlueprints(UBlueprint* Blueprint, TArray<UBlueprint*>& DependentBlueprints);
+	static UNREALED_API void FindDependentBlueprints(UBlueprint* Blueprint, TArray<UBlueprint*>& DependentBlueprints);
 
 	/** Ensures, that CachedDependencies in BP are up to date */
-	static void EnsureCachedDependenciesUpToDate(UBlueprint* Blueprint);
+	static UNREALED_API void EnsureCachedDependenciesUpToDate(UBlueprint* Blueprint);
 
 	/** returns if a graph is an intermediate build product */
-	static bool IsGraphIntermediate(const UEdGraph* Graph);
+	static UNREALED_API bool IsGraphIntermediate(const UEdGraph* Graph);
 
 	/** @return true if the blueprint does not contain any special logic or variables or other elements that require a full compile. */
-	static bool IsDataOnlyBlueprint(const UBlueprint* Blueprint);
+	static UNREALED_API bool IsDataOnlyBlueprint(const UBlueprint* Blueprint);
 
 	/** Returns whether or not the blueprint is const during execution */
-	static bool IsBlueprintConst(const UBlueprint* Blueprint);
+	static UNREALED_API bool IsBlueprintConst(const UBlueprint* Blueprint);
 
 	/** Returns whether or not the blueprint is an editor utility blueprint or widget */
-	static bool IsEditorUtilityBlueprint(const UBlueprint* Blueprint);
+	static UNREALED_API bool IsEditorUtilityBlueprint(const UBlueprint* Blueprint);
 
 	/**
 	 * Whether or not this is an actor-based blueprint, and supports features like the uber-graph, components, etc
 	 *
 	 * @return	Whether or not this is an actor based blueprint
 	 */
-	static bool IsActorBased(const UBlueprint* Blueprint);
+	static UNREALED_API bool IsActorBased(const UBlueprint* Blueprint);
 
 	/**
 	 * @return Whether or not this is a component-based blueprint
 	 */
-	static bool IsComponentBased(const UBlueprint* Blueprint);
+	static UNREALED_API bool IsComponentBased(const UBlueprint* Blueprint);
 
 	/**
 	 * Whether or not this blueprint is an interface, used only for defining functions to implement
 	 *
 	 * @return	Whether or not this is an interface blueprint
 	 */
-	static bool IsInterfaceBlueprint(const UBlueprint* Blueprint);
+	static UNREALED_API bool IsInterfaceBlueprint(const UBlueprint* Blueprint);
 
 	/**
 	* Whether or not this graph is an interface graph (i.e. is from an interface blueprint)
 	* 
 	* @return	Whether or not this is an interface graph
 	*/
-	static bool IsInterfaceGraph(const UEdGraph* Graph);
+	static UNREALED_API bool IsInterfaceGraph(const UEdGraph* Graph);
 
 	/**
 	 * Whether or not this blueprint is an interface, used only for defining functions to implement
 	 *
 	 * @return	Whether or not this is a level script blueprint
 	 */
-	static bool IsLevelScriptBlueprint(const UBlueprint* Blueprint);
+	static UNREALED_API bool IsLevelScriptBlueprint(const UBlueprint* Blueprint);
 
 	/** Returns whether the parent class of the specified blueprint is also a blueprint */
-	static bool IsParentClassABlueprint(const UBlueprint* Blueprint);
+	static UNREALED_API bool IsParentClassABlueprint(const UBlueprint* Blueprint);
 
 	/** Returns whether the parent class of the specified blueprint is an editable blueprint */
-	static bool IsParentClassAnEditableBlueprint(const UBlueprint* Blueprint);
+	static UNREALED_API bool IsParentClassAnEditableBlueprint(const UBlueprint* Blueprint);
 
 	/**
 	 * Whether or not this class represents a class generated by an anonymous actor class stored in a level 
 	 *
 	 * @return	Whether or not this is an anonymous blueprint
 	 */
-	static bool IsAnonymousBlueprintClass(const UClass* Class);
+	static UNREALED_API bool IsAnonymousBlueprintClass(const UClass* Class);
 
 	/**
 	 * Checks for events in the argument class
 	 * @param Class	The class to check for events.
 	 */
-	static bool CanClassGenerateEvents(const UClass* Class);
+	static UNREALED_API bool CanClassGenerateEvents(const UClass* Class);
 
 	/**
 	 * If a blueprint is directly tied to a level (level script and anonymous blueprints), this will return a pointer to that level
 	 *
 	 * @return	The level, if any, tied to this blueprint
 	 */
-	static class ULevel* GetLevelFromBlueprint(const UBlueprint* Blueprint);
+	static UNREALED_API class ULevel* GetLevelFromBlueprint(const UBlueprint* Blueprint);
 
 	/** Do we support construction scripts */
-	static bool SupportsConstructionScript(const UBlueprint* Blueprint);
+	static UNREALED_API bool SupportsConstructionScript(const UBlueprint* Blueprint);
 
 	/** Returns the user construction script, if any */
-	static UEdGraph* FindUserConstructionScript(const UBlueprint* Blueprint);
+	static UNREALED_API UEdGraph* FindUserConstructionScript(const UBlueprint* Blueprint);
 
 	/** Returns the event graph, if any */
-	static UEdGraph* FindEventGraph(const UBlueprint* Blueprint);
+	static UNREALED_API UEdGraph* FindEventGraph(const UBlueprint* Blueprint);
 
 	/** Checks if given graph is an event graph */
-	static bool IsEventGraph(const UEdGraph* InGraph);
+	static UNREALED_API bool IsEventGraph(const UEdGraph* InGraph);
 
 	/** Checks if given node is a tunnel instance node */
-	static bool IsTunnelInstanceNode(const UEdGraphNode* InGraphNode);
+	static UNREALED_API bool IsTunnelInstanceNode(const UEdGraphNode* InGraphNode);
 
 	/** See if a class is the one generated by this blueprint */
-	static bool DoesBlueprintDeriveFrom(const UBlueprint* Blueprint, UClass* TestClass);
+	static UNREALED_API bool DoesBlueprintDeriveFrom(const UBlueprint* Blueprint, UClass* TestClass);
 
 	/** See if a field (property, function etc) is part of the blueprint chain, or  */
-	static bool DoesBlueprintContainField(const UBlueprint* Blueprint, UField* TestField);
+	static UNREALED_API bool DoesBlueprintContainField(const UBlueprint* Blueprint, UField* TestField);
 
 	/** Returns whether or not the blueprint supports overriding functions */
-	static bool DoesSupportOverridingFunctions(const UBlueprint* Blueprint);
+	static UNREALED_API bool DoesSupportOverridingFunctions(const UBlueprint* Blueprint);
 
 	/** Returns whether or not the blueprint supports timelines */
-	static bool DoesSupportTimelines(const UBlueprint* Blueprint);
+	static UNREALED_API bool DoesSupportTimelines(const UBlueprint* Blueprint);
 
 	/** Returns whether or not the blueprint supports event graphs*/
-	static bool DoesSupportEventGraphs(const UBlueprint* Blueprint);
+	static UNREALED_API bool DoesSupportEventGraphs(const UBlueprint* Blueprint);
 
 	/** Returns whether or not the blueprint supports implementing interfaces */
-	static bool DoesSupportImplementingInterfaces(const UBlueprint* Blueprint);
+	static UNREALED_API bool DoesSupportImplementingInterfaces(const UBlueprint* Blueprint);
 
 	/** Returns whether or not the blueprint supports components */
-	static bool DoesSupportComponents(UBlueprint const* Blueprint);
+	static UNREALED_API bool DoesSupportComponents(UBlueprint const* Blueprint);
 
 	/** Returns whether or not the blueprint supports default values (IE has a CDO) */
-	static bool DoesSupportDefaults(UBlueprint const* Blueprint);
+	static UNREALED_API bool DoesSupportDefaults(UBlueprint const* Blueprint);
 
 	/** Returns whether or not the blueprint graph supports local variables */
-	static bool DoesSupportLocalVariables(UEdGraph const* InGraph);
+	static UNREALED_API bool DoesSupportLocalVariables(UEdGraph const* InGraph);
 
 	// Returns a descriptive name of the type of blueprint passed in
-	static FString GetBlueprintTypeDescription(const UBlueprint* Blueprint);
+	static UNREALED_API FString GetBlueprintTypeDescription(const UBlueprint* Blueprint);
 
 	/** Constructs a class picker widget for reparenting the specified blueprint(s) */
-	static TSharedRef<SWidget> ConstructBlueprintParentClassPicker( const TArray< UBlueprint* >& Blueprints, const FOnClassPicked& OnPicked);
+	static UNREALED_API TSharedRef<SWidget> ConstructBlueprintParentClassPicker( const TArray< UBlueprint* >& Blueprints, const FOnClassPicked& OnPicked);
 
 	/** Try to open reparent menu for specified blueprint */
-	static void OpenReparentBlueprintMenu( UBlueprint* Blueprint, const TSharedRef<SWidget>& ParentContent, const FOnClassPicked& OnPicked);
-	static void OpenReparentBlueprintMenu( const TArray< UBlueprint* >& Blueprints, const TSharedRef<SWidget>& ParentContent, const FOnClassPicked& OnPicked);
+	static UNREALED_API void OpenReparentBlueprintMenu( UBlueprint* Blueprint, const TSharedRef<SWidget>& ParentContent, const FOnClassPicked& OnPicked);
+	static UNREALED_API void OpenReparentBlueprintMenu( const TArray< UBlueprint* >& Blueprints, const TSharedRef<SWidget>& ParentContent, const FOnClassPicked& OnPicked);
 
 	/** Constructs a class picker widget for adding interfaces for the specified blueprint(s) */
-	static TSharedRef<SWidget> ConstructBlueprintInterfaceClassPicker( const TArray< UBlueprint* >& Blueprints, const FOnClassPicked& OnPicked);
+	static UNREALED_API TSharedRef<SWidget> ConstructBlueprintInterfaceClassPicker( const TArray< UBlueprint* >& Blueprints, const FOnClassPicked& OnPicked);
 
 	/** return find first native class in the hierarchy */
-	static UClass* FindFirstNativeClass(UClass* Class);
+	static UNREALED_API UClass* FindFirstNativeClass(UClass* Class);
 
 	/** returns true if this blueprints signature (inc. visibility) was determined by UHT, rather than the blueprint compiler */
-	static bool IsNativeSignature(const UFunction* Fn);
+	static UNREALED_API bool IsNativeSignature(const UFunction* Fn);
 
 	/**
 	 * Gets the names of all graphs in the Blueprint
 	 *
 	 * @param [in,out]	GraphNames	The graph names will be appended to this array.
 	 */
-	static void GetAllGraphNames(const UBlueprint* Blueprint, TSet<FName>& GraphNames);
+	static UNREALED_API void GetAllGraphNames(const UBlueprint* Blueprint, TSet<FName>& GraphNames);
 
 	/**
 	 * Gets the compiler-relevant (i.e. non-ignorable) node links from the given pin.
@@ -705,7 +705,7 @@ public:
 	 * @param			FromPin			The pin to start searching from.
 	 * @param			OutNodeLinks	Will contain the given pin + owning node if compiler-relevant, or all nodes linked to the owning node at the matching "pass-through" pin that are compiler-relevant. Empty if no compiler-relevant node links can be found from the given pin.
 	 */
-	static void GetCompilerRelevantNodeLinks(UEdGraphPin* FromPin, FCompilerRelevantNodeLinkArray& OutNodeLinks);
+	static UNREALED_API void GetCompilerRelevantNodeLinks(UEdGraphPin* FromPin, FCompilerRelevantNodeLinkArray& OutNodeLinks);
 
 	/**
 	 * Finds the first compiler-relevant (i.e. non-ignorable) node from the given pin.
@@ -714,7 +714,7 @@ public:
 	 *
 	 * @return			The given pin's owning node if compiler-relevant, or the first node linked to the owning node at the matching "pass-through" pin that is compiler-relevant. May be NULL if no compiler-relevant nodes can be found from the given pin.
 	 */
-	static UK2Node* FindFirstCompilerRelevantNode(UEdGraphPin* FromPin);
+	static UNREALED_API UK2Node* FindFirstCompilerRelevantNode(UEdGraphPin* FromPin);
 
 	/**
 	 * Finds the first compiler-relevant (i.e. non-ignorable) node from the given pin and returns the owned pin that's linked.
@@ -723,14 +723,14 @@ public:
 	 *
 	 * @return			The given pin if its owning node is compiler-relevant, or the first pin linked to the owning node at the matching "pass-through" pin that is owned by a compiler-relevant node. May be NULL if no compiler-relevant nodes can be found from the given pin.
 	 */
-	static UEdGraphPin* FindFirstCompilerRelevantLinkedPin(UEdGraphPin* FromPin);
+	static UNREALED_API UEdGraphPin* FindFirstCompilerRelevantLinkedPin(UEdGraphPin* FromPin);
 
 	/**
 	 * Removes all local bookmarks that reference the given Blueprint asset.
 	 *
 	 * @param			ForBlueprint	The Blueprint asset for which to remove local Bookmarks.
 	 */
-	static void RemoveAllLocalBookmarks(const UBlueprint* ForBlueprint);
+	static UNREALED_API void RemoveAllLocalBookmarks(const UBlueprint* ForBlueprint);
 
 	//////////////////////////////////////////////////////////////////////////
 	// Functions
@@ -741,7 +741,7 @@ public:
 	 * @param			Blueprint		The blueprint to check
 	 * @param [in,out]	FunctionNames	List of function names currently in use
 	 */
-	static void GetFunctionNameList(const UBlueprint* Blueprint, TSet<FName>& FunctionNames);
+	static UNREALED_API void GetFunctionNameList(const UBlueprint* Blueprint, TSet<FName>& FunctionNames);
 
 	/**
 	 * Gets a list of delegates names in the blueprint, based on the skeleton class
@@ -749,7 +749,7 @@ public:
 	 * @param			Blueprint		The blueprint to check
 	 * @param [in,out]	DelegatesNames	List of function names currently in use
 	 */
-	static void GetDelegateNameList(const UBlueprint* Blueprint, TSet<FName>& DelegatesNames);
+	static UNREALED_API void GetDelegateNameList(const UBlueprint* Blueprint, TSet<FName>& DelegatesNames);
 
 	/** 
 	 * Get a graph for delegate signature with given name, from given blueprint.
@@ -758,13 +758,13 @@ public:
 	 * @param			DelegateName	Name of delegate.
 	 * @return			Graph of delegate-signature function.
 	 */
-	static UEdGraph* GetDelegateSignatureGraphByName(UBlueprint* Blueprint, FName DelegateName);
+	static UNREALED_API UEdGraph* GetDelegateSignatureGraphByName(UBlueprint* Blueprint, FName DelegateName);
 
 	/** Checks if given graph contains a delegate signature */
-	static bool IsDelegateSignatureGraph(const UEdGraph* Graph);
+	static UNREALED_API bool IsDelegateSignatureGraph(const UEdGraph* Graph);
 
 	/** Checks if given graph is owned by a Math Expression node */
-	static bool IsMathExpressionGraph(const UEdGraph* InGraph);
+	static UNREALED_API bool IsMathExpressionGraph(const UEdGraph* InGraph);
 
 	/**
 	 * Gets a list of pins that should hidden for a given function in a given graph
@@ -774,16 +774,16 @@ public:
 	 * @param [out]		HiddenPins		Set of pins that should be hidden
 	 * @param [out]		OutInternalPins	Subset of hidden pins that are marked for internal use only rather than marked as hidden (optional)
 	 */
-	static void GetHiddenPinsForFunction(UEdGraph const* Graph, UFunction const* Function, TSet<FName>& HiddenPins, TSet<FName>* OutInternalPins = nullptr);
+	static UNREALED_API void GetHiddenPinsForFunction(UEdGraph const* Graph, UFunction const* Function, TSet<FName>& HiddenPins, TSet<FName>* OutInternalPins = nullptr);
 
 	/** Makes sure that calls to parent functions are valid, and removes them if not */
-	static void ConformCallsToParentFunctions(UBlueprint* Blueprint);
+	static UNREALED_API void ConformCallsToParentFunctions(UBlueprint* Blueprint);
 
 	//////////////////////////////////////////////////////////////////////////
 	// Events
 
 	/** Makes sure that all events we handle exist, and replace with custom events if not */
-	static void ConformImplementedEvents(UBlueprint* Blueprint);
+	static UNREALED_API void ConformImplementedEvents(UBlueprint* Blueprint);
 
 	//////////////////////////////////////////////////////////////////////////
 	// Variables
@@ -795,7 +795,7 @@ public:
 	 *
 	 * @return				if type is valid
 	 */
-	static bool IsPinTypeValid(const FEdGraphPinType& Type);
+	static UNREALED_API bool IsPinTypeValid(const FEdGraphPinType& Type);
 
 	/**
 	* Ensures the validity of each pin connection on the given node. Outputs compiler error if invalid
@@ -803,7 +803,7 @@ public:
 	* @param Node			The node to check all linked pins on
 	* @param MessageLog		BP compiler results log to output any error messages to
 	*/
-	static void ValidatePinConnections(const UEdGraphNode* Node, FCompilerResultsLog& MessageLog);
+	static UNREALED_API void ValidatePinConnections(const UEdGraphNode* Node, FCompilerResultsLog& MessageLog);
 
 	/**
 	* If the given node is from an editor only module but is placed in a runtime blueprint
@@ -812,14 +812,14 @@ public:
 	* @param Node			Node to check the outer package on
 	* @param MessageLog		BP Compiler results log to output messages to
 	*/
-	static void ValidateEditorOnlyNodes(const UK2Node* Node, FCompilerResultsLog& MessageLog);
+	static UNREALED_API void ValidateEditorOnlyNodes(const UK2Node* Node, FCompilerResultsLog& MessageLog);
 
 	/**
 	 * Gets the visible class variable list.  This includes both variables introduced here and in all superclasses.
 	 *
 	 * @param [in,out]	VisibleVariables	The visible variables will be appended to this array.
 	 */
-	static void GetClassVariableList(const UBlueprint* Blueprint, TSet<FName>& VisibleVariables, bool bIncludePrivateVars=false);
+	static UNREALED_API void GetClassVariableList(const UBlueprint* Blueprint, TSet<FName>& VisibleVariables, bool bIncludePrivateVars=false);
 
 	/**
 	 * Gets variables of specified type
@@ -827,7 +827,7 @@ public:
 	 * @param 			FEdGraphPinType	 			Type of variables to look for
 	 * @param [in,out]	VisibleVariables			The visible variables will be appended to this array.
 	 */
-	static void GetNewVariablesOfType( const UBlueprint* Blueprint, const FEdGraphPinType& Type, TArray<FName>& OutVars);
+	static UNREALED_API void GetNewVariablesOfType( const UBlueprint* Blueprint, const FEdGraphPinType& Type, TArray<FName>& OutVars);
 
 	/**
 	 * Gets local variables of specified type
@@ -835,7 +835,7 @@ public:
 	 * @param 			FEdGraphPinType	 			Type of variables to look for
 	 * @param [in,out]	VisibleVariables			The visible variables will be appended to this array.
 	 */
-	static void GetLocalVariablesOfType( const UEdGraph* Graph, const FEdGraphPinType& Type, TArray<FName>& OutVars);
+	static UNREALED_API void GetLocalVariablesOfType( const UEdGraph* Graph, const FEdGraphPinType& Type, TArray<FName>& OutVars);
 
 	/**
 	 * Adds a member variable to the blueprint.  It cannot mask a variable in any superclass.
@@ -846,7 +846,7 @@ public:
 	 *
 	 * @return	true if it succeeds, false if it fails.
 	 */
-	static bool AddMemberVariable(UBlueprint* Blueprint, const FName& NewVarName, const FEdGraphPinType& NewVarType, const FString& DefaultValue = FString());
+	static UNREALED_API bool AddMemberVariable(UBlueprint* Blueprint, const FName& NewVarName, const FEdGraphPinType& NewVarType, const FString& DefaultValue = FString());
 
 	/**
 	 * Duplicates a variable from one Blueprint to another blueprint
@@ -857,7 +857,7 @@ public:
 	 *
 	 * @return								Returns the name of the new variable or NAME_None if failed to duplicate
 	 */
-	static FName DuplicateMemberVariable(UBlueprint* InFromBlueprint, UBlueprint* InToBlueprint, FName InVariableToDuplicate);
+	static UNREALED_API FName DuplicateMemberVariable(UBlueprint* InFromBlueprint, UBlueprint* InToBlueprint, FName InVariableToDuplicate);
 
 	/**
 	 * Duplicates a variable given its name and Blueprint
@@ -868,7 +868,7 @@ public:
 	 *
 	 * @return								Returns the name of the new variable or NAME_None if failed to duplicate
 	 */
-	static FName DuplicateVariable(UBlueprint* InBlueprint, const UStruct* InScope, FName InVariableToDuplicate);
+	static UNREALED_API FName DuplicateVariable(UBlueprint* InBlueprint, const UStruct* InScope, FName InVariableToDuplicate);
 
 	/**
 	 * Internal function that deep copies a variable description
@@ -878,28 +878,28 @@ public:
 	 *
 	 * @return								Returns a copy of the passed in variable description.
 	 */
-	static FBPVariableDescription DuplicateVariableDescription(UBlueprint* InBlueprint, FBPVariableDescription& InVariableDescription);
+	static UNREALED_API FBPVariableDescription DuplicateVariableDescription(UBlueprint* InBlueprint, FBPVariableDescription& InVariableDescription);
 
 	/**
 	 * Removes a member variable if it was declared in this blueprint and not in a base class.
 	 *
 	 * @param	VarName	Name of the variable to be removed.
 	 */
-	static void RemoveMemberVariable(UBlueprint* Blueprint, const FName VarName);
+	static UNREALED_API void RemoveMemberVariable(UBlueprint* Blueprint, const FName VarName);
 	
 	/**
 	 * Removes member variables if they were declared in this blueprint and not in a base class.
 	 *
 	 * @param	VarNames	Names of the variable to be removed.
 	 */
-	static void BulkRemoveMemberVariables(UBlueprint* Blueprint, const TArray<FName>& VarNames);
+	static UNREALED_API void BulkRemoveMemberVariables(UBlueprint* Blueprint, const TArray<FName>& VarNames);
 
 	/**
 	 * Removes a field notify variable from the metadata of all other field notify variables and functions.
 	 *
 	 * @param	VarName	Name of the variable to be removed.
 	 */
-	static void RemoveFieldNotifyFromAllMetadata(UBlueprint* Blueprint, const FName VarName);
+	static UNREALED_API void RemoveFieldNotifyFromAllMetadata(UBlueprint* Blueprint, const FName VarName);
 
 	/**
 	 * Removes all unused member variables.
@@ -907,7 +907,7 @@ public:
 	 * @param	OutUsedProperties	The list of used variables in the blueprint
 	 * @param	OutUnusedProperties	The list of unused variables in the blueprint
 	 */
-	static void GetUsedAndUnusedVariables(UBlueprint* Blueprint, TArray<FProperty*>& OutUsedVariables, TArray<FProperty*>& OutUnusedVariables);
+	static UNREALED_API void GetUsedAndUnusedVariables(UBlueprint* Blueprint, TArray<FProperty*>& OutUsedVariables, TArray<FProperty*>& OutUnusedVariables);
 
 	/**
 	 * Finds a member variable Guid using the variable's name
@@ -916,7 +916,7 @@ public:
 	 * @param InVariableGuid	Local variable's name to search for
 	 * @return					The Guid associated with the local variable
 	 */
-	static FGuid FindMemberVariableGuidByName(UBlueprint* InBlueprint, const FName InVariableName);
+	static UNREALED_API FGuid FindMemberVariableGuidByName(UBlueprint* InBlueprint, const FName InVariableName);
 
 	/**
 	 * Finds a member variable name using the variable's Guid
@@ -925,7 +925,7 @@ public:
 	 * @param InVariableGuid	Guid to identify the local variable with
 	 * @return					Local variable's name
 	 */
-	static FName FindMemberVariableNameByGuid(UBlueprint* InBlueprint, const FGuid& InVariableGuid);
+	static UNREALED_API FName FindMemberVariableNameByGuid(UBlueprint* InBlueprint, const FGuid& InVariableGuid);
 
 	/**
 	 * Removes the variable nodes associated with the specified var name
@@ -936,16 +936,16 @@ public:
 	 * @param	LocalGraphScope		Local scope graph of variables
 	 *								false if you just want everything with the specified name removed (variables from other classes too).
 	 */
-	static void RemoveVariableNodes(UBlueprint* Blueprint, const FName VarName, bool const bForSelfOnly = true, UEdGraph* LocalGraphScope = nullptr);
+	static UNREALED_API void RemoveVariableNodes(UBlueprint* Blueprint, const FName VarName, bool const bForSelfOnly = true, UEdGraph* LocalGraphScope = nullptr);
 
 	/**Rename a member variable*/
-	static void RenameMemberVariable(UBlueprint* Blueprint, const FName OldName, const FName NewName);
+	static UNREALED_API void RenameMemberVariable(UBlueprint* Blueprint, const FName OldName, const FName NewName);
 
 	/** Rename a member variable created by a SCS entry */
-	static void RenameComponentMemberVariable(UBlueprint* Blueprint, USCS_Node* Node, const FName NewName);
+	static UNREALED_API void RenameComponentMemberVariable(UBlueprint* Blueprint, USCS_Node* Node, const FName NewName);
 	
 	/** Changes the type of a member variable */
-	static void ChangeMemberVariableType(UBlueprint* Blueprint, const FName VariableName, const FEdGraphPinType& NewPinType);
+	static UNREALED_API void ChangeMemberVariableType(UBlueprint* Blueprint, const FName VariableName, const FEdGraphPinType& NewPinType);
 
 	/**
 	 * Finds the scope's associated graph for local variables (or any passed UFunction)
@@ -953,7 +953,7 @@ public:
 	 * @param	InBlueprint			The Blueprint the local variable can be found in
 	 * @param	InScope				Local variable's scope
 	 */
-	static UEdGraph* FindScopeGraph(const UBlueprint* InBlueprint, const UStruct* InScope);
+	static UNREALED_API UEdGraph* FindScopeGraph(const UBlueprint* InBlueprint, const UStruct* InScope);
 
 	/**
 	 * Adds a local variable to the function graph.  It cannot mask a member variable or a variable in any superclass.
@@ -964,7 +964,7 @@ public:
 	 *
 	 * @return	true if it succeeds, false if it fails.
 	 */
-	static bool AddLocalVariable(UBlueprint* Blueprint, UEdGraph* InTargetGraph, const FName InNewVarName, const FEdGraphPinType& InNewVarType, const FString& DefaultValue = FString());
+	static UNREALED_API bool AddLocalVariable(UBlueprint* Blueprint, UEdGraph* InTargetGraph, const FName InNewVarName, const FEdGraphPinType& InNewVarType, const FString& DefaultValue = FString());
 
 	/**
 	 * Removes a member variable if it was declared in this blueprint and not in a base class.
@@ -973,7 +973,7 @@ public:
 	 * @param	InScope				Local variable's scope
 	 * @param	InVarName			Name of the variable to be removed.
 	 */
-	static void RemoveLocalVariable(UBlueprint* InBlueprint, const UStruct* InScope, const FName InVarName);
+	static UNREALED_API void RemoveLocalVariable(UBlueprint* InBlueprint, const UStruct* InScope, const FName InVarName);
 
 	/**
 	 * Returns a local variable with the function entry it was found in
@@ -982,7 +982,7 @@ public:
 	 * @param InVariableName	Name of the variable to search for
 	 * @return					The local variable description
 	 */
-	static FBPVariableDescription* FindLocalVariable(UBlueprint* InBlueprint, const UStruct* InScope, const FName InVariableName);
+	static UNREALED_API FBPVariableDescription* FindLocalVariable(UBlueprint* InBlueprint, const UStruct* InScope, const FName InVariableName);
 
 	/**
 	 * Returns a local variable
@@ -993,7 +993,7 @@ public:
 	 * @param OutFunctionEntry	Optional output parameter. If not null, the found function entry is returned.
 	 * @return					The local variable description
 	 */
-	static FBPVariableDescription* FindLocalVariable(const UBlueprint* InBlueprint, const UEdGraph* InScopeGraph, const FName InVariableName, class UK2Node_FunctionEntry** OutFunctionEntry = NULL);
+	static UNREALED_API FBPVariableDescription* FindLocalVariable(const UBlueprint* InBlueprint, const UEdGraph* InScopeGraph, const FName InVariableName, class UK2Node_FunctionEntry** OutFunctionEntry = NULL);
 
 	/**
 	 * Returns a local variable
@@ -1004,7 +1004,7 @@ public:
 	 * @param OutFunctionEntry	Optional output parameter. If not null, the found function entry is returned.
 	 * @return					The local variable description
 	 */
-	static FBPVariableDescription* FindLocalVariable(const UBlueprint* InBlueprint, const UStruct* InScope, const FName InVariableName, class UK2Node_FunctionEntry** OutFunctionEntry = NULL);
+	static UNREALED_API FBPVariableDescription* FindLocalVariable(const UBlueprint* InBlueprint, const UStruct* InScope, const FName InVariableName, class UK2Node_FunctionEntry** OutFunctionEntry = NULL);
 
 	/**
 	 * Finds a local variable name using the variable's Guid
@@ -1013,7 +1013,7 @@ public:
 	 * @param InVariableGuid	Guid to identify the local variable with
 	 * @return					Local variable's name
 	 */
-	static FName FindLocalVariableNameByGuid(UBlueprint* InBlueprint, const FGuid& InVariableGuid);
+	static UNREALED_API FName FindLocalVariableNameByGuid(UBlueprint* InBlueprint, const FGuid& InVariableGuid);
 
 	/**
 	 * Finds a local variable Guid using the variable's name
@@ -1023,7 +1023,7 @@ public:
 	 * @param InVariableGuid	Local variable's name to search for
 	 * @return					The Guid associated with the local variable
 	 */
-	static FGuid FindLocalVariableGuidByName(UBlueprint* InBlueprint, const UStruct* InScope, const FName InVariableName);
+	static UNREALED_API FGuid FindLocalVariableGuidByName(UBlueprint* InBlueprint, const UStruct* InScope, const FName InVariableName);
 
 	/**
 	 * Finds a local variable Guid using the variable's name
@@ -1033,7 +1033,7 @@ public:
 	 * @param InVariableGuid	Local variable's name to search for
 	 * @return					The Guid associated with the local variable
 	 */
-	static FGuid FindLocalVariableGuidByName(UBlueprint* InBlueprint, const UEdGraph* InScopeGraph, const FName InVariableName);
+	static UNREALED_API FGuid FindLocalVariableGuidByName(UBlueprint* InBlueprint, const UEdGraph* InScopeGraph, const FName InVariableName);
 
 	/**
 	 * Rename a local variable
@@ -1043,7 +1043,7 @@ public:
 	 * @param InOldName			The name of the local variable to change
 	 * @param InNewName			The new name of the local variable
 	 */
-	static void RenameLocalVariable(UBlueprint* InBlueprint, const UStruct* InScope, const FName InOldName, const FName InNewName);
+	static UNREALED_API void RenameLocalVariable(UBlueprint* InBlueprint, const UStruct* InScope, const FName InOldName, const FName InNewName);
 
 	/**
 	 * Changes the type of a local variable
@@ -1053,16 +1053,16 @@ public:
 	 * @param InVariableName	Name of the local variable to change the type of
 	 * @param InNewPinType		The pin type to change the local variable type to
 	 */
-	static void ChangeLocalVariableType(UBlueprint* InBlueprint, const UStruct* InScope, const FName InVariableName, const FEdGraphPinType& InNewPinType);
+	static UNREALED_API void ChangeLocalVariableType(UBlueprint* InBlueprint, const UStruct* InScope, const FName InVariableName, const FEdGraphPinType& InNewPinType);
 
 	/** Replaces all variable references in the specified blueprint */
-	static void ReplaceVariableReferences(UBlueprint* Blueprint, const FName OldName, const FName NewName);
+	static UNREALED_API void ReplaceVariableReferences(UBlueprint* Blueprint, const FName OldName, const FName NewName);
 
 	/** Replaces all variable references in the specified blueprint */
-	static void ReplaceVariableReferences(UBlueprint* Blueprint, const FProperty* OldVariable, const FProperty* NewVariable);
+	static UNREALED_API void ReplaceVariableReferences(UBlueprint* Blueprint, const FProperty* OldVariable, const FProperty* NewVariable);
 
 	/** Check blueprint variable metadata keys/values for validity and make adjustments if needed */
-	static void FixupVariableDescription(UBlueprint* Blueprint, FBPVariableDescription& VarDesc);
+	static UNREALED_API void FixupVariableDescription(UBlueprint* Blueprint, FBPVariableDescription& VarDesc);
 
 	/**
 	  * Validate child blueprint component member variables, member variables, and timelines, and function graphs against the given variable name
@@ -1071,7 +1071,7 @@ public:
 	  * @param	InVariableName				Variable or function name to validate child blueprints against.
 	  * @param	CustomValidationCallback	Optional callback that allows for doing any post-validation tasks.
 	  */
-	static void ValidateBlueprintChildVariables(UBlueprint* InBlueprint, const FName InVariableName,
+	static UNREALED_API void ValidateBlueprintChildVariables(UBlueprint* InBlueprint, const FName InVariableName,
 		TFunction<void(UBlueprint* InChildBP, const FName InVariableName, bool bValidatedVariable)> PostValidationCallback = TFunction<void(UBlueprint*, FName, bool)>());
 
 	/**
@@ -1082,13 +1082,13 @@ public:
 	 * @param bInRecursive   if true, will return classes derived from child classes as well
 	 * @return Number of child blueprints found
 	 */
-	static int32 GetChildrenOfBlueprint(UBlueprint* InBlueprint, TArray<FAssetData>& OutChildren, bool bInRecursive = true);
+	static UNREALED_API int32 GetChildrenOfBlueprint(UBlueprint* InBlueprint, TArray<FAssetData>& OutChildren, bool bInRecursive = true);
 
 	/** Marks all children of a blueprint as modified */
-	static void MarkBlueprintChildrenAsModified(UBlueprint* InBlueprint);
+	static UNREALED_API void MarkBlueprintChildrenAsModified(UBlueprint* InBlueprint);
 
 	/** Rename a Timeline. If bRenameNodes is true, will also rename any timeline nodes associated with this timeline */
-	static bool RenameTimeline (UBlueprint* Blueprint, const FName OldVarName, const FName NewVarName);
+	static UNREALED_API bool RenameTimeline (UBlueprint* Blueprint, const FName OldVarName, const FName NewVarName);
 
 	/**
 	 * Sets the Blueprint edit-only flag on the variable with the specified name
@@ -1096,7 +1096,7 @@ public:
 	 * @param	VarName				Name of the var to set the flag on
 	 * @param	bNewBlueprintOnly	The new value to set the bitflag to
 	 */
-	static void SetBlueprintOnlyEditableFlag(UBlueprint* Blueprint, const FName& VarName, const bool bNewBlueprintOnly);
+	static UNREALED_API void SetBlueprintOnlyEditableFlag(UBlueprint* Blueprint, const FName& VarName, const bool bNewBlueprintOnly);
 
 	/**
 	 * Sets the Blueprint read-only flag on the variable with the specified name
@@ -1104,7 +1104,7 @@ public:
 	 * @param	VarName				Name of the var to set the flag on
 	 * @param	bVariableReadOnly	The new value to set the bitflag to
 	 */
-	static void SetBlueprintPropertyReadOnlyFlag(UBlueprint* Blueprint, const FName& VarName, const bool bVariableReadOnly);
+	static UNREALED_API void SetBlueprintPropertyReadOnlyFlag(UBlueprint* Blueprint, const FName& VarName, const bool bVariableReadOnly);
 
 	/**
 	 * Sets the Interp flag on the variable with the specified name to make available to sequencer
@@ -1112,7 +1112,7 @@ public:
 	 * @param	VarName				Name of the var to set the flag on
 	 * @param	bInterp	true to make variable available to sequencer, false otherwise
 	 */
-	static void SetInterpFlag(UBlueprint* Blueprint, const FName& VarName, const bool bInterp);
+	static UNREALED_API void SetInterpFlag(UBlueprint* Blueprint, const FName& VarName, const bool bInterp);
 
 	/**
 	 * Sets the Transient flag on the variable with the specified name
@@ -1120,7 +1120,7 @@ public:
 	 * @param	InVarName				Name of the var to set the flag on
 	 * @param	bInIsTransient			The new value to set the bitflag to
 	 */
-	static void SetVariableTransientFlag(UBlueprint* InBlueprint, const FName& InVarName, const bool bInIsTransient);
+	static UNREALED_API void SetVariableTransientFlag(UBlueprint* InBlueprint, const FName& InVarName, const bool bInIsTransient);
 
 	/**
 	 * Sets the Save Game flag on the variable with the specified name
@@ -1128,7 +1128,7 @@ public:
 	 * @param	InVarName				Name of the var to set the flag on
 	 * @param	bInIsSaveGame			The new value to set the bitflag to
 	 */
-	static void SetVariableSaveGameFlag(UBlueprint* InBlueprint, const FName& InVarName, const bool bInIsSaveGame);
+	static UNREALED_API void SetVariableSaveGameFlag(UBlueprint* InBlueprint, const FName& InVarName, const bool bInIsSaveGame);
 
 	/**
 	 * Sets the Advanced Display flag on the variable with the specified name
@@ -1136,7 +1136,7 @@ public:
 	 * @param	InVarName				Name of the var to set the flag on
 	 * @param	bInIsAdvancedDisplay	The new value to set the bitflag to
 	 */
-	static void SetVariableAdvancedDisplayFlag(UBlueprint* InBlueprint, const FName& InVarName, const bool bInIsAdvancedDisplay);
+	static UNREALED_API void SetVariableAdvancedDisplayFlag(UBlueprint* InBlueprint, const FName& InVarName, const bool bInIsAdvancedDisplay);
 
 	/**
 	 * Sets the Deprecated flag on the variable with the specified name
@@ -1144,7 +1144,7 @@ public:
 	 * @param	InVarName				Name of the var to set the flag on
 	 * @param	bInIsDeprecated			The new value to set the bitflag to
 	 */
-	static void SetVariableDeprecatedFlag(UBlueprint* InBlueprint, const FName& InVarName, const bool bInIsDeprecated);
+	static UNREALED_API void SetVariableDeprecatedFlag(UBlueprint* InBlueprint, const FName& InVarName, const bool bInIsDeprecated);
 
 	/** Sets a metadata key/value on the specified variable
 	 *
@@ -1154,7 +1154,7 @@ public:
 	 * @param MetaDataKey			Key name for the metadata to change
 	 * @param MetaDataValue			Value to change the metadata to
 	 */
-	static void SetBlueprintVariableMetaData(UBlueprint* Blueprint, const FName& VarName, const UStruct* InLocalVarScope, const FName& MetaDataKey, const FString& MetaDataValue);
+	static UNREALED_API void SetBlueprintVariableMetaData(UBlueprint* Blueprint, const FName& VarName, const UStruct* InLocalVarScope, const FName& MetaDataKey, const FString& MetaDataValue);
 
 	/** Get a metadata key/value on the specified variable, or timeline if it exists, returning false if it does not exist
 	 *
@@ -1165,7 +1165,7 @@ public:
 	 * @param OutMetaDataValue		Value of the metadata
 	 * @return						TRUE if finding the metadata was successful
 	 */
-	static bool GetBlueprintVariableMetaData(const UBlueprint* Blueprint, const FName& VarName, const UStruct* InLocalVarScope, const FName& MetaDataKey, FString& OutMetaDataValue);
+	static UNREALED_API bool GetBlueprintVariableMetaData(const UBlueprint* Blueprint, const FName& VarName, const UStruct* InLocalVarScope, const FName& MetaDataKey, FString& OutMetaDataValue);
 
 	/** Clear metadata key on specified variable, or timeline
 	 * @param Blueprint				The Blueprint to find the variable in
@@ -1173,7 +1173,7 @@ public:
 	 * @param InLocalVarScope		Local variable's scope, if looking to modify a local variable
 	 * @param MetaDataKey			Key name for the metadata to change
 	 */
-	static void RemoveBlueprintVariableMetaData(UBlueprint* Blueprint, const FName& VarName, const UStruct* InLocalVarScope, const FName& MetaDataKey);
+	static UNREALED_API void RemoveBlueprintVariableMetaData(UBlueprint* Blueprint, const FName& VarName, const UStruct* InLocalVarScope, const FName& MetaDataKey);
 
 	/**
 	 * Sets the custom category on the variable with the specified name.
@@ -1184,7 +1184,7 @@ public:
 	 * @param	VarCategory			The new value of the custom category for the variable
 	 * @param	bDontRecompile		If true, the blueprint will not be marked as modified, and will not be recompiled.  
 	 */
-	static void SetBlueprintVariableCategory(UBlueprint* Blueprint, const FName& VarName, const UStruct* InLocalVarScope, const FText& NewCategory, bool bDontRecompile=false);
+	static UNREALED_API void SetBlueprintVariableCategory(UBlueprint* Blueprint, const FName& VarName, const UStruct* InLocalVarScope, const FText& NewCategory, bool bDontRecompile=false);
 
 
 	/**
@@ -1195,7 +1195,7 @@ public:
 	 * @param	NewCategory			The new value of the custom category for the function
 	 * @param	bDontRecompile		If true, the blueprint will not be marked as modified, and will not be recompiled.  
 	 */
-	static void SetBlueprintFunctionOrMacroCategory(UEdGraph* Graph, const FText& NewCategory, bool bDontRecompile=false);
+	static UNREALED_API void SetBlueprintFunctionOrMacroCategory(UEdGraph* Graph, const FText& NewCategory, bool bDontRecompile=false);
 
 	/** 
 	 * Helper function to grab the root node from an anim graph. 
@@ -1203,7 +1203,7 @@ public:
 	 * @param	InGraph		The graph to check
 	 * @return the one and only root, if this is an anim graph, or null otherwise
 	 */
-	static UAnimGraphNode_Root* GetAnimGraphRoot(UEdGraph* InGraph);
+	static UNREALED_API UAnimGraphNode_Root* GetAnimGraphRoot(UEdGraph* InGraph);
 
 	/**
 	 * Sets the layer group on the anim graph
@@ -1212,13 +1212,13 @@ public:
 	 * @param	InGraph				Graph associated with the layer
 	 * @param	InGroupName			The new value of the group for the layer
 	 */
-	static void SetAnimationGraphLayerGroup(UEdGraph* InGraph, const FText& InGroupName);
+	static UNREALED_API void SetAnimationGraphLayerGroup(UEdGraph* InGraph, const FText& InGroupName);
 
 	/** Finds the index of the specified graph (function or macro) in the parent (if it is not reorderable, then we will return INDEX_NONE) */
-	static int32 FindIndexOfGraphInParent(UEdGraph* Graph);
+	static UNREALED_API int32 FindIndexOfGraphInParent(UEdGraph* Graph);
 
 	/** Reorders the specified graph (function or macro) to be at the new index in the parent (moving whatever was there to be after it), assuming it is reorderable and that is a valid index */
-	static bool MoveGraphBeforeOtherGraph(UEdGraph* Graph, int32 NewIndex, bool bDontRecompile);
+	static UNREALED_API bool MoveGraphBeforeOtherGraph(UEdGraph* Graph, int32 NewIndex, bool bDontRecompile);
 
 	/**
 	 * Gets the custom category on the variable with the specified name.
@@ -1227,22 +1227,22 @@ public:
 	 * @param	InLocalVarScope		Local variable's scope, if looking to modify a local variable
 	 * @return						The custom category (None indicates the name will be the same as the blueprint)
 	 */
-	static FText GetBlueprintVariableCategory(UBlueprint* Blueprint, const FName& VarName, const UStruct* InLocalVarScope);
+	static UNREALED_API FText GetBlueprintVariableCategory(UBlueprint* Blueprint, const FName& VarName, const UStruct* InLocalVarScope);
 
 	/** Gets pointer to PropertyFlags of variable */
-	static uint64* GetBlueprintVariablePropertyFlags(UBlueprint* Blueprint, const FName& VarName);
+	static UNREALED_API uint64* GetBlueprintVariablePropertyFlags(UBlueprint* Blueprint, const FName& VarName);
 
 	/** Gets the variable linked to a RepNotify function, returns nullptr if not found */
-	static FBPVariableDescription* GetVariableFromOnRepFunction(UBlueprint* Blueprint, FName FuncName);
+	static UNREALED_API FBPVariableDescription* GetVariableFromOnRepFunction(UBlueprint* Blueprint, FName FuncName);
 
 	/** Get RepNotify function name of variable */
-	static FName GetBlueprintVariableRepNotifyFunc(UBlueprint* Blueprint, const FName& VarName);
+	static UNREALED_API FName GetBlueprintVariableRepNotifyFunc(UBlueprint* Blueprint, const FName& VarName);
 
 	/** Set RepNotify function of variable */
-	static void SetBlueprintVariableRepNotifyFunc(UBlueprint* Blueprint, const FName& VarName, const FName& RepNotifyFunc);
+	static UNREALED_API void SetBlueprintVariableRepNotifyFunc(UBlueprint* Blueprint, const FName& VarName, const FName& RepNotifyFunc);
 
 	/** Returns TRUE if the variable was created by the Blueprint */
-	static bool IsVariableCreatedByBlueprint(UBlueprint* InBlueprint, FProperty* InVariableProperty);
+	static UNREALED_API bool IsVariableCreatedByBlueprint(UBlueprint* InBlueprint, FProperty* InVariableProperty);
 
 	/**
 	 * Find the index of a variable first declared in this blueprint. Returns INDEX_NONE if not found.
@@ -1251,7 +1251,7 @@ public:
 	 *
 	 * @return	The index of the variable, or INDEX_NONE if it wasn't introduced in this blueprint.
 	 */
-	static int32 FindNewVariableIndex(const UBlueprint* Blueprint, const FName& InName);
+	static UNREALED_API int32 FindNewVariableIndex(const UBlueprint* Blueprint, const FName& InName);
 
 	/**
 	 * Find the index of a variable first declared in this blueprint or its parents. Returns INDEX_NONE if not found.
@@ -1262,7 +1262,7 @@ public:
 	 *
 	 * @return	The index of the variable, or INDEX_NONE if it wasn't introduced in this blueprint.
 	 */
-	static int32 FindNewVariableIndexAndBlueprint(UBlueprint* InBlueprint, FName InName, UBlueprint*& OutFoundBlueprint);
+	static UNREALED_API int32 FindNewVariableIndexAndBlueprint(UBlueprint* InBlueprint, FName InName, UBlueprint*& OutFoundBlueprint);
 
 	/**
 	 * Find the index of a local variable declared in this blueprint. Returns INDEX_NONE if not found.
@@ -1273,10 +1273,10 @@ public:
 	 *
 	 * @return	The index of the variable, or INDEX_NONE if it wasn't introduced in this blueprint.
 	 */
-	static int32 FindLocalVariableIndex(const UBlueprint* Blueprint, UStruct* VariableScope, const FName& InVariableName);
+	static UNREALED_API int32 FindLocalVariableIndex(const UBlueprint* Blueprint, UStruct* VariableScope, const FName& InVariableName);
 
 	/** Change the order of variables in the Blueprint */
-	static bool MoveVariableBeforeVariable(UBlueprint* Blueprint, UStruct* VariableScope, FName VarNameToMove, FName TargetVarName, bool bDontRecompile);
+	static UNREALED_API bool MoveVariableBeforeVariable(UBlueprint* Blueprint, UStruct* VariableScope, FName VarNameToMove, FName TargetVarName, bool bDontRecompile);
 
 	/**
 	 * Find the index of a timeline first declared in this blueprint. Returns INDEX_NONE if not found.
@@ -1285,28 +1285,28 @@ public:
 	 *
 	 * @return	The index of the variable, or INDEX_NONE if it wasn't introduced in this blueprint.
 	 */
-	static int32 FindTimelineIndex(const UBlueprint* Blueprint, const FName& InName);
+	static UNREALED_API int32 FindTimelineIndex(const UBlueprint* Blueprint, const FName& InName);
 
 	/** 
 	 * Gets a list of SCS node variable names for the given blueprint.
 	 *
 	 * @param [in,out]	VariableNames		The list of variable names for the SCS node array.
 	 */
-	static void GetSCSVariableNameList(const UBlueprint* Blueprint, TSet<FName>& VariableNames);
+	static UNREALED_API void GetSCSVariableNameList(const UBlueprint* Blueprint, TSet<FName>& VariableNames);
 
 	/**
 	 * Gets a list of SCS node variable names for the given BPGC.
 	 *
 	 * @param [in,out]	VariableNames		The list of variable names for the SCS node array.
 	 */
-	static void GetSCSVariableNameList(const UBlueprintGeneratedClass* BPGC, TSet<FName>& VariableNames);
+	static UNREALED_API void GetSCSVariableNameList(const UBlueprintGeneratedClass* BPGC, TSet<FName>& VariableNames);
 
 	/**
 	 * Gets a list of SCS node variable names for the given SCS.
 	 *
 	 * @param [in,out]	VariableNames		The list of variable names for the SCS node array.
 	 */
-	static void GetSCSVariableNameList(const USimpleConstructionScript* SCS, TSet<FName>& VariableNames);
+	static UNREALED_API void GetSCSVariableNameList(const USimpleConstructionScript* SCS, TSet<FName>& VariableNames);
 
 	/** 
 	 * Gets a list of function names in blueprints that implement the interface defined by the given blueprint.
@@ -1314,7 +1314,7 @@ public:
 	 * @param [in,out]	Blueprint			The interface blueprint to check.
 	 * @param [in,out]	VariableNames		The list of function names for implementing blueprints.
 	 */
-	static void GetImplementingBlueprintsFunctionNameList(const UBlueprint* Blueprint, TSet<FName>& FunctionNames);
+	static UNREALED_API void GetImplementingBlueprintsFunctionNameList(const UBlueprint* Blueprint, TSet<FName>& FunctionNames);
 
 	/**
 	 * Finds an SCSNode by variable name
@@ -1322,46 +1322,46 @@ public:
 	 * @param	InName	Name of the variable to find.
 	 * @return	The SCS Node that corresponds to the variable name or null if one is not found
 	 */
-	static USCS_Node* FindSCS_Node(const UBlueprint* Blueprint, const FName InName);
+	static UNREALED_API USCS_Node* FindSCS_Node(const UBlueprint* Blueprint, const FName InName);
 
 	/** Returns whether or not the specified member var is a component */
-	static bool IsVariableComponent(const FBPVariableDescription& Variable);
+	static UNREALED_API bool IsVariableComponent(const FBPVariableDescription& Variable);
 
 	/** Indicates if the variable is used on any graphs in this Blueprint */
-	static bool IsVariableUsed(const UBlueprint* VariableBlueprint, const FName& VariableName, const UEdGraph* LocalGraphScope = nullptr);
+	static UNREALED_API bool IsVariableUsed(const UBlueprint* VariableBlueprint, const FName& VariableName, const UEdGraph* LocalGraphScope = nullptr);
 
 	/** Indicates if the function is used on any graphs in this Blueprint */
-	static bool IsFunctionUsed(const UBlueprint* FunctionBlueprint, const FName& FunctionName, const UEdGraph* LocalGraphScope = nullptr);
+	static UNREALED_API bool IsFunctionUsed(const UBlueprint* FunctionBlueprint, const FName& FunctionName, const UEdGraph* LocalGraphScope = nullptr);
 
 	/** 
 	 * Copies the value from the passed in string into a property. ContainerMem points to the Struct or Class containing Property 
 	 * NOTE: This function does not work correctly with static arrays.
 	 */
-	static bool PropertyValueFromString(const FProperty* Property, const FString& StrValue, uint8* Container, UObject* OwningObject = nullptr, int32 PortFlags = PPF_None);
+	static UNREALED_API bool PropertyValueFromString(const FProperty* Property, const FString& StrValue, uint8* Container, UObject* OwningObject = nullptr, int32 PortFlags = PPF_None);
 
 	/** 
 	 * Copies the value from the passed in string into a property. DirectValue is the raw memory address of the property value 
 	 * NOTE: This function does not work correctly with static arrays.
 	 */
-	static bool PropertyValueFromString_Direct(const FProperty* Property, const FString& StrValue, uint8* DirectValue, UObject* OwningObject = nullptr, int32 PortFlags = PPF_None);
+	static UNREALED_API bool PropertyValueFromString_Direct(const FProperty* Property, const FString& StrValue, uint8* DirectValue, UObject* OwningObject = nullptr, int32 PortFlags = PPF_None);
 
 	/** 
 	 * Copies the value from a property into the string OutForm. ContainerMem points to the Struct or Class containing Property 
 	 * NOTE: This function does not work correctly with static arrays.
 	 */
-	static bool PropertyValueToString(const FProperty* Property, const uint8* Container, FString& OutForm, UObject* OwningObject = nullptr, int32 PortFlags = PPF_None);
+	static UNREALED_API bool PropertyValueToString(const FProperty* Property, const uint8* Container, FString& OutForm, UObject* OwningObject = nullptr, int32 PortFlags = PPF_None);
 
 	/** 
 	 * Copies the value from a property into the string OutForm. DirectValue is the raw memory address of the property value 
 	 * NOTE: This function does not work correctly with static arrays.
 	 */
-	static bool PropertyValueToString_Direct(const FProperty* Property, const uint8* DirectValue, FString& OutForm, UObject* OwningObject = nullptr, int32 PortFlags = PPF_None);
+	static UNREALED_API bool PropertyValueToString_Direct(const FProperty* Property, const uint8* DirectValue, FString& OutForm, UObject* OwningObject = nullptr, int32 PortFlags = PPF_None);
 
 	/** Call PostEditChange() on all Actors based on the given Blueprint */
-	static void PostEditChangeBlueprintActors(UBlueprint* Blueprint, bool bComponentEditChange = false);
+	static UNREALED_API void PostEditChangeBlueprintActors(UBlueprint* Blueprint, bool bComponentEditChange = false);
 
 	/** @return whether a property is private or not (whether native or via BP metadata) */
-	static bool IsPropertyPrivate(const FProperty* Property);
+	static UNREALED_API bool IsPropertyPrivate(const FProperty* Property);
 	
 	/** Enumeration of whether a property is writable or if not, why. */
 	enum class EPropertyWritableState : uint8
@@ -1373,7 +1373,7 @@ public:
 	};
 
 	/** Returns an enumeration indicating if the property can be written to by the given Blueprint */
-	static EPropertyWritableState IsPropertyWritableInBlueprint(const UBlueprint* Blueprint, const FProperty* Property);
+	static UNREALED_API EPropertyWritableState IsPropertyWritableInBlueprint(const UBlueprint* Blueprint, const FProperty* Property);
 
 	/** Enumeration of whether a property is readable or if not, why. */
 	enum class EPropertyReadableState : uint8
@@ -1384,16 +1384,16 @@ public:
 	};
 
 	/** Returns an enumeration indicating if the property can be read by the given Blueprint */
-	static EPropertyReadableState IsPropertyReadableInBlueprint(const UBlueprint* Blueprint, const FProperty* Property);
+	static UNREALED_API EPropertyReadableState IsPropertyReadableInBlueprint(const UBlueprint* Blueprint, const FProperty* Property);
 
 	/** Ensures that the CDO root component reference is valid for Actor-based Blueprints */
-	static void UpdateRootComponentReference(UBlueprint* Blueprint);
+	static UNREALED_API void UpdateRootComponentReference(UBlueprint* Blueprint);
 
 	/** Determines if this property is associated with a component that would be displayed in the SCS editor */
-	static bool IsSCSComponentProperty(FObjectProperty* MemberProperty);
+	static UNREALED_API bool IsSCSComponentProperty(FObjectProperty* MemberProperty);
 
 	/** Attempts to match up the FComponentKey with a ComponentTemplate from the Blueprint's UCS. Will fall back to try matching the given template name if the key cannot be used. */
-	static UActorComponent* FindUCSComponentTemplate(const FComponentKey& ComponentKey, const FName& TemplateName);
+	static UNREALED_API UActorComponent* FindUCSComponentTemplate(const FComponentKey& ComponentKey, const FName& TemplateName);
 
 	/** Takes the Blueprint's NativizedFlag property and applies it to the authoritative config (does the same for flagged dependencies) */
 	UE_DEPRECATED(5.0, "Blueprint Nativization has been removed as a supported feature. This API will eventually be removed.")
@@ -1416,7 +1416,7 @@ public:
 	 * @param	GraphName		The graph name to find a GUID for.
 	 * @param	InterfaceClass	The interface's generated class.
 	 */
-	static FGuid FindInterfaceGraphGuid(const FName& GraphName, const UClass* InterfaceClass);
+	static UNREALED_API FGuid FindInterfaceGraphGuid(const FName& GraphName, const UClass* InterfaceClass);
 
 	/** 
 	 * Find the interface Guid for a function if it exists.
@@ -1424,21 +1424,21 @@ public:
 	 * @param	Function		The function to find a graph for.
 	 * @param	InterfaceClass	The interface's generated class.
 	 */
-	static FGuid FindInterfaceFunctionGuid(const UFunction* Function, const UClass* InterfaceClass);
+	static UNREALED_API FGuid FindInterfaceFunctionGuid(const UFunction* Function, const UClass* InterfaceClass);
 
 	/** Add a new interface, and member function graphs to the blueprint */
 	UE_DEPRECATED(5.1, "Short class names are no longer supported. Use a version of this function that takes FTopLevelAssetPath.")
-	static bool ImplementNewInterface(UBlueprint* Blueprint, const FName& InterfaceClassName);
+	static UNREALED_API bool ImplementNewInterface(UBlueprint* Blueprint, const FName& InterfaceClassName);
 
 	/** Add a new interface, and member function graphs to the blueprint */
-	static bool ImplementNewInterface(UBlueprint* Blueprint, FTopLevelAssetPath InterfaceClassPathName);
+	static UNREALED_API bool ImplementNewInterface(UBlueprint* Blueprint, FTopLevelAssetPath InterfaceClassPathName);
 
 	/** Remove an implemented interface, and its associated member function graphs.  If bPreserveFunctions is true, then the interface will move its functions to be normal implemented blueprint functions */
 	UE_DEPRECATED(5.1, "Short class names are no longer supported. Use a version of this function that takes FTopLevelAssetPath.")
-	static void RemoveInterface(UBlueprint* Blueprint, const FName& InterfaceClassName, bool bPreserveFunctions = false);
+	static UNREALED_API void RemoveInterface(UBlueprint* Blueprint, const FName& InterfaceClassName, bool bPreserveFunctions = false);
 	
 	/** Remove an implemented interface, and its associated member function graphs.  If bPreserveFunctions is true, then the interface will move its functions to be normal implemented blueprint functions */
-	static void RemoveInterface(UBlueprint* Blueprint, FTopLevelAssetPath InterfaceClassPathName, bool bPreserveFunctions = false);
+	static UNREALED_API void RemoveInterface(UBlueprint* Blueprint, FTopLevelAssetPath InterfaceClassPathName, bool bPreserveFunctions = false);
 
 	/**
 	* Attempt to remove a function from an interfaces list of function graphs.
@@ -1446,7 +1446,7 @@ public:
 	* 
 	* @return	True if the function was removed from the blueprint
 	*/
-	static bool RemoveInterfaceFunction(UBlueprint* Blueprint, FBPInterfaceDescription& Interface, UFunction* Function, bool bPreserveFunction);
+	static UNREALED_API bool RemoveInterfaceFunction(UBlueprint* Blueprint, FBPInterfaceDescription& Interface, UFunction* Function, bool bPreserveFunction);
 
 	/**
 	* Promotes a Graph from being an Interface Override to a full member function
@@ -1454,14 +1454,14 @@ public:
 	* @param InBlueprint			Blueprint the graph is contained within
 	* @param InInterfaceGraph		The graph to promote
 	*/
-	static void PromoteGraphFromInterfaceOverride(UBlueprint* InBlueprint, UEdGraph* InInterfaceGraph);
+	static UNREALED_API void PromoteGraphFromInterfaceOverride(UBlueprint* InBlueprint, UEdGraph* InInterfaceGraph);
 
 	/** Gets the graphs currently in the blueprint associated with the specified interface */
 	UE_DEPRECATED(5.1, "Short class names are no longer supported. Use a version of this function that takes FTopLevelAssetPath.")
-	static void GetInterfaceGraphs(UBlueprint* Blueprint, const FName& InterfaceClassName, TArray<UEdGraph*>& ChildGraphs);
+	static UNREALED_API void GetInterfaceGraphs(UBlueprint* Blueprint, const FName& InterfaceClassName, TArray<UEdGraph*>& ChildGraphs);
 
 	/** Gets the graphs currently in the blueprint associated with the specified interface */
-	static void GetInterfaceGraphs(UBlueprint* Blueprint, FTopLevelAssetPath InterfaceClassPathName, TArray<UEdGraph*>& ChildGraphs);
+	static UNREALED_API void GetInterfaceGraphs(UBlueprint* Blueprint, FTopLevelAssetPath InterfaceClassPathName, TArray<UEdGraph*>& ChildGraphs);
 
 	/**
 	* Checks if the given function is a part of an interface on this blueprint
@@ -1470,7 +1470,7 @@ public:
 	* @param Function		Function to check if it is an interface or not
 	* @return	True if the given function is implemented as part of an interface
 	*/
-	static bool IsInterfaceFunction(UBlueprint* Blueprint, UFunction* Function);
+	static UNREALED_API bool IsInterfaceFunction(UBlueprint* Blueprint, UFunction* Function);
 
 	/**
 	* Get the corresponding UFunction pointer to the name given on the blueprint.
@@ -1483,46 +1483,46 @@ public:
 	* @return	Corresponding UFunction pointer to the name given; Nullptr if not 
 	*			part of any interfaces
 	*/
-	static UFunction* GetInterfaceFunction(UBlueprint* Blueprint, const FName FuncName);
+	static UNREALED_API UFunction* GetInterfaceFunction(UBlueprint* Blueprint, const FName FuncName);
 
 	/** Makes sure that all graphs for all interfaces we implement exist, and add if not */
-	static void ConformImplementedInterfaces(UBlueprint* Blueprint);
+	static UNREALED_API void ConformImplementedInterfaces(UBlueprint* Blueprint);
 
 	/** Makes sure that all delegate graphs have a corresponding variable declaration, removing the graph if not */
-	static void ConformDelegateSignatureGraphs(UBlueprint* Blueprint);
+	static UNREALED_API void ConformDelegateSignatureGraphs(UBlueprint* Blueprint);
 
 	/** Makes sure that all function graphs are flagged as bAllowDeletion=true, except for construction script and animgraph: */
-	static void ConformAllowDeletionFlag(UBlueprint* Blueprint);
+	static UNREALED_API void ConformAllowDeletionFlag(UBlueprint* Blueprint);
 	
 	/** Makes sure that all NULL graph references are removed from SubGraphs and top-level graph arrays */
-	static void PurgeNullGraphs(UBlueprint* Blueprint);
+	static UNREALED_API void PurgeNullGraphs(UBlueprint* Blueprint);
 
 	/** Handle old AnimBlueprints (state machines in the wrong position, transition graphs with the wrong schema, etc...) */
-	static void UpdateOutOfDateAnimBlueprints(UBlueprint* Blueprint);
+	static UNREALED_API void UpdateOutOfDateAnimBlueprints(UBlueprint* Blueprint);
 
 	/** Handle fixing up composite nodes within the blueprint*/
-	static void UpdateOutOfDateCompositeNodes(UBlueprint* Blueprint);
+	static UNREALED_API void UpdateOutOfDateCompositeNodes(UBlueprint* Blueprint);
 
 	/** Handle fixing up composite nodes within the specified Graph of Blueprint, and correctly setting the Outer */
-	static void UpdateOutOfDateCompositeWithOuter(UBlueprint* Blueprint, UEdGraph* Outer );
+	static UNREALED_API void UpdateOutOfDateCompositeWithOuter(UBlueprint* Blueprint, UEdGraph* Outer );
 
 	/** Handle stale components and ensure correct flags are set */
-	static void UpdateComponentTemplates(UBlueprint* Blueprint);
+	static UNREALED_API void UpdateComponentTemplates(UBlueprint* Blueprint);
 
 	/** Handle stale transactional flags on blueprints */
-	static void UpdateTransactionalFlags(UBlueprint* Blueprint);
+	static UNREALED_API void UpdateTransactionalFlags(UBlueprint* Blueprint);
 
 	/** Handle stale pin watches */
-	static void UpdateStalePinWatches( UBlueprint* Blueprint );
+	static UNREALED_API void UpdateStalePinWatches( UBlueprint* Blueprint );
 
 	/** Updates the cosmetic information cache for macros */
-	static void ClearMacroCosmeticInfoCache(UBlueprint* Blueprint);
+	static UNREALED_API void ClearMacroCosmeticInfoCache(UBlueprint* Blueprint);
 
 	/** Returns the cosmetic information for the specified macro graph, caching it if necessary */
-	static FBlueprintMacroCosmeticInfo GetCosmeticInfoForMacro(UEdGraph* MacroGraph);
+	static UNREALED_API FBlueprintMacroCosmeticInfo GetCosmeticInfoForMacro(UEdGraph* MacroGraph);
 
 	/** Return the first function from implemented interface with given name */
-	static UFunction* FindFunctionInImplementedInterfaces(const UBlueprint* Blueprint, const FName& FunctionName, bool* bOutInvalidInterface = nullptr, bool bGetAllInterfaces = false);
+	static UNREALED_API UFunction* FindFunctionInImplementedInterfaces(const UBlueprint* Blueprint, const FName& FunctionName, bool* bOutInvalidInterface = nullptr, bool bGetAllInterfaces = false);
 
 	/** 
 	 * Build a list of all interface classes either implemented by this blueprint or through inheritance
@@ -1530,7 +1530,7 @@ public:
 	 * @param		bGetAllInterfaces		If true, get all the implemented and inherited. False, just gets the interfaces implemented directly.
 	 * @param [out]	ImplementedInterfaces	The interfaces implemented by this blueprint
 	 */
-	static void FindImplementedInterfaces(const UBlueprint* Blueprint, bool bGetAllInterfaces, TArray<UClass*>& ImplementedInterfaces);
+	static UNREALED_API void FindImplementedInterfaces(const UBlueprint* Blueprint, bool bGetAllInterfaces, TArray<UClass*>& ImplementedInterfaces);
 
 	/**
 	 * Returns true if the interfaces is implemented. It can be implemented directly or inherited.
@@ -1538,7 +1538,7 @@ public:
 	 * @param		bIncludeInherited		If true, check in all the implemented and inherited interfaces. False, check in the interfaces implemented directly.
 	 * @param		SomeInterface			The interface to check.
 	*/
-	static bool ImplementsInterface(const UBlueprint* Blueprint, bool bIncludeInherited, UClass* SomeInterface);
+	static UNREALED_API bool ImplementsInterface(const UBlueprint* Blueprint, bool bIncludeInherited, UClass* SomeInterface);
 
 	/**
 	 * Finds a unique name with a base of the passed in string, appending numbers as needed
@@ -1549,7 +1549,7 @@ public:
 	 *
 	 * @return					A unique name that will not conflict in the Blueprint
 	 */
-	static FName FindUniqueKismetName(const UBlueprint* InBlueprint, const FString& InBaseName, UStruct* InScope = NULL);
+	static UNREALED_API FName FindUniqueKismetName(const UBlueprint* InBlueprint, const FString& InBaseName, UStruct* InScope = NULL);
 	
 	/**
 	 * Cleanses a name of invalid characters and replaces them with '_'.
@@ -1557,7 +1557,7 @@ public:
 	 *
 	 * @param InBaseName The base name to assign
 	 */
-	static void ReplaceInvalidBlueprintNameCharacters(FString& InBaseName);
+	static UNREALED_API void ReplaceInvalidBlueprintNameCharacters(FString& InBaseName);
 	
 	/** Util version of ReplaceInvalidBlueprintNameCharacters that performs the operation inline. */
 	static FString ReplaceInvalidBlueprintNameCharactersInline(FString InBaseName)
@@ -1567,16 +1567,16 @@ public:
 	}
 
 	/** Finds a unique and valid name for a custom event */
-	static FName FindUniqueCustomEventName(const UBlueprint* Blueprint);
+	static UNREALED_API FName FindUniqueCustomEventName(const UBlueprint* Blueprint);
 
 	//////////////////////////////////////////////////////////////////////////
 	// Timeline
 
 	/** Finds a name for a timeline that is not already in use */
-	static FName FindUniqueTimelineName(const UBlueprint* Blueprint);
+	static UNREALED_API FName FindUniqueTimelineName(const UBlueprint* Blueprint);
 
 	/** Add a new timeline with the supplied name to the blueprint */
-	static class UTimelineTemplate* AddNewTimeline(UBlueprint* Blueprint, const FName& TimelineVarName);
+	static UNREALED_API class UTimelineTemplate* AddNewTimeline(UBlueprint* Blueprint, const FName& TimelineVarName);
 
 	/** Remove the timeline from the blueprint 
 	 * @note Just removes the timeline from the associated timelist in the Blueprint. Does not remove the node graph
@@ -1584,22 +1584,22 @@ public:
 	 * @param Timeline			The timeline to remove
 	 * @param bDontRecompile	If true, the blueprint will not be marked as modified, and will not be recompiled.  
 	 */
-	static void RemoveTimeline(UBlueprint* Blueprint, class UTimelineTemplate* Timeline, bool bDontRecompile=false);
+	static UNREALED_API void RemoveTimeline(UBlueprint* Blueprint, class UTimelineTemplate* Timeline, bool bDontRecompile=false);
 
 	/** Find the node that owns the supplied timeline template */
-	static class UK2Node_Timeline* FindNodeForTimeline(UBlueprint* Blueprint, UTimelineTemplate* Timeline);
+	static UNREALED_API class UK2Node_Timeline* FindNodeForTimeline(UBlueprint* Blueprint, UTimelineTemplate* Timeline);
 
 	//////////////////////////////////////////////////////////////////////////
 	// LevelScriptBlueprint
 
 	/** Find how many nodes reference the supplied actor */
-	static bool FindReferencesToActorFromLevelScript(ULevelScriptBlueprint* LevelScriptBlueprint, AActor* InActor, TArray<UK2Node*>& ReferencedToActors);
+	static UNREALED_API bool FindReferencesToActorFromLevelScript(ULevelScriptBlueprint* LevelScriptBlueprint, AActor* InActor, TArray<UK2Node*>& ReferencedToActors);
 
 	/** Replace all references of the old actor with the new actor */
-	static void ReplaceAllActorRefrences(ULevelScriptBlueprint* InLevelScriptBlueprint, AActor* InOldActor, AActor* InNewActor);
+	static UNREALED_API void ReplaceAllActorRefrences(ULevelScriptBlueprint* InLevelScriptBlueprint, AActor* InOldActor, AActor* InNewActor);
 
 	/** Function to call modify() on all graph nodes which reference this actor */
-	static void  ModifyActorReferencedGraphNodes(ULevelScriptBlueprint* LevelScriptBlueprint, const AActor* InActor);
+	static UNREALED_API void  ModifyActorReferencedGraphNodes(ULevelScriptBlueprint* LevelScriptBlueprint, const AActor* InActor);
 
 	/**
 	 * Called after a level script blueprint is changed and nodes should be refreshed for it's new level script actor
@@ -1607,7 +1607,7 @@ public:
 	 * @param	LevelScriptActor	The newly-created level script actor that should be (re-)bound to
 	 * @param	ScriptBlueprint		The level scripting blueprint that contains the bound events to try and bind delegates to this actor for
 	 */
-	static void FixLevelScriptActorBindings(ALevelScriptActor* LevelScriptActor, const class ULevelScriptBlueprint* ScriptBlueprint);
+	static UNREALED_API void FixLevelScriptActorBindings(ALevelScriptActor* LevelScriptActor, const class ULevelScriptBlueprint* ScriptBlueprint);
 
 	/**
 	 * Find how many actors reference the supplied actor
@@ -1616,7 +1616,7 @@ public:
 	 * @param InClassesToIgnore An array of class types to ignore, even if there is an instance of one that references the InActor
 	 * @param OutReferencingActors An array of actors found that reference the specified InActor
 	 */
-	static void FindActorsThatReferenceActor( AActor* InActor, TArray<UClass*>& InClassesToIgnore, TArray<AActor*>& OutReferencingActors );
+	static UNREALED_API void FindActorsThatReferenceActor( AActor* InActor, TArray<UClass*>& InClassesToIgnore, TArray<AActor*>& OutReferencingActors );
 
 	/**
 	 * Go through the world and build a map of all actors that are referenced by other actors.
@@ -1624,27 +1624,27 @@ public:
 	 * @param InClassesToIgnore  An array of class types to ignore, even if there is an instance of one that references another Actor
 	 * @param OutReferencingActors A map of Actors that are referenced by a list of other Actors.
 	*/
-	static void GetActorReferenceMap(UWorld* InWorld, TArray<UClass*>& InClassesToIgnore, TMap<AActor*, TArray<AActor*> >& OutReferencingActors);
+	static UNREALED_API void GetActorReferenceMap(UWorld* InWorld, TArray<UClass*>& InClassesToIgnore, TMap<AActor*, TArray<AActor*> >& OutReferencingActors);
 
 	//////////////////////////////////////////////////////////////////////////
 	// Diagnostics
 
 	// Diagnostic use only: Lists all of the objects have a direct outer of Package
-	static void ListPackageContents(UPackage* Package, FOutputDevice& Ar);
+	static UNREALED_API void ListPackageContents(UPackage* Package, FOutputDevice& Ar);
 
 	// Diagnostic exec commands
-	static bool KismetDiagnosticExec(const TCHAR* Stream, FOutputDevice& Ar);
+	static UNREALED_API bool KismetDiagnosticExec(const TCHAR* Stream, FOutputDevice& Ar);
 
 	/**
 	 * Searches the world for any blueprints that are open and do not have a debug instances set and sets one if possible.
 	 * It will favor a selected instance over a non selected one
 	 */
-	static void FindAndSetDebuggableBlueprintInstances();
+	static UNREALED_API void FindAndSetDebuggableBlueprintInstances();
 
 	/**
 	 * Records node create events for analytics
 	 */
-	static void AnalyticsTrackNewNode( UEdGraphNode* NewNode );
+	static UNREALED_API void AnalyticsTrackNewNode( UEdGraphNode* NewNode );
 
 	/**
 	 * Generates a unique graph name for the supplied blueprint (guaranteed to not 
@@ -1654,7 +1654,7 @@ public:
 	 * @param  ProposedName		The name you want to give the graph (the result will be some permutation of this string).
 	 * @return A unique name intended for a new graph.
 	 */
-	static FName GenerateUniqueGraphName(UObject* const InOuter, FString const& ProposedName);
+	static UNREALED_API FName GenerateUniqueGraphName(UObject* const InOuter, FString const& ProposedName);
 
 	/* Checks if the passed in selection set causes cycling on compile
 	 *
@@ -1663,7 +1663,7 @@ public:
 	 *
 	 * @return						Returns TRUE if the selection does cause cycling
 	 */
-	static bool CheckIfSelectionIsCycling(const TSet<UEdGraphNode*>& InSelectionSet, FCompilerResultsLog& InMessageLog);
+	static UNREALED_API bool CheckIfSelectionIsCycling(const TSet<UEdGraphNode*>& InSelectionSet, FCompilerResultsLog& InMessageLog);
 	
 	/**
 	 * A utility function intended to aid the construction of a specific blueprint 
@@ -1674,7 +1674,7 @@ public:
 	 * @param  BlueprintEditorIn	The blueprint editor owning the palette item you're querying for.
 	 * @return True is the item CANNOT be renamed, false if it can.
 	 */
-	static bool IsPaletteActionReadOnly(TSharedPtr<FEdGraphSchemaAction> ActionIn, TSharedPtr<class FBlueprintEditor> const BlueprintEditorIn);
+	static UNREALED_API bool IsPaletteActionReadOnly(TSharedPtr<FEdGraphSchemaAction> ActionIn, TSharedPtr<class FBlueprintEditor> const BlueprintEditorIn);
 
 	/**
 	 * Finds the entry and result nodes for a function or macro graph
@@ -1683,7 +1683,7 @@ public:
 	 * @param OutEntryNode		The found entry node for the graph
 	 * @param OutResultNode		The found result node for the graph
 	 */
-	static void GetEntryAndResultNodes(const UEdGraph* InGraph, TWeakObjectPtr<class UK2Node_EditablePinBase>& OutEntryNode, TWeakObjectPtr<class UK2Node_EditablePinBase>& OutResultNode);
+	static UNREALED_API void GetEntryAndResultNodes(const UEdGraph* InGraph, TWeakObjectPtr<class UK2Node_EditablePinBase>& OutEntryNode, TWeakObjectPtr<class UK2Node_EditablePinBase>& OutResultNode);
 
 
 	/**
@@ -1692,7 +1692,7 @@ public:
 	 * @param InGraph			The graph to search through
 	 * @return		The found entry node for the graph
 	 */
-	static class UK2Node_EditablePinBase* GetEntryNode(const UEdGraph* InGraph);
+	static UNREALED_API class UK2Node_EditablePinBase* GetEntryNode(const UEdGraph* InGraph);
 
 	/**
 	 * Returns the function meta data block for the graph entry node.
@@ -1700,14 +1700,14 @@ public:
 	 * @param InGraph			The graph to search through
 	 * @return					If valid a pointer to the user declared function meta data structure otherwise nullptr.
 	 */
-	static FKismetUserDeclaredFunctionMetadata* GetGraphFunctionMetaData(const UEdGraph* InGraph);
+	static UNREALED_API FKismetUserDeclaredFunctionMetadata* GetGraphFunctionMetaData(const UEdGraph* InGraph);
 
 	/**
 	 * Modifies the graph entry node that contains the function metadata block, used in metadata transactions.
 	 *
 	 * @param InGraph			The graph to modify
 	 */
-	static void ModifyFunctionMetaData(const UEdGraph* InGraph);
+	static UNREALED_API void ModifyFunctionMetaData(const UEdGraph* InGraph);
 
 	/**
 	 * Returns the description of the graph from the metadata
@@ -1715,10 +1715,10 @@ public:
 	 * @param InGraph			Graph to find the description of
 	 * @return					The description of the graph
 	 */
-	static FText GetGraphDescription(const UEdGraph* InGraph);
+	static UNREALED_API FText GetGraphDescription(const UEdGraph* InGraph);
 
 	/** Checks if a graph (or any sub-graphs or referenced graphs) have latent function nodes */
-	static bool CheckIfGraphHasLatentFunctions(UEdGraph* InGraph);
+	static UNREALED_API bool CheckIfGraphHasLatentFunctions(UEdGraph* InGraph);
 
 	/**
 	 * Creates a function result node or returns the current one if one exists
@@ -1726,7 +1726,7 @@ public:
 	 * @param InFunctionEntryNode		The function entry node to spawn the result node for
 	 * @return							Spawned result node
 	 */
-	static class UK2Node_FunctionResult* FindOrCreateFunctionResultNode(class UK2Node_EditablePinBase* InFunctionEntryNode);
+	static UNREALED_API class UK2Node_FunctionResult* FindOrCreateFunctionResultNode(class UK2Node_EditablePinBase* InFunctionEntryNode);
 
 	/** 
 	 * Determine the best icon to represent the given pin.
@@ -1734,27 +1734,27 @@ public:
 	 * @param PinType		The pin get the icon for.
 	 * @param returns a brush that best represents the icon (or Kismet.VariableList.TypeIcon if none is available )
 	 */
-	static const struct FSlateBrush* GetIconFromPin(const FEdGraphPinType& PinType, bool bIsLarge = false);
+	static UNREALED_API const struct FSlateBrush* GetIconFromPin(const FEdGraphPinType& PinType, bool bIsLarge = false);
 
 	/**
 	 * Determine the best secondary icon icon to represent the given pin.
 	 */
-	static const struct FSlateBrush* GetSecondaryIconFromPin(const FEdGraphPinType& PinType);
+	static UNREALED_API const struct FSlateBrush* GetSecondaryIconFromPin(const FEdGraphPinType& PinType);
 
 	/**
 	 * Returns true if this terminal type can be hashed (native types need GetTypeHash, script types are always hashable).
 	 */
-	static bool HasGetTypeHash(const FEdGraphPinType& PinType);
+	static UNREALED_API bool HasGetTypeHash(const FEdGraphPinType& PinType);
 
 	/**
 	 * Returns true if this type of FProperty can be hashed. Matches native constructors of FNumericProperty, etc.
 	 */
-	static bool PropertyHasGetTypeHash(const FProperty* PropertyType);
+	static UNREALED_API bool PropertyHasGetTypeHash(const FProperty* PropertyType);
 
 	/**
 	 * Returns true if the StructType is native and has a GetTypeHash or is non-native and all of its member types are handled by UScriptStruct::GetStructTypeHash
 	 */
-	static bool StructHasGetTypeHash(const UScriptStruct* StructType);
+	static UNREALED_API bool StructHasGetTypeHash(const UScriptStruct* StructType);
 
 	/**
 	 * Generate component instancing data (for cooked builds).
@@ -1764,7 +1764,7 @@ public:
 	 * @param bUseTemplateArchetype	Whether or not to use the template archetype or the template CDO for delta serialization (default is to use the template CDO).
 	 * @return						TRUE if component instancing data was built, FALSE otherwise.
 	 */
-	static void BuildComponentInstancingData(UActorComponent* ComponentTemplate, FBlueprintCookedComponentInstancingData& OutData, bool bUseTemplateArchetype = false);
+	static UNREALED_API void BuildComponentInstancingData(UActorComponent* ComponentTemplate, FBlueprintCookedComponentInstancingData& OutData, bool bUseTemplateArchetype = false);
 
 	/**
 	 * Callback for when a node has been found
@@ -1774,34 +1774,34 @@ public:
 	/**
 	 * Search the blueprints looking for nodes that contain the given script structs
 	 */
-	static void FindScriptStructsInNodes(const TSet<UScriptStruct*>& Structs, FOnNodeFoundOrUpdated InOnNodeFoundOrUpdated);
+	static UNREALED_API void FindScriptStructsInNodes(const TSet<UScriptStruct*>& Structs, FOnNodeFoundOrUpdated InOnNodeFoundOrUpdated);
 
 	/**
 	 * Search the blueprints looking for nodes that contain the given enumerations
 	 */
-	static void FindEnumsInNodes(const TSet<UEnum*>& UEnums, FOnNodeFoundOrUpdated InOnNodeFoundOrUpdated);
+	static UNREALED_API void FindEnumsInNodes(const TSet<UEnum*>& UEnums, FOnNodeFoundOrUpdated InOnNodeFoundOrUpdated);
 
 	/**
 	 * Search the blueprints looking for nodes that contain the given script structs and replace the references
 	 */
-	static void UpdateScriptStructsInNodes(const TMap<UScriptStruct*, UScriptStruct*>& Structs, FOnNodeFoundOrUpdated InOnNodeFoundOrUpdated);
+	static UNREALED_API void UpdateScriptStructsInNodes(const TMap<UScriptStruct*, UScriptStruct*>& Structs, FOnNodeFoundOrUpdated InOnNodeFoundOrUpdated);
 
 	/**
 	 * Search the blueprints looking for nodes that contain the given enumerations and replace the references
 	 */
-	static void UpdateEnumsInNodes(const TMap<UEnum*, UEnum*>& Structs, FOnNodeFoundOrUpdated InOnNodeFoundOrUpdated);
+	static UNREALED_API void UpdateEnumsInNodes(const TMap<UEnum*, UEnum*>& Structs, FOnNodeFoundOrUpdated InOnNodeFoundOrUpdated);
 
 	/**
 	 * Recombine any nested subpins on this node
 	 */
-	static void RecombineNestedSubPins(UK2Node* Node);
+	static UNREALED_API void RecombineNestedSubPins(UK2Node* Node);
 
 protected:
 	// Removes all NULL graph references from the SubGraphs array and recurses thru the non-NULL ones
-	static void CleanNullGraphReferencesRecursive(UEdGraph* Graph);
+	static UNREALED_API void CleanNullGraphReferencesRecursive(UEdGraph* Graph);
 
 	// Removes all NULL graph references in the specified array
-	static void CleanNullGraphReferencesInArray(UBlueprint* Blueprint, TArray<UEdGraph*>& GraphArray);
+	static UNREALED_API void CleanNullGraphReferencesInArray(UBlueprint* Blueprint, TArray<UEdGraph*>& GraphArray);
 
 	/**
 	 * Checks that the actor type matches the blueprint type (or optionally is BASED on the same type. 
@@ -1810,19 +1810,19 @@ protected:
 	 * @param InBlueprint						The blueprint to check against
 	 * @param bInDisallowDerivedBlueprints		if true will only allow exact type matches, if false derived types are allowed.
 	 */
-	static bool IsObjectADebugCandidate( AActor* InActorObject, UBlueprint* InBlueprint , bool bInDisallowDerivedBlueprints );
+	static UNREALED_API bool IsObjectADebugCandidate( AActor* InActorObject, UBlueprint* InBlueprint , bool bInDisallowDerivedBlueprints );
 
 	/** Validate child blueprint member variables against the given variable name */
-	static bool ValidateAllMemberVariables(UBlueprint* InBlueprint, UBlueprint* InParentBlueprint, const FName InVariableName);
+	static UNREALED_API bool ValidateAllMemberVariables(UBlueprint* InBlueprint, UBlueprint* InParentBlueprint, const FName InVariableName);
 
 	/** Validate child blueprint component member variables against the given variable name */
-	static bool ValidateAllComponentMemberVariables(UBlueprint* InBlueprint, UBlueprint* InParentBlueprint, const FName& InVariableName);
+	static UNREALED_API bool ValidateAllComponentMemberVariables(UBlueprint* InBlueprint, UBlueprint* InParentBlueprint, const FName& InVariableName);
 
 	/** Validates all timelines of the passed blueprint against the given variable name */
-	static bool ValidateAllTimelines(UBlueprint* InBlueprint, UBlueprint* InParentBlueprint, const FName& InVariableName);
+	static UNREALED_API bool ValidateAllTimelines(UBlueprint* InBlueprint, UBlueprint* InParentBlueprint, const FName& InVariableName);
 
 	/** Validates all function graphs of the passed blueprint against the given variable name */
-	static bool ValidateAllFunctionGraphs(UBlueprint* InBlueprint, UBlueprint* InParentBlueprint, const FName& InVariableName);
+	static UNREALED_API bool ValidateAllFunctionGraphs(UBlueprint* InBlueprint, UBlueprint* InParentBlueprint, const FName& InVariableName);
 
 	/**
 	 * Checks if the passed node connects to the selection set
@@ -1832,7 +1832,7 @@ protected:
 	 *
 	 * @return						Returns TRUE if the node does connect to the selection set
 	 */
-	static bool CheckIfNodeConnectsToSelection(UEdGraphNode* InNode, const TSet<UEdGraphNode*>& InSelectionSet);
+	static UNREALED_API bool CheckIfNodeConnectsToSelection(UEdGraphNode* InNode, const TSet<UEdGraphNode*>& InSelectionSet);
 
 	/**
 	 * Returns an array of variables Get/Set nodes of the current variable
@@ -1842,7 +1842,7 @@ protected:
 	 * @param InScope		Option scope for local variables
 	 * @return				Array of variable nodes
 	 */
-	static TArray<UK2Node*> GetNodesForVariable(const FName& InVarName, const UBlueprint* InBlueprint, const UStruct* InScope = nullptr);
+	static UNREALED_API TArray<UK2Node*> GetNodesForVariable(const FName& InVarName, const UBlueprint* InBlueprint, const UStruct* InScope = nullptr);
 
 	/**
 	 * Helper function to warn user of the results of changing var type by displaying a suppressible dialog
@@ -1850,7 +1850,7 @@ protected:
 	 * @param InVarName		Variable name to display in the dialog message
 	 * @return				TRUE if the user wants to change the variable type
 	 */
-	static bool VerifyUserWantsVariableTypeChanged(const FName& InVarName);
+	static UNREALED_API bool VerifyUserWantsVariableTypeChanged(const FName& InVarName);
 
 	/**
 	 * Helper function to warn user of the results of changing a RepNotify variable name by displaying a suppressible dialog
@@ -1859,7 +1859,7 @@ protected:
 	 * @param InFuncName	Associated OnRep function name to display in the dialog message
 	 * @return				TRUE if the user wants to change the variable name
 	 */
-	static bool VerifyUserWantsRepNotifyVariableNameChanged(const FName& InVarName, const FName& InFuncName);
+	static UNREALED_API bool VerifyUserWantsRepNotifyVariableNameChanged(const FName& InVarName, const FName& InFuncName);
 
 	/**
 	 * Helper function to get all loaded Blueprints that are children (or using as an interface) the passed Blueprint
@@ -1867,19 +1867,19 @@ protected:
 	 * @param InBlueprint		Blueprint to find children of
 	 * @param OutBlueprints		Filled out with child Blueprints
 	 */
-	static void GetLoadedChildBlueprints(UBlueprint* InBlueprint, TArray<UBlueprint*>& OutBlueprints);
+	static UNREALED_API void GetLoadedChildBlueprints(UBlueprint* InBlueprint, TArray<UBlueprint*>& OutBlueprints);
 
 	/**
 	 * Validates flags and settings on object pins, keeping them from being given default values and from being in invalid states
 	 *
 	 * @param InOutVarDesc		The variable description to validate
 	 */
-	static void PostSetupObjectPinType(UBlueprint* InBlueprint, FBPVariableDescription& InOutVarDesc);
+	static UNREALED_API void PostSetupObjectPinType(UBlueprint* InBlueprint, FBPVariableDescription& InOutVarDesc);
 
 public:
 	/** Event fired after RenameVariableReferences is called */
 	DECLARE_MULTICAST_DELEGATE_FourParams(FOnRenameVariableReferences, UBlueprint* /*Blueprint*/, UClass* /*VariableClass*/, const FName& /*OldVarName*/, const FName& /*NewVarName*/);
-	static FOnRenameVariableReferences OnRenameVariableReferencesEvent;
+	static UNREALED_API FOnRenameVariableReferences OnRenameVariableReferencesEvent;
 
 	/**
 	 * Looks through the specified blueprint for any references to the specified 
@@ -1890,28 +1890,28 @@ public:
 	 * @param  OldVarName		The current name of the variable we want to replace
 	 * @param  NewVarName		The name that we wish to change all references to
 	 */
-	static void RenameVariableReferences(UBlueprint* Blueprint, UClass* VariableClass, const FName& OldVarName, const FName& NewVarName);
+	static UNREALED_API void RenameVariableReferences(UBlueprint* Blueprint, UClass* VariableClass, const FName& OldVarName, const FName& NewVarName);
 
 public:
-	static FName GetFunctionNameFromClassByGuid(const UClass* InClass, const FGuid FunctionGuid);
-	static bool GetFunctionGuidFromClassByFieldName(const UClass* InClass, const FName FunctionName, FGuid& FunctionGuid);
+	static UNREALED_API FName GetFunctionNameFromClassByGuid(const UClass* InClass, const FGuid FunctionGuid);
+	static UNREALED_API bool GetFunctionGuidFromClassByFieldName(const UClass* InClass, const FName FunctionName, FGuid& FunctionGuid);
 
 	/**
 	 * Returns a friendly class display name for the specified class (removing things like _C from the end, may localize the class name).  Class can be nullptr.
 	 */
-	static FText GetFriendlyClassDisplayName(const UClass* Class);
+	static UNREALED_API FText GetFriendlyClassDisplayName(const UClass* Class);
 
 	/**
 	 * Returns a class name for the specified class that has no automatic suffixes, but is otherwise unmodified.  Class can be nullptr.
 	 */
-	static FString GetClassNameWithoutSuffix(const UClass* Class);
+	static UNREALED_API FString GetClassNameWithoutSuffix(const UClass* Class);
 
 	/**
 	 * Returns a formatted menu item label for a deprecated variable or function member with the given name.
 	 *
 	 * @param MemberName		(Required) User-facing name of the deprecated variable or function.
 	 */
-	static FText GetDeprecatedMemberMenuItemName(const FText& MemberName);
+	static UNREALED_API FText GetDeprecatedMemberMenuItemName(const FText& MemberName);
 
 	/**
 	 * Returns a formatted warning message regarding usage of a deprecated variable or function member with the given name.
@@ -1919,25 +1919,25 @@ public:
 	 * @param MemberName		(Required) User-facing name of the deprecated variable or function.
 	 * @param DetailedMessage	(Optional) Instructional text or other details from the owner. If empty, a default message will be used.
 	 */
-	static FText GetDeprecatedMemberUsageNodeWarning(const FText& MemberName, const FText& DetailedMessage);
+	static UNREALED_API FText GetDeprecatedMemberUsageNodeWarning(const FText& MemberName, const FText& DetailedMessage);
 
 	/**
 	 * Remove overridden component templates from instance component handlers when a parent class disables editable when inherited boolean.
 	 */
-	static void HandleDisableEditableWhenInherited(UObject* ModifiedObject, TArray<UObject*>& ArchetypeInstances);
+	static UNREALED_API void HandleDisableEditableWhenInherited(UObject* ModifiedObject, TArray<UObject*>& ArchetypeInstances);
 
 	/**
 	 * Returns the BPs most derived native parent type:
 	 */
-	static UClass* GetNativeParent(const UBlueprint* BP);
+	static UNREALED_API UClass* GetNativeParent(const UBlueprint* BP);
 
 	/** Returns the UClass type for an object pin, if any */
-	static UClass* GetTypeForPin(const UEdGraphPin& Pin);
+	static UNREALED_API UClass* GetTypeForPin(const UEdGraphPin& Pin);
 
 	/**
 	 * Returns true if this BP is currently based on a type that returns true for the UObject::ImplementsGetWorld() call:
 	 */
-	static bool ImplementsGetWorld(const UBlueprint* BP);
+	static UNREALED_API bool ImplementsGetWorld(const UBlueprint* BP);
 
 	/*
 	 * Checks a function for 'blueprint thread safety'. Note this is not strict thread safety and does not perform any
@@ -1945,7 +1945,7 @@ public:
 	 * that are marked thread safe and metadata is used both here and in the compiler for native functions that have
 	 * been validated for thread safe use in a blueprint.
 	 */
-	static bool HasFunctionBlueprintThreadSafeMetaData(const UFunction* InFunction);
+	static UNREALED_API bool HasFunctionBlueprintThreadSafeMetaData(const UFunction* InFunction);
 
 	/**
 	 * Returns true if the given Blueprint is found to contain one or more nodes that are disallowed (restricted) within
@@ -1953,20 +1953,20 @@ public:
 	 * 
 	 * @param OutRestrictedNodes	(Optional) If non-NULL, will be populated with references to any nodes that are restricted.
 	 */
-	static bool HasRestrictedNodes(const UBlueprint* BP, TArray<UEdGraphNode*>* OutRestrictedNodes = nullptr);
+	static UNREALED_API bool HasRestrictedNodes(const UBlueprint* BP, TArray<UEdGraphNode*>* OutRestrictedNodes = nullptr);
 
 	/**
 	 * Checks the given Blueprint for any restricted content within the current editor context and sanitizes it away.
 	 * 
 	 * Note: If any restricted content is removed, this will also recompile the Blueprint.
 	 */
-	static void SanitizeRestrictedContent(UBlueprint* BP);
+	static UNREALED_API void SanitizeRestrictedContent(UBlueprint* BP);
 
 public:
-	static bool ShouldOpenWithDataOnlyEditor(const UBlueprint* Blueprint);
+	static UNREALED_API bool ShouldOpenWithDataOnlyEditor(const UBlueprint* Blueprint);
 };
 
-struct UNREALED_API FBlueprintDuplicationScopeFlags
+struct FBlueprintDuplicationScopeFlags
 {
 	enum EFlags : uint32
 	{
@@ -1978,7 +1978,7 @@ struct UNREALED_API FBlueprintDuplicationScopeFlags
 		TheSameNodeGuid = 1 << 3,
 	};
 
-	static uint32 bStaticFlags;
+	static UNREALED_API uint32 bStaticFlags;
 	static bool HasAnyFlag(uint32 InFlags)
 	{
 		return 0 != (bStaticFlags & InFlags);
@@ -1987,7 +1987,7 @@ struct UNREALED_API FBlueprintDuplicationScopeFlags
 	TGuardValue<uint32> Guard;
 	FBlueprintDuplicationScopeFlags(uint32 InFlags) : Guard(bStaticFlags, InFlags) {}
 };
-struct UNREALED_API FMakeClassSpawnableOnScope
+struct FMakeClassSpawnableOnScope
 {
 	UClass* Class;
 	bool bIsDeprecated;
