[UnrealEd]

diff --git a/Engine/Source/Editor/UnrealEd/Public/Kismet2/Breakpoint.h b/Engine/Source/Editor/UnrealEd/Public/Kismet2/Breakpoint.h
--- a/Engine/Source/Editor/UnrealEd/Public/Kismet2/Breakpoint.h
+++ b/Engine/Source/Editor/UnrealEd/Public/Kismet2/Breakpoint.h
@@ -13,7 +13,7 @@
 
 
 USTRUCT()
-struct UNREALED_API FBlueprintBreakpoint
+struct FBlueprintBreakpoint
 {
 	GENERATED_BODY()
 
@@ -37,7 +37,7 @@ private:
 	uint8 bStepOnce_RemoveAfterHit:1;
 
 public:
-	FBlueprintBreakpoint();
+	UNREALED_API FBlueprintBreakpoint();
 
 	/** Get the target node for the breakpoint */
 	UEdGraphNode* GetLocation() const
@@ -63,7 +63,7 @@ public:
 	}
 
 	/** Gets a string that describes the location */
-	FText GetLocationDescription() const;
+	UNREALED_API FText GetLocationDescription() const;
 
 	friend class FKismetDebugUtilities;
 };
diff --git a/Engine/Source/Editor/UnrealEd/Public/Kismet2/ChildActorComponentEditorUtils.h b/Engine/Source/Editor/UnrealEd/Public/Kismet2/ChildActorComponentEditorUtils.h
--- a/Engine/Source/Editor/UnrealEd/Public/Kismet2/ChildActorComponentEditorUtils.h
+++ b/Engine/Source/Editor/UnrealEd/Public/Kismet2/ChildActorComponentEditorUtils.h
@@ -8,48 +8,48 @@
 class UToolMenu;
 class FSCSEditorTreeNode;
 
-class UNREALED_API FChildActorComponentEditorUtils
+class FChildActorComponentEditorUtils
 {
 public:
 	/** Returns true if the given SCS editor tree node is a child actor node */
-	static bool IsChildActorNode(TSharedPtr<const FSCSEditorTreeNode> InNodePtr);
+	static UNREALED_API bool IsChildActorNode(TSharedPtr<const FSCSEditorTreeNode> InNodePtr);
 
 	/** Returns true if the given SCS editor tree node belongs to a child actor template */
-	static bool IsChildActorSubtreeNode(TSharedPtr<const FSCSEditorTreeNode> InNodePtr);
+	static UNREALED_API bool IsChildActorSubtreeNode(TSharedPtr<const FSCSEditorTreeNode> InNodePtr);
 
 	/** Returns true if any array element belongs to a child actor template */
-	static bool ContainsChildActorSubtreeNode(const TArray<TSharedPtr<FSCSEditorTreeNode>>& InNodePtrs);
+	static UNREALED_API bool ContainsChildActorSubtreeNode(const TArray<TSharedPtr<FSCSEditorTreeNode>>& InNodePtrs);
 
 	/** Returns the child actor component root node of a child actor subtree */
-	static TSharedPtr<FSCSEditorTreeNode> GetOuterChildActorComponentNode(TSharedPtr<const FSCSEditorTreeNode> InNodePtr);
+	static UNREALED_API TSharedPtr<FSCSEditorTreeNode> GetOuterChildActorComponentNode(TSharedPtr<const FSCSEditorTreeNode> InNodePtr);
 
 	/** Returns whether child actor tree view expansion is enabled in project settings */
-	static bool IsChildActorTreeViewExpansionEnabled();
+	static UNREALED_API bool IsChildActorTreeViewExpansionEnabled();
 
 	/** Returns the default visualization mode for child actors in a component tree view */
-	static EChildActorComponentTreeViewVisualizationMode GetProjectDefaultTreeViewVisualizationMode();
+	static UNREALED_API EChildActorComponentTreeViewVisualizationMode GetProjectDefaultTreeViewVisualizationMode();
 
 	/** 
 	 * Returns the component-specific visualization mode for the given child actor component
 	 * @param DefaultVisOverride If different than UseDefault, overrides the default ChildActorComponentTreeViewVisualizationMode from the project settings.
 	 */
-	static EChildActorComponentTreeViewVisualizationMode GetChildActorTreeViewVisualizationMode(UChildActorComponent* ChildActorComponent, EChildActorComponentTreeViewVisualizationMode DefaultVisOverride = EChildActorComponentTreeViewVisualizationMode::UseDefault);
+	static UNREALED_API EChildActorComponentTreeViewVisualizationMode GetChildActorTreeViewVisualizationMode(UChildActorComponent* ChildActorComponent, EChildActorComponentTreeViewVisualizationMode DefaultVisOverride = EChildActorComponentTreeViewVisualizationMode::UseDefault);
 
 	/** 
 	 * Whether to expand the given child actor component in a component tree view 
 	 * @param DefaultVisOverride If different than UseDefault, overrides the default ChildActorComponentTreeViewVisualizationMode from the project settings and forces child actor tree view expansion to be enabled.
 	 */
-	static bool ShouldExpandChildActorInTreeView(UChildActorComponent* ChildActorComponent, EChildActorComponentTreeViewVisualizationMode DefaultVisOverride = EChildActorComponentTreeViewVisualizationMode::UseDefault);
+	static UNREALED_API bool ShouldExpandChildActorInTreeView(UChildActorComponent* ChildActorComponent, EChildActorComponentTreeViewVisualizationMode DefaultVisOverride = EChildActorComponentTreeViewVisualizationMode::UseDefault);
 
 	/** 
 	 * Whether the Child Actor should be shown in a component tree view for the given component
 	 * @param DefaultVisOverride If different than UseDefault, overrides the default ChildActorComponentTreeViewVisualizationMode from the project settings and forces child actor tree view expansion to be enabled.
 	 */
-	static bool ShouldShowChildActorNodeInTreeView(UChildActorComponent* ChildActorComponent, EChildActorComponentTreeViewVisualizationMode DefaultVisOverride = EChildActorComponentTreeViewVisualizationMode::UseDefault);
+	static UNREALED_API bool ShouldShowChildActorNodeInTreeView(UChildActorComponent* ChildActorComponent, EChildActorComponentTreeViewVisualizationMode DefaultVisOverride = EChildActorComponentTreeViewVisualizationMode::UseDefault);
 
 	/** Populates the given menu with options for the given Child Actor component */
-	static void FillComponentContextMenuOptions(UToolMenu* Menu, UChildActorComponent* ChildActorComponent);
+	static UNREALED_API void FillComponentContextMenuOptions(UToolMenu* Menu, UChildActorComponent* ChildActorComponent);
 
 	/** Populates the given menu with additional options if the given SCS editor tree node represents a Child Actor node */
-	static void FillChildActorContextMenuOptions(UToolMenu* Menu, TSharedPtr<const FSCSEditorTreeNode> InNodePtr);
-};
\ No newline at end of file
+	static UNREALED_API void FillChildActorContextMenuOptions(UToolMenu* Menu, TSharedPtr<const FSCSEditorTreeNode> InNodePtr);
+};
diff --git a/Engine/Source/Editor/UnrealEd/Public/Kismet2/CompilerResultsLog.h b/Engine/Source/Editor/UnrealEd/Public/Kismet2/CompilerResultsLog.h
--- a/Engine/Source/Editor/UnrealEd/Public/Kismet2/CompilerResultsLog.h
+++ b/Engine/Source/Editor/UnrealEd/Public/Kismet2/CompilerResultsLog.h
@@ -43,7 +43,7 @@ struct FObjectKey;
 #if WITH_EDITOR
 
 /** This class maps from final objects to their original source object, across cloning, autoexpansion, etc... */
-class UNREALED_API FBacktrackMap
+class FBacktrackMap
 {
 protected:
 	// Maps from transient object created during compiling to original 'source code' object
@@ -53,20 +53,20 @@ protected:
 
 public:
 	/** Update the source backtrack map to note that NewObject was most closely generated/caused by the SourceObject */
-	void NotifyIntermediateObjectCreation(UObject* NewObject, UObject* SourceObject);
+	UNREALED_API void NotifyIntermediateObjectCreation(UObject* NewObject, UObject* SourceObject);
 
 	/** Update the pin source backtrack map to note that NewPin was most closely generated/caused by the SourcePin */
-	void NotifyIntermediatePinCreation(UEdGraphPin* NewPin, UEdGraphPin* SourcePin);
+	UNREALED_API void NotifyIntermediatePinCreation(UEdGraphPin* NewPin, UEdGraphPin* SourcePin);
 
 	/** Returns the true source object for the passed in object */
-	UObject* FindSourceObject(UObject* PossiblyDuplicatedObject);
-	UObject const* FindSourceObject(UObject const* PossiblyDuplicatedObject) const;
-	UEdGraphPin* FindSourcePin(UEdGraphPin* PossiblyDuplicatedPin);
-	UEdGraphPin const* FindSourcePin(UEdGraphPin const* PossiblyDuplicatedPin) const;
+	UNREALED_API UObject* FindSourceObject(UObject* PossiblyDuplicatedObject);
+	UNREALED_API UObject const* FindSourceObject(UObject const* PossiblyDuplicatedObject) const;
+	UNREALED_API UEdGraphPin* FindSourcePin(UEdGraphPin* PossiblyDuplicatedPin);
+	UNREALED_API UEdGraphPin const* FindSourcePin(UEdGraphPin const* PossiblyDuplicatedPin) const;
 };
 
 /** This class represents a log of compiler output lines (errors, warnings, and information notes), each of which can be a rich tokenized message */
-class UNREALED_API FCompilerResultsLog
+class FCompilerResultsLog
 {
 	// Compiler event
 	struct FCompilerEvent
@@ -144,14 +144,14 @@ protected:
 	TMap<TWeakObjectPtr<const UEdGraphNode>, TWeakObjectPtr<UEdGraphNode>> FullMacroBacktrackMap;
 
 public:
-	FCompilerResultsLog(bool bIsCompatibleWithEvents = true);
-	virtual ~FCompilerResultsLog();
+	UNREALED_API FCompilerResultsLog(bool bIsCompatibleWithEvents = true);
+	UNREALED_API virtual ~FCompilerResultsLog();
 
 	/** Register this log with the MessageLog module */
-	static void Register();
+	static UNREALED_API void Register();
 
 	/** Unregister this log from the MessageLog module */
-	static void Unregister();
+	static UNREALED_API void Unregister();
 
 	/** Accessor for the LogName, so it can be opened elsewhere */
 	static FName GetLogName(){ return Name; }
@@ -282,28 +282,28 @@ public:
 	/**
 	 * Commit all stored potential messages for a given node. Returns true if any messages were written.
 	 */
-	bool CommitPotentialMessages(UEdGraphNode* Source);
+	UNREALED_API bool CommitPotentialMessages(UEdGraphNode* Source);
 
 	/** Update the source backtrack map to note that NewObject was most closely generated/caused by the SourceObject */
-	void NotifyIntermediateObjectCreation(UObject* NewObject, UObject* SourceObject);
-	void NotifyIntermediatePinCreation(UEdGraphPin* NewObject, UEdGraphPin* SourceObject);
+	UNREALED_API void NotifyIntermediateObjectCreation(UObject* NewObject, UObject* SourceObject);
+	UNREALED_API void NotifyIntermediatePinCreation(UEdGraphPin* NewObject, UEdGraphPin* SourceObject);
 
 	/** Update the expansion map to note that Node was expanded from OuterTunnelInstance, both the node and tunnel instance should be intermediate nodes */
-	void NotifyIntermediateTunnelNode(const UEdGraphNode* Node, const UEdGraphNode* OuterTunnelInstance);
+	UNREALED_API void NotifyIntermediateTunnelNode(const UEdGraphNode* Node, const UEdGraphNode* OuterTunnelInstance);
 
 	/** Update the map that tracks nodes created by macro instance nodes */
-	void NotifyIntermediateMacroNode(UEdGraphNode* SourceNode, const UEdGraphNode* IntermediateNode);
+	UNREALED_API void NotifyIntermediateMacroNode(UEdGraphNode* SourceNode, const UEdGraphNode* IntermediateNode);
 
 	/** Returns the true source object for the passed in object */
-	UObject* FindSourceObject(UObject* PossiblyDuplicatedObject);
-	UObject const* FindSourceObject(UObject const* PossiblyDuplicatedObject) const;
-	UObject* FindSourceMacroInstance(const UEdGraphNode* IntermediateNode) const;
+	UNREALED_API UObject* FindSourceObject(UObject* PossiblyDuplicatedObject);
+	UNREALED_API UObject const* FindSourceObject(UObject const* PossiblyDuplicatedObject) const;
+	UNREALED_API UObject* FindSourceMacroInstance(const UEdGraphNode* IntermediateNode) const;
 
 	/** Returns the intermediate tunnel instance that generated the node */
-	const UEdGraphNode* GetIntermediateTunnelInstance(const UEdGraphNode* IntermediateNode) const;
+	UNREALED_API const UEdGraphNode* GetIntermediateTunnelInstance(const UEdGraphNode* IntermediateNode) const;
 
 	/** Returns a int32 used to uniquely identify an action for the latent action manager */
-	int32 CalculateStableIdentifierForLatentActionManager( const UEdGraphNode* Node );
+	UNREALED_API int32 CalculateStableIdentifierForLatentActionManager( const UEdGraphNode* Node );
 
 	/** Returns the true source object for the passed in object; does type checking on the result */
 	template <typename T>
@@ -318,17 +318,17 @@ public:
 		return CastChecked<T const>(FindSourceObject(PossiblyDuplicatedObject));
 	}
 
-	UEdGraphPin* FindSourcePin(UEdGraphPin* PossiblyDuplicatedPin);
-	const UEdGraphPin* FindSourcePin(const UEdGraphPin* PossiblyDuplicatedPin) const;
+	UNREALED_API UEdGraphPin* FindSourcePin(UEdGraphPin* PossiblyDuplicatedPin);
+	UNREALED_API const UEdGraphPin* FindSourcePin(const UEdGraphPin* PossiblyDuplicatedPin) const;
 	
 	/** Copy errors from an existing log into this one, and optionally write out to log if it was suppressed the first time */
-	void Append(FCompilerResultsLog const& Other, bool bWriteToSystemLog = false);
+	UNREALED_API void Append(FCompilerResultsLog const& Other, bool bWriteToSystemLog = false);
 
 	/** Begin a new compiler event */
-	void BeginEvent(const TCHAR* InName);
+	UNREALED_API void BeginEvent(const TCHAR* InName);
 
 	/** End the current compiler event */
-	void EndEvent();
+	UNREALED_API void EndEvent();
 
 	/** Access the current event target log */
 	static FCompilerResultsLog* GetEventTarget()
@@ -337,13 +337,13 @@ public:
 	}
 
 	/** Get the message log listing for this blueprint */
-	static TSharedRef<IMessageLogListing> GetBlueprintMessageLog(UBlueprint* InBlueprint);
+	static UNREALED_API TSharedRef<IMessageLogListing> GetBlueprintMessageLog(UBlueprint* InBlueprint);
 
 protected:
 	/** Helper method to add a child event to the given parent event scope */
-	void AddChildEvent(TSharedPtr<FCompilerEvent>& ParentEventScope, TSharedRef<FCompilerEvent>& ChildEventScope);
+	UNREALED_API void AddChildEvent(TSharedPtr<FCompilerEvent>& ParentEventScope, TSharedRef<FCompilerEvent>& ChildEventScope);
 
-	void InternalLogMessage(FName MessageID, const TSharedRef<FTokenizedMessage>& Message, const TArray<UEdGraphNode*>& SourceNodes );
+	UNREALED_API void InternalLogMessage(FName MessageID, const TSharedRef<FTokenizedMessage>& Message, const TArray<UEdGraphNode*>& SourceNodes );
 
 	void Tokenize(const TCHAR* Text, FTokenizedMessage &OutMessage, TArray<UEdGraphNode*>& OutSourceNode)
 	{
@@ -384,45 +384,45 @@ protected:
 	}
 
 	/** Links the UEdGraphNode with the LogLine: */
-	void AnnotateNode(const TArray<UEdGraphNode*>& Nodes, TSharedRef<FTokenizedMessage> LogLine);
+	UNREALED_API void AnnotateNode(const TArray<UEdGraphNode*>& Nodes, TSharedRef<FTokenizedMessage> LogLine);
 
 	/** Internal method to append the final compiler results summary to the MessageLog */
-	void InternalLogSummary();
+	UNREALED_API void InternalLogSummary();
 
 	/** Internal helper method to recursively append event details into the MessageLog */
-	void InternalLogEvent(const FCompilerEvent& InEvent, int32 InDepth = 0);
+	UNREALED_API void InternalLogEvent(const FCompilerEvent& InEvent, int32 InDepth = 0);
 
 	/** Returns true if the user has requested this compiler message be suppressed */
-	bool IsMessageEnabled(FName ID);
+	UNREALED_API bool IsMessageEnabled(FName ID);
 private:
 
 	/** Map of stored potential messages indexed by a node. Can be committed to the results log by calling CommitPotentialMessages for that node. */
 	TMap< FObjectKey, TArray< TSharedRef<FTokenizedMessage> > > PotentialMessages;
 
 	/** Parses a compiler log dump to generate tokenized output */
-	static TArray< TSharedRef<FTokenizedMessage> > ParseCompilerLogDump(const FString& LogDump);
+	static UNREALED_API TArray< TSharedRef<FTokenizedMessage> > ParseCompilerLogDump(const FString& LogDump);
 
 	/** Goes to an error given a Message Token */
-	static void OnGotoError(const class TSharedRef<IMessageToken>& Token);
+	static UNREALED_API void OnGotoError(const class TSharedRef<IMessageToken>& Token);
 
 	/** Callback function for binding the global compiler dump to open the static compiler log */
-	static void GetGlobalModuleCompilerDump(const FString& LogDump, ECompilationResult::Type CompilationResult, bool bShowLog);
+	static UNREALED_API void GetGlobalModuleCompilerDump(const FString& LogDump, ECompilationResult::Type CompilationResult, bool bShowLog);
 	
 	/** Searches a token list for referenced UEdGraphNodes, used to update the nodes when a log is committed */
-	static void GetNodesFromTokens(const TArray<TSharedRef<IMessageToken> >& MessageTokens, TArray<UEdGraphNode*>& OutOwnerNodes);
+	static UNREALED_API void GetNodesFromTokens(const TArray<TSharedRef<IMessageToken> >& MessageTokens, TArray<UEdGraphNode*>& OutOwnerNodes);
 
 	/** The log's name, for easy re-use */
-	static const FName Name;
+	static UNREALED_API const FName Name;
 	
 	/** The log target for compile events */
-	static FCompilerResultsLog* CurrentEventTarget;
+	static UNREALED_API FCompilerResultsLog* CurrentEventTarget;
 
 	/** Handle to the registered GetGlobalModuleCompilerDump delegate. */
-	static FDelegateHandle GetGlobalModuleCompilerDumpDelegateHandle;
+	static UNREALED_API FDelegateHandle GetGlobalModuleCompilerDumpDelegateHandle;
 };
 
 /** This class will begin a new compile event on construction, and automatically end it when the instance goes out of scope */
-class UNREALED_API FScopedCompilerEvent
+class FScopedCompilerEvent
 {
 public:
 	/** Constructor; automatically begins a new event */
@@ -447,11 +447,11 @@ public:
 };
 
 /** Scope wrapper for the blueprint message log. Ensures we dont leak logs that we dont need (i.e. those that have no messages) */
-class UNREALED_API FScopedBlueprintMessageLog
+class FScopedBlueprintMessageLog
 {
 public:
-	FScopedBlueprintMessageLog(UBlueprint* InBlueprint);
-	~FScopedBlueprintMessageLog();
+	UNREALED_API FScopedBlueprintMessageLog(UBlueprint* InBlueprint);
+	UNREALED_API ~FScopedBlueprintMessageLog();
 
 public:
 	/** The listing we wrap */
diff --git a/Engine/Source/Editor/UnrealEd/Public/Kismet2/ComponentEditorUtils.h b/Engine/Source/Editor/UnrealEd/Public/Kismet2/ComponentEditorUtils.h
--- a/Engine/Source/Editor/UnrealEd/Public/Kismet2/ComponentEditorUtils.h
+++ b/Engine/Source/Editor/UnrealEd/Public/Kismet2/ComponentEditorUtils.h
@@ -12,55 +12,55 @@
 class UToolMenu;
 class UMaterialInterface;
 
-class UNREALED_API FComponentEditorUtils
+class FComponentEditorUtils
 {
 public:
 	/** Is the instance component is editable */
-	static bool CanEditComponentInstance(const UActorComponent* ActorComp, const UActorComponent* ParentSceneComp, bool bAllowUserContructionScript);
+	static UNREALED_API bool CanEditComponentInstance(const UActorComponent* ActorComp, const UActorComponent* ParentSceneComp, bool bAllowUserContructionScript);
 
 	/** 
 	* Test if the native component is editable. If it is, return a valid pointer to it's FProperty
 	* Otherwise, return nullptr. A native component is editable if it is marked as EditAnywhere
 	* via meta data tags or is within an editable property container 
 	*/
-	static FProperty* GetPropertyForEditableNativeComponent(const UActorComponent* NativeComponent);
+	static UNREALED_API FProperty* GetPropertyForEditableNativeComponent(const UActorComponent* NativeComponent);
 
 	/** Test whether or not the given string is a valid variable name string for the given component instance */
-	static bool IsValidVariableNameString(const UActorComponent* InComponent, const FString& InString);
+	static UNREALED_API bool IsValidVariableNameString(const UActorComponent* InComponent, const FString& InString);
 
 	/** 
 	 * Test whether or not the given string is already the name string of a component on the the owner
 	 * Optionally excludes an existing component from the check (ex. a component currently being renamed)
 	 * @return True if the InString is an available name for a component of ComponentOwner
 	 */
-	static bool IsComponentNameAvailable(const FString& InString, UObject* ComponentOwner, const UActorComponent* ComponentToIgnore = nullptr);
+	static UNREALED_API bool IsComponentNameAvailable(const FString& InString, UObject* ComponentOwner, const UActorComponent* ComponentToIgnore = nullptr);
 		
 	/** Generate a valid variable name string for the given component instance */
-	static FString GenerateValidVariableName(TSubclassOf<UActorComponent> InComponentClass, AActor* ComponentOwner);
+	static UNREALED_API FString GenerateValidVariableName(TSubclassOf<UActorComponent> InComponentClass, AActor* ComponentOwner);
 
 	/** Generate a valid variable name string for the given component instance based on the name of the asset referenced by the component */
-	static FString GenerateValidVariableNameFromAsset(UObject* Asset, AActor* ComponentOwner);
+	static UNREALED_API FString GenerateValidVariableNameFromAsset(UObject* Asset, AActor* ComponentOwner);
 
 	/**
 	* Checks whether it is valid to copy the given component
 	* @param ComponentToCopy The component to check
 	* @return Whether the given component can be copied
 	*/
-	static bool CanCopyComponent(const UActorComponent* ComponentToCopy);
+	static UNREALED_API bool CanCopyComponent(const UActorComponent* ComponentToCopy);
 
 	/**
 	 * Checks whether it is valid to copy the indicated components
 	 * @param ComponentsToCopy The list of components to check
 	 * @return Whether the indicated components can be copied
 	 */
-	static bool CanCopyComponents(const TArray<UActorComponent*>& ComponentsToCopy);
+	static UNREALED_API bool CanCopyComponents(const TArray<UActorComponent*>& ComponentsToCopy);
 
 	/**
 	 * Copies the selected components to the clipboard
 	 * @param ComponentsToCopy The list of components to copy
 	 * @param DestinationData Buffer to fill with the copied data, or null to use the clipboard
 	 */
-	static void CopyComponents(const TArray<UActorComponent*>& ComponentsToCopy, FString* DestinationData = nullptr);
+	static UNREALED_API void CopyComponents(const TArray<UActorComponent*>& ComponentsToCopy, FString* DestinationData = nullptr);
 
 	/**
 	 * Determines whether the current contents of the clipboard contain paste-able component information
@@ -69,7 +69,7 @@ public:
 	 * @param SourceData Component data to paste, or null to use the clipboard
 	 * @return Whether components can be pasted
 	 */
-	static bool CanPasteComponents(const USceneComponent* RootComponent, bool bOverrideCanAttach = false, bool bPasteAsArchetypes = false, const FString* SourceData = nullptr);
+	static UNREALED_API bool CanPasteComponents(const USceneComponent* RootComponent, bool bOverrideCanAttach = false, bool bPasteAsArchetypes = false, const FString* SourceData = nullptr);
 
 	/**
 	 * Attempts to paste components from the clipboard as siblings of the target component
@@ -78,28 +78,28 @@ public:
 	 * @param TargetComponent The component the paste is targeting (will attempt to paste components as siblings). If null, will attach pasted components to the root.
 	 * @param SourceData Component data to paste, or null to use the clipboard
 	 */
-	static void PasteComponents(TArray<UActorComponent*>& OutPastedComponents, AActor* TargetActor, USceneComponent* TargetComponent = nullptr, const FString* SourceData = nullptr);
+	static UNREALED_API void PasteComponents(TArray<UActorComponent*>& OutPastedComponents, AActor* TargetActor, USceneComponent* TargetComponent = nullptr, const FString* SourceData = nullptr);
 
 	/**
 	 * Gets the copied components from the clipboard without attempting to paste/apply them in any way
 	 * @param OutParentMap Contains the child->parent name relationships of the copied components
 	 * @param OutNewObjectMap Contains the name->instance object mapping of the copied components
 	 */
-	static void GetComponentsFromClipboard(TMap<FName, FName>& OutParentMap, TMap<FName, UActorComponent*>& OutNewObjectMap, bool bGetComponentsAsArchetypes);
+	static UNREALED_API void GetComponentsFromClipboard(TMap<FName, FName>& OutParentMap, TMap<FName, UActorComponent*>& OutNewObjectMap, bool bGetComponentsAsArchetypes);
 
 	/**
 	 * Determines whether the indicated component can be deleted
 	 * @param ComponentToDelete The component to determine can be deleted
 	 * @return Whether the indicated component can be deleted
 	 */
-	static bool CanDeleteComponent(const UActorComponent* ComponentToDelete);
+	static UNREALED_API bool CanDeleteComponent(const UActorComponent* ComponentToDelete);
 
 	/**
 	 * Determines whether the indicated components can be deleted
 	 * @param ComponentsToDelete The list of components to determine can be deleted
 	 * @return Whether the indicated components can be deleted
 	 */
-	static bool CanDeleteComponents(const TArray<UActorComponent*>& ComponentsToDelete);
+	static UNREALED_API bool CanDeleteComponents(const TArray<UActorComponent*>& ComponentsToDelete);
 
 	/**
 	 * Deletes the indicated components and identifies the component that should be selected following the operation.
@@ -109,14 +109,14 @@ public:
 	 * @param OutComponentToSelect The component that should be selected after the deletion
 	 * @return The number of components that were actually deleted
 	 */
-	static int32 DeleteComponents(const TArray<UActorComponent*>& ComponentsToDelete, UActorComponent*& OutComponentToSelect);
+	static UNREALED_API int32 DeleteComponents(const TArray<UActorComponent*>& ComponentsToDelete, UActorComponent*& OutComponentToSelect);
 
 	/** 
 	 * Duplicates a component instance and takes care of attachment and registration.
 	 * @param TemplateComponent The component to duplicate
 	 * @return The created clone of the provided TemplateComponent. nullptr if the duplication failed.
 	 */
-	static UActorComponent* DuplicateComponent(UActorComponent* TemplateComponent);
+	static UNREALED_API UActorComponent* DuplicateComponent(UActorComponent* TemplateComponent);
 
 	/**
 	 * Ensures that the selection override delegate is properly bound for the supplied component
@@ -125,7 +125,7 @@ public:
 	 * @param SceneComponent The component to set the selection override for
 	 * @param bBind Whether the override should be bound
 	*/
-	static void BindComponentSelectionOverride(USceneComponent* SceneComponent, bool bBind);
+	static UNREALED_API void BindComponentSelectionOverride(USceneComponent* SceneComponent, bool bBind);
 
 	/**
 	 * Attempts to apply a material to a component at the specified slot.
@@ -136,10 +136,10 @@ public:
 	 *
 	 * @return	True if the material was successfully applied to the component.
 	 */
-	static bool AttemptApplyMaterialToComponent( USceneComponent* SceneComponent, UMaterialInterface* MaterialToApply, int32 OptionalMaterialSlot = -1 );
+	static UNREALED_API bool AttemptApplyMaterialToComponent( USceneComponent* SceneComponent, UMaterialInterface* MaterialToApply, int32 OptionalMaterialSlot = -1 );
 
 	/** Potentially transforms the delta to be applied to a component into the appropriate space */
-	static void AdjustComponentDelta(const USceneComponent* Component, FVector& Drag, FRotator& Rotation);
+	static UNREALED_API void AdjustComponentDelta(const USceneComponent* Component, FVector& Drag, FRotator& Rotation);
 
 	// Given a template and a property, propagates a default value change to all instances (only if applicable)
 	template<typename T>
@@ -260,14 +260,14 @@ public:
 	}
 
 	// Try to find the correct variable name for a given native component template or instance (which can have a mismatch)
-	static FName FindVariableNameGivenComponentInstance(const UActorComponent* ComponentInstance);
+	static UNREALED_API FName FindVariableNameGivenComponentInstance(const UActorComponent* ComponentInstance);
 
 	/**
 	* Populates the given menu with basic options for operations on components in the world.
 	* @param Menu Used to register the menu options
 	* @param SelectedComponents The selected components to create menu options for
 	*/
-	static void FillComponentContextMenuOptions(UToolMenu* Menu, const TArray<UActorComponent*>& SelectedComponents);
+	static UNREALED_API void FillComponentContextMenuOptions(UToolMenu* Menu, const TArray<UActorComponent*>& SelectedComponents);
 
 	/**
 	 * Tries to find a match for ComponentInstance in the ComponentList. First by name and then if multiple Components have a matching name try to match the SceneComponent hierarchy to find the best match.
@@ -275,19 +275,19 @@ public:
 	 * @param ComponentList List containing possible matches
 	 * @return Valid Component pointer if match was found. nullptr otherwise.
 	 */
-	static UActorComponent* FindMatchingComponent(const UActorComponent* ComponentInstance, const TInlineComponentArray<UActorComponent*>& ComponentList);
+	static UNREALED_API UActorComponent* FindMatchingComponent(const UActorComponent* ComponentInstance, const TInlineComponentArray<UActorComponent*>& ComponentList);
 
 	/**
 	 * Make a FComponentReference from a component pointer.
 	 * @param ExpectedComponentOwner The expected component owner. Should be the same as OwningActor from FComponentReference::GetComponent().
 	 * @param Component The component we would like to initialize the FComponentReference with.
 	 */
-	static FComponentReference MakeComponentReference(const AActor* ExpectedComponentOwner, const UActorComponent* Component);
+	static UNREALED_API FComponentReference MakeComponentReference(const AActor* ExpectedComponentOwner, const UActorComponent* Component);
 
 private:	
-	static USceneComponent* FindClosestParentInList(UActorComponent* ChildComponent, const TArray<UActorComponent*>& ComponentList);
+	static UNREALED_API USceneComponent* FindClosestParentInList(UActorComponent* ChildComponent, const TArray<UActorComponent*>& ComponentList);
 
-	static void OnGoToComponentAssetInBrowser(UObject* Asset);
-	static void OnOpenComponentCodeFile(const FString CodeFileName);
-	static void OnEditBlueprintComponent(UObject* Blueprint);
+	static UNREALED_API void OnGoToComponentAssetInBrowser(UObject* Asset);
+	static UNREALED_API void OnOpenComponentCodeFile(const FString CodeFileName);
+	static UNREALED_API void OnEditBlueprintComponent(UObject* Blueprint);
 };
diff --git a/Engine/Source/Editor/UnrealEd/Public/Kismet2/DebuggerCommands.h b/Engine/Source/Editor/UnrealEd/Public/Kismet2/DebuggerCommands.h
--- a/Engine/Source/Editor/UnrealEd/Public/Kismet2/DebuggerCommands.h
+++ b/Engine/Source/Editor/UnrealEd/Public/Kismet2/DebuggerCommands.h
@@ -15,7 +15,7 @@ class FToolBarBuilder;
 struct FToolMenuSection;
 
 /** This class acts as a generic widget that listens to and process global play world actions */
-class UNREALED_API SGlobalPlayWorldActions : public SCompoundWidget
+class SGlobalPlayWorldActions : public SCompoundWidget
 {
 public:
 
@@ -24,56 +24,56 @@ public:
 
 	SLATE_END_ARGS()
 
-	void Construct(const FArguments& InArgs);
+	UNREALED_API void Construct(const FArguments& InArgs);
 
-	virtual FReply OnKeyDown(const FGeometry& MyGeometry, const FKeyEvent& InKeyEvent) override;
+	UNREALED_API virtual FReply OnKeyDown(const FGeometry& MyGeometry, const FKeyEvent& InKeyEvent) override;
 
-	virtual bool SupportsKeyboardFocus() const override;
+	UNREALED_API virtual bool SupportsKeyboardFocus() const override;
 };
 
 //////////////////////////////////////////////////////////////////////////
 // FPlayWorldCommands
 
-class UNREALED_API FPlayWorldCommands : public TCommands<FPlayWorldCommands>
+class FPlayWorldCommands : public TCommands<FPlayWorldCommands>
 {
 private:
 
 	friend class TCommands<FPlayWorldCommands>;
 
-	FPlayWorldCommands();
+	UNREALED_API FPlayWorldCommands();
 
 public:
 
 	// TCommands interface
-	virtual void RegisterCommands() override;
+	UNREALED_API virtual void RegisterCommands() override;
 	// End of TCommands interface
 
 	/**
 	 * Binds all global kismet commands to delegates
 	 */
-	static void BindGlobalPlayWorldCommands();
+	static UNREALED_API void BindGlobalPlayWorldCommands();
 
 	/** Populates a toolbar with the menu commands for play-world control (pause/resume/stop/possess/eject/step/show current loc) */
-	static void BuildToolbar( FToolMenuSection& InSection, bool bIncludeLaunchButtonAndOptions = false );
+	static UNREALED_API void BuildToolbar( FToolMenuSection& InSection, bool bIncludeLaunchButtonAndOptions = false );
 
 	/**
 	* Return the active widget that processes play world actions for PIE
 	*
 	*/
-	static TWeakPtr<SGlobalPlayWorldActions> GetActiveGlobalPlayWorldActionsWidget();
+	static UNREALED_API TWeakPtr<SGlobalPlayWorldActions> GetActiveGlobalPlayWorldActionsWidget();
 
 	/**
 	* Set the active widget that processes play world actions for PIE
 	*
 	*/
-	static void SetActiveGlobalPlayWorldActionsWidget(TWeakPtr<SGlobalPlayWorldActions> ActiveWidget);
+	static UNREALED_API void SetActiveGlobalPlayWorldActionsWidget(TWeakPtr<SGlobalPlayWorldActions> ActiveWidget);
 
 public:
 
 	/** 
 	 * A command list that can be passed around and isn't bound to an instance of any tool or editor. 
 	 */
-	static TSharedPtr<FUICommandList> GlobalPlayWorldActions;
+	static UNREALED_API TSharedPtr<FUICommandList> GlobalPlayWorldActions;
 
 public:
 
@@ -116,46 +116,46 @@ public:
 protected:
 
 	/** A weak pointer to the current active widget that processes PIE actions */
-	static TWeakPtr<SGlobalPlayWorldActions> ActiveGlobalPlayWorldActionsWidget;
+	static UNREALED_API TWeakPtr<SGlobalPlayWorldActions> ActiveGlobalPlayWorldActionsWidget;
 
 	/**
 	 * Generates menu content for the PIE combo button drop down menu
 	 *
 	 * @return	Menu content widget
 	 */
-	static TSharedRef< SWidget > GeneratePlayMenuContent( TSharedRef<FUICommandList> InCommandList );
+	static UNREALED_API TSharedRef< SWidget > GeneratePlayMenuContent( TSharedRef<FUICommandList> InCommandList );
 
 	// Add mobile PIE preview device commands
-	void AddPIEPreviewDeviceCommands();
+	UNREALED_API void AddPIEPreviewDeviceCommands();
 
 	// Add mobile PIE preview device actions
-	static void AddPIEPreviewDeviceActions(const FPlayWorldCommands &Commands, FUICommandList &ActionList);
+	static UNREALED_API void AddPIEPreviewDeviceActions(const FPlayWorldCommands &Commands, FUICommandList &ActionList);
 };
 
 
 //////////////////////////////////////////////////////////////////////////
 // FPlayWorldCommandCallbacks
 
-class UNREALED_API FPlayWorldCommandCallbacks
+class FPlayWorldCommandCallbacks
 {
 public:
 	/**
 	 * Called from the context menu to start previewing the game at the clicked location                   
 	 */
-	static void StartPlayFromHere();
-	static void StartPlayFromHere(const TOptional<FVector>& Location, const TOptional<FRotator>& Rotation, const TSharedPtr<IAssetViewport>& ActiveLevelViewport);
+	static UNREALED_API void StartPlayFromHere();
+	static UNREALED_API void StartPlayFromHere(const TOptional<FVector>& Location, const TOptional<FRotator>& Rotation, const TSharedPtr<IAssetViewport>& ActiveLevelViewport);
 
-	static void ResumePlaySession_Clicked();
-	static void PausePlaySession_Clicked();
-	static void SingleFrameAdvance_Clicked();
+	static UNREALED_API void ResumePlaySession_Clicked();
+	static UNREALED_API void PausePlaySession_Clicked();
+	static UNREALED_API void SingleFrameAdvance_Clicked();
 
-	static bool IsInSIE();
-	static bool IsInPIE();
+	static UNREALED_API bool IsInSIE();
+	static UNREALED_API bool IsInPIE();
 
-	static bool IsInSIE_AndRunning();
-	static bool IsInPIE_AndRunning();
+	static UNREALED_API bool IsInSIE_AndRunning();
+	static UNREALED_API bool IsInPIE_AndRunning();
 
-	static bool HasPlayWorld();
-	static bool HasPlayWorldAndPaused();
-	static bool HasPlayWorldAndRunning();
+	static UNREALED_API bool HasPlayWorld();
+	static UNREALED_API bool HasPlayWorldAndPaused();
+	static UNREALED_API bool HasPlayWorldAndRunning();
 };
diff --git a/Engine/Source/Editor/UnrealEd/Public/Kismet2/EnumEditorUtils.h b/Engine/Source/Editor/UnrealEd/Public/Kismet2/EnumEditorUtils.h
--- a/Engine/Source/Editor/UnrealEd/Public/Kismet2/EnumEditorUtils.h
+++ b/Engine/Source/Editor/UnrealEd/Public/Kismet2/EnumEditorUtils.h
@@ -17,7 +17,7 @@ class UEnum;
 class UObject;
 class UUserDefinedEnum;
 
-class UNREALED_API FEnumEditorUtils
+class FEnumEditorUtils
 {
 	static void PrepareForChange(UUserDefinedEnum* Enum);
 	static void BroadcastChanges(const UUserDefinedEnum* Enum, const TArray<TPair<FName, int64>>& OldNames, bool bResolveData = true);
@@ -37,7 +37,7 @@ public:
 	public:
 		UNREALED_API static FEnumEditorManager& Get();
 
-		class UNREALED_API ListenerType : public InnerListenerType<FEnumEditorManager>
+		class ListenerType : public InnerListenerType<FEnumEditorManager>
 		{
 		};
 	};
@@ -48,37 +48,37 @@ public:
 	// User defined enumerations
 
 	/** Creates new user defined enum in given blueprint. */
-	static UEnum* CreateUserDefinedEnum(UObject* InParent, FName EnumName, EObjectFlags Flags);
+	static UNREALED_API UEnum* CreateUserDefinedEnum(UObject* InParent, FName EnumName, EObjectFlags Flags);
 
 	/** return if an enum can be named/renamed with given name*/
-	static bool IsNameAvailebleForUserDefinedEnum(FName Name);
+	static UNREALED_API bool IsNameAvailebleForUserDefinedEnum(FName Name);
 
 	/** Updates enumerators names after name or path of the Enum was changed */
-	static void UpdateAfterPathChanged(UEnum* Enum);
+	static UNREALED_API void UpdateAfterPathChanged(UEnum* Enum);
 
 	/** adds new enumerator (with default unique name) for user defined enum */
-	static void AddNewEnumeratorForUserDefinedEnum(class UUserDefinedEnum* Enum);
+	static UNREALED_API void AddNewEnumeratorForUserDefinedEnum(class UUserDefinedEnum* Enum);
 
 	/** Removes enumerator from enum*/
-	static void RemoveEnumeratorFromUserDefinedEnum(class UUserDefinedEnum* Enum, int32 EnumeratorIndex);
+	static UNREALED_API void RemoveEnumeratorFromUserDefinedEnum(class UUserDefinedEnum* Enum, int32 EnumeratorIndex);
 
 	/**
 	 * Reorder enumerators in enum. Moves the enumerator at the given initial index to a new target index, shifting other enumerators as needed.
 	 * E.g. with enum [A, B, C, D, E], moving index 1 to index 3 results in [A, C, D, B, E].
 	 */
-	static void MoveEnumeratorInUserDefinedEnum(class UUserDefinedEnum* Enum, int32 InitialEnumeratorIndex, int32 TargetIndex);
+	static UNREALED_API void MoveEnumeratorInUserDefinedEnum(class UUserDefinedEnum* Enum, int32 InitialEnumeratorIndex, int32 TargetIndex);
 
 	/** Check if the enumerator-as-bitflags meta data is set */
-	static bool IsEnumeratorBitflagsType(class UUserDefinedEnum* Enum);
+	static UNREALED_API bool IsEnumeratorBitflagsType(class UUserDefinedEnum* Enum);
 
 	/** Set the state of the enumerator-as-bitflags meta data */
-	static void SetEnumeratorBitflagsTypeState(class UUserDefinedEnum* Enum, bool bBitflagsType);
+	static UNREALED_API void SetEnumeratorBitflagsTypeState(class UUserDefinedEnum* Enum, bool bBitflagsType);
 
 	/** check if NewName is a short name and is acceptable as name in given enum */
-	static bool IsProperNameForUserDefinedEnumerator(const UEnum* Enum, FString NewName);
+	static UNREALED_API bool IsProperNameForUserDefinedEnumerator(const UEnum* Enum, FString NewName);
 
 	/** Handles necessary notifications when the Enum has had a transaction undone or redone on it. */
-	static void PostEditUndo(UUserDefinedEnum* Enum);
+	static UNREALED_API void PostEditUndo(UUserDefinedEnum* Enum);
 
 	/*
 	 *	Try to update an out-of-date enum index after an enum's change
@@ -89,11 +89,11 @@ public:
 	 *
 	 *	@return new enum 
 	 */
-	static int64 ResolveEnumerator(const UEnum* Enum, FArchive& Ar, int64 EnumeratorValue);
+	static UNREALED_API int64 ResolveEnumerator(const UEnum* Enum, FArchive& Ar, int64 EnumeratorValue);
 
 	//DISPLAY NAME
-	static bool SetEnumeratorDisplayName(UUserDefinedEnum* Enum, int32 EnumeratorIndex, FText NewDisplayName);
-	static bool IsEnumeratorDisplayNameValid(const UUserDefinedEnum* Enum, int32 EnumeratorIndex, FText NewDisplayName);
-	static void EnsureAllDisplayNamesExist(class UUserDefinedEnum* Enum);
-	static void UpgradeDisplayNamesFromMetaData(class UUserDefinedEnum* Enum);
+	static UNREALED_API bool SetEnumeratorDisplayName(UUserDefinedEnum* Enum, int32 EnumeratorIndex, FText NewDisplayName);
+	static UNREALED_API bool IsEnumeratorDisplayNameValid(const UUserDefinedEnum* Enum, int32 EnumeratorIndex, FText NewDisplayName);
+	static UNREALED_API void EnsureAllDisplayNamesExist(class UUserDefinedEnum* Enum);
+	static UNREALED_API void UpgradeDisplayNamesFromMetaData(class UUserDefinedEnum* Enum);
 };
diff --git a/Engine/Source/Editor/UnrealEd/Public/Kismet2/Kismet2NameValidators.h b/Engine/Source/Editor/UnrealEd/Public/Kismet2/Kismet2NameValidators.h
--- a/Engine/Source/Editor/UnrealEd/Public/Kismet2/Kismet2NameValidators.h
+++ b/Engine/Source/Editor/UnrealEd/Public/Kismet2/Kismet2NameValidators.h
@@ -15,10 +15,10 @@ class UBlueprint;
 
 //////////////////////////////////////////////////////////////////////////
 // FNameValidtorFactory
-class UNREALED_API FNameValidatorFactory
+class FNameValidatorFactory
 {
 public:
-	static TSharedPtr<class INameValidatorInterface> MakeValidator(class UEdGraphNode* Node);
+	static UNREALED_API TSharedPtr<class INameValidatorInterface> MakeValidator(class UEdGraphNode* Node);
 };
 
 enum class EValidatorResult
@@ -43,7 +43,7 @@ enum class EValidatorResult
 //////////////////////////////////////////////////////////////////////////
 // FNameValidatorInterface
 
-class UNREALED_API INameValidatorInterface
+class INameValidatorInterface
 {
 public:
 	virtual ~INameValidatorInterface() {}
@@ -55,7 +55,7 @@ public:
 	virtual EValidatorResult IsValid (const FString& Name, bool bOriginal = false) = 0;
 
 	/** @return A text string describing the type of error in ErrorCode for Name */
-	static FText GetErrorText(const FString& Name, EValidatorResult ErrorCode);
+	static UNREALED_API FText GetErrorText(const FString& Name, EValidatorResult ErrorCode);
 
 	/** @return A string describing the type of error in ErrorCode for Name */
 	static FString GetErrorString(const FString& Name, EValidatorResult ErrorCode)
@@ -64,7 +64,7 @@ public:
 	}
 
 	/** @return Ok if was valid and doesn't require anything, AlreadyInUse if was already in use and is replaced with new one */
-	EValidatorResult FindValidString(FString& InOutName);
+	UNREALED_API EValidatorResult FindValidString(FString& InOutName);
 
 	/** 
 	 * Helper method to see if object exists with this name in the blueprint. Useful for 
@@ -75,23 +75,23 @@ public:
 	 *
 	 * @return True if name is not used by object in the blueprint, False otherwise
 	 */
-	static bool BlueprintObjectNameIsUnique(class UBlueprint* Blueprint, const FName& Name);
+	static UNREALED_API bool BlueprintObjectNameIsUnique(class UBlueprint* Blueprint, const FName& Name);
 };
 
 /////////////////////////////////////////////////////
 // FKismetNameValidator
-class UNREALED_API FKismetNameValidator : public INameValidatorInterface
+class FKismetNameValidator : public INameValidatorInterface
 {
 public:
-	FKismetNameValidator(const class UBlueprint* Blueprint, FName InExistingName = NAME_None, const UStruct* InScope = nullptr);
+	UNREALED_API FKismetNameValidator(const class UBlueprint* Blueprint, FName InExistingName = NAME_None, const UStruct* InScope = nullptr);
 	~FKismetNameValidator() {}
 
 	/** Return the name validator maximum string length */
-	static int32 GetMaximumNameLength();
+	static UNREALED_API int32 GetMaximumNameLength();
 
 	// Begin FNameValidatorInterface
-	virtual EValidatorResult IsValid( const FString& Name, bool bOriginal = false) override;
-	virtual EValidatorResult IsValid( const FName& Name, bool bOriginal = false) override;
+	UNREALED_API virtual EValidatorResult IsValid( const FString& Name, bool bOriginal = false) override;
+	UNREALED_API virtual EValidatorResult IsValid( const FName& Name, bool bOriginal = false) override;
 	// End FNameValidatorInterface
 private:
 	/** Name set to validate */
@@ -107,12 +107,12 @@ private:
 /////////////////////////////////////////////////////
 // FStringSetNameValidator
 
-class UNREALED_API FStringSetNameValidator : public INameValidatorInterface
+class FStringSetNameValidator : public INameValidatorInterface
 {
 public:
 	// Begin FNameValidatorInterface
-	virtual EValidatorResult IsValid(const FString& Name, bool bOriginal) override;
-	virtual EValidatorResult IsValid(const FName& Name, bool bOriginal) override;
+	UNREALED_API virtual EValidatorResult IsValid(const FString& Name, bool bOriginal) override;
+	UNREALED_API virtual EValidatorResult IsValid(const FName& Name, bool bOriginal) override;
 	// End FNameValidatorInterface
 
 	// This class is a base class for anything that just needs to validate a string is unique
@@ -130,26 +130,26 @@ protected:
 /////////////////////////////////////////////////////
 // FAnimStateTransitionNodeSharedRulesNameValidator
 
-class UNREALED_API FAnimStateTransitionNodeSharedRulesNameValidator : public FStringSetNameValidator
+class FAnimStateTransitionNodeSharedRulesNameValidator : public FStringSetNameValidator
 {
 public:
-	FAnimStateTransitionNodeSharedRulesNameValidator(class UAnimStateTransitionNode* InStateTransitionNode);
+	UNREALED_API FAnimStateTransitionNodeSharedRulesNameValidator(class UAnimStateTransitionNode* InStateTransitionNode);
 };
 
 /////////////////////////////////////////////////////
 // FAnimStateTransitionNodeSharedCrossfadeNameValidator
 
-class UNREALED_API FAnimStateTransitionNodeSharedCrossfadeNameValidator : public FStringSetNameValidator
+class FAnimStateTransitionNodeSharedCrossfadeNameValidator : public FStringSetNameValidator
 {
 public:
-	FAnimStateTransitionNodeSharedCrossfadeNameValidator(class UAnimStateTransitionNode* InStateTransitionNode);
+	UNREALED_API FAnimStateTransitionNodeSharedCrossfadeNameValidator(class UAnimStateTransitionNode* InStateTransitionNode);
 };
 
 /////////////////////////////////////////////////////
 // FDummyNameValidator
 
 // Always returns the same value
-class UNREALED_API FDummyNameValidator : public INameValidatorInterface
+class FDummyNameValidator : public INameValidatorInterface
 {
 public:
 	FDummyNameValidator(EValidatorResult InReturnValue) : ReturnValue(InReturnValue) {}
diff --git a/Engine/Source/Editor/UnrealEd/Public/Kismet2/KismetDebugUtilities.h b/Engine/Source/Editor/UnrealEd/Public/Kismet2/KismetDebugUtilities.h
--- a/Engine/Source/Editor/UnrealEd/Public/Kismet2/KismetDebugUtilities.h
+++ b/Engine/Source/Editor/UnrealEd/Public/Kismet2/KismetDebugUtilities.h
@@ -121,19 +121,19 @@ inline bool operator==(const FPropertyInstanceInfo::FPropertyInstance& A, const
 // FObjectsBeingDebuggedIterator
 
 // Helper class to iterate over all objects that should be visible in the debugger
-struct UNREALED_API FObjectsBeingDebuggedIterator
+struct FObjectsBeingDebuggedIterator
 {
 public:
-	FObjectsBeingDebuggedIterator();
+	UNREALED_API FObjectsBeingDebuggedIterator();
 
 	/** @name Element access */
 	//@{
-	UObject* operator* () const;
-	UObject* operator-> () const;
+	UNREALED_API UObject* operator* () const;
+	UNREALED_API UObject* operator-> () const;
 	//@}
 
 	/** Advances iterator to the next element in the container. */
-	FObjectsBeingDebuggedIterator& operator++();
+	UNREALED_API FObjectsBeingDebuggedIterator& operator++();
 
 	/** conversion to "bool" returning true if the iterator has not reached the last element. */
 	FORCEINLINE explicit operator bool() const
@@ -150,9 +150,9 @@ private:
 	FSelectionIterator SelectedActorsIter;
 	int32 LevelScriptActorIndex;
 private:
-	void FindNextLevelScriptActor();
-	bool IsValid() const;
-	UWorld* GetWorld() const;
+	UNREALED_API void FindNextLevelScriptActor();
+	UNREALED_API bool IsValid() const;
+	UNREALED_API UWorld* GetWorld() const;
 };
 
 
@@ -161,19 +161,19 @@ private:
 // FObjectsBeingDebuggedIterator
 
 // Helper class to iterate over all objects that should be visible in the debugger
-struct UNREALED_API FBlueprintObjectsBeingDebuggedIterator
+struct FBlueprintObjectsBeingDebuggedIterator
 {
 public:
-	FBlueprintObjectsBeingDebuggedIterator(UBlueprint* InBlueprint);
+	UNREALED_API FBlueprintObjectsBeingDebuggedIterator(UBlueprint* InBlueprint);
 
 	/** @name Element access */
 	//@{
-	UObject* operator* () const;
-	UObject* operator-> () const;
+	UNREALED_API UObject* operator* () const;
+	UNREALED_API UObject* operator-> () const;
 	//@}
 
 	/** Advances iterator to the next element in the container. */
-	FBlueprintObjectsBeingDebuggedIterator& operator++();
+	UNREALED_API FBlueprintObjectsBeingDebuggedIterator& operator++();
 
 	/** conversion to "bool" returning true if the iterator has not reached the last element. */
 	FORCEINLINE explicit operator bool() const
@@ -189,154 +189,154 @@ public:
 private:
 	UBlueprint* Blueprint;
 private:
-	bool IsValid() const;
+	UNREALED_API bool IsValid() const;
 };
 
 //////////////////////////////////////////////////////////////////////////
 // FKismetDebugUtilities
 
-class UNREALED_API FKismetDebugUtilities
+class FKismetDebugUtilities
 {
 public:
 	// Delegate for when pins are added or removed from the watchlist
 	DECLARE_MULTICAST_DELEGATE_OneParam(FOnWatchedPinsListChanged, class UBlueprint*);
 
-	static void OnScriptException(const UObject* ActiveObject, const FFrame& StackFrame, const FBlueprintExceptionInfo& Info);
+	static UNREALED_API void OnScriptException(const UObject* ActiveObject, const FFrame& StackFrame, const FBlueprintExceptionInfo& Info);
 
 	/** Returns the current instruction; if a debugging session has started; otherwise none */
-	static class UEdGraphNode* GetCurrentInstruction();
+	static UNREALED_API class UEdGraphNode* GetCurrentInstruction();
 
 	/** Returns the most recent hit breakpoint; if a debugging session has started; otherwise none */
-	static class UEdGraphNode* GetMostRecentBreakpointHit();
+	static UNREALED_API class UEdGraphNode* GetMostRecentBreakpointHit();
 
 	/** Returns the most recent hit breakpoint; if a debugging session has started in PIE/SIE; otherwise none */
-	static UWorld* GetCurrentDebuggingWorld();
+	static UNREALED_API UWorld* GetCurrentDebuggingWorld();
 
 	/** Request abort the current frame execution */
-	static void RequestAbortingExecution();
+	static UNREALED_API void RequestAbortingExecution();
 
 	/** Request an attempt to single-step to the next node */
-	static void RequestSingleStepIn();
+	static UNREALED_API void RequestSingleStepIn();
 	
 	/** Request an attempt to step over to the next node in this graph or the calling graph */
-	static void RequestStepOver();
+	static UNREALED_API void RequestStepOver();
 
 	/** Request an attempt to step out of the current graph */
-	static void RequestStepOut();
+	static UNREALED_API void RequestStepOut();
 
 	/** Called on terminatation of the current script execution so we can reset any break conditions */
-	static void EndOfScriptExecution(const FBlueprintContextTracker& BlueprintContext);
+	static UNREALED_API void EndOfScriptExecution(const FBlueprintContextTracker& BlueprintContext);
 
 	// The maximum number of trace samples to gather before overwriting old ones
 	enum { MAX_TRACE_STACK_SAMPLES = 1024 };
 
 	// Get the trace stack
-	static const TSimpleRingBuffer<FKismetTraceSample>& GetTraceStack();
+	static UNREALED_API const TSimpleRingBuffer<FKismetTraceSample>& GetTraceStack();
 
 	// Find the node that resulted in code at the specified location in the Object, or NULL if there was a problem (e.g., no debugging information was generated)
-	static class UEdGraphNode* FindSourceNodeForCodeLocation(const UObject* Object, UFunction* Function, int32 DebugOpcodeOffset, bool bAllowImpreciseHit = false);
+	static UNREALED_API class UEdGraphNode* FindSourceNodeForCodeLocation(const UObject* Object, UFunction* Function, int32 DebugOpcodeOffset, bool bAllowImpreciseHit = false);
 
 	// Return proper class for breakpoint
-	static UClass* FindClassForNode(const UObject* Object, UFunction* Function);
+	static UNREALED_API UClass* FindClassForNode(const UObject* Object, UFunction* Function);
 
 	// Notify the debugger of the start of the game frame
-	static void NotifyDebuggerOfStartOfGameFrame(UWorld* CurrentWorld);
+	static UNREALED_API void NotifyDebuggerOfStartOfGameFrame(UWorld* CurrentWorld);
 
 	// Notify the debugger of the end of the game frame
-	static void NotifyDebuggerOfEndOfGameFrame(UWorld* CurrentWorld);
+	static UNREALED_API void NotifyDebuggerOfEndOfGameFrame(UWorld* CurrentWorld);
 
 	// Whether or not we are single stepping
-	static bool IsSingleStepping();
+	static UNREALED_API bool IsSingleStepping();
 
 	// Breakpoint utils
 
 	/** Is the node a valid breakpoint target? (i.e., the node is impure and ended up generating code) */
-	static bool IsBreakpointValid(const FBlueprintBreakpoint& Breakpoint);
+	static UNREALED_API bool IsBreakpointValid(const FBlueprintBreakpoint& Breakpoint);
 
 	/** Set the node that the breakpoint should focus on */
-	static void SetBreakpointLocation(FBlueprintBreakpoint& Breakpoint, UEdGraphNode* NewNode);
+	static UNREALED_API void SetBreakpointLocation(FBlueprintBreakpoint& Breakpoint, UEdGraphNode* NewNode);
 
 	/** Set or clear the enabled flag for the breakpoint */
-	static void SetBreakpointEnabled(FBlueprintBreakpoint& Breakpoint, bool bIsEnabled);
-	static void SetBreakpointEnabled(const UEdGraphNode* OwnerNode, const UBlueprint* OwnerBlueprint, bool bIsEnabled);
+	static UNREALED_API void SetBreakpointEnabled(FBlueprintBreakpoint& Breakpoint, bool bIsEnabled);
+	static UNREALED_API void SetBreakpointEnabled(const UEdGraphNode* OwnerNode, const UBlueprint* OwnerBlueprint, bool bIsEnabled);
 
 	/** Sets this breakpoint up as a single-step breakpoint (will disable or delete itself after one go if the breakpoint wasn't already enabled) */
-	static void SetBreakpointEnabledForSingleStep(FBlueprintBreakpoint& Breakpoint, bool bDeleteAfterStep);
+	static UNREALED_API void SetBreakpointEnabledForSingleStep(FBlueprintBreakpoint& Breakpoint, bool bDeleteAfterStep);
 
 	/** Reapplies the breakpoint (used after recompiling to ensure it is set if needed) */
-	static void ReapplyBreakpoint(FBlueprintBreakpoint& Breakpoint);
+	static UNREALED_API void ReapplyBreakpoint(FBlueprintBreakpoint& Breakpoint);
 
 	/** Start the process of deleting this breakpoint */
-	static void RemoveBreakpointFromNode(const UEdGraphNode* OwnerNode, const UBlueprint* OwnerBlueprint);
+	static UNREALED_API void RemoveBreakpointFromNode(const UEdGraphNode* OwnerNode, const UBlueprint* OwnerBlueprint);
 
 	/** Update the internal state of the breakpoint when it got hit */
-	static void UpdateBreakpointStateWhenHit(const UEdGraphNode* OwnerNode, const UBlueprint* OwnerBlueprint);
+	static UNREALED_API void UpdateBreakpointStateWhenHit(const UEdGraphNode* OwnerNode, const UBlueprint* OwnerBlueprint);
 
 	/** Returns the installation site(s); don't cache these pointers! */
-	static void GetBreakpointInstallationSites(const FBlueprintBreakpoint& Breakpoint, TArray<uint8*>& InstallSites);
+	static UNREALED_API void GetBreakpointInstallationSites(const FBlueprintBreakpoint& Breakpoint, TArray<uint8*>& InstallSites);
 
 	/** Install/uninstall the breakpoint into/from the script code for the generated class that contains the node */
-	static void SetBreakpointInternal(FBlueprintBreakpoint& Breakpoint, bool bShouldBeEnabled);
+	static UNREALED_API void SetBreakpointInternal(FBlueprintBreakpoint& Breakpoint, bool bShouldBeEnabled);
 
 	/** Returns the set of valid macro source node breakpoint location(s) for the given macro instance node. The set may be empty. */
-	static void GetValidBreakpointLocations(const class UK2Node_MacroInstance* MacroInstanceNode, TArray<const UEdGraphNode*>& BreakpointLocations);
+	static UNREALED_API void GetValidBreakpointLocations(const class UK2Node_MacroInstance* MacroInstanceNode, TArray<const UEdGraphNode*>& BreakpointLocations);
 
 	/** Adds a breakpoint to the provided node */
-	static void CreateBreakpoint(const UBlueprint* Blueprint, UEdGraphNode* Node, bool bIsEnabled = true);
+	static UNREALED_API void CreateBreakpoint(const UBlueprint* Blueprint, UEdGraphNode* Node, bool bIsEnabled = true);
 
 	/**
 	 * Performs a task on every breakpoint in the provided blueprint
 	 * @param Blueprint The owning blueprint of the breakpoints to iterate
 	 * @param Task function to be called on every element
 	 */
-	static void ForeachBreakpoint(const UBlueprint* Blueprint, TFunctionRef<void(FBlueprintBreakpoint &)> Task);
+	static UNREALED_API void ForeachBreakpoint(const UBlueprint* Blueprint, TFunctionRef<void(FBlueprintBreakpoint &)> Task);
 
 	/**
 	* Removes any breakpoint that matches the provided predicate
 	* @param Blueprint The owning blueprint of the breakpoints to iterate
 	* @param Predicate function that returns true if a breakpoint should be removed
 	*/
-	static void RemoveBreakpointsByPredicate(const UBlueprint* Blueprint, const TFunctionRef<bool(const FBlueprintBreakpoint&)> Predicate);
+	static UNREALED_API void RemoveBreakpointsByPredicate(const UBlueprint* Blueprint, const TFunctionRef<bool(const FBlueprintBreakpoint&)> Predicate);
 
 	/**
 	* Returns the first breakpoint that matches the provided predicate or nullptr if nothing matched
 	* @param Blueprint The owning blueprint of the breakpoints to iterate
 	* @param Predicate function that returns true for the found breakpoint
 	*/
-	static FBlueprintBreakpoint* FindBreakpointByPredicate(const UBlueprint* Blueprint, const TFunctionRef<bool(const FBlueprintBreakpoint&)> Predicate);
+	static UNREALED_API FBlueprintBreakpoint* FindBreakpointByPredicate(const UBlueprint* Blueprint, const TFunctionRef<bool(const FBlueprintBreakpoint&)> Predicate);
 
 	/** Queries whether a blueprint has breakpoints in it */
-	static bool BlueprintHasBreakpoints(const UBlueprint* Blueprint);
+	static UNREALED_API bool BlueprintHasBreakpoints(const UBlueprint* Blueprint);
 
 	/** Handles breakpoint validation/restoration after loading the given Blueprint */
-	static void RestoreBreakpointsOnLoad(const UBlueprint* Blueprint);
+	static UNREALED_API void RestoreBreakpointsOnLoad(const UBlueprint* Blueprint);
 	
 	// Blueprint utils 
 
 	/** Duplicates debug data from original blueprint to new duplicated blueprint */
-	static void PostDuplicateBlueprint(UBlueprint* SrcBlueprint, UBlueprint* DupBlueprint, const TArray<UEdGraphNode*>& DupNodes);
+	static UNREALED_API void PostDuplicateBlueprint(UBlueprint* SrcBlueprint, UBlueprint* DupBlueprint, const TArray<UEdGraphNode*>& DupNodes);
 
 	// Looks thru the debugging data for any class variables associated with the pin
-	static class FProperty* FindClassPropertyForPin(UBlueprint* Blueprint, const UEdGraphPin* Pin);
+	static UNREALED_API class FProperty* FindClassPropertyForPin(UBlueprint* Blueprint, const UEdGraphPin* Pin);
 
 	// Looks thru the debugging data for any class variables associated with the node (e.g., temporary variables or timelines)
-	static class FProperty* FindClassPropertyForNode(UBlueprint* Blueprint, const UEdGraphNode* Node);
+	static UNREALED_API class FProperty* FindClassPropertyForNode(UBlueprint* Blueprint, const UEdGraphNode* Node);
 
 	// Is there debugging data available for this blueprint?
-	static bool HasDebuggingData(const UBlueprint* Blueprint);
+	static UNREALED_API bool HasDebuggingData(const UBlueprint* Blueprint);
 
 	/** Determines if the given pin's value can be inspected */
-	static bool CanInspectPinValue(const UEdGraphPin* Pin);
+	static UNREALED_API bool CanInspectPinValue(const UEdGraphPin* Pin);
 
 	/** Returns the breakpoint associated with a node, or NULL */
-	static FBlueprintBreakpoint* FindBreakpointForNode(const UEdGraphNode* OwnerNode, const UBlueprint* OwnerBlueprint, bool bCheckSubLocations = false);
+	static UNREALED_API FBlueprintBreakpoint* FindBreakpointForNode(const UEdGraphNode* OwnerNode, const UBlueprint* OwnerBlueprint, bool bCheckSubLocations = false);
 
 	/** Deletes all breakpoints in this blueprint */
-	static void ClearBreakpoints(const UBlueprint* Blueprint);
-	static void ClearBreakpointsForPath(const FString &BlueprintPath);
+	static UNREALED_API void ClearBreakpoints(const UBlueprint* Blueprint);
+	static UNREALED_API void ClearBreakpointsForPath(const FString &BlueprintPath);
 
 	// Notifies listeners when a watched pin is added or removed
-	static FOnWatchedPinsListChanged WatchedPinsListChangedEvent;
+	static UNREALED_API FOnWatchedPinsListChanged WatchedPinsListChangedEvent;
 
 	/**
 	 * Returns whether a pin property can be watched
@@ -344,7 +344,7 @@ public:
 	 * @param Pin The Pin to check if it can be watched
 	 * @param InPathToProperty Path to the property on the pin to watch
 	 */
-	static bool CanWatchPin(const UBlueprint* Blueprint, const UEdGraphPin* Pin, const TArray<FName>& InPathToProperty = TArray<FName>());
+	static UNREALED_API bool CanWatchPin(const UBlueprint* Blueprint, const UEdGraphPin* Pin, const TArray<FName>& InPathToProperty = TArray<FName>());
 
 	/** 
 	 * Returns whether a pin property is being watched
@@ -352,21 +352,21 @@ public:
 	 * @param Pin The Pin to check if it's being watched
 	 * @param InPathToProperty Path to the property on the pin to watch
 	 */
-	static bool IsPinBeingWatched(const UBlueprint* Blueprint, const UEdGraphPin* Pin, const TArray<FName>& InPathToProperty = TArray<FName>());
+	static UNREALED_API bool IsPinBeingWatched(const UBlueprint* Blueprint, const UEdGraphPin* Pin, const TArray<FName>& InPathToProperty = TArray<FName>());
 	
 	/**
 	 * Returns whether there are any watched properties for a given pin 
 	 * @param Blueprint The blueprint that owns the pin
 	 * @param Pin The Pin to check if it has any watches
 	 */
-	static bool DoesPinHaveWatches(const UBlueprint* Blueprint, const UEdGraphPin* Pin);
+	static UNREALED_API bool DoesPinHaveWatches(const UBlueprint* Blueprint, const UEdGraphPin* Pin);
 
 	/**
 	 * Toggles whether a pin is being watched 
 	 * @param Blueprint The blueprint that owns the pin
 	 * @param Pin The Pin to watch
 	 */
-	static void TogglePinWatch(const UBlueprint* Blueprint, const UEdGraphPin* Pin);
+	static UNREALED_API void TogglePinWatch(const UBlueprint* Blueprint, const UEdGraphPin* Pin);
 
 	/**
 	 * Removes a pin property Watch 
@@ -375,55 +375,55 @@ public:
 	 * @param InPathToProperty Path to the property on the pin to stop watching
 	 * @return true if a watch was found and removed
 	 */
-	static bool RemovePinWatch(const UBlueprint* Blueprint, const UEdGraphPin* Pin, const TArray<FName>& InPathToProperty = TArray<FName>());
+	static UNREALED_API bool RemovePinWatch(const UBlueprint* Blueprint, const UEdGraphPin* Pin, const TArray<FName>& InPathToProperty = TArray<FName>());
 
 	/** 
 	 * Adds a pin property watch 
 	 * @param Blueprint The blueprint that owns the pin
 	 * @param WatchedPin The Pin to watch
 	 */
-	static void AddPinWatch(const UBlueprint* Blueprint, FBlueprintWatchedPin&& WatchedPin);
+	static UNREALED_API void AddPinWatch(const UBlueprint* Blueprint, FBlueprintWatchedPin&& WatchedPin);
 
 	/** Removes all Watched pins from a blueprint */
-	static void ClearPinWatches(const UBlueprint* Blueprint);
+	static UNREALED_API void ClearPinWatches(const UBlueprint* Blueprint);
 
 	/** Returns whether any pins are watched for a Blueprint */
-	static bool BlueprintHasPinWatches(const UBlueprint* Blueprint);
+	static UNREALED_API bool BlueprintHasPinWatches(const UBlueprint* Blueprint);
 
 	/**
 	* Performs a task on every watched pin in the provided blueprint
 	* @param Blueprint The owning blueprint of the watched pins to iterate
 	* @param Task function to be called on every element
 	*/
-	static void ForeachPinWatch(const UBlueprint* Blueprint, TFunctionRef<void(UEdGraphPin*)> Task);
+	static UNREALED_API void ForeachPinWatch(const UBlueprint* Blueprint, TFunctionRef<void(UEdGraphPin*)> Task);
 	
 	/**
 	* Performs a task on every watched pin in the provided blueprint
 	* @param Blueprint The owning blueprint of the watched pins to iterate
 	* @param Task function to be called on every element
 	*/
-	static void ForeachPinPropertyWatch(const UBlueprint* Blueprint, TFunctionRef<void(FBlueprintWatchedPin&)> Task);
+	static UNREALED_API void ForeachPinPropertyWatch(const UBlueprint* Blueprint, TFunctionRef<void(FBlueprintWatchedPin&)> Task);
 
 	/**
 	* Removes any watched pin that matches the provided predicate
 	* @param Blueprint The owning blueprint of the watched pins to iterate
 	* @param Predicate function that returns true if a watched pin should be removed
 	*/
-	static bool RemovePinWatchesByPredicate(const UBlueprint* Blueprint, const TFunctionRef<bool(const UEdGraphPin*)> Predicate);
+	static UNREALED_API bool RemovePinWatchesByPredicate(const UBlueprint* Blueprint, const TFunctionRef<bool(const UEdGraphPin*)> Predicate);
 	
 	/**
 	* Removes any watched pin that matches the provided predicate
 	* @param Blueprint The owning blueprint of the watched pins to iterate
 	* @param Predicate function that returns true if a watched pin should be removed
 	*/
-	static bool RemovePinPropertyWatchesByPredicate(const UBlueprint* Blueprint, const TFunctionRef<bool(const FBlueprintWatchedPin&)> Predicate);
+	static UNREALED_API bool RemovePinPropertyWatchesByPredicate(const UBlueprint* Blueprint, const TFunctionRef<bool(const FBlueprintWatchedPin&)> Predicate);
 
 	/**
 	* Returns the first watched pin that matches the provided predicate or nullptr if nothing matched
 	* @param Blueprint The owning blueprint of the watched pins to iterate
 	* @param Predicate function that returns true for the found watched pin
 	*/
-	static UEdGraphPin* FindPinWatchByPredicate(const UBlueprint* Blueprint, const TFunctionRef<bool(const UEdGraphPin*)> Predicate);
+	static UNREALED_API UEdGraphPin* FindPinWatchByPredicate(const UBlueprint* Blueprint, const TFunctionRef<bool(const UEdGraphPin*)> Predicate);
 
 	enum EWatchTextResult
 	{
@@ -441,22 +441,22 @@ public:
 	};
 
 	// Gets the watched tooltip for a specified site
-	static EWatchTextResult GetWatchText(FString& OutWatchText, UBlueprint* Blueprint, UObject* ActiveObject, const UEdGraphPin* WatchPin);
+	static UNREALED_API EWatchTextResult GetWatchText(FString& OutWatchText, UBlueprint* Blueprint, UObject* ActiveObject, const UEdGraphPin* WatchPin);
 
 	// Gets the debug info for a specified site
-	static EWatchTextResult GetDebugInfo(TSharedPtr<FPropertyInstanceInfo> &OutDebugInfo, UBlueprint* Blueprint, UObject* ActiveObject, const UEdGraphPin* WatchPin);
+	static UNREALED_API EWatchTextResult GetDebugInfo(TSharedPtr<FPropertyInstanceInfo> &OutDebugInfo, UBlueprint* Blueprint, UObject* ActiveObject, const UEdGraphPin* WatchPin);
 
 	// Retrieves Debug info from a Property and a pointer to it's associated data
-	static void GetDebugInfoInternal(TSharedPtr<FPropertyInstanceInfo> &DebugInfo, const FProperty* Property, const void* PropertyValue);
+	static UNREALED_API void GetDebugInfoInternal(TSharedPtr<FPropertyInstanceInfo> &DebugInfo, const FProperty* Property, const void* PropertyValue);
 
 	//@TODO: Pretty lame way to handle this messaging, ideally the entire Info object gets pushed into the editor when intraframe debugging is triggered!
 	// This doesn't work properly if there is more than one blueprint editor open at once either (one will consume it, the others will be left in the cold)
-	static FText GetAndClearLastExceptionMessage();
+	static UNREALED_API FText GetAndClearLastExceptionMessage();
 protected:
-	static void CheckBreakConditions(UEdGraphNode* NodeStoppedAt, bool bHitBreakpoint, int32 BreakpointOffset, bool& InOutBreakExecution);
-	static void AttemptToBreakExecution(UBlueprint* BlueprintObj, const UObject* ActiveObject, const FFrame& StackFrame, const FBlueprintExceptionInfo& Info, UEdGraphNode* NodeStoppedAt, int32 DebugOpcodeOffset);
+	static UNREALED_API void CheckBreakConditions(UEdGraphNode* NodeStoppedAt, bool bHitBreakpoint, int32 BreakpointOffset, bool& InOutBreakExecution);
+	static UNREALED_API void AttemptToBreakExecution(UBlueprint* BlueprintObj, const UObject* ActiveObject, const FFrame& StackFrame, const FBlueprintExceptionInfo& Info, UEdGraphNode* NodeStoppedAt, int32 DebugOpcodeOffset);
 
-	static void GetDebugInfo_InContainer(int32 Index, TSharedPtr<FPropertyInstanceInfo> &DebugInfo, const FProperty* Property, const void* Data);
+	static UNREALED_API void GetDebugInfo_InContainer(int32 Index, TSharedPtr<FPropertyInstanceInfo> &DebugInfo, const FProperty* Property, const void* Data);
 
 	/**
 	* @brief	Helper function for converting between blueprint and debuggable data
@@ -473,29 +473,29 @@ protected:
 	* @param 	bOutIsDirectPtr	True if OutData/OutDelta point directly at the property rather than its base address
 	* @return	EWTR_Valid if the debug data could be found, otherwise an appropriate error code
 	*/
-	static EWatchTextResult FindDebuggingData(UBlueprint* Blueprint, UObject* ActiveObject, const UEdGraphPin* WatchPin, const FProperty*& OutProperty, const void*& OutData, const void*& OutDelta, UObject*& OutParent, TArray<UObject*>& SeenObjects, bool* bOutIsDirectPtr = nullptr);
+	static UNREALED_API EWatchTextResult FindDebuggingData(UBlueprint* Blueprint, UObject* ActiveObject, const UEdGraphPin* WatchPin, const FProperty*& OutProperty, const void*& OutData, const void*& OutDelta, UObject*& OutParent, TArray<UObject*>& SeenObjects, bool* bOutIsDirectPtr = nullptr);
 
 	/**	Retrieve the user settings associated with a blueprint.
 	*	returns null if the blueprint has default settings (no breakpoints and no watches) */
-	static struct FPerBlueprintSettings* GetPerBlueprintSettings(const UBlueprint* Blueprint);
+	static UNREALED_API struct FPerBlueprintSettings* GetPerBlueprintSettings(const UBlueprint* Blueprint);
 
 	/**	Retrieve the Array of breakpoints associated with a blueprint.
 	*	returns null if there are no breakpoints associated with this blueprint */
-	static TArray<FBlueprintBreakpoint>* GetBreakpoints(const UBlueprint* Blueprint);
+	static UNREALED_API TArray<FBlueprintBreakpoint>* GetBreakpoints(const UBlueprint* Blueprint);
 
 	/**	Retrieve the Array of watched pins associated with a blueprint.
 	*	returns null if there are no watched pins associated with this blueprint */
-	static TArray<FBlueprintWatchedPin>* GetWatchedPins(const UBlueprint* Blueprint);
+	static UNREALED_API TArray<FBlueprintWatchedPin>* GetWatchedPins(const UBlueprint* Blueprint);
 	
 	/** Save any modifications made to breakpoints */
-	static void SaveBlueprintEditorSettings();
+	static UNREALED_API void SaveBlueprintEditorSettings();
 
-	static void CleanupBreakpoints(const UBlueprint* Blueprint);
-	static void CleanupWatches(const UBlueprint* Blueprint);
-	static void RemoveEmptySettings(const FString& BlueprintPath);
+	static UNREALED_API void CleanupBreakpoints(const UBlueprint* Blueprint);
+	static UNREALED_API void CleanupWatches(const UBlueprint* Blueprint);
+	static UNREALED_API void RemoveEmptySettings(const FString& BlueprintPath);
 
 	/** Looks along the outer chain to find the owning world and disallowing breaking on tracepoints for editor preview and inactive worlds. */
-	static bool TracepointBreakAllowedOnOwningWorld(const UObject* ObjOuter);
+	static UNREALED_API bool TracepointBreakAllowedOnOwningWorld(const UObject* ObjOuter);
 
 private:
 	FKismetDebugUtilities() {}
diff --git a/Engine/Source/Editor/UnrealEd/Public/Kismet2/KismetEditorUtilities.h b/Engine/Source/Editor/UnrealEd/Public/Kismet2/KismetEditorUtilities.h
--- a/Engine/Source/Editor/UnrealEd/Public/Kismet2/KismetEditorUtilities.h
+++ b/Engine/Source/Editor/UnrealEd/Public/Kismet2/KismetEditorUtilities.h
@@ -67,7 +67,7 @@ ENUM_CLASS_FLAGS(EBlueprintCompileOptions)
 //////////////////////////////////////////////////////////////////////////
 // FKismetEditorUtilities
 
-class UNREALED_API FKismetEditorUtilities
+class FKismetEditorUtilities
 {
 public:
 	/** 
@@ -107,7 +107,7 @@ public:
 	 * @param CallingContext				the name of the calling method or module used to identify creation methods to engine analytics/usage stats (default None will be ignored)
 	 * @return								the new blueprint
 	 */
-	static UBlueprint* CreateBlueprint(UClass* ParentClass, UObject* Outer, const FName NewBPName, enum EBlueprintType BlueprintType, FName CallingContext = NAME_None);
+	static UNREALED_API UBlueprint* CreateBlueprint(UClass* ParentClass, UObject* Outer, const FName NewBPName, enum EBlueprintType BlueprintType, FName CallingContext = NAME_None);
 
 	/**
 	 * Create a new Blueprint and initialize it to a valid state.
@@ -121,7 +121,7 @@ public:
 	 * @param CallingContext				the name of the calling method or module used to identify creation methods to engine analytics/usage stats (default None will be ignored)
 	 * @return								the new blueprint
 	 */
-	static UBlueprint* CreateBlueprint(UClass* ParentClass, UObject* Outer, const FName NewBPName, enum EBlueprintType BlueprintType, TSubclassOf<UBlueprint> BlueprintClassType, TSubclassOf<UBlueprintGeneratedClass> BlueprintGeneratedClassType, FName CallingContext = NAME_None);
+	static UNREALED_API UBlueprint* CreateBlueprint(UClass* ParentClass, UObject* Outer, const FName NewBPName, enum EBlueprintType BlueprintType, TSubclassOf<UBlueprint> BlueprintClassType, TSubclassOf<UBlueprintGeneratedClass> BlueprintGeneratedClassType, FName CallingContext = NAME_None);
 
 	/**
 	 * Creates a user construction script graph for the blueprint.
@@ -129,18 +129,18 @@ public:
 	 * @param Blueprint					the blueprint
 	 * @return							the new UCS Graph, does not register it.
 	 */
-	static UEdGraph* CreateUserConstructionScript(UBlueprint* Blueprint);
+	static UNREALED_API UEdGraph* CreateUserConstructionScript(UBlueprint* Blueprint);
 
 	/** 
 	 * Event that's broadcast anytime a blueprint is unloaded, and becomes 
 	 * invalid (with calls to ReplaceBlueprint(), for example).
 	 */
 	DECLARE_MULTICAST_DELEGATE_OneParam(FOnBlueprintUnloaded, UBlueprint*);
-	static FOnBlueprintUnloaded OnBlueprintUnloaded;
+	static UNREALED_API FOnBlueprintUnloaded OnBlueprintUnloaded;
 
 	/** Event that's broadcast anytime a blueprint generated class is unloaded */
 	DECLARE_MULTICAST_DELEGATE_OneParam(FOnBlueprintGeneratedClassUnloaded, UBlueprintGeneratedClass*);
-	static FOnBlueprintGeneratedClassUnloaded OnBlueprintGeneratedClassUnloaded;
+	static UNREALED_API FOnBlueprintGeneratedClassUnloaded OnBlueprintGeneratedClassUnloaded;
 
 	/** 
 	 * Unloads the specified Blueprint (marking it pending-kill, and removing it 
@@ -151,7 +151,7 @@ public:
 	 * @param  Replacement	The Blueprint you cloned and used to replace Target.
 	 * @return The duplicated replacement Blueprint.
 	 */
-	static UBlueprint* ReplaceBlueprint(UBlueprint* Target, UBlueprint const* Replacement);
+	static UNREALED_API UBlueprint* ReplaceBlueprint(UBlueprint* Target, UBlueprint const* Replacement);
 
 	/** 
 	 * Determines if the specified blueprint is referenced currently in the undo 
@@ -160,22 +160,22 @@ public:
 	 * @param  Blueprint	The Blueprint you want to query about.
 	 * @return True if the Blueprint is saved in the undo buffer, false if not.
 	 */
-	static bool IsReferencedByUndoBuffer(UBlueprint* Blueprint);
+	static UNREALED_API bool IsReferencedByUndoBuffer(UBlueprint* Blueprint);
 
 	/** Create the correct event graphs for this blueprint */
-	static void CreateDefaultEventGraphs(UBlueprint* Blueprint);
+	static UNREALED_API void CreateDefaultEventGraphs(UBlueprint* Blueprint);
 
 	/** Tries to compile a blueprint, updating any actors in the editor who are using the old class, etc... */
-	static void CompileBlueprint(UBlueprint* BlueprintObj, EBlueprintCompileOptions CompileFlags = EBlueprintCompileOptions::None, class FCompilerResultsLog* pResults = nullptr );
+	static UNREALED_API void CompileBlueprint(UBlueprint* BlueprintObj, EBlueprintCompileOptions CompileFlags = EBlueprintCompileOptions::None, class FCompilerResultsLog* pResults = nullptr );
 
 	/** Generates a blueprint skeleton only.  Minimal compile, no notifications will be sent, no GC, etc.  Only successful if there isn't already a skeleton generated */
-	static bool GenerateBlueprintSkeleton(UBlueprint* BlueprintObj, bool bForceRegeneration = false);
+	static UNREALED_API bool GenerateBlueprintSkeleton(UBlueprint* BlueprintObj, bool bForceRegeneration = false);
 
 	/** Tries to make sure that a data-only blueprint is conformed to its native parent, in case any native class flags have changed */
-	static void ConformBlueprintFlagsAndComponents(UBlueprint* BlueprintObj);
+	static UNREALED_API void ConformBlueprintFlagsAndComponents(UBlueprint* BlueprintObj);
 
 	/** @return true is it's possible to create a blueprint from the specified class */
-	static bool CanCreateBlueprintOfClass(const UClass* Class);
+	static UNREALED_API bool CanCreateBlueprintOfClass(const UClass* Class);
 
 	/** 
 	 * Given an array of Actors, identifies which of those Actors are roots of attachment hierarchies (and implicitly which are attached to another actor in the array)
@@ -183,7 +183,7 @@ public:
 	 * For example if A is attached to B is attached to C and E is attached to D, and A, C, and E are in the Actors array, C and E will be in the RootActors, 
 	 * and the AttachmentMap will indicate that C has A as an attachment 
 	 */
-	static void IdentifyRootActors(const TArray<AActor*>& Actors, TArray<AActor*>& RootActors, TMap<AActor*, TArray<AActor*>>* AttachmentMap = nullptr);
+	static UNREALED_API void IdentifyRootActors(const TArray<AActor*>& Actors, TArray<AActor*>& RootActors, TMap<AActor*, TArray<AActor*>>* AttachmentMap = nullptr);
 
 	enum class EAddComponentToBPHarvestMode : uint8
 	{
@@ -222,7 +222,7 @@ public:
 	};
 
 	/** Take a list of components that belong to a single Actor and add them to a blueprint as SCSNodes */
-	static void AddComponentsToBlueprint(UBlueprint* Blueprint, const TArray<UActorComponent*>& Components, const FAddComponentsToBlueprintParams& Params = FAddComponentsToBlueprintParams());
+	static UNREALED_API void AddComponentsToBlueprint(UBlueprint* Blueprint, const TArray<UActorComponent*>& Components, const FAddComponentsToBlueprintParams& Params = FAddComponentsToBlueprintParams());
 
 	UE_DEPRECATED(4.26, "Use version that uses parameter struct instead of individual parameters")
 	static void AddComponentsToBlueprint(UBlueprint* Blueprint, const TArray<UActorComponent*>& Components, EAddComponentToBPHarvestMode HarvestMode, USCS_Node* OptionalNewRootNode = nullptr, bool bKeepMobility = false)
@@ -275,7 +275,7 @@ public:
 	};
 
 	/** Take a list of actors and add them to a blueprint as Child Actor Components */
-	static void AddActorsToBlueprint(UBlueprint* Blueprint, const TArray<AActor*>& Actors, const FAddActorsToBlueprintParams& Params = FAddActorsToBlueprintParams());
+	static UNREALED_API void AddActorsToBlueprint(UBlueprint* Blueprint, const TArray<AActor*>& Actors, const FAddActorsToBlueprintParams& Params = FAddActorsToBlueprintParams());
 
 	/** Parameter struct for customizing calls to CreateBlueprintFromActor */
 	struct FCreateBlueprintFromActorParams
@@ -312,7 +312,7 @@ public:
 	 * @param Params				The parameter struct of additional behaviors
 	 * @return The blueprint created from the actor
 	 */
-	static UBlueprint* CreateBlueprintFromActor(const FString& Path, AActor* Actor, const FCreateBlueprintFromActorParams& Params = FCreateBlueprintFromActorParams());
+	static UNREALED_API UBlueprint* CreateBlueprintFromActor(const FString& Path, AActor* Actor, const FCreateBlueprintFromActorParams& Params = FCreateBlueprintFromActorParams());
 
 	/** 
 	 * Take an Actor and generate a blueprint based on it. Uses the Actors type as the parent class. 
@@ -322,7 +322,7 @@ public:
 	 * @param Params				The parameter struct of additional behaviors
 	 * @return The blueprint created from the actor
 	 */
-	static UBlueprint* CreateBlueprintFromActor(const FName BlueprintName, UObject* Outer, AActor* Actor, const FCreateBlueprintFromActorParams& Params = FCreateBlueprintFromActorParams());
+	static UNREALED_API UBlueprint* CreateBlueprintFromActor(const FName BlueprintName, UObject* Outer, AActor* Actor, const FCreateBlueprintFromActorParams& Params = FCreateBlueprintFromActorParams());
 
 	/** 
 	 * Take an Actor and generate a blueprint based on it. Uses the Actors type as the parent class. 
@@ -414,7 +414,7 @@ public:
 	 * @param Params				The parameter struct containing actors and additional behavior definitions
 	 * @return The blueprint created from the actor
 	 */
-	static UBlueprint* CreateBlueprintFromActors(const FString& Path, const FCreateBlueprintFromActorsParams& Params);
+	static UNREALED_API UBlueprint* CreateBlueprintFromActors(const FString& Path, const FCreateBlueprintFromActorsParams& Params);
 
 	/**
 	 * Take a collection of Actors and generate a blueprint based on it
@@ -423,7 +423,7 @@ public:
 	 * @param Params				The parameter struct containing actors and additional behavior definitions
 	 * @return The blueprint created from the actor
 	 */
-	static UBlueprint* CreateBlueprintFromActors(const FName BlueprintName, UPackage* Package, const FCreateBlueprintFromActorsParams& Params);
+	static UNREALED_API UBlueprint* CreateBlueprintFromActors(const FName BlueprintName, UPackage* Package, const FCreateBlueprintFromActorsParams& Params);
 
 	/** 
 	 * Take a list of Actors and generate a blueprint based on it using the Actors as templates for child actor components.
@@ -492,7 +492,7 @@ public:
 	 * @param Params				The parameter struct containing actors and additional behavior definitions
 	 * @return The blueprint created from the actors
 	 */
-	static UBlueprint* HarvestBlueprintFromActors(const FString& Path, const TArray<AActor*>& Actors, const FHarvestBlueprintFromActorsParams& Params = FHarvestBlueprintFromActorsParams());
+	static UNREALED_API UBlueprint* HarvestBlueprintFromActors(const FString& Path, const TArray<AActor*>& Actors, const FHarvestBlueprintFromActorsParams& Params = FHarvestBlueprintFromActorsParams());
 
 	/**
 	 * Take a list of Actors and generate a blueprint by harvesting the components they have.
@@ -502,7 +502,7 @@ public:
 	 * @param Params				The parameter struct containing actors and additional behavior definitions
 	 * @return The blueprint created from the actors
 	 */
-	static UBlueprint* HarvestBlueprintFromActors(const FName BlueprintName, UPackage* Package, const TArray<AActor*>& Actors, const FHarvestBlueprintFromActorsParams& Params = FHarvestBlueprintFromActorsParams());
+	static UNREALED_API UBlueprint* HarvestBlueprintFromActors(const FName BlueprintName, UPackage* Package, const TArray<AActor*>& Actors, const FHarvestBlueprintFromActorsParams& Params = FHarvestBlueprintFromActorsParams());
 
 	/** 
 	 * Take a list of Actors and generate a blueprint  by harvesting the components they have.
@@ -528,7 +528,7 @@ public:
 	 * @param Actors				The actor list to use as the template for the blueprint.
 	 * @return The updated blueprint, or null if it failed somehow
 	 */
-	static UBlueprint* UpdateExistingBlueprintFromActors(const FString& Path, const TArray<AActor*>& Actors);
+	static UNREALED_API UBlueprint* UpdateExistingBlueprintFromActors(const FString& Path, const TArray<AActor*>& Actors);
 
 	/**
 	 * Take a list of Actors and generate a blueprint  by harvesting the components they have. Uses AActor as parent class type as the parent class.
@@ -552,7 +552,7 @@ public:
 	 * Updates this Actor's blueprint based on the actor itself. 
 	 * @return The number of properties that changes in the blueprint.
 	 */
-	static int32 ApplyInstanceChangesToBlueprint(AActor* Actor);
+	static UNREALED_API int32 ApplyInstanceChangesToBlueprint(AActor* Actor);
 
 	/** 
 	 * Creates a new blueprint instance and replaces the provided actor list with the new actor
@@ -562,7 +562,7 @@ public:
 	 * @param Rotator               The rotation of the newly created actor
 	 * @param AttachParent          The actor the newly created instance should be attached to if any
 	 */
-	static AActor* CreateBlueprintInstanceFromSelection(class UBlueprint* Blueprint, const TArray<AActor*>& SelectedActors, const FVector& Location, const FRotator& Rotator, AActor* AttachParent = nullptr);
+	static UNREALED_API AActor* CreateBlueprintInstanceFromSelection(class UBlueprint* Blueprint, const TArray<AActor*>& SelectedActors, const FVector& Location, const FRotator& Rotator, AActor* AttachParent = nullptr);
 
 	/** 
 	 * Create a new Blueprint from the supplied base class. Pops up window to let user select location and name.
@@ -570,92 +570,92 @@ public:
 	 * @param InWindowTitle			The window title
 	 * @param InParentClass			The class to create a Blueprint based on
 	 */
-	static UBlueprint* CreateBlueprintFromClass(FText InWindowTitle, UClass* InParentClass, FString NewNameSuggestion = TEXT(""));
+	static UNREALED_API UBlueprint* CreateBlueprintFromClass(FText InWindowTitle, UClass* InParentClass, FString NewNameSuggestion = TEXT(""));
 
 	/** Create a new Actor Blueprint and add the supplied asset to it. */
-	static UBlueprint* CreateBlueprintUsingAsset(UObject* Asset, bool bOpenInEditor);
+	static UNREALED_API UBlueprint* CreateBlueprintUsingAsset(UObject* Asset, bool bOpenInEditor);
 
 	/** Open a Kismet window, focusing on the specified object (either a node, or a graph).  Prefers existing windows, but will open a new application if required. */
-	static void BringKismetToFocusAttentionOnObject(const UObject* ObjectToFocusOn, bool bRequestRename=false);
+	static UNREALED_API void BringKismetToFocusAttentionOnObject(const UObject* ObjectToFocusOn, bool bRequestRename=false);
 
 	/** Open a Kismet window, focusing on the specified pin.  Prefers existing windows, but will open a new application if required. */
-	static void BringKismetToFocusAttentionOnPin(const UEdGraphPin* PinToFocusOn );
+	static UNREALED_API void BringKismetToFocusAttentionOnPin(const UEdGraphPin* PinToFocusOn );
 
 	/** Open level script kismet window and show any references to the selected actor */
-	static void ShowActorReferencesInLevelScript(const AActor* Actor);
+	static UNREALED_API void ShowActorReferencesInLevelScript(const AActor* Actor);
 
 	/** Upgrade any cosmetically stale information in a blueprint (done when edited instead of PostLoad to make certain operations easier)
 		@returns True if blueprint modified, False otherwise */
-	static void UpgradeCosmeticallyStaleBlueprint(UBlueprint* Blueprint);
+	static UNREALED_API void UpgradeCosmeticallyStaleBlueprint(UBlueprint* Blueprint);
 
 	/** Create a new event node in the level script blueprint, for the supplied Actor and event (multicast delegate property) name */
-	static void CreateNewBoundEventForActor(AActor* Actor, FName EventName);
+	static UNREALED_API void CreateNewBoundEventForActor(AActor* Actor, FName EventName);
 
 	/** Create a new event node in the  blueprint, for the supplied component, event name and blueprint */
-	static void CreateNewBoundEventForComponent(UObject* Component, FName EventName, UBlueprint* Blueprint, FObjectProperty* ComponentProperty);
+	static UNREALED_API void CreateNewBoundEventForComponent(UObject* Component, FName EventName, UBlueprint* Blueprint, FObjectProperty* ComponentProperty);
 
 	/** Create a new event node in the  blueprint, for the supplied class, event name and blueprint */
-	static void CreateNewBoundEventForClass(UClass* Class, FName EventName, UBlueprint* Blueprint, FObjectProperty* ComponentProperty);
+	static UNREALED_API void CreateNewBoundEventForClass(UClass* Class, FName EventName, UBlueprint* Blueprint, FObjectProperty* ComponentProperty);
 
 	/** Can we paste to this graph? */
-	static bool CanPasteNodes(const class UEdGraph* Graph);
+	static UNREALED_API bool CanPasteNodes(const class UEdGraph* Graph);
 
 	/** Perform paste on graph, at location */
-	static void  PasteNodesHere( class UEdGraph* Graph, const FVector2D& Location);
+	static UNREALED_API void  PasteNodesHere( class UEdGraph* Graph, const FVector2D& Location);
 
 	/** Attempt to get the bounds for currently selected nodes
 		@returns false if no nodes are selected */
-	static bool GetBoundsForSelectedNodes(const class UBlueprint* Blueprint, class FSlateRect& Rect, float Padding = 0.0f);
+	static UNREALED_API bool GetBoundsForSelectedNodes(const class UBlueprint* Blueprint, class FSlateRect& Rect, float Padding = 0.0f);
 
-	static int32 GetNumberOfSelectedNodes(const class UBlueprint* Blueprint);
+	static UNREALED_API int32 GetNumberOfSelectedNodes(const class UBlueprint* Blueprint);
 
 	/** Find the event node for this actor with the given event name */
-	static const class UK2Node_ActorBoundEvent* FindBoundEventForActor(AActor const* Actor, FName EventName);
+	static UNREALED_API const class UK2Node_ActorBoundEvent* FindBoundEventForActor(AActor const* Actor, FName EventName);
 
 	/** Find the event node for the component property with the given event name */
-	static const class UK2Node_ComponentBoundEvent* FindBoundEventForComponent(const UBlueprint* Blueprint, FName EventName, FName PropertyName);
+	static UNREALED_API const class UK2Node_ComponentBoundEvent* FindBoundEventForComponent(const UBlueprint* Blueprint, FName EventName, FName PropertyName);
 
 	/** Finds all bound component nodes for the given property on this blueprint */
-	static void FindAllBoundEventsForComponent(const UBlueprint* Blueprint, FName PropertyName, TArray<UK2Node_ComponentBoundEvent*>& OutNodes);
+	static UNREALED_API void FindAllBoundEventsForComponent(const UBlueprint* Blueprint, FName PropertyName, TArray<UK2Node_ComponentBoundEvent*>& OutNodes);
 
 	/** Returns true if the given property name has any bound component events in any blueprint graphs */
-	static bool PropertyHasBoundEvents(const UBlueprint* Blueprint, FName PropertyName);
+	static UNREALED_API bool PropertyHasBoundEvents(const UBlueprint* Blueprint, FName PropertyName);
 
 	/** Checks to see if the class is an interface class of any type, including native interfaces that are blueprint accessible */
-	static bool IsClassABlueprintInterface(const UClass* Class);
+	static UNREALED_API bool IsClassABlueprintInterface(const UClass* Class);
 
 	/** Checks to see if a given class is implementable by any blueprints, if false a native class needs to implement it */
-	static bool IsClassABlueprintImplementableInterface(const UClass* Class);
+	static UNREALED_API bool IsClassABlueprintImplementableInterface(const UClass* Class);
 
 	/** Checks to see if a blueprint can implement the specified class as an interface */
-	static bool CanBlueprintImplementInterface(UBlueprint const* Blueprint, UClass const* Class);
+	static UNREALED_API bool CanBlueprintImplementInterface(UBlueprint const* Blueprint, UClass const* Class);
 
 	/** Check to see if a given class is blueprint skeleton class. */
-	static bool IsClassABlueprintSkeleton (const UClass* Class);
+	static UNREALED_API bool IsClassABlueprintSkeleton (const UClass* Class);
 
 	/** Check to see if a given class is blueprint spawnable component class. */
-	static bool IsClassABlueprintSpawnableComponent(const UClass* Class);
+	static UNREALED_API bool IsClassABlueprintSpawnableComponent(const UClass* Class);
 
 	/** Check to see if a given class is a blueprint macro library */
-	static bool IsClassABlueprintMacroLibrary(const UClass* Class);
+	static UNREALED_API bool IsClassABlueprintMacroLibrary(const UClass* Class);
 
 	/** Run over the components in the blueprint, and then remove any that fall outside this blueprint's scope (e.g. components brought over after reparenting from another class) */
-	static void StripExternalComponents(class UBlueprint* Blueprint);
+	static UNREALED_API void StripExternalComponents(class UBlueprint* Blueprint);
 
 	/** Whether or not the specified actor is a valid target for bound events */
-	static bool IsActorValidForLevelScript(const AActor* Actor);
+	static UNREALED_API bool IsActorValidForLevelScript(const AActor* Actor);
 
 	/** 
 	 *	if bCouldAddAny is true it returns if any event can be bound in LevelScript for given Actor
 	 *	else it returns if there exists any event in level script bound with the actor
 	 */
-	static bool AnyBoundLevelScriptEventForActor(AActor* Actor, bool bCouldAddAny);
+	static UNREALED_API bool AnyBoundLevelScriptEventForActor(AActor* Actor, bool bCouldAddAny);
 
 	/** It lists bounded LevelScriptEvents for given actor */
-	static void AddLevelScriptEventOptionsForActor(UToolMenu* Menu, TWeakObjectPtr<AActor> ActorPtr, bool bExistingEvents, bool bNewEvents, bool bOnlyEventName);
+	static UNREALED_API void AddLevelScriptEventOptionsForActor(UToolMenu* Menu, TWeakObjectPtr<AActor> ActorPtr, bool bExistingEvents, bool bNewEvents, bool bOnlyEventName);
 	
 	/** Return information about the given macro graph */
-	static void GetInformationOnMacro(UEdGraph* MacroGraph, /*out*/ class UK2Node_Tunnel*& EntryNode, /*out*/ class UK2Node_Tunnel*& ExitNode, bool& bIsPure);
+	static UNREALED_API void GetInformationOnMacro(UEdGraph* MacroGraph, /*out*/ class UK2Node_Tunnel*& EntryNode, /*out*/ class UK2Node_Tunnel*& ExitNode, bool& bIsPure);
 	
 	/** 
 	 * Add information about any interfaces that have been implemented to the OutTags array
@@ -663,7 +663,7 @@ public:
 	 * @param	Blueprint		Blueprint to harvest interface data from
 	 * @param	OutTags			Array to add tags to
 	 */
-	static void AddInterfaceTags(const UBlueprint* Blueprint, TArray<UObject::FAssetRegistryTag>& OutTags);
+	static UNREALED_API void AddInterfaceTags(const UBlueprint* Blueprint, TArray<UObject::FAssetRegistryTag>& OutTags);
 
 	/**
 	 * Add a default event node to the graph, this node will also be in a disabled state and will spawn
@@ -676,7 +676,7 @@ public:
 	 * @param InOutNodePosY		Position to spawn the node at, will return with an offset more suitable to offset the next node
 	 * @return					The K2Node_Event will be returned
 	 */
-	static class UK2Node_Event* AddDefaultEventNode(UBlueprint* InBlueprint, UEdGraph* InGraph, FName InEventName, UClass* InEventClass, int32& InOutNodePosY);
+	static UNREALED_API class UK2Node_Event* AddDefaultEventNode(UBlueprint* InBlueprint, UEdGraph* InGraph, FName InEventName, UClass* InEventClass, int32& InOutNodePosY);
 
 	/**
 	 * Will add an event to the list of default event nodes to be auto-generated for the class or a child of the class
@@ -685,7 +685,7 @@ public:
 	 * @param InTargetClass		If a new Blueprint is a child of the target class, the event will attempt to be placed in the main event graph
 	 * @param InEventName		Name of event to place
 	 */
-	static void RegisterAutoGeneratedDefaultEvent(void* InOwner, UClass* InTargetClass, FName InEventName);
+	static UNREALED_API void RegisterAutoGeneratedDefaultEvent(void* InOwner, UClass* InTargetClass, FName InEventName);
 
 	/**
 	 * Will add an event to a list of callbacks to occur post Blueprint creation if the Blueprint is a child of the class
@@ -694,28 +694,28 @@ public:
 	 * @param InTargetClass						If a new Blueprint is a child of the target class, the event will attempt to be placed in the main event graph
 	 * @param InOnBlueprintCreatedCallback		Callback to call when the Blueprint is created
 	 */
-	static void RegisterOnBlueprintCreatedCallback(void* InOwner, UClass* InTargetClass, FOnBlueprintCreated InOnBlueprintCreatedCallback);
+	static UNREALED_API void RegisterOnBlueprintCreatedCallback(void* InOwner, UClass* InTargetClass, FOnBlueprintCreated InOnBlueprintCreatedCallback);
 
 	/** Unregisters a class from having AutoGeneratedDefaultEvent nodes or callbacks for OnBlueprintCreated */
-	static void UnregisterAutoBlueprintNodeCreation(void* InOwner);
+	static UNREALED_API void UnregisterAutoBlueprintNodeCreation(void* InOwner);
 
 	/** Add InNode to selection of editor */
-	static void AddToSelection(const class UEdGraph* Graph, UEdGraphNode* InNode);
+	static UNREALED_API void AddToSelection(const class UEdGraph* Graph, UEdGraphNode* InNode);
 
 	/** Get IBlueprintEditor for given object, if it exists */
-	static TSharedPtr<class IBlueprintEditor> GetIBlueprintEditorForObject(const UObject* ObjectToFocusOn, bool bOpenEditor);
+	static UNREALED_API TSharedPtr<class IBlueprintEditor> GetIBlueprintEditorForObject(const UObject* ObjectToFocusOn, bool bOpenEditor);
 private:
 	/** Stores whether we are already listening for kismet clicks */
-	static bool bIsListeningForClicksOnKismetLog;
+	static UNREALED_API bool bIsListeningForClicksOnKismetLog;
 
 	/** List of blueprint parent class names cached by IsTrackedBlueprintParent() */
-	static TArray<FString> TrackedBlueprintParentList;
+	static UNREALED_API TArray<FString> TrackedBlueprintParentList;
 
 	/** Mapping of classes to names of Events that should be automatically spawned */
-	static TMultiMap<void*, FDefaultEventNodeData> AutoGeneratedDefaultEventsMap;
+	static UNREALED_API TMultiMap<void*, FDefaultEventNodeData> AutoGeneratedDefaultEventsMap;
 
 	/** Mapping of classes to delegate callbacks when a Blueprint is created, occurs post Event node creation */
-	static TMultiMap<void*, FOnBlueprintCreatedData> OnBlueprintCreatedCallbacks;
+	static UNREALED_API TMultiMap<void*, FOnBlueprintCreatedData> OnBlueprintCreatedCallbacks;
 private:
 	
 	/**
@@ -725,7 +725,7 @@ private:
 	 * 
 	 * @return	True if the parent class is one we wish to track by reporting creation of children to analytics, otherwise false
 	 */
-	static bool IsTrackedBlueprintParent(const UClass* ParentClass);
+	static UNREALED_API bool IsTrackedBlueprintParent(const UClass* ParentClass);
 
 	FKismetEditorUtilities() {}
 };
