[UnrealEd]

diff --git a/Engine/Source/Editor/UnrealEd/Public/Kismet2/KismetReinstanceUtilities.h b/Engine/Source/Editor/UnrealEd/Public/Kismet2/KismetReinstanceUtilities.h
--- a/Engine/Source/Editor/UnrealEd/Public/Kismet2/KismetReinstanceUtilities.h
+++ b/Engine/Source/Editor/UnrealEd/Public/Kismet2/KismetReinstanceUtilities.h
@@ -27,17 +27,17 @@ ENUM_CLASS_FLAGS(EBlueprintCompileReinstancerFlags)
 
 class FReinstanceFinalizer;
 
-struct UNREALED_API FRecreateUberGraphFrameScope
+struct FRecreateUberGraphFrameScope
 {
 private:
 	TArray<UObject*> Objects;
 	UClass* RecompiledClass;
 public:
-	FRecreateUberGraphFrameScope(UClass* InClass, bool bRecreate);
-	~FRecreateUberGraphFrameScope();
+	UNREALED_API FRecreateUberGraphFrameScope(UClass* InClass, bool bRecreate);
+	UNREALED_API ~FRecreateUberGraphFrameScope();
 };
 
-struct UNREALED_API FReplaceInstancesOfClassParameters
+struct FReplaceInstancesOfClassParameters
 {
 	UE_NONCOPYABLE(FReplaceInstancesOfClassParameters)
 
@@ -108,16 +108,16 @@ struct UNREALED_API UE_DEPRECATED(5.2, "This type is no longer in use.") FBatchR
 	bool bReplaceReferencesToOldClasses = false;
 };
 
-class UNREALED_API FBlueprintCompileReinstancer : public TSharedFromThis<FBlueprintCompileReinstancer>, public FGCObject
+class FBlueprintCompileReinstancer : public TSharedFromThis<FBlueprintCompileReinstancer>, public FGCObject
 {
 protected:
 	friend struct FBlueprintCompilationManagerImpl;
 	friend struct FReinstancingJob;
 
-	static TSet<TWeakObjectPtr<UBlueprint>> CompiledBlueprintsToSave;
+	static UNREALED_API TSet<TWeakObjectPtr<UBlueprint>> CompiledBlueprintsToSave;
 
-	static UClass* HotReloadedOldClass;
-	static UClass* HotReloadedNewClass;
+	static UNREALED_API UClass* HotReloadedOldClass;
+	static UNREALED_API UClass* HotReloadedNewClass;
 
 	/** Reference to the class we're actively reinstancing */
 	UClass* ClassToReinstance;
@@ -164,61 +164,61 @@ protected:
 
 public:
 	// FSerializableObject interface
-	virtual void AddReferencedObjects(FReferenceCollector& Collector) override;
+	UNREALED_API virtual void AddReferencedObjects(FReferenceCollector& Collector) override;
 	virtual FString GetReferencerName() const override
 	{
 		return TEXT("FBlueprintCompileReinstancer");
 	}
 	// End of FSerializableObject interface
 
-	static void OptionallyRefreshNodes(UBlueprint* BP);
+	static UNREALED_API void OptionallyRefreshNodes(UBlueprint* BP);
 
 	UE_DEPRECATED(5.0, "This method performs no function and isn't invoked by the base class.  Remove all calls to the base class method.")
-	virtual void EnlistDependentBlueprintToRecompile(UBlueprint* BP, bool bBytecodeOnly);
-	virtual void BlueprintWasRecompiled(UBlueprint* BP, bool bBytecodeOnly);
+	UNREALED_API virtual void EnlistDependentBlueprintToRecompile(UBlueprint* BP, bool bBytecodeOnly);
+	UNREALED_API virtual void BlueprintWasRecompiled(UBlueprint* BP, bool bBytecodeOnly);
 
 	static TSharedPtr<FBlueprintCompileReinstancer> Create(UClass* InClassToReinstance, EBlueprintCompileReinstancerFlags Flags = EBlueprintCompileReinstancerFlags::AutoInferSaveOnCompile)
 	{
 		return MakeShareable(new FBlueprintCompileReinstancer(InClassToReinstance, Flags));
 	}
 
-	virtual ~FBlueprintCompileReinstancer();
+	UNREALED_API virtual ~FBlueprintCompileReinstancer();
 
 	/** Saves a mapping of field names to their UField equivalents, so we can remap any bytecode that references them later */
-	void SaveClassFieldMapping(UClass* InClassToReinstance);
+	UNREALED_API void SaveClassFieldMapping(UClass* InClassToReinstance);
 
 	/** Helper to gather mappings from the old class's fields to the new class's version */
-	void GenerateFieldMappings(TMap<FFieldVariant, FFieldVariant>& FieldMapping);
+	UNREALED_API void GenerateFieldMappings(TMap<FFieldVariant, FFieldVariant>& FieldMapping);
 
 	/** Reinstances all objects in the ObjectReinstancingMap */
-	void ReinstanceObjects(bool bForceAlwaysReinstance = false);
+	UNREALED_API void ReinstanceObjects(bool bForceAlwaysReinstance = false);
 
 	/** Updates references to properties and functions of the class that has in the bytecode of dependent blueprints */
-	void UpdateBytecodeReferences( TSet<UBlueprint*>& OutDependentBlueprints, TMap<FFieldVariant, FFieldVariant>& OutFieldMapping);
+	UNREALED_API void UpdateBytecodeReferences( TSet<UBlueprint*>& OutDependentBlueprints, TMap<FFieldVariant, FFieldVariant>& OutFieldMapping);
 
 	/** Consumes the set and map populated by calls to UpdateBytecodeReferences */
-	static void FinishUpdateBytecodeReferences( const TSet<UBlueprint*>& DependentBPs, const TMap<FFieldVariant, FFieldVariant>& FieldMappings);
+	static UNREALED_API void FinishUpdateBytecodeReferences( const TSet<UBlueprint*>& DependentBPs, const TMap<FFieldVariant, FFieldVariant>& FieldMappings);
 
 	// @todo_deprecated - To be removed in a future release.
 	UE_DEPRECATED(5.2, "Please use the version that takes an FReplaceInstancesOfClassParameters input instead.")
-	static void ReplaceInstancesOfClass(UClass* OldClass, UClass* NewClass, UObject* OriginalCDO = nullptr, TSet<UObject*>* ObjectsThatShouldUseOldStuff = nullptr, bool bClassObjectReplaced = false, bool bPreserveRootComponent = true);
+	static UNREALED_API void ReplaceInstancesOfClass(UClass* OldClass, UClass* NewClass, UObject* OriginalCDO = nullptr, TSet<UObject*>* ObjectsThatShouldUseOldStuff = nullptr, bool bClassObjectReplaced = false, bool bPreserveRootComponent = true);
 	UE_DEPRECATED(5.2, "Please use ReplaceInstancesOfClass() instead.")
-	static void ReplaceInstancesOfClassEx(const FReplaceInstancesOfClassParameters& Parameters );
+	static UNREALED_API void ReplaceInstancesOfClassEx(const FReplaceInstancesOfClassParameters& Parameters );
 	UE_DEPRECATED(5.2, "Please use the version that takes an FReplaceInstancesOfClassParameters input instead.")
-	static void BatchReplaceInstancesOfClass(TMap<UClass*, UClass*>& InOldToNewClassMap,
+	static UNREALED_API void BatchReplaceInstancesOfClass(TMap<UClass*, UClass*>& InOldToNewClassMap,
 PRAGMA_DISABLE_DEPRECATION_WARNINGS
 		const FBatchReplaceInstancesOfClassParameters& BatchParams = FBatchReplaceInstancesOfClassParameters()
 PRAGMA_ENABLE_DEPRECATION_WARNINGS
 	);
 
 	/** Replace all instances of OldClass with a new instance of NewClass */
-	static void ReplaceInstancesOfClass(UClass* OldClass, UClass* NewClass, const FReplaceInstancesOfClassParameters& Params);
+	static UNREALED_API void ReplaceInstancesOfClass(UClass* OldClass, UClass* NewClass, const FReplaceInstancesOfClassParameters& Params);
 
 	/** Batch replaces a mapping of one or more classes to their new class */
-	static void BatchReplaceInstancesOfClass(const TMap<UClass*, UClass*>& InOldToNewClassMap, const FReplaceInstancesOfClassParameters& Params);
+	static UNREALED_API void BatchReplaceInstancesOfClass(const TMap<UClass*, UClass*>& InOldToNewClassMap, const FReplaceInstancesOfClassParameters& Params);
 	
 	/** Function used to safely discard a CDO, so that the class can have its layout changed, callers must move parent CDOs aside before moving child CDOs aside: */
-	static UClass* MoveCDOToNewClass(UClass* OwnerClass, const TMap<UClass*, UClass*>& OldToNewMap, bool bAvoidCDODuplication);
+	static UNREALED_API UClass* MoveCDOToNewClass(UClass* OwnerClass, const TMap<UClass*, UClass*>& OldToNewMap, bool bAvoidCDODuplication);
 
 	/**
 	* Moves CDOs aside to immutable versions of classes(`REINST`) so that the CDO's can safely be GC'd.
@@ -226,13 +226,13 @@ PRAGMA_ENABLE_DEPRECATION_WARNINGS
 	* through this function again later, so we avoid O(N^2) processing of REINST classes.
 	* Maps each given `SKEL` class to its appropriate `REINST` version of itself
 	*/
-	static void MoveDependentSkelToReinst(UClass* OwnerClass, TMap<UClass*, UClass*>& OldToNewMap);
+	static UNREALED_API void MoveDependentSkelToReinst(UClass* OwnerClass, TMap<UClass*, UClass*>& OldToNewMap);
 
 	/** Gathers the full class Hierarchy of the ClassToSearch, sorted top down (0 index being UObject, n being the subclasses) */
-	static void GetSortedClassHierarchy(UClass* ClassToSearch, TArray<UClass*>& OutHierarchy, UClass** OutNativeParent);
+	static UNREALED_API void GetSortedClassHierarchy(UClass* ClassToSearch, TArray<UClass*>& OutHierarchy, UClass** OutNativeParent);
 
 	/** Returns true if the given class is a REINST class (starts with the 'REINST_' prefix) */
-	static bool IsReinstClass(const UClass* Class);
+	static UNREALED_API bool IsReinstClass(const UClass* Class);
 
 	/**
 	 * When re-instancing a component, we have to make sure all instance owners' 
@@ -241,21 +241,21 @@ PRAGMA_ENABLE_DEPRECATION_WARNINGS
 	 * 
 	 * @param  ComponentClass    Identifies the component that was altered (used to find all its instances, and thusly all instance owners).
 	 */
-	static void ReconstructOwnerInstances(TSubclassOf<UActorComponent> ComponentClass);
+	static UNREALED_API void ReconstructOwnerInstances(TSubclassOf<UActorComponent> ComponentClass);
 	
 	/** Verify that all instances of the duplicated class have been replaced and collected */
-	void VerifyReplacement();
+	UNREALED_API void VerifyReplacement();
 
 	virtual bool IsClassObjectReplaced() const { return false; }
 
-	void FinalizeFastReinstancing(TArray<UObject*>& ObjectsToReplace);
+	UNREALED_API void FinalizeFastReinstancing(TArray<UObject*>& ObjectsToReplace);
 protected:
 
-	TSharedPtr<FReinstanceFinalizer> ReinstanceInner(bool bForceAlwaysReinstance);
+	UNREALED_API TSharedPtr<FReinstanceFinalizer> ReinstanceInner(bool bForceAlwaysReinstance);
 
-	TSharedPtr<FReinstanceFinalizer> ReinstanceFast();
+	UNREALED_API TSharedPtr<FReinstanceFinalizer> ReinstanceFast();
 
-	void CompileChildren();
+	UNREALED_API void CompileChildren();
 
 	bool IsReinstancingSkeleton() const { return (ReinstClassType == RCT_BpSkeleton); }
 
@@ -278,11 +278,11 @@ protected:
 	 * @param bSkipGC					TRUE if garbage collection should be skipped
 	 * @param bAutoInferSaveOnCompile	TRUE if the reinstancer should infer whether or not save on compile should occur, FALSE if it should never save on compile
 	 */
-	FBlueprintCompileReinstancer(UClass* InClassToReinstance, EBlueprintCompileReinstancerFlags Flags = EBlueprintCompileReinstancerFlags::AutoInferSaveOnCompile);
+	UNREALED_API FBlueprintCompileReinstancer(UClass* InClassToReinstance, EBlueprintCompileReinstancerFlags Flags = EBlueprintCompileReinstancerFlags::AutoInferSaveOnCompile);
 
 	/** Reparents the specified blueprint or class to be the duplicated class in order to allow properties to be copied from the previous CDO to the new one */
-	void ReparentChild(UBlueprint* ChildBP);
-	void ReparentChild(UClass* ChildClass);
+	UNREALED_API void ReparentChild(UBlueprint* ChildBP);
+	UNREALED_API void ReparentChild(UClass* ChildClass);
 
 	/** Determine whether reinstancing actors should preserve the root component of the new actor */
 	virtual bool ShouldPreserveRootComponentOfReinstancedActor() const { return true; }
@@ -296,12 +296,12 @@ protected:
 	* @param bClearExternalReferences		If true then attempt to replace references to old classes and instances on this object with the corresponding new ones
 	* @param bForceDeltaSerialization		If true the delta serialization will be used when copying
 	*/
-	static void CopyPropertiesForUnrelatedObjects(UObject* OldObject, UObject* NewObject, bool bClearExternalReferences, bool bForceDeltaSerialization = false);
+	static UNREALED_API void CopyPropertiesForUnrelatedObjects(UObject* OldObject, UObject* NewObject, bool bClearExternalReferences, bool bForceDeltaSerialization = false);
 
 private:
 	/** Handles the work of ReplaceInstancesOfClass, handling both normal replacement of instances and batch */
-	static void ReplaceInstancesOfClass_Inner(const TMap<UClass*, UClass*>& InOldToNewClassMap, const FReplaceInstancesOfClassParameters& Params);
+	static UNREALED_API void ReplaceInstancesOfClass_Inner(const TMap<UClass*, UClass*>& InOldToNewClassMap, const FReplaceInstancesOfClassParameters& Params);
 
 	/** Returns true if A is higher up the class hierarchy  */
-	static bool ReinstancerOrderingFunction(UClass* A, UClass* B);
+	static UNREALED_API bool ReinstancerOrderingFunction(UClass* A, UClass* B);
 };
diff --git a/Engine/Source/Editor/UnrealEd/Public/Kismet2/ReloadUtilities.h b/Engine/Source/Editor/UnrealEd/Public/Kismet2/ReloadUtilities.h
--- a/Engine/Source/Editor/UnrealEd/Public/Kismet2/ReloadUtilities.h
+++ b/Engine/Source/Editor/UnrealEd/Public/Kismet2/ReloadUtilities.h
@@ -12,7 +12,7 @@
 class UBlueprint;
 class UK2Node;
 
-class UNREALED_API FReload : public IReload
+class FReload : public IReload
 {
 public:
 	struct FNativeFuncPtrMapKeyFuncs : TDefaultMapKeyFuncs<FNativeFuncPtr, FNativeFuncPtr, false>
@@ -27,33 +27,33 @@ public:
 
 public:
 
-	FReload(EActiveReloadType InType, const TCHAR* InPrefix, const TArray<UPackage*>& InPackages, FOutputDevice& InAr);
-	FReload(EActiveReloadType InType, const TCHAR* InPrefix, FOutputDevice& InAr);
+	UNREALED_API FReload(EActiveReloadType InType, const TCHAR* InPrefix, const TArray<UPackage*>& InPackages, FOutputDevice& InAr);
+	UNREALED_API FReload(EActiveReloadType InType, const TCHAR* InPrefix, FOutputDevice& InAr);
 
-	virtual ~FReload();
+	UNREALED_API virtual ~FReload();
 
 	// IReload interface 
 	virtual EActiveReloadType GetType() const override { return Type; }
 	virtual const TCHAR* GetPrefix() const override { return Prefix; };
-	virtual bool GetEnableReinstancing(bool bHasChanged) const override;
-	virtual void NotifyFunctionRemap(FNativeFuncPtr NewFunctionPointer, FNativeFuncPtr OldFunctionPointer) override;
-	virtual void NotifyChange(UClass* New, UClass* Old) override;
-	virtual void NotifyChange(UEnum* New, UEnum* Old) override;
-	virtual void NotifyChange(UScriptStruct* New, UScriptStruct* Old) override;
-	virtual void NotifyChange(UPackage* New, UPackage* Old) override;
-	virtual void Reinstance() override;
-	virtual UObject* GetReinstancedCDO(UObject* CDO) override;
-	virtual const UObject* GetReinstancedCDO(const UObject* CDO) override;
+	UNREALED_API virtual bool GetEnableReinstancing(bool bHasChanged) const override;
+	UNREALED_API virtual void NotifyFunctionRemap(FNativeFuncPtr NewFunctionPointer, FNativeFuncPtr OldFunctionPointer) override;
+	UNREALED_API virtual void NotifyChange(UClass* New, UClass* Old) override;
+	UNREALED_API virtual void NotifyChange(UEnum* New, UEnum* Old) override;
+	UNREALED_API virtual void NotifyChange(UScriptStruct* New, UScriptStruct* Old) override;
+	UNREALED_API virtual void NotifyChange(UPackage* New, UPackage* Old) override;
+	UNREALED_API virtual void Reinstance() override;
+	UNREALED_API virtual UObject* GetReinstancedCDO(UObject* CDO) override;
+	UNREALED_API virtual const UObject* GetReinstancedCDO(const UObject* CDO) override;
 
 	/**
 	 * If you wish to reuse the same reload object, invoke this method to reset the state
 	 */
-	void Reset();
+	UNREALED_API void Reset();
 
 	/**
 	 * Perform any finalize processing for reloading.
 	 */
-	void Finalize(bool bRunGC = true);
+	UNREALED_API void Finalize(bool bRunGC = true);
 
 	/**
 	 * Set the sending of the complete notification
@@ -108,22 +108,22 @@ private:
 	 * Skipping UBlueprintGeneratedClass::OverridenArchetypeForCDO as it's the
 	 * only one needed.
 	 */
-	void ReplaceReferencesToReconstructedCDOs();
+	UNREALED_API void ReplaceReferencesToReconstructedCDOs();
 
 	/**
 	 * Based on the pointers, update the given stat
 	 */
-	void UpdateStats(FReinstanceStats& Stats, void* New, void* Old);
+	UNREALED_API void UpdateStats(FReinstanceStats& Stats, void* New, void* Old);
 
 	/**
 	 * Helper method to format all the stats
 	 */
-	static void FormatStats(FStringBuilderBase& Out, const TCHAR* Singular, const TCHAR* Plural, const FReinstanceStats& Stats);
+	static UNREALED_API void FormatStats(FStringBuilderBase& Out, const TCHAR* Singular, const TCHAR* Plural, const FReinstanceStats& Stats);
 
 	/**
 	 * Helper method to format a specific stat
 	 */
-	static void FormatStat(FStringBuilderBase& Out, const TCHAR* Singular, const TCHAR* Plural, const TCHAR* What, int32 Value);
+	static UNREALED_API void FormatStat(FStringBuilderBase& Out, const TCHAR* Singular, const TCHAR* Plural, const TCHAR* What, int32 Value);
 
 	/** Type of the active reload */
 	EActiveReloadType Type = EActiveReloadType::None;
diff --git a/Engine/Source/Editor/UnrealEd/Public/Kismet2/StructureEditorUtils.h b/Engine/Source/Editor/UnrealEd/Public/Kismet2/StructureEditorUtils.h
--- a/Engine/Source/Editor/UnrealEd/Public/Kismet2/StructureEditorUtils.h
+++ b/Engine/Source/Editor/UnrealEd/Public/Kismet2/StructureEditorUtils.h
@@ -11,7 +11,7 @@ struct FStructVariableDescription;
 class UBlueprint;
 class UUserDefinedStruct;
 
-class UNREALED_API FStructureEditorUtils
+class FStructureEditorUtils
 {
 public:
 	enum EStructureEditorChangeInfo
@@ -31,7 +31,7 @@ public:
 	public:
 		UNREALED_API static FStructEditorManager& Get();
 
-		class UNREALED_API ListenerType : public InnerListenerType<FStructEditorManager>
+		class ListenerType : public InnerListenerType<FStructEditorManager>
 		{
 		};
 
@@ -68,42 +68,42 @@ public:
 	};
 
 	//STRUCTURE
-	static UUserDefinedStruct* CreateUserDefinedStruct(UObject* InParent, FName Name, EObjectFlags Flags);
+	static UNREALED_API UUserDefinedStruct* CreateUserDefinedStruct(UObject* InParent, FName Name, EObjectFlags Flags);
 
-	static void CompileStructure(UUserDefinedStruct* Struct);
+	static UNREALED_API void CompileStructure(UUserDefinedStruct* Struct);
 
-	static FString GetTooltip(const UUserDefinedStruct* Struct);
+	static UNREALED_API FString GetTooltip(const UUserDefinedStruct* Struct);
 
-	static bool ChangeTooltip(UUserDefinedStruct* Struct, const FString& InTooltip);
+	static UNREALED_API bool ChangeTooltip(UUserDefinedStruct* Struct, const FString& InTooltip);
 
 	//VARIABLE
-	static bool AddVariable(UUserDefinedStruct* Struct, const FEdGraphPinType& VarType);
+	static UNREALED_API bool AddVariable(UUserDefinedStruct* Struct, const FEdGraphPinType& VarType);
 
-	static bool RemoveVariable(UUserDefinedStruct* Struct, FGuid VarGuid);
+	static UNREALED_API bool RemoveVariable(UUserDefinedStruct* Struct, FGuid VarGuid);
 
-	static bool RenameVariable(UUserDefinedStruct* Struct, FGuid VarGuid, const FString& NewDisplayNameStr);
+	static UNREALED_API bool RenameVariable(UUserDefinedStruct* Struct, FGuid VarGuid, const FString& NewDisplayNameStr);
 
-	static bool RenameVariable(UUserDefinedStruct* Struct, const FString& OldDisplayNameStr, const FString& NewDisplayNameStr);
+	static UNREALED_API bool RenameVariable(UUserDefinedStruct* Struct, const FString& OldDisplayNameStr, const FString& NewDisplayNameStr);
 
-	static bool ChangeVariableType(UUserDefinedStruct* Struct, FGuid VarGuid, const FEdGraphPinType& NewType);
+	static UNREALED_API bool ChangeVariableType(UUserDefinedStruct* Struct, FGuid VarGuid, const FEdGraphPinType& NewType);
 
-	static bool ChangeVariableDefaultValue(UUserDefinedStruct* Struct, FGuid VarGuid, const FString& NewDefaultValue);
+	static UNREALED_API bool ChangeVariableDefaultValue(UUserDefinedStruct* Struct, FGuid VarGuid, const FString& NewDefaultValue);
 
-	static bool IsUniqueVariableFriendlyName(const UUserDefinedStruct* Struct, const FString& DisplayName);
+	static UNREALED_API bool IsUniqueVariableFriendlyName(const UUserDefinedStruct* Struct, const FString& DisplayName);
 
-	static FString GetVariableFriendlyName(const UUserDefinedStruct* Struct, FGuid VarGuid);
+	static UNREALED_API FString GetVariableFriendlyName(const UUserDefinedStruct* Struct, FGuid VarGuid);
 
-	static FString GetVariableFriendlyNameForProperty(const UUserDefinedStruct* Struct, const FProperty* Property);
+	static UNREALED_API FString GetVariableFriendlyNameForProperty(const UUserDefinedStruct* Struct, const FProperty* Property);
 
-	static FProperty* GetPropertyByFriendlyName(const UUserDefinedStruct* Struct, FString DisplayName);
+	static UNREALED_API FProperty* GetPropertyByFriendlyName(const UUserDefinedStruct* Struct, FString DisplayName);
 
-	static FString GetVariableTooltip(const UUserDefinedStruct* Struct, FGuid VarGuid);
+	static UNREALED_API FString GetVariableTooltip(const UUserDefinedStruct* Struct, FGuid VarGuid);
 	
-	static bool ChangeVariableTooltip(UUserDefinedStruct* Struct, FGuid VarGuid, const FString& InTooltip);
+	static UNREALED_API bool ChangeVariableTooltip(UUserDefinedStruct* Struct, FGuid VarGuid, const FString& InTooltip);
 
-	static bool ChangeEditableOnBPInstance(UUserDefinedStruct* Struct, FGuid VarGuid, bool bInIsEditable);
+	static UNREALED_API bool ChangeEditableOnBPInstance(UUserDefinedStruct* Struct, FGuid VarGuid, bool bInIsEditable);
 
-	static bool ChangeSaveGameEnabled(UUserDefinedStruct* Struct, FGuid VarGuid, bool bInSaveGame);
+	static UNREALED_API bool ChangeSaveGameEnabled(UUserDefinedStruct* Struct, FGuid VarGuid, bool bInSaveGame);
 
 	enum EMovePosition
 	{
@@ -121,56 +121,56 @@ public:
 	 * 
 	 * @return Whether the variable was actually moved.
 	 */
-	static bool MoveVariable(UUserDefinedStruct* Struct, FGuid MoveVarGuid, FGuid RelativeToGuid, EMovePosition Position);
+	static UNREALED_API bool MoveVariable(UUserDefinedStruct* Struct, FGuid MoveVarGuid, FGuid RelativeToGuid, EMovePosition Position);
 
 	/** Checks whether MoveVariable can actually move the variable with the given guid. */
-	static bool CanMoveVariable(UUserDefinedStruct* Struct, FGuid MoveVarGuid, FGuid RelativeToGuid, EMovePosition Position);
+	static UNREALED_API bool CanMoveVariable(UUserDefinedStruct* Struct, FGuid MoveVarGuid, FGuid RelativeToGuid, EMovePosition Position);
 
 	//Multi-line text
-	static bool CanEnableMultiLineText(const UUserDefinedStruct* Struct, FGuid VarGuid);
+	static UNREALED_API bool CanEnableMultiLineText(const UUserDefinedStruct* Struct, FGuid VarGuid);
 
-	static bool ChangeMultiLineTextEnabled(UUserDefinedStruct* Struct, FGuid VarGuid, bool bIsEnabled);
+	static UNREALED_API bool ChangeMultiLineTextEnabled(UUserDefinedStruct* Struct, FGuid VarGuid, bool bIsEnabled);
 
-	static bool IsMultiLineTextEnabled(const UUserDefinedStruct* Struct, FGuid VarGuid);
+	static UNREALED_API bool IsMultiLineTextEnabled(const UUserDefinedStruct* Struct, FGuid VarGuid);
 
 	//3D Widget
-	static bool CanEnable3dWidget(const UUserDefinedStruct* Struct, FGuid VarGuid);
+	static UNREALED_API bool CanEnable3dWidget(const UUserDefinedStruct* Struct, FGuid VarGuid);
 
-	static bool Change3dWidgetEnabled(UUserDefinedStruct* Struct, FGuid VarGuid, bool bIsEnabled);
+	static UNREALED_API bool Change3dWidgetEnabled(UUserDefinedStruct* Struct, FGuid VarGuid, bool bIsEnabled);
 
-	static bool Is3dWidgetEnabled(const UUserDefinedStruct* Struct, FGuid VarGuid);
+	static UNREALED_API bool Is3dWidgetEnabled(const UUserDefinedStruct* Struct, FGuid VarGuid);
 
 	//GUID AND VAR DESC
-	static TArray<FStructVariableDescription>& GetVarDesc(UUserDefinedStruct* Struct);
+	static UNREALED_API TArray<FStructVariableDescription>& GetVarDesc(UUserDefinedStruct* Struct);
 
-	static const TArray<FStructVariableDescription>& GetVarDesc(const UUserDefinedStruct* Struct);
+	static UNREALED_API const TArray<FStructVariableDescription>& GetVarDesc(const UUserDefinedStruct* Struct);
 
-	static TArray<FStructVariableDescription>* GetVarDescPtr(UUserDefinedStruct* Struct);
+	static UNREALED_API TArray<FStructVariableDescription>* GetVarDescPtr(UUserDefinedStruct* Struct);
 
-	static const TArray<FStructVariableDescription>* GetVarDescPtr(const UUserDefinedStruct* Struct);
+	static UNREALED_API const TArray<FStructVariableDescription>* GetVarDescPtr(const UUserDefinedStruct* Struct);
 
-	static FStructVariableDescription* GetVarDescByGuid(UUserDefinedStruct* Struct, FGuid VarGuid);
+	static UNREALED_API FStructVariableDescription* GetVarDescByGuid(UUserDefinedStruct* Struct, FGuid VarGuid);
 
-	static const FStructVariableDescription* GetVarDescByGuid(const UUserDefinedStruct* Struct, FGuid VarGuid);
+	static UNREALED_API const FStructVariableDescription* GetVarDescByGuid(const UUserDefinedStruct* Struct, FGuid VarGuid);
 
-	static FGuid GetGuidForProperty(const FProperty* Property);
+	static UNREALED_API FGuid GetGuidForProperty(const FProperty* Property);
 
-	static FProperty* GetPropertyByGuid(const UUserDefinedStruct* Struct, FGuid VarGuid);
+	static UNREALED_API FProperty* GetPropertyByGuid(const UUserDefinedStruct* Struct, FGuid VarGuid);
 
-	static FGuid GetGuidFromPropertyName(FName Name);
+	static UNREALED_API FGuid GetGuidFromPropertyName(FName Name);
 
 	//MISC
-	static void ModifyStructData(UUserDefinedStruct* Struct);
+	static UNREALED_API void ModifyStructData(UUserDefinedStruct* Struct);
 
-	static bool UserDefinedStructEnabled();
+	static UNREALED_API bool UserDefinedStructEnabled();
 
-	static void RemoveInvalidStructureMemberVariableFromBlueprint(UBlueprint* Blueprint);
+	static UNREALED_API void RemoveInvalidStructureMemberVariableFromBlueprint(UBlueprint* Blueprint);
 
 	//DEFAULT VALUE
-	static void RecreateDefaultInstanceInEditorData(UUserDefinedStruct* Struct);
+	static UNREALED_API void RecreateDefaultInstanceInEditorData(UUserDefinedStruct* Struct);
 
 	//VALIDATION
-	static bool CanHaveAMemberVariableOfType(const UUserDefinedStruct* Struct, const FEdGraphPinType& VarType, FString* OutMsg = NULL);
+	static UNREALED_API bool CanHaveAMemberVariableOfType(const UUserDefinedStruct* Struct, const FEdGraphPinType& VarType, FString* OutMsg = NULL);
 
 	enum EStructureError
 	{
@@ -184,12 +184,12 @@ public:
 	};
 
 	/** Can the structure be a member variable for a BPGClass or BPGStruct */
-	static EStructureError IsStructureValid(const UScriptStruct* Struct, const UStruct* RecursionParent = NULL, FString* OutMsg = NULL);
+	static UNREALED_API EStructureError IsStructureValid(const UScriptStruct* Struct, const UStruct* RecursionParent = NULL, FString* OutMsg = NULL);
 
 	/** called after UDS was changed by editor*/
-	static void OnStructureChanged(UUserDefinedStruct* Struct, EStructureEditorChangeInfo ChangeReason = EStructureEditorChangeInfo::Unknown);
+	static UNREALED_API void OnStructureChanged(UUserDefinedStruct* Struct, EStructureEditorChangeInfo ChangeReason = EStructureEditorChangeInfo::Unknown);
 
-	static void BroadcastPreChange(UUserDefinedStruct* Struct);
-	static void BroadcastPostChange(UUserDefinedStruct* Struct);
+	static UNREALED_API void BroadcastPreChange(UUserDefinedStruct* Struct);
+	static UNREALED_API void BroadcastPostChange(UUserDefinedStruct* Struct);
 };
 
diff --git a/Engine/Source/Editor/UnrealEd/Public/Kismet2/WatchedPin.h b/Engine/Source/Editor/UnrealEd/Public/Kismet2/WatchedPin.h
--- a/Engine/Source/Editor/UnrealEd/Public/Kismet2/WatchedPin.h
+++ b/Engine/Source/Editor/UnrealEd/Public/Kismet2/WatchedPin.h
@@ -19,25 +19,25 @@ class UEdGraphNode;
 /** Contains information about a watched pin in a Blueprint graph for local settings data.
  */
 USTRUCT()
-struct UNREALED_API FBlueprintWatchedPin
+struct FBlueprintWatchedPin
 {
 	GENERATED_BODY()
 
-	FBlueprintWatchedPin();
-	FBlueprintWatchedPin(const UEdGraphPin* Pin);
-	FBlueprintWatchedPin(const UEdGraphPin* Pin, TArray<FName>&& InPathToProperty);
+	UNREALED_API FBlueprintWatchedPin();
+	UNREALED_API FBlueprintWatchedPin(const UEdGraphPin* Pin);
+	UNREALED_API FBlueprintWatchedPin(const UEdGraphPin* Pin, TArray<FName>&& InPathToProperty);
 
 	/** Returns a reference to the underlying graph pin */
-	UEdGraphPin* Get() const;
+	UNREALED_API UEdGraphPin* Get() const;
 
 	/** Returns a reference to the path to the property we're watching on this pin */
 	const TArray<FName>& GetPathToProperty() const { return PathToProperty; }
 
 	/** Resets the pin watch to the given graph pin */
-	void SetFromPin(const UEdGraphPin* Pin);
+	UNREALED_API void SetFromPin(const UEdGraphPin* Pin);
 
 	/** Move another watched pin struct into this one */
-	void SetFromWatchedPin(FBlueprintWatchedPin&& Other);
+	UNREALED_API void SetFromWatchedPin(FBlueprintWatchedPin&& Other);
 
 	bool operator==(const FBlueprintWatchedPin& Other) const
 	{
diff --git a/Engine/Source/Editor/UnrealEd/Public/Kismet2/WildcardNodeUtils.h b/Engine/Source/Editor/UnrealEd/Public/Kismet2/WildcardNodeUtils.h
--- a/Engine/Source/Editor/UnrealEd/Public/Kismet2/WildcardNodeUtils.h
+++ b/Engine/Source/Editor/UnrealEd/Public/Kismet2/WildcardNodeUtils.h
@@ -11,7 +11,7 @@
 * A collection of utilities for creating, using, and changing the value of wildcard pin types
 * throughout the graph editor.
 */
-class UNREALED_API FWildcardNodeUtils
+class FWildcardNodeUtils
 {
 public:
 	
@@ -19,7 +19,7 @@ public:
 	* Gets the default wildcard pin type. Useful for doing comparisons on other pin types 
 	* and checks during compilation of nodes.
 	*/
-	static FEdGraphPinType GetDefaultWildcardPinType();
+	static UNREALED_API FEdGraphPinType GetDefaultWildcardPinType();
 
 	/**
 	* Checks if the given pin is in a wildcard state
@@ -27,14 +27,14 @@ public:
 	* @param	Pin		The pin the consider
 	* @return	True if the given pin is a Wildcard pin
 	*/
-	static bool IsWildcardPin(const UEdGraphPin* const Pin);
+	static UNREALED_API bool IsWildcardPin(const UEdGraphPin* const Pin);
 
 	/**
 	* Checks if the given pin is linked to any wildcard pins
 	* 
 	* @return	True if the given pin is linked to any wildcard pins
 	*/
-	static bool IsLinkedToWildcard(const UEdGraphPin* const Pin);
+	static UNREALED_API bool IsLinkedToWildcard(const UEdGraphPin* const Pin);
 
 	/**
 	* Add a default wildcard pin to the given node
@@ -45,12 +45,12 @@ public:
 	* @param ContainerType		
 	* @return	The newly created pin or nullptr if failed
 	*/
-	static UEdGraphPin* CreateWildcardPin(UEdGraphNode* Node, const FName PinName, const EEdGraphPinDirection Direction, const EPinContainerType ContainerType = EPinContainerType::None);
+	static UNREALED_API UEdGraphPin* CreateWildcardPin(UEdGraphNode* Node, const FName PinName, const EEdGraphPinDirection Direction, const EPinContainerType ContainerType = EPinContainerType::None);
 
 	/**
 	* Check this node for any wildcard pins
 	*
 	* @return	True if the given node has any wildcard pins on it
 	*/
-	static bool NodeHasAnyWildcards(const UEdGraphNode* const Node);
-};
\ No newline at end of file
+	static UNREALED_API bool NodeHasAnyWildcards(const UEdGraphNode* const Node);
+};
diff --git a/Engine/Source/Editor/UnrealEd/Public/Layers/LayersSubsystem.h b/Engine/Source/Editor/UnrealEd/Public/Layers/LayersSubsystem.h
--- a/Engine/Source/Editor/UnrealEd/Public/Layers/LayersSubsystem.h
+++ b/Engine/Source/Editor/UnrealEd/Public/Layers/LayersSubsystem.h
@@ -40,8 +40,8 @@ namespace ELayersAction
 	};
 }
 
-UCLASS()
-class UNREALED_API ULayersSubsystem : public UEditorSubsystem
+UCLASS(MinimalAPI)
+class ULayersSubsystem : public UEditorSubsystem
 {
 	GENERATED_BODY()
 
@@ -51,17 +51,17 @@ public:
 	/**
 	 *	Prepares for use
 	 */
-	virtual void Initialize(FSubsystemCollectionBase& Collection) override final;
+	UNREALED_API virtual void Initialize(FSubsystemCollectionBase& Collection) override final;
 
 	/**
 	 *	Internal cleanup
 	 */
-	virtual void Deinitialize() override final;
+	UNREALED_API virtual void Deinitialize() override final;
 
 	/**
 	 *	Destructor
 	 */
-	virtual ~ULayersSubsystem();
+	UNREALED_API virtual ~ULayersSubsystem();
 
 	/** Broadcasts whenever one or more Layers are modified*/
 	DECLARE_EVENT_ThreeParams(ULayersSubsystem, FOnLayersChanged, const ELayersAction::Type /*Action*/, const TWeakObjectPtr< ULayer >& /*ChangedLayer*/, const FName& /*ChangedProperty*/);
@@ -80,14 +80,14 @@ public:
 	 *	@param	Level	The process
 	 */
 	UFUNCTION(BlueprintCallable, Category = Layers)
-	virtual void AddLevelLayerInformation(ULevel* Level) final;
+	UNREALED_API virtual void AddLevelLayerInformation(ULevel* Level) final;
 	/**
 	 *	Purges any information regarding layers associated with the level and it contents
 	 *
 	 *	@param	Level	The process
 	 */
 	UFUNCTION(BlueprintCallable, Category = Layers)
-	virtual void RemoveLevelLayerInformation(ULevel* Level) final;
+	UNREALED_API virtual void RemoveLevelLayerInformation(ULevel* Level) final;
 
 	///////////////////////////////////////////////////////////////////////////////////////////////////////////////////////
 	// Operations on an individual actor.
@@ -98,21 +98,21 @@ public:
 	 *	@param	Actor	The actor to validate
 	 */
 	UFUNCTION(BlueprintCallable, Category = Layers)
-	virtual bool IsActorValidForLayer(AActor* Actor) final;
+	UNREALED_API virtual bool IsActorValidForLayer(AActor* Actor) final;
 	/**
 	 *	Synchronizes an newly created Actor's layers with the layer system
 	 *
 	 *	@param	Actor	The actor to initialize
 	 */
 	UFUNCTION(BlueprintCallable, Category = Layers)
-	virtual bool InitializeNewActorLayers(AActor* Actor) final;
+	UNREALED_API virtual bool InitializeNewActorLayers(AActor* Actor) final;
 	/**
 	 *	Disassociates an Actor's layers from the layer system, general used before deleting the Actor
 	 *
 	 *	@param	Actor	The actor to disassociate from the layer system
 	 */
 	UFUNCTION(BlueprintCallable, Category = Layers)
-	virtual bool DisassociateActorFromLayers(AActor* Actor) final;
+	UNREALED_API virtual bool DisassociateActorFromLayers(AActor* Actor) final;
 
 	/**
 	 * Adds the actor to the named layer.
@@ -122,7 +122,7 @@ public:
 	 * @return				true if the actor was added.  false is returned if the actor already belongs to the layer.
 	 */
 	UFUNCTION(BlueprintCallable, Category = Layers)
-	virtual bool AddActorToLayer(AActor* Actor, const FName& LayerName) final;
+	UNREALED_API virtual bool AddActorToLayer(AActor* Actor, const FName& LayerName) final;
 	/**
 	 * Adds the provided actor to the named layers.
 	 *
@@ -131,7 +131,7 @@ public:
 	 * @return				true if the actor was added to at least one of the provided layers.
 	 */
 	UFUNCTION(BlueprintCallable, Category = Layers)
-	virtual bool AddActorToLayers(AActor* Actor, const TArray< FName >& LayerNames) final;
+	UNREALED_API virtual bool AddActorToLayers(AActor* Actor, const TArray< FName >& LayerNames) final;
 
 	/**
 	 * Removes an actor from the specified layer.
@@ -141,7 +141,7 @@ public:
 	 * @return					true if the actor was removed from the layer.  false is returned if the actor already belonged to the layer.
 	 */
 	UFUNCTION(BlueprintCallable, Category = Layers)
-	virtual bool RemoveActorFromLayer(AActor* Actor, const FName& LayerToRemove, const bool bUpdateStats = true) final;
+	UNREALED_API virtual bool RemoveActorFromLayer(AActor* Actor, const FName& LayerToRemove, const bool bUpdateStats = true) final;
 	/**
 	 * Removes the provided actor from the named layers.
 	 *
@@ -150,7 +150,7 @@ public:
 	 * @return				true if the actor was removed from at least one of the provided layers.
 	 */
 	UFUNCTION(BlueprintCallable, Category = Layers)
-	virtual bool RemoveActorFromLayers(AActor* Actor, const TArray< FName >& LayerNames, const bool bUpdateStats = true) final;
+	UNREALED_API virtual bool RemoveActorFromLayers(AActor* Actor, const TArray< FName >& LayerNames, const bool bUpdateStats = true) final;
 
 
 	/////////////////////////////////////////////////
@@ -164,7 +164,7 @@ public:
 	 * @return				true if at least one actor was added to the layer.  false is returned if all the actors already belonged to the layer.
 	 */
 	UFUNCTION(BlueprintCallable, Category = Layers)
-	bool AddActorsToLayer(const TArray< AActor* >& Actors, const FName& LayerName);
+	UNREALED_API bool AddActorsToLayer(const TArray< AActor* >& Actors, const FName& LayerName);
 	/**
 	 * Add the actors to the named layer
 	 *
@@ -172,7 +172,7 @@ public:
 	 * @param	LayerName	The name of the layer to add to
 	 * @return				true if at least one actor was added to the layer.  false is returned if all the actors already belonged to the layer.
 	 */
-	virtual bool AddActorsToLayer(const TArray< TWeakObjectPtr< AActor > >& Actors, const FName& LayerName) final;
+	UNREALED_API virtual bool AddActorsToLayer(const TArray< TWeakObjectPtr< AActor > >& Actors, const FName& LayerName) final;
 	/**
 	 * Add the actors to the named layers
 	 *
@@ -181,7 +181,7 @@ public:
 	 * @return				true if at least one actor was added to at least one layer.  false is returned if all the actors already belonged to all specified layers.
 	 */
 	UFUNCTION(BlueprintCallable, Category = Layers)
-	bool AddActorsToLayers(const TArray< AActor* >& Actors, const TArray< FName >& LayerNames);
+	UNREALED_API bool AddActorsToLayers(const TArray< AActor* >& Actors, const TArray< FName >& LayerNames);
 	/**
 	 * Add the actors to the named layers
 	 *
@@ -189,7 +189,7 @@ public:
 	 * @param	LayerNames	A valid list of layer names.
 	 * @return				true if at least one actor was added to at least one layer.  false is returned if all the actors already belonged to all specified layers.
 	 */
-	virtual bool AddActorsToLayers(const TArray< TWeakObjectPtr< AActor > >& Actors, const TArray< FName >& LayerNames) final;
+	UNREALED_API virtual bool AddActorsToLayers(const TArray< TWeakObjectPtr< AActor > >& Actors, const TArray< FName >& LayerNames) final;
 
 	/**
 	*	Disassociates actors from the layer system, generally used before deleting the Actors
@@ -197,7 +197,7 @@ public:
 	*	@param	Actors	The actors to disassociate from the layer system
 	*/
 	UFUNCTION(BlueprintCallable, Category = Layers)
-	virtual bool DisassociateActorsFromLayers(const TArray< AActor* >& Actors) final;
+	UNREALED_API virtual bool DisassociateActorsFromLayers(const TArray< AActor* >& Actors) final;
 	
 	/**
 	 * Removes the actors from the specified layer.
@@ -207,7 +207,7 @@ public:
 	 * @return					true if at least one actor was removed from the layer.  false is returned if all the actors already belonged to the layer.
 	 */
 	UFUNCTION(BlueprintCallable, Category = Layers)
-	bool RemoveActorsFromLayer(const TArray< AActor* >& Actors, const FName& LayerName, const bool bUpdateStats = true);
+	UNREALED_API bool RemoveActorsFromLayer(const TArray< AActor* >& Actors, const FName& LayerName, const bool bUpdateStats = true);
 	/**
 	 * Removes the actors from the specified layer.
 	 *
@@ -215,7 +215,7 @@ public:
 	 * @param	LayerToRemove	The name of the layer to remove the actors from
 	 * @return					true if at least one actor was removed from the layer.  false is returned if all the actors already belonged to the layer.
 	 */
-	virtual bool RemoveActorsFromLayer(const TArray< TWeakObjectPtr< AActor > >& Actors, const FName& LayerName, const bool bUpdateStats = true) final;
+	UNREALED_API virtual bool RemoveActorsFromLayer(const TArray< TWeakObjectPtr< AActor > >& Actors, const FName& LayerName, const bool bUpdateStats = true) final;
 	/**
 	 * Remove the actors to the named layers
 	 *
@@ -224,7 +224,7 @@ public:
 	 * @return				true if at least one actor was removed from at least one layer.  false is returned if none of the actors belonged to any of the specified layers.
 	 */
 	UFUNCTION(BlueprintCallable, Category = Layers)
-	bool RemoveActorsFromLayers(const TArray< AActor* >& Actors, const TArray< FName >& LayerNames, const bool bUpdateStats = true);
+	UNREALED_API bool RemoveActorsFromLayers(const TArray< AActor* >& Actors, const TArray< FName >& LayerNames, const bool bUpdateStats = true);
 	/**
 	 * Remove the actors to the named layers
 	 *
@@ -232,7 +232,7 @@ public:
 	 * @param	LayerNames	A valid list of layer names.
 	 * @return				true if at least one actor was removed from at least one layer.  false is returned if none of the actors belonged to any of the specified layers.
 	 */
-	virtual bool RemoveActorsFromLayers(const TArray< TWeakObjectPtr< AActor > >& Actors, const TArray< FName >& LayerNames, const bool bUpdateStats = true) final;
+	UNREALED_API virtual bool RemoveActorsFromLayers(const TArray< TWeakObjectPtr< AActor > >& Actors, const TArray< FName >& LayerNames, const bool bUpdateStats = true) final;
 
 	
 	/////////////////////////////////////////////////
@@ -244,7 +244,7 @@ public:
 	 * @return				The selected AActor's as a TArray.
 	 */
 	UFUNCTION(BlueprintCallable, Category = Layers)
-	TArray< AActor* > GetSelectedActors() const;
+	UNREALED_API TArray< AActor* > GetSelectedActors() const;
 
 
 	/**
@@ -254,7 +254,7 @@ public:
 	 * @return				true if at least one actor was added.  false is returned if all selected actors already belong to the named layer.
 	 */
 	UFUNCTION(BlueprintCallable, Category = Layers)
-	virtual bool AddSelectedActorsToLayer(const FName& LayerName) final;
+	UNREALED_API virtual bool AddSelectedActorsToLayer(const FName& LayerName) final;
 	/**
 	 * Adds selected actors to the named layers.
 	 *
@@ -262,7 +262,7 @@ public:
 	 * @return				true if at least one actor was added.  false is returned if all selected actors already belong to the named layers.
 	 */
 	UFUNCTION(BlueprintCallable, Category = Layers)
-	virtual bool AddSelectedActorsToLayers(const TArray< FName >& LayerNames) final;
+	UNREALED_API virtual bool AddSelectedActorsToLayers(const TArray< FName >& LayerNames) final;
 
 	/**
 	 * Removes the selected actors from the named layer.
@@ -271,7 +271,7 @@ public:
 	 * @return				true if at least one actor was added.  false is returned if all selected actors already belong to the named layer.
 	 */
 	UFUNCTION(BlueprintCallable, Category = Layers)
-	virtual bool RemoveSelectedActorsFromLayer(const FName& LayerName) final;
+	UNREALED_API virtual bool RemoveSelectedActorsFromLayer(const FName& LayerName) final;
 	/**
 	 * Removes selected actors from the named layers.
 	 *
@@ -279,7 +279,7 @@ public:
 	 * @return				true if at least one actor was removed.
 	 */
 	UFUNCTION(BlueprintCallable, Category = Layers)
-	virtual bool RemoveSelectedActorsFromLayers(const TArray< FName >& LayerNames) final;
+	UNREALED_API virtual bool RemoveSelectedActorsFromLayers(const TArray< FName >& LayerNames) final;
 
 
 	/////////////////////////////////////////////////
@@ -295,7 +295,7 @@ public:
 	 * @return									true if at least one actor was selected/deselected.
 	 */
 	UFUNCTION(BlueprintCallable, Category = Layers)
-	bool SelectActorsInLayer(const FName& LayerName, const bool bSelect, const bool bNotify, const bool bSelectEvenIfHidden = false);
+	UNREALED_API bool SelectActorsInLayer(const FName& LayerName, const bool bSelect, const bool bNotify, const bool bSelectEvenIfHidden = false);
 	/**
 	 * Selects/de-selects actors belonging to the named layer.
 	 *
@@ -306,7 +306,7 @@ public:
 	 * @param	Filter	[optional]				Actor that don't pass the specified filter restrictions won't be selected.
 	 * @return									true if at least one actor was selected/deselected.
 	 */
-	virtual bool SelectActorsInLayer(const FName& LayerName, const bool bSelect, const bool bNotify, const bool bSelectEvenIfHidden, const TSharedPtr< ActorFilter >& Filter) final;
+	UNREALED_API virtual bool SelectActorsInLayer(const FName& LayerName, const bool bSelect, const bool bNotify, const bool bSelectEvenIfHidden, const TSharedPtr< ActorFilter >& Filter) final;
 	/**
 	 * Selects/de-selects actors belonging to the named layers.
 	 *
@@ -317,7 +317,7 @@ public:
 	 * @return									true if at least one actor was selected/deselected.
 	 */
 	UFUNCTION(BlueprintCallable, Category = Layers)
-	bool SelectActorsInLayers(const TArray< FName >& LayerNames, const bool bSelect, const bool bNotify, const bool bSelectEvenIfHidden = false);
+	UNREALED_API bool SelectActorsInLayers(const TArray< FName >& LayerNames, const bool bSelect, const bool bNotify, const bool bSelectEvenIfHidden = false);
 	/**
 	 * Selects/de-selects actors belonging to the named layers.
 	 *
@@ -328,7 +328,7 @@ public:
 	 * @param	Filter	[optional]				Actor that don't pass the specified filter restrictions won't be selected.
 	 * @return									true if at least one actor was selected/deselected.
 	 */
-	virtual bool SelectActorsInLayers(const TArray< FName >& LayerNames, const bool bSelect, const bool bNotify, const bool bSelectEvenIfHidden, const TSharedPtr< ActorFilter >& Filter) final;
+	UNREALED_API virtual bool SelectActorsInLayers(const TArray< FName >& LayerNames, const bool bSelect, const bool bNotify, const bool bSelectEvenIfHidden, const TSharedPtr< ActorFilter >& Filter) final;
 
 
 	/////////////////////////////////////////////////
@@ -340,14 +340,14 @@ public:
 	 * @param LayerThatChanged  If one layer was changed (toggled in view pop-up, etc), then we only need to modify actors that use that layer.
 	 */
 	UFUNCTION(BlueprintCallable, Category = Layers)
-	virtual void UpdateAllViewVisibility(const FName& LayerThatChanged) final;
+	UNREALED_API virtual void UpdateAllViewVisibility(const FName& LayerThatChanged) final;
 	/**
 	 * Updates the per-view visibility for all actors for the given view
 	 *
 	 * @param ViewportClient				The viewport client to update visibility on
 	 * @param LayerThatChanged [optional]	If one layer was changed (toggled in view pop-up, etc), then we only need to modify actors that use that layer
 	 */
-	virtual void UpdatePerViewVisibility(FLevelEditorViewportClient* ViewportClient, const FName& LayerThatChanged = NAME_Skip) final;
+	UNREALED_API virtual void UpdatePerViewVisibility(FLevelEditorViewportClient* ViewportClient, const FName& LayerThatChanged = NAME_Skip) final;
 
 	/**
 	 * Updates per-view visibility for the given actor in the given view
@@ -356,21 +356,21 @@ public:
 	 * @param Actor								Actor to update
 	 * @param bReregisterIfDirty [optional]		If true, the actor will reregister itself to give the rendering thread updated information
 	 */
-	virtual void UpdateActorViewVisibility(FLevelEditorViewportClient* ViewportClient, AActor* Actor, const bool bReregisterIfDirty = true) final;
+	UNREALED_API virtual void UpdateActorViewVisibility(FLevelEditorViewportClient* ViewportClient, AActor* Actor, const bool bReregisterIfDirty = true) final;
 	/**
 	 * Updates per-view visibility for the given actor for all views
 	 *
 	 * @param Actor		Actor to update
 	 */
 	UFUNCTION(BlueprintCallable, Category = Layers)
-	virtual void UpdateActorAllViewsVisibility(AActor* Actor) final;
+	UNREALED_API virtual void UpdateActorAllViewsVisibility(AActor* Actor) final;
 
 	/**
 	 * Removes the corresponding visibility bit from all actors (slides the later bits down 1)
 	 *
 	 * @param ViewportClient	The viewport client to update visibility on
 	 */
-	virtual void RemoveViewFromActorViewVisibility(FLevelEditorViewportClient* ViewportClient) final;
+	UNREALED_API virtual void RemoveViewFromActorViewVisibility(FLevelEditorViewportClient* ViewportClient) final;
 
 	/**
 	 * Updates the provided actors visibility in the viewports
@@ -382,7 +382,7 @@ public:
 	 * @param	bRedrawViewports			If true the viewports will be redrawn; if false, they will not
 	 */
 	UFUNCTION(BlueprintCallable, Category = Layers)
-	virtual bool UpdateActorVisibility(AActor* Actor, bool& bOutSelectionChanged, bool& bOutActorModified, const bool bNotifySelectionChange, const bool bRedrawViewports) final;
+	UNREALED_API virtual bool UpdateActorVisibility(AActor* Actor, bool& bOutSelectionChanged, bool& bOutActorModified, const bool bNotifySelectionChange, const bool bRedrawViewports) final;
 	/**
 	 * Updates the visibility of all actors in the viewports
 	 *
@@ -390,7 +390,7 @@ public:
 	 * @param	bRedrawViewports			If true the viewports will be redrawn; if false, they will not
 	 */
 	UFUNCTION(BlueprintCallable, Category = Layers)
-	virtual bool UpdateAllActorsVisibility(const bool bNotifySelectionChange, const bool bRedrawViewports) final;
+	UNREALED_API virtual bool UpdateAllActorsVisibility(const bool bNotifySelectionChange, const bool bRedrawViewports) final;
 
 
 	/////////////////////////////////////////////////
@@ -403,7 +403,7 @@ public:
 	 *	@param	InOutActors			The list to append the found actors to.
 	 */
 	UFUNCTION(BlueprintCallable, Category = Layers)
-	void AppendActorsFromLayer(const FName& LayerName, TArray< AActor* >& InOutActors) const;
+	UNREALED_API void AppendActorsFromLayer(const FName& LayerName, TArray< AActor* >& InOutActors) const;
 	/**
 	 *	Appends all the actors associated with the specified layer.
 	 *
@@ -411,7 +411,7 @@ public:
 	 *	@param	InOutActors			The list to append the found actors to.
 	 *  @param	Filter	[optional]	Actor that don't pass the specified filter restrictions won't be selected.
 	 */
-	void AppendActorsFromLayer(const FName& LayerName, TArray< AActor* >& InOutActors, const TSharedPtr< ActorFilter >& Filter) const;
+	UNREALED_API void AppendActorsFromLayer(const FName& LayerName, TArray< AActor* >& InOutActors, const TSharedPtr< ActorFilter >& Filter) const;
 	/**
 	 *	Appends all the actors associated with the specified layer.
 	 *
@@ -419,7 +419,7 @@ public:
 	 *	@param	InOutActors			The list to append the found actors to.
 	 *  @param	Filter	[optional]	Actor that don't pass the specified filter restrictions won't be selected.
 	 */
-	virtual void AppendActorsFromLayer(const FName& LayerName, TArray< TWeakObjectPtr< AActor > >& InOutActors, const TSharedPtr< ActorFilter >& Filter = TSharedPtr< ActorFilter >(nullptr)) const final;
+	UNREALED_API virtual void AppendActorsFromLayer(const FName& LayerName, TArray< TWeakObjectPtr< AActor > >& InOutActors, const TSharedPtr< ActorFilter >& Filter = TSharedPtr< ActorFilter >(nullptr)) const final;
 	/**
 	 *	Appends all the actors associated with ANY of the specified layers.
 	 *
@@ -427,7 +427,7 @@ public:
 	 *	@param	InOutActors			The list to append the found actors to.
 	 */
 	UFUNCTION(BlueprintCallable, Category = Layers)
-	void AppendActorsFromLayers(const TArray< FName >& LayerNames, TArray< AActor* >& InOutActors) const;
+	UNREALED_API void AppendActorsFromLayers(const TArray< FName >& LayerNames, TArray< AActor* >& InOutActors) const;
 	/**
 	 *	Appends all the actors associated with ANY of the specified layers.
 	 *
@@ -435,7 +435,7 @@ public:
 	 *	@param	InOutActors			The list to append the found actors to.
 	 *  @param	Filter	[optional]	Actor that don't pass the specified filter restrictions won't be selected.
 	 */
-	void AppendActorsFromLayers(const TArray< FName >& LayerNames, TArray< AActor* >& InOutActors, const TSharedPtr< ActorFilter >& Filter) const;
+	UNREALED_API void AppendActorsFromLayers(const TArray< FName >& LayerNames, TArray< AActor* >& InOutActors, const TSharedPtr< ActorFilter >& Filter) const;
 	/**
 	 *	Appends all the actors associated with ANY of the specified layers.
 	 *
@@ -443,7 +443,7 @@ public:
 	 *	@param	InOutActors			The list to append the found actors to.
 	 *  @param	Filter	[optional]	Actor that don't pass the specified filter restrictions won't be selected.
 	 */
-	virtual void AppendActorsFromLayers(const TArray< FName >& LayerNames, TArray< TWeakObjectPtr< AActor > >& InOutActors, const TSharedPtr< ActorFilter >& Filter = TSharedPtr< ActorFilter >(nullptr)) const final;
+	UNREALED_API virtual void AppendActorsFromLayers(const TArray< FName >& LayerNames, TArray< TWeakObjectPtr< AActor > >& InOutActors, const TSharedPtr< ActorFilter >& Filter = TSharedPtr< ActorFilter >(nullptr)) const final;
 
 	/**
 	 *	Gets all the actors associated with the specified layer. Analog to AppendActorsFromLayer but it returns rather than appends the actors.
@@ -452,7 +452,7 @@ public:
 	 *	@return						The list to assign the found actors to.
 	 */
 	UFUNCTION(BlueprintCallable, Category = Layers)
-	TArray< AActor* > GetActorsFromLayer(const FName& LayerName) const;
+	UNREALED_API TArray< AActor* > GetActorsFromLayer(const FName& LayerName) const;
 	/**
 	 *	Gets all the actors associated with the specified layer. Analog to AppendActorsFromLayer but it returns rather than appends the actors.
 	 *
@@ -460,7 +460,7 @@ public:
 	 *  @param	Filter	[optional]	Actor that don't pass the specified filter restrictions won't be selected.
 	 *	@return						The list to assign the found actors to.
 	 */
-	TArray< AActor* > GetActorsFromLayer(const FName& LayerName, const TSharedPtr< ActorFilter >& Filter) const;
+	UNREALED_API TArray< AActor* > GetActorsFromLayer(const FName& LayerName, const TSharedPtr< ActorFilter >& Filter) const;
 	/**
 	 *	Gets all the actors associated with ANY of the specified layers. Analog to AppendActorsFromLayers but it returns rather than appends the actors.
 	 *
@@ -468,7 +468,7 @@ public:
 	 *	@return						The list to assign the found actors to.
 	 */
 	UFUNCTION(BlueprintCallable, Category = Layers)
-	TArray< AActor* > GetActorsFromLayers(const TArray< FName >& LayerNames) const;
+	UNREALED_API TArray< AActor* > GetActorsFromLayers(const TArray< FName >& LayerNames) const;
 	/**
 	 *	Gets all the actors associated with ANY of the specified layers. Analog to AppendActorsFromLayers but it returns rather than appends the actors.
 	 *
@@ -476,7 +476,7 @@ public:
 	 *  @param	Filter	[optional]	Actor that don't pass the specified filter restrictions won't be selected.
 	 *	@return						The list to assign the found actors to.
 	 */
-	TArray< AActor* > GetActorsFromLayers(const TArray< FName >& LayerNames, const TSharedPtr< ActorFilter >& Filter) const;
+	UNREALED_API TArray< AActor* > GetActorsFromLayers(const TArray< FName >& LayerNames, const TSharedPtr< ActorFilter >& Filter) const;
 
 	/**
 	 * Changes the named layer's visibility to the provided state
@@ -485,7 +485,7 @@ public:
 	 * @param	bIsVisible	If true the layer will be visible; if false, the layer will not be visible.
 	 */
 	UFUNCTION(BlueprintCallable, Category = Layers)
-	virtual void SetLayerVisibility(const FName& LayerName, const bool bIsVisible) final;
+	UNREALED_API virtual void SetLayerVisibility(const FName& LayerName, const bool bIsVisible) final;
 	/**
 	 * Changes visibility of the named layers to the provided state
 	 *
@@ -493,7 +493,7 @@ public:
 	 * @param	bIsVisible	If true the layers will be visible; if false, the layers will not be visible
 	 */
 	UFUNCTION(BlueprintCallable, Category = Layers)
-	virtual void SetLayersVisibility(const TArray< FName >& LayerNames, const bool bIsVisible) final;
+	UNREALED_API virtual void SetLayersVisibility(const TArray< FName >& LayerNames, const bool bIsVisible) final;
 
 	/**
 	 * Toggles the named layer's visibility
@@ -501,20 +501,20 @@ public:
 	 * @param LayerName	The name of the layer to affect
 	 */
 	UFUNCTION(BlueprintCallable, Category = Layers)
-	virtual void ToggleLayerVisibility(const FName& LayerName) final;
+	UNREALED_API virtual void ToggleLayerVisibility(const FName& LayerName) final;
 	/**
 	 * Toggles the visibility of all of the named layers
 	 *
 	 * @param	LayerNames	The names of the layers to affect
 	 */
 	UFUNCTION(BlueprintCallable, Category = Layers)
-	virtual void ToggleLayersVisibility(const TArray< FName >& LayerNames) final;
+	UNREALED_API virtual void ToggleLayersVisibility(const TArray< FName >& LayerNames) final;
 
 	/**
 	 * Set the visibility of all layers to true
 	 */
 	UFUNCTION(BlueprintCallable, Category = Layers)
-	virtual void MakeAllLayersVisible() final;
+	UNREALED_API virtual void MakeAllLayersVisible() final;
 
 	/**
 	 * Gets the ULayer Object of the named layer
@@ -523,7 +523,7 @@ public:
 	 * @return				The ULayer Object of the provided layer name
 	 */
 	UFUNCTION(BlueprintCallable, Category = Layers)
-	ULayer* GetLayer(const FName& LayerName) const;
+	UNREALED_API ULayer* GetLayer(const FName& LayerName) const;
 	/**
 	 * Checks whether the ULayer Object of the provided layer name exists.
 	 *
@@ -531,7 +531,7 @@ public:
 	 * @return					If true a valid ULayer Object was found; if false, a valid ULayer object was not found
 	 */
 	UFUNCTION(BlueprintCallable, Category = Layers)
-	bool IsLayer(const FName& LayerName);
+	UNREALED_API bool IsLayer(const FName& LayerName);
 	/**
 	 * Attempts to get the ULayer Object of the provided layer name.
 	 *
@@ -540,7 +540,7 @@ public:
 	 * @return					If true a valid ULayer Object was found and set to OutLayer; if false, a valid ULayer object was not found and invalid set to OutLayer
 	 */
 	UFUNCTION(BlueprintCallable, Category = Layers)
-	virtual bool TryGetLayer(const FName& LayerName, ULayer*& OutLayer) final;
+	UNREALED_API virtual bool TryGetLayer(const FName& LayerName, ULayer*& OutLayer) final;
 
 	/**
 	 * Gets all known layers and appends their names to the provide array
@@ -548,20 +548,20 @@ public:
 	 * @param OutLayers[OUT] Output array to store all known layers
 	 */
 	UFUNCTION(BlueprintCallable, Category = Layers)
-	virtual void AddAllLayerNamesTo(TArray< FName >& OutLayerNames) const final;
+	UNREALED_API virtual void AddAllLayerNamesTo(TArray< FName >& OutLayerNames) const final;
 	/**
 	 * Gets all known layers and appends them to the provided array
 	 *
 	 * @param OutLayers[OUT] Output array to store all known layers
 	 */
 	UFUNCTION(BlueprintCallable, Category = Layers)
-	virtual void AddAllLayersTo(TArray< ULayer* > & OutLayers) const final;
+	UNREALED_API virtual void AddAllLayersTo(TArray< ULayer* > & OutLayers) const final;
 	/**
 	 * Gets all known layers and appends them to the provided array
 	 *
 	 * @param OutLayers[OUT] Output array to store all known layers
 	 */
-	virtual void AddAllLayersTo(TArray< TWeakObjectPtr< ULayer > >& OutLayers) const final;
+	UNREALED_API virtual void AddAllLayersTo(TArray< TWeakObjectPtr< ULayer > >& OutLayers) const final;
 
 	/**
 	 * Creates a ULayer Object for the named layer
@@ -570,7 +570,7 @@ public:
 	 * @return				The newly created ULayer Object for the named layer
 	 */
 	UFUNCTION(BlueprintCallable, Category = Layers)
-	ULayer* CreateLayer(const FName& LayerName);
+	UNREALED_API ULayer* CreateLayer(const FName& LayerName);
 
 	/**
 	 * Deletes all of the provided layers, disassociating all actors from them
@@ -578,14 +578,14 @@ public:
 	 * @param LayersToDelete	A valid list of layer names.
 	 */
 	UFUNCTION(BlueprintCallable, Category = Layers)
-	virtual void DeleteLayers(const TArray< FName >& LayersToDelete) final;
+	UNREALED_API virtual void DeleteLayers(const TArray< FName >& LayersToDelete) final;
 	/**
 	 * Deletes the provided layer, disassociating all actors from them
 	 *
 	 * @param LayerToDelete		A valid layer name
 	 */
 	UFUNCTION(BlueprintCallable, Category = Layers)
-	virtual void DeleteLayer(const FName& LayerToDelete) final;
+	UNREALED_API virtual void DeleteLayer(const FName& LayerToDelete) final;
 
 	/**
 	 * Renames the provided originally named layer to the provided new name
@@ -594,7 +594,7 @@ public:
 	 * @param	NewLayerName		The new name for the layer to be renamed
 	 */
 	UFUNCTION(BlueprintCallable, Category = Layers)
-	virtual bool RenameLayer(const FName& OriginalLayerName, const FName& NewLayerName) final;
+	UNREALED_API virtual bool RenameLayer(const FName& OriginalLayerName, const FName& NewLayerName) final;
 
 	/**
 	 * Get the current UWorld object.
@@ -602,26 +602,26 @@ public:
 	 * @return						The UWorld* object
 	 */
 	UFUNCTION(BlueprintCallable, Category = Layers)
-	UWorld* GetWorld() const; // Fallback to GWorld
+	UNREALED_API UWorld* GetWorld() const; // Fallback to GWorld
 
 	/**
 	 * Delegate handler for FEditorDelegates::MapChange. It internally calls LayersChanged.Broadcast.
 	 **/
 	UFUNCTION(BlueprintCallable, Category = Layers)
-	void EditorMapChange();
+	UNREALED_API void EditorMapChange();
 	/**
 	 * Delegate handler for FEditorDelegates::RefreshLayerBrowser. It internally calls UpdateAllActorsVisibility to refresh the actors of each layer.
 	 **/
 	UFUNCTION(BlueprintCallable, Category = Layers)
-	void EditorRefreshLayerBrowser();
+	UNREALED_API void EditorRefreshLayerBrowser();
 	/**
 	 * Delegate handler for FEditorDelegates::PostUndoRedo. It internally calls LayersChanged.Broadcast and UpdateAllActorsVisibility to refresh the actors of each layer.
 	 **/
-	void PostUndoRedo();
+	UNREALED_API void PostUndoRedo();
 
 private:
-	void AddActorToStats(ULayer* Layer, AActor* Actor);
-	void RemoveActorFromStats(ULayer* Layer, AActor* Actor);
+	UNREALED_API void AddActorToStats(ULayer* Layer, AActor* Actor);
+	UNREALED_API void RemoveActorFromStats(ULayer* Layer, AActor* Actor);
 
 	/**
 	 *	Checks to see if the named layer exists, and if it doesn't creates it.
@@ -629,7 +629,7 @@ private:
 	 * @param	LayerName	A valid layer name
 	 * @return				The ULayer Object of the named layer
 	 */
-	ULayer* EnsureLayerExists(const FName& LayerName);
+	UNREALED_API ULayer* EnsureLayerExists(const FName& LayerName);
 
 	/**	Fires whenever one or more layer changes */
 	FOnLayersChanged LayersChanged;
diff --git a/Engine/Source/Editor/UnrealEd/Public/LevelEditorDragDropHandler.h b/Engine/Source/Editor/UnrealEd/Public/LevelEditorDragDropHandler.h
--- a/Engine/Source/Editor/UnrealEd/Public/LevelEditorDragDropHandler.h
+++ b/Engine/Source/Editor/UnrealEd/Public/LevelEditorDragDropHandler.h
@@ -9,18 +9,18 @@
 class FViewport;
 struct FAssetData;
 
-UCLASS(transient)
-class UNREALED_API ULevelEditorDragDropHandler : public UObject
+UCLASS(transient, MinimalAPI)
+class ULevelEditorDragDropHandler : public UObject
 {
 public:
 	GENERATED_BODY()
 	
 public:
-	ULevelEditorDragDropHandler();
+	UNREALED_API ULevelEditorDragDropHandler();
 	
-	virtual bool PreviewDropObjectsAtCoordinates(int32 MouseX, int32 MouseY, UWorld* World, FViewport* Viewport, const FAssetData& AssetData);
-	virtual bool PreDropObjectsAtCoordinates(int32 MouseX, int32 MouseY, UWorld* World, FViewport* Viewport, const TArray<UObject*>& DroppedObjects, TArray<AActor*>& OutNewActors);
-	virtual bool PostDropObjectsAtCoordinates(int32 MouseX, int32 MouseY, UWorld* World, FViewport* Viewport, const TArray<UObject*>& DroppedObjects);
+	UNREALED_API virtual bool PreviewDropObjectsAtCoordinates(int32 MouseX, int32 MouseY, UWorld* World, FViewport* Viewport, const FAssetData& AssetData);
+	UNREALED_API virtual bool PreDropObjectsAtCoordinates(int32 MouseX, int32 MouseY, UWorld* World, FViewport* Viewport, const TArray<UObject*>& DroppedObjects, TArray<AActor*>& OutNewActors);
+	UNREALED_API virtual bool PostDropObjectsAtCoordinates(int32 MouseX, int32 MouseY, UWorld* World, FViewport* Viewport, const TArray<UObject*>& DroppedObjects);
 
 	bool GetCanDrop() const { return bCanDrop; }
 	FText GetPreviewDropHintText() const { return HintText; }
@@ -33,4 +33,4 @@ protected:
 
 	/** Optional hint text that may be returned to the user. */
 	FText HintText;
-};
\ No newline at end of file
+};
diff --git a/Engine/Source/Editor/UnrealEd/Public/LevelEditorViewport.h b/Engine/Source/Editor/UnrealEd/Public/LevelEditorViewport.h
--- a/Engine/Source/Editor/UnrealEd/Public/LevelEditorViewport.h
+++ b/Engine/Source/Editor/UnrealEd/Public/LevelEditorViewport.h
@@ -69,7 +69,7 @@ struct FViewportHoverTarget
 	}
 };
 
-struct UNREALED_API FTrackingTransaction
+struct FTrackingTransaction
 {
 	/** State of this transaction */
 	struct ETransactionState
@@ -82,23 +82,23 @@ struct UNREALED_API FTrackingTransaction
 		};
 	};
 
-	FTrackingTransaction();
-	~FTrackingTransaction();
+	UNREALED_API FTrackingTransaction();
+	UNREALED_API ~FTrackingTransaction();
 
 	/**
 	 * Initiates a transaction.
 	 */
-	void Begin(const FText& Description, AActor* AdditionalActor = nullptr);
+	UNREALED_API void Begin(const FText& Description, AActor* AdditionalActor = nullptr);
 
-	void End();
+	UNREALED_API void End();
 
-	void Cancel();
+	UNREALED_API void Cancel();
 
 	/** Begin a pending transaction, which won't become a real transaction until PromotePendingToActive is called */
-	void BeginPending(const FText& Description);
+	UNREALED_API void BeginPending(const FText& Description);
 
 	/** Promote a pending transaction (if any) to an active transaction */
-	void PromotePendingToActive();
+	UNREALED_API void PromotePendingToActive();
 
 	bool IsActive() const { return TrackingTransactionState == ETransactionState::Active; }
 
@@ -108,11 +108,11 @@ struct UNREALED_API FTrackingTransaction
 
 private:
 
-	const UTypedElementSelectionSet* GetSelectionSet() const;
-	UTypedElementSelectionSet* GetMutableSelectionSet() const;
+	UNREALED_API const UTypedElementSelectionSet* GetSelectionSet() const;
+	UNREALED_API UTypedElementSelectionSet* GetMutableSelectionSet() const;
 
 	/** Editor selection changed delegate handler */	
-	void OnEditorSelectionChanged(const UTypedElementSelectionSet* InSelectionSet);
+	UNREALED_API void OnEditorSelectionChanged(const UTypedElementSelectionSet* InSelectionSet);
 
 	/** The current transaction. */
 	class FScopedTransaction* ScopedTransaction = nullptr;
@@ -129,10 +129,10 @@ private:
 };
 
 /** Interface for objects who want to lock the viewport to an actor. */
-struct UNREALED_API FLevelViewportActorLock
+struct FLevelViewportActorLock
 {
 	/** Represents no lock. */
-	static const FLevelViewportActorLock None;
+	static UNREALED_API const FLevelViewportActorLock None;
 
 	/** Creates a new instance of FLevelViewportActorLock. */
 	FLevelViewportActorLock() 
@@ -160,7 +160,7 @@ struct UNREALED_API FLevelViewportActorLock
 };
 
 /** */
-class UNREALED_API FLevelEditorViewportClient : public FEditorViewportClient
+class FLevelEditorViewportClient : public FEditorViewportClient
 {
 	friend class FActorElementLevelEditorViewportInteractionCustomization;
 	friend class FComponentElementLevelEditorViewportInteractionCustomization;
@@ -178,55 +178,55 @@ public:
 	/**
 	 * Constructor
 	 */
-	FLevelEditorViewportClient(const TSharedPtr<class SLevelViewport>& InLevelViewport);
+	UNREALED_API FLevelEditorViewportClient(const TSharedPtr<class SLevelViewport>& InLevelViewport);
 
 	/**
 	 * Destructor
 	 */
-	virtual ~FLevelEditorViewportClient();
+	UNREALED_API virtual ~FLevelEditorViewportClient();
 
 	////////////////////////////
 	// FViewElementDrawer interface
-	virtual void Draw(const FSceneView* View,FPrimitiveDrawInterface* PDI) override;
+	UNREALED_API virtual void Draw(const FSceneView* View,FPrimitiveDrawInterface* PDI) override;
 	// End of FViewElementDrawer interface
 	
-	virtual FSceneView* CalcSceneView(FSceneViewFamily* ViewFamily, const int32 StereoViewIndex = INDEX_NONE) override;
+	UNREALED_API virtual FSceneView* CalcSceneView(FSceneViewFamily* ViewFamily, const int32 StereoViewIndex = INDEX_NONE) override;
 
 	////////////////////////////
 	// FEditorViewportClient interface
-	virtual void DrawCanvas( FViewport& InViewport, FSceneView& View, FCanvas& Canvas ) override;
-	virtual bool InputKey(const FInputKeyEventArgs& InEventArgs) override;
-	virtual bool InputAxis(FViewport* Viewport, FInputDeviceId DeviceId, FKey Key, float Delta, float DeltaTime, int32 NumSamples=1, bool bGamepad=false) override;
-	virtual EMouseCursor::Type GetCursor(FViewport* Viewport,int32 X,int32 Y) override;
-	virtual void CapturedMouseMove(FViewport* InViewport, int32 InMouseX, int32 InMouseY) override;
-	virtual void MouseMove(FViewport* InViewport, int32 x, int32 y) override;
-	virtual void Tick(float DeltaSeconds) override;
-	virtual bool InputWidgetDelta( FViewport* Viewport, EAxisList::Type CurrentAxis, FVector& Drag, FRotator& Rot, FVector& Scale ) override;
-	virtual TSharedPtr<FDragTool> MakeDragTool( EDragTool::Type DragToolType ) override;
+	UNREALED_API virtual void DrawCanvas( FViewport& InViewport, FSceneView& View, FCanvas& Canvas ) override;
+	UNREALED_API virtual bool InputKey(const FInputKeyEventArgs& InEventArgs) override;
+	UNREALED_API virtual bool InputAxis(FViewport* Viewport, FInputDeviceId DeviceId, FKey Key, float Delta, float DeltaTime, int32 NumSamples=1, bool bGamepad=false) override;
+	UNREALED_API virtual EMouseCursor::Type GetCursor(FViewport* Viewport,int32 X,int32 Y) override;
+	UNREALED_API virtual void CapturedMouseMove(FViewport* InViewport, int32 InMouseX, int32 InMouseY) override;
+	UNREALED_API virtual void MouseMove(FViewport* InViewport, int32 x, int32 y) override;
+	UNREALED_API virtual void Tick(float DeltaSeconds) override;
+	UNREALED_API virtual bool InputWidgetDelta( FViewport* Viewport, EAxisList::Type CurrentAxis, FVector& Drag, FRotator& Rot, FVector& Scale ) override;
+	UNREALED_API virtual TSharedPtr<FDragTool> MakeDragTool( EDragTool::Type DragToolType ) override;
 	virtual bool IsLevelEditorClient() const override { return ParentLevelEditor.IsValid(); }
-	virtual void TrackingStarted( const struct FInputEventState& InInputState, bool bIsDraggingWidget, bool bNudge ) override;
-	virtual void TrackingStopped() override;
-	virtual void AbortTracking() override;
-	virtual UE::Widget::EWidgetMode GetWidgetMode() const override;
-	virtual FVector GetWidgetLocation() const override;
-	virtual FMatrix GetWidgetCoordSystem() const override;
-	virtual void SetupViewForRendering( FSceneViewFamily& ViewFamily, FSceneView& View ) override;
-	virtual FLinearColor GetBackgroundColor() const override;
-	virtual int32 GetCameraSpeedSetting() const override;
-	virtual void SetCameraSpeedSetting(int32 SpeedSetting) override;
-	virtual float GetCameraSpeedScalar() const override;
-	virtual void SetCameraSpeedScalar(float SpeedScalar) override;
-	virtual void ReceivedFocus(FViewport* InViewport) override;
-	virtual void LostFocus(FViewport* InViewport) override;
-	virtual void ProcessClick(FSceneView& View, HHitProxy* HitProxy, FKey Key, EInputEvent Event, uint32 HitX, uint32 HitY) override;
-	virtual UWorld* GetWorld() const override;
-	virtual void BeginCameraMovement(bool bHasMovement) override;
-	virtual void EndCameraMovement() override;
-	virtual void SetVREditView(bool bGameViewEnable) override;
-	virtual bool GetPivotForOrbit(FVector& Pivot) const override;
-	virtual bool ShouldScaleCameraSpeedByDistance() const override;
-
-	virtual bool OverrideHighResScreenshotCaptureRegion(FIntRect& OutCaptureRegion) override;
+	UNREALED_API virtual void TrackingStarted( const struct FInputEventState& InInputState, bool bIsDraggingWidget, bool bNudge ) override;
+	UNREALED_API virtual void TrackingStopped() override;
+	UNREALED_API virtual void AbortTracking() override;
+	UNREALED_API virtual UE::Widget::EWidgetMode GetWidgetMode() const override;
+	UNREALED_API virtual FVector GetWidgetLocation() const override;
+	UNREALED_API virtual FMatrix GetWidgetCoordSystem() const override;
+	UNREALED_API virtual void SetupViewForRendering( FSceneViewFamily& ViewFamily, FSceneView& View ) override;
+	UNREALED_API virtual FLinearColor GetBackgroundColor() const override;
+	UNREALED_API virtual int32 GetCameraSpeedSetting() const override;
+	UNREALED_API virtual void SetCameraSpeedSetting(int32 SpeedSetting) override;
+	UNREALED_API virtual float GetCameraSpeedScalar() const override;
+	UNREALED_API virtual void SetCameraSpeedScalar(float SpeedScalar) override;
+	UNREALED_API virtual void ReceivedFocus(FViewport* InViewport) override;
+	UNREALED_API virtual void LostFocus(FViewport* InViewport) override;
+	UNREALED_API virtual void ProcessClick(FSceneView& View, HHitProxy* HitProxy, FKey Key, EInputEvent Event, uint32 HitX, uint32 HitY) override;
+	UNREALED_API virtual UWorld* GetWorld() const override;
+	UNREALED_API virtual void BeginCameraMovement(bool bHasMovement) override;
+	UNREALED_API virtual void EndCameraMovement() override;
+	UNREALED_API virtual void SetVREditView(bool bGameViewEnable) override;
+	UNREALED_API virtual bool GetPivotForOrbit(FVector& Pivot) const override;
+	UNREALED_API virtual bool ShouldScaleCameraSpeedByDistance() const override;
+
+	UNREALED_API virtual bool OverrideHighResScreenshotCaptureRegion(FIntRect& OutCaptureRegion) override;
 
 	/** Sets a flag for this frame indicating that the camera has been cut, and temporal effects (such as motion blur) should be reset */
 	void SetIsCameraCut()
@@ -242,115 +242,115 @@ public:
 	/** 
 	 * Initialize visibility flags
 	 */
-	void InitializeVisibilityFlags();
+	UNREALED_API void InitializeVisibilityFlags();
 
 	/**
 	 * Initialize viewport interaction
 	 */
-	void InitializeViewportInteraction();
+	UNREALED_API void InitializeViewportInteraction();
 
 	/**
 	 * Reset the camera position and rotation.  Used when creating a new level.
 	 */
-	void ResetCamera();
+	UNREALED_API void ResetCamera();
 
 	/**
 	 * Reset the view for a new map 
 	 */
-	void ResetViewForNewMap();
+	UNREALED_API void ResetViewForNewMap();
 
 	/**
 	 * Stores camera settings that may be adversely affected by PIE, so that they may be restored later
 	 */
-	void PrepareCameraForPIE();
+	UNREALED_API void PrepareCameraForPIE();
 
 	/**
 	 * Restores camera settings that may be adversely affected by PIE
 	 */
-	void RestoreCameraFromPIE();
+	UNREALED_API void RestoreCameraFromPIE();
 
 	/**
 	 * Updates the audio listener for this viewport 
 	 *
 	 * @param View	The scene view to use when calculate the listener position
 	 */
-	void UpdateAudioListener( const FSceneView& View );
+	UNREALED_API void UpdateAudioListener( const FSceneView& View );
 
 	/** Determines if the new MoveCanvas movement should be used */
-	bool ShouldUseMoveCanvasMovement (void);
+	UNREALED_API bool ShouldUseMoveCanvasMovement (void);
 
 	/** 
 	 * Returns true if the passed in volume is visible in the viewport (due to volume actor visibility flags)
 	 *
 	 * @param VolumeActor	The volume to check
 	 */
-	bool IsVolumeVisibleInViewport( const AActor& VolumeActor ) const;
+	UNREALED_API bool IsVolumeVisibleInViewport( const AActor& VolumeActor ) const;
 
 	/**
 	 * Updates or resets view properties such as aspect ratio, FOV, location etc to match that of any actor we are locked to
 	 */
-	void UpdateViewForLockedActor(float DeltaTime=0.f);
+	UNREALED_API void UpdateViewForLockedActor(float DeltaTime=0.f);
 
 	/**
 	 * Returns the horizontal axis for this viewport.
 	 */
-	EAxisList::Type GetHorizAxis() const;
+	UNREALED_API EAxisList::Type GetHorizAxis() const;
 
 	/**
 	 * Returns the vertical axis for this viewport.
 	 */
-	EAxisList::Type GetVertAxis() const;
+	UNREALED_API EAxisList::Type GetVertAxis() const;
 
-	virtual void NudgeSelectedObjects( const struct FInputEventState& InputState ) override;
+	UNREALED_API virtual void NudgeSelectedObjects( const struct FInputEventState& InputState ) override;
 
 	/**
 	 * Moves the viewport camera according to the locked actors location and rotation
 	 */
-	void MoveCameraToLockedActor();
+	UNREALED_API void MoveCameraToLockedActor();
 
 	/**
 	 * Check to see if this actor is locked by the viewport
 	 */
-	bool IsActorLocked(const TWeakObjectPtr<const AActor> InActor) const;
+	UNREALED_API bool IsActorLocked(const TWeakObjectPtr<const AActor> InActor) const;
 
 	/**
 	 * Check to see if any actor is locked by the viewport
 	 */
-	bool IsAnyActorLocked() const;
+	UNREALED_API bool IsAnyActorLocked() const;
 
-	void ApplyDeltaToActors( const FVector& InDrag, const FRotator& InRot, const FVector& InScale );
-	void ApplyDeltaToActor( AActor* InActor, const FVector& InDeltaDrag, const FRotator& InDeltaRot, const FVector& InDeltaScale );
-	void ApplyDeltaToComponent(USceneComponent* InComponent, const FVector& InDeltaDrag, const FRotator& InDeltaRot, const FVector& InDeltaScale);
+	UNREALED_API void ApplyDeltaToActors( const FVector& InDrag, const FRotator& InRot, const FVector& InScale );
+	UNREALED_API void ApplyDeltaToActor( AActor* InActor, const FVector& InDeltaDrag, const FRotator& InDeltaRot, const FVector& InDeltaScale );
+	UNREALED_API void ApplyDeltaToComponent(USceneComponent* InComponent, const FVector& InDeltaDrag, const FRotator& InDeltaRot, const FVector& InDeltaScale);
 	
-	void ApplyDeltaToSelectedElements(const FTransform& InDeltaTransform);
-	void ApplyDeltaToElement(const FTypedElementHandle& InElementHandle, const FTransform& InDeltaTransform);
+	UNREALED_API void ApplyDeltaToSelectedElements(const FTransform& InDeltaTransform);
+	UNREALED_API void ApplyDeltaToElement(const FTypedElementHandle& InElementHandle, const FTransform& InDeltaTransform);
 
-	void MirrorSelectedActors(const FVector& InMirrorScale);
-	void MirrorSelectedElements(const FVector& InMirrorScale);
+	UNREALED_API void MirrorSelectedActors(const FVector& InMirrorScale);
+	UNREALED_API void MirrorSelectedElements(const FVector& InMirrorScale);
 
-	bool GetFocusBounds(FTypedElementListConstRef InElements, FBoxSphereBounds& OutBounds);
+	UNREALED_API bool GetFocusBounds(FTypedElementListConstRef InElements, FBoxSphereBounds& OutBounds);
 
 	/**
 	 * Get the elements (from the current selection set) that this viewport can manipulate (eg, via the transform gizmo).
 	 */
-	FTypedElementListConstRef GetElementsToManipulate(const bool bForceRefresh = false);
+	UNREALED_API FTypedElementListConstRef GetElementsToManipulate(const bool bForceRefresh = false);
 
-	virtual void SetIsSimulateInEditorViewport( bool bInIsSimulateInEditorViewport ) override;
+	UNREALED_API virtual void SetIsSimulateInEditorViewport( bool bInIsSimulateInEditorViewport ) override;
 
 	/**
 	 *	Draw the texture streaming bounds.
 	 */
-	void DrawTextureStreamingBounds(const FSceneView* View,FPrimitiveDrawInterface* PDI);
+	UNREALED_API void DrawTextureStreamingBounds(const FSceneView* View,FPrimitiveDrawInterface* PDI);
 
 	/** GC references. */
-	void AddReferencedObjects( FReferenceCollector& Collector ) override;
+	UNREALED_API void AddReferencedObjects( FReferenceCollector& Collector ) override;
 	
 	/**
 	 * Copies layout and camera settings from the specified viewport
 	 *
 	 * @param InViewport The viewport to copy settings from
 	 */
-	void CopyLayoutFromViewport( const FLevelEditorViewportClient& InViewport );
+	UNREALED_API void CopyLayoutFromViewport( const FLevelEditorViewportClient& InViewport );
 
 	/**
 	 * Returns whether the provided unlocalized sprite category is visible in the viewport or not
@@ -359,7 +359,7 @@ public:
 	 *
 	 * @return	true if the specified category is visible in the viewport; false if it is not
 	 */
-	bool GetSpriteCategoryVisibility( const FName& InSpriteCategory ) const;
+	UNREALED_API bool GetSpriteCategoryVisibility( const FName& InSpriteCategory ) const;
 
 	/**
 	 * Returns whether the sprite category specified by the provided index is visible in the viewport or not
@@ -368,7 +368,7 @@ public:
 	 *
 	 * @return	true if the category specified by the index is visible in the viewport; false if it is not
 	 */
-	bool GetSpriteCategoryVisibility( int32 Index ) const;
+	UNREALED_API bool GetSpriteCategoryVisibility( int32 Index ) const;
 
 	/**
 	 * Sets the visibility of the provided unlocalized category to the provided value
@@ -376,7 +376,7 @@ public:
 	 * @param	InSpriteCategory	Sprite category to get the index of
 	 * @param	bVisible			true if the category should be made visible, false if it should be hidden
 	 */
-	void SetSpriteCategoryVisibility( const FName& InSpriteCategory, bool bVisible );
+	UNREALED_API void SetSpriteCategoryVisibility( const FName& InSpriteCategory, bool bVisible );
 
 	/**
 	 * Sets the visibility of the category specified by the provided index to the provided value
@@ -384,21 +384,21 @@ public:
 	 * @param	Index		Index of the sprite category to set the visibility of
 	 * @param	bVisible	true if the category should be made visible, false if it should be hidden
 	 */
-	void SetSpriteCategoryVisibility( int32 Index, bool bVisible );
+	UNREALED_API void SetSpriteCategoryVisibility( int32 Index, bool bVisible );
 
 	/**
 	 * Sets the visibility of all sprite categories to the provided value
 	 *
 	 * @param	bVisible	true if all the categories should be made visible, false if they should be hidden
 	 */
-	void SetAllSpriteCategoryVisibility( bool bVisible );
+	UNREALED_API void SetAllSpriteCategoryVisibility( bool bVisible );
 
-	void SetReferenceToWorldContext(FWorldContext& WorldContext);
+	UNREALED_API void SetReferenceToWorldContext(FWorldContext& WorldContext);
 
-	void RemoveReferenceToWorldContext(FWorldContext& WorldContext);
+	UNREALED_API void RemoveReferenceToWorldContext(FWorldContext& WorldContext);
 
 	/** Returns true if a placement dragging actor exists */
-	virtual bool HasDropPreviewActors() const override;
+	UNREALED_API virtual bool HasDropPreviewActors() const override;
 
 	/**
 	 * If dragging an actor for placement, this function updates its position.
@@ -410,12 +410,12 @@ public:
 	 *
 	 * Returns true if preview actors were updated
 	 */
-	virtual bool UpdateDropPreviewActors(int32 MouseX, int32 MouseY, const TArray<UObject*>& DroppedObjects, bool& out_bDroppedObjectsVisible, class UActorFactory* FactoryToUse = NULL) override;
+	UNREALED_API virtual bool UpdateDropPreviewActors(int32 MouseX, int32 MouseY, const TArray<UObject*>& DroppedObjects, bool& out_bDroppedObjectsVisible, class UActorFactory* FactoryToUse = NULL) override;
 
 	/**
 	 * If dragging an actor for placement, this function destroys the actor.
 	 */
-	virtual void DestroyDropPreviewActors() override;
+	UNREALED_API virtual void DestroyDropPreviewActors() override;
 
 	/**
 	 * Checks the viewport to see if the given object can be dropped using the given mouse coordinates local to this viewport
@@ -424,7 +424,7 @@ public:
 	 * @param MouseY			The position of the mouse's Y coordinate
 	 * @param AssetInfo			Asset in question to be dropped
 	 */
-	virtual FDropQuery CanDropObjectsAtCoordinates(int32 MouseX, int32 MouseY, const FAssetData& AssetInfo) override;
+	UNREALED_API virtual FDropQuery CanDropObjectsAtCoordinates(int32 MouseX, int32 MouseY, const FAssetData& AssetInfo) override;
 
 	/**
 	 * Attempts to intelligently drop the given objects in the viewport, using the given mouse coordinates local to this viewport
@@ -438,51 +438,51 @@ public:
 	 * @param bSelectActors		 If true, select the newly dropped actors (defaults: true)
 	 * @param FactoryToUse		 The preferred actor factory to use (optional)
 	 */
-	virtual bool DropObjectsAtCoordinates(int32 MouseX, int32 MouseY, const TArray<UObject*>& DroppedObjects, TArray<AActor*>& OutNewActors, bool bOnlyDropOnTarget = false, bool bCreateDropPreview = false, bool bSelectActors = true, UActorFactory* FactoryToUse = NULL ) override;
+	UNREALED_API virtual bool DropObjectsAtCoordinates(int32 MouseX, int32 MouseY, const TArray<UObject*>& DroppedObjects, TArray<AActor*>& OutNewActors, bool bOnlyDropOnTarget = false, bool bCreateDropPreview = false, bool bSelectActors = true, UActorFactory* FactoryToUse = NULL ) override;
 
 	/**
 	 * Sets GWorld to the appropriate world for this client
 	 * 
 	 * @return the previous GWorld
 	 */
-	virtual UWorld* ConditionalSetWorld() override;
+	UNREALED_API virtual UWorld* ConditionalSetWorld() override;
 
 	/**
 	 * Restores GWorld to InWorld
 	 *
 	 * @param InWorld	The world to restore
 	 */
-	virtual void ConditionalRestoreWorld( UWorld* InWorld  ) override;
+	UNREALED_API virtual void ConditionalRestoreWorld( UWorld* InWorld  ) override;
 
 	/**
 	 *	Called to check if a material can be applied to an object, given the hit proxy
 	 */
-	bool CanApplyMaterialToHitProxy( const HHitProxy* HitProxy ) const;
+	UNREALED_API bool CanApplyMaterialToHitProxy( const HHitProxy* HitProxy ) const;
 
 	/**
 	 * Static: Adds a hover effect to the specified object
 	 *
 	 * @param	InHoverTarget	The hoverable object to add the effect to
 	 */
-	static void AddHoverEffect( const struct FViewportHoverTarget& InHoverTarget );
+	static UNREALED_API void AddHoverEffect( const struct FViewportHoverTarget& InHoverTarget );
 
 	/**
 	 * Static: Removes a hover effect to the specified object
 	 *
 	 * @param	InHoverTarget	The hoverable object to remove the effect from
 	 */
-	static void RemoveHoverEffect( const struct FViewportHoverTarget& InHoverTarget );
+	static UNREALED_API void RemoveHoverEffect( const struct FViewportHoverTarget& InHoverTarget );
 
 	/**
 	 * Static: Clears viewport hover effects from any objects that currently have that
 	 */
-	static void ClearHoverFromObjects();
+	static UNREALED_API void ClearHoverFromObjects();
 
 	/** Set the global ptr to the current viewport */
-	void SetCurrentViewport();
+	UNREALED_API void SetCurrentViewport();
 
 	/** Set the global ptr to the last viewport to receive a key press */
-	void SetLastKeyViewport();
+	UNREALED_API void SetLastKeyViewport();
 
 	/** 
 	 * Access the 'active' actor lock.
@@ -506,7 +506,7 @@ public:
 	 * Find a view component to use for the specified actor. Prioritizes selected 
 	 * components first, followed by camera components (then falls through to the first component that implements GetEditorPreviewInfo)
 	 */
-	static UActorComponent* FindViewComponentForActor(AActor const* Actor);
+	static UNREALED_API UActorComponent* FindViewComponentForActor(AActor const* Actor);
 
 	/** 
 	 * Find the camera component that is driving this viewport, in the following order of preference:
@@ -540,12 +540,12 @@ public:
 	/** 
 	 * Set the actor lock. This is the actor locked to the viewport via the viewport menus.
 	 */
-	void SetActorLock(AActor* Actor);
+	UNREALED_API void SetActorLock(AActor* Actor);
 
 	/** 
 	 * Set the actor lock. This is the actor locked to the viewport via the viewport menus.
 	 */
-	void SetActorLock(const FLevelViewportActorLock& InActorLock);
+	UNREALED_API void SetActorLock(const FLevelViewportActorLock& InActorLock);
 
 	/**
 	 * Get the actor locked to the viewport by cinematic tools like Sequencer.
@@ -566,12 +566,12 @@ public:
 	/**
 	 * Set the actor locked to the viewport by cinematic tools like Sequencer.
 	 */
-	void SetCinematicActorLock(AActor* Actor);
+	UNREALED_API void SetCinematicActorLock(AActor* Actor);
 
 	/**
 	 * Set the actor locked to the viewport by cinematic tools like Sequencer.
 	 */
-	void SetCinematicActorLock(const FLevelViewportActorLock& InActorLock);
+	UNREALED_API void SetCinematicActorLock(const FLevelViewportActorLock& InActorLock);
 
 	/**
 	 * Gets the previous actor lock. This is the actor locked to the viewport via the viewport menus.
@@ -605,12 +605,12 @@ public:
 		return ActorLocks.CinematicActorLock.HasValidLockedActor();
 	}
 
-	void UpdateHoveredObjects( const TSet<FViewportHoverTarget>& NewHoveredObjects );
+	UNREALED_API void UpdateHoveredObjects( const TSet<FViewportHoverTarget>& NewHoveredObjects );
 
 	/**
 	 * Calling SetViewportType from Dragtool_ViewportChange
 	 */
-	void SetViewportTypeFromTool(ELevelViewportType InViewportType);
+	UNREALED_API void SetViewportTypeFromTool(ELevelViewportType InViewportType);
 
 	/**
 	 * Static: Attempts to place the specified object in the level, returning one or more newly-created actors if successful.
@@ -626,7 +626,7 @@ public:
 	 *
 	 * @return	true if the object was successfully used to place an actor; false otherwise
 	 */
-	static TArray<AActor*> TryPlacingActorFromObject( ULevel* InLevel, UObject* ObjToUse, bool bSelectActors, EObjectFlags ObjectFlags, UActorFactory* FactoryToUse, const FName Name = NAME_None, const FViewportCursorLocation* Cursor = nullptr);
+	static UNREALED_API TArray<AActor*> TryPlacingActorFromObject( ULevel* InLevel, UObject* ObjToUse, bool bSelectActors, EObjectFlags ObjectFlags, UActorFactory* FactoryToUse, const FName Name = NAME_None, const FViewportCursorLocation* Cursor = nullptr);
 
 	/** 
 	 * Returns true if creating a preview actor in the viewport. 
@@ -644,7 +644,7 @@ public:
 	 *
 	 * @return	The material that uses this texture, or null if we couldn't find or create one
 	 */
-	static UObject* GetOrCreateMaterialFromTexture( UTexture* UnrealTexture );
+	static UNREALED_API UObject* GetOrCreateMaterialFromTexture( UTexture* UnrealTexture );
 
 	virtual bool UseAppTime() const override { return false; }
 
@@ -652,7 +652,7 @@ public:
 	 * Informs the renderer that the view is being interactively edited. (ex. rotation/translation gizmo).
 	 * This state is reset on tick.
 	 */
-	void SetEditingThroughMovementWidget();
+	UNREALED_API void SetEditingThroughMovementWidget();
 
 protected:
 	/**
@@ -670,71 +670,71 @@ protected:
 	 *
 	 * @return true if asset can be dropped, false otherwise
 	 */
-	bool CanDropBlueprintAsset ( const struct FSelectedAssetInfo& );
+	UNREALED_API bool CanDropBlueprintAsset ( const struct FSelectedAssetInfo& );
 
 	/** Called when editor cleanse event is triggered */
-	void OnEditorCleanse();
+	UNREALED_API void OnEditorCleanse();
 
 	/** Called before the editor tries to begin PIE */
-	void OnPreBeginPIE(const bool bIsSimulating);
+	UNREALED_API void OnPreBeginPIE(const bool bIsSimulating);
 
 	/** Callback for when an editor user setting has changed */
-	void HandleViewportSettingChanged(FName PropertyName);
+	UNREALED_API void HandleViewportSettingChanged(FName PropertyName);
 
 	/** Callback for when a map is created or destroyed */
-	void OnMapChanged(UWorld* InWorld, EMapChangeType MapChangeType);
+	UNREALED_API void OnMapChanged(UWorld* InWorld, EMapChangeType MapChangeType);
 
 	/** Delegate handler for ActorMoved events */
-	void OnActorMoved(AActor* InActor);
+	UNREALED_API void OnActorMoved(AActor* InActor);
 
 public:
 	/** FEditorViewportClient Interface*/
-	virtual void UpdateLinkedOrthoViewports(bool bInvalidate = false) override;
-	virtual ELevelViewportType GetViewportType() const override;
-	virtual void SetViewportType(ELevelViewportType InViewportType) override;
-	virtual void RotateViewportType() override;
-	virtual void OverridePostProcessSettings(FSceneView& View) override;
-	virtual bool ShouldLockPitch() const override;
-	virtual void CheckHoveredHitProxy(HHitProxy* HoveredHitProxy) override;
+	UNREALED_API virtual void UpdateLinkedOrthoViewports(bool bInvalidate = false) override;
+	UNREALED_API virtual ELevelViewportType GetViewportType() const override;
+	UNREALED_API virtual void SetViewportType(ELevelViewportType InViewportType) override;
+	UNREALED_API virtual void RotateViewportType() override;
+	UNREALED_API virtual void OverridePostProcessSettings(FSceneView& View) override;
+	UNREALED_API virtual bool ShouldLockPitch() const override;
+	UNREALED_API virtual void CheckHoveredHitProxy(HHitProxy* HoveredHitProxy) override;
 
 protected:
 
-	virtual void PerspectiveCameraMoved() override;
-	virtual bool GetActiveSafeFrame(float& OutAspectRatio) const override;
-	virtual void RedrawAllViewportsIntoThisScene() override;
+	UNREALED_API virtual void PerspectiveCameraMoved() override;
+	UNREALED_API virtual bool GetActiveSafeFrame(float& OutAspectRatio) const override;
+	UNREALED_API virtual void RedrawAllViewportsIntoThisScene() override;
 
 private:
-	FTransform CachePreDragActorTransform(const AActor* InActor);
+	UNREALED_API FTransform CachePreDragActorTransform(const AActor* InActor);
 
 	/**
 	 * Checks to see the viewports locked actor need updating
 	 */
-	void UpdateLockedActorViewports(const AActor* InActor, const bool bCheckRealtime);
-	void UpdateLockedActorViewport(const AActor* InActor, const bool bCheckRealtime);
+	UNREALED_API void UpdateLockedActorViewports(const AActor* InActor, const bool bCheckRealtime);
+	UNREALED_API void UpdateLockedActorViewport(const AActor* InActor, const bool bCheckRealtime);
 
 	/**
 	 * Moves the locked actor according to the viewport cameras location and rotation
 	 */
-	void MoveLockedActorToCamera();
+	UNREALED_API void MoveLockedActorToCamera();
 	
 	/** @return	Returns true if the delta tracker was used to modify any selected actors or BSP.  Must be called before EndTracking(). */
-	bool HaveSelectedObjectsBeenChanged() const;
+	UNREALED_API bool HaveSelectedObjectsBeenChanged() const;
 
 	/** Cache the list of elements to manipulate based on the current selection set. */
-	void CacheElementsToManipulate(const bool bForceRefresh = false);
+	UNREALED_API void CacheElementsToManipulate(const bool bForceRefresh = false);
 
 	/** Reset the list of elements to manipulate */
-	void ResetElementsToManipulate(const bool bClearList = true);
+	UNREALED_API void ResetElementsToManipulate(const bool bClearList = true);
 
 	/** Reset the list of elements to manipulate, because the selection set they were cached from has changed */
-	void ResetElementsToManipulateFromSelectionChange(const UTypedElementSelectionSet* InSelectionSet);
+	UNREALED_API void ResetElementsToManipulateFromSelectionChange(const UTypedElementSelectionSet* InSelectionSet);
 
 	/** Reset the list of elements to manipulate, because the typed element registry is about to process deferred deletion */
-	void ResetElementsToManipulateFromProcessingDeferredElementsToDestroy();
+	UNREALED_API void ResetElementsToManipulateFromProcessingDeferredElementsToDestroy();
 
 	/** Get the selection set that associated with our level editor. */
-	const UTypedElementSelectionSet* GetSelectionSet() const;
-	UTypedElementSelectionSet* GetMutableSelectionSet() const;
+	UNREALED_API const UTypedElementSelectionSet* GetSelectionSet() const;
+	UNREALED_API UTypedElementSelectionSet* GetMutableSelectionSet() const;
 
 	/**
 	 * Called when to attempt to apply an object to a BSP surface
@@ -745,7 +745,7 @@ private:
 	 *
 	 * @return	true if the object was applied to the object
 	 */
-	bool AttemptApplyObjAsMaterialToSurface( UObject* ObjToUse, class HModel* ModelHitProxy, FViewportCursorLocation& Cursor );
+	UNREALED_API bool AttemptApplyObjAsMaterialToSurface( UObject* ObjToUse, class HModel* ModelHitProxy, FViewportCursorLocation& Cursor );
 
 	/**
 	 * Called when an asset is dropped onto the blank area of a viewport.
@@ -760,7 +760,7 @@ private:
 	 *
 	 * @return	true if the drop operation was successfully handled; false otherwise
 	 */
-	bool DropObjectsOnBackground( struct FViewportCursorLocation& Cursor, const TArray<UObject*>& DroppedObjects, EObjectFlags ObjectFlags, TArray<AActor*>& OutNewActors, bool bCreateDropPreview = false, bool bSelectActors = true, class UActorFactory* FactoryToUse = NULL );
+	UNREALED_API bool DropObjectsOnBackground( struct FViewportCursorLocation& Cursor, const TArray<UObject*>& DroppedObjects, EObjectFlags ObjectFlags, TArray<AActor*>& OutNewActors, bool bCreateDropPreview = false, bool bSelectActors = true, class UActorFactory* FactoryToUse = NULL );
 
 	/**
 	* Called when an asset is dropped upon an existing actor.
@@ -777,7 +777,7 @@ private:
 	*
 	* @return	true if the drop operation was successfully handled; false otherwise
 	*/
-	bool DropObjectsOnActor(struct FViewportCursorLocation& Cursor, const TArray<UObject*>& DroppedObjects, AActor* DroppedUponActor, int32 DroppedUponSlot, EObjectFlags ObjectFlags, TArray<AActor*>& OutNewActors, bool bCreateDropPreview = false, bool bSelectActors = true, class UActorFactory* FactoryToUse = NULL);
+	UNREALED_API bool DropObjectsOnActor(struct FViewportCursorLocation& Cursor, const TArray<UObject*>& DroppedObjects, AActor* DroppedUponActor, int32 DroppedUponSlot, EObjectFlags ObjectFlags, TArray<AActor*>& OutNewActors, bool bCreateDropPreview = false, bool bSelectActors = true, class UActorFactory* FactoryToUse = NULL);
 
 	/**
 	 * Called when an asset is dropped upon a BSP surface.
@@ -794,7 +794,7 @@ private:
 	 *
 	 * @return	true if the drop operation was successfully handled; false otherwise
 	 */
-	bool DropObjectsOnBSPSurface(FSceneView* View, struct FViewportCursorLocation& Cursor, const TArray<UObject*>& DroppedObjects, HModel* TargetProxy, EObjectFlags ObjectFlags, TArray<AActor*>& OutNewActors, bool bCreateDropPreview = false, bool bSelectActors = true, UActorFactory* FactoryToUse = NULL);
+	UNREALED_API bool DropObjectsOnBSPSurface(FSceneView* View, struct FViewportCursorLocation& Cursor, const TArray<UObject*>& DroppedObjects, HModel* TargetProxy, EObjectFlags ObjectFlags, TArray<AActor*>& OutNewActors, bool bCreateDropPreview = false, bool bSelectActors = true, UActorFactory* FactoryToUse = NULL);
 
 	/**
 	 * Called when an asset is dropped upon a manipulation widget.
@@ -806,20 +806,20 @@ private:
 	 *
 	 * @return	true if the drop operation was successfully handled; false otherwise
 	 */
-	bool DropObjectsOnWidget(FSceneView* View, struct FViewportCursorLocation& Cursor, const TArray<UObject*>& DroppedObjects, bool bCreateDropPreview = false);
+	UNREALED_API bool DropObjectsOnWidget(FSceneView* View, struct FViewportCursorLocation& Cursor, const TArray<UObject*>& DroppedObjects, bool bCreateDropPreview = false);
 
 	/** Project the specified actors into the world according to the current drag parameters */
-	void ProjectActorsIntoWorld(const TArray<AActor*>& Actors, FViewport* Viewport, const FVector& Drag, const FRotator& Rot);
+	UNREALED_API void ProjectActorsIntoWorld(const TArray<AActor*>& Actors, FViewport* Viewport, const FVector& Drag, const FRotator& Rot);
 
 	/** Draw additional details for brushes in the world */
-	void DrawBrushDetails(const FSceneView* View, FPrimitiveDrawInterface* PDI);
+	UNREALED_API void DrawBrushDetails(const FSceneView* View, FPrimitiveDrawInterface* PDI);
 
 	/** Internal function for public FindViewComponentForActor, which finds a view component to use for the specified actor. */
-	static UActorComponent* FindViewComponentForActor(AActor const* Actor, TSet<AActor const*>& CheckedActors);
+	static UNREALED_API UActorComponent* FindViewComponentForActor(AActor const* Actor, TSet<AActor const*>& CheckedActors);
 
 public:
 	/** Static: List of objects we're hovering over */
-	static TSet< FViewportHoverTarget > HoveredObjects;
+	static UNREALED_API TSet< FViewportHoverTarget > HoveredObjects;
 	
 	/** Parent level editor that owns this viewport.  Currently, this may be null if the parent doesn't happen to be a level editor. */
 	TWeakPtr< class ILevelEditor > ParentLevelEditor;
@@ -890,10 +890,10 @@ public:
 
 private:
 	/** The actors that are currently being placed in the viewport via dragging */
-	static TArray< TWeakObjectPtr< AActor > > DropPreviewActors;
+	static UNREALED_API TArray< TWeakObjectPtr< AActor > > DropPreviewActors;
 
 	/** If currently creating a preview actor. */
-	static bool bIsDroppingPreviewActor;
+	static UNREALED_API bool bIsDroppingPreviewActor;
 
 	/** A map of actor locations before a drag operation */
 	mutable TMap<TWeakObjectPtr<const AActor>, FTransform> PreDragActorTransforms;
@@ -953,7 +953,7 @@ private:
 	FActorLockStack PreviousActorLocks;
 
 	/** Caching for expensive FindViewComponentForActor. Invalidated once per Tick. */
-	static TMap<TObjectKey<AActor>, TWeakObjectPtr<UActorComponent>> ViewComponentForActorCache;
+	static UNREALED_API TMap<TObjectKey<AActor>, TWeakObjectPtr<UActorComponent>> ViewComponentForActorCache;
 
 	/** If true, we switched between two different cameras. Set by cinematics, used by the motion blur to invalidate this frames motion vectors */
 	bool					bEditorCameraCut;
