[UnrealEd]

diff --git a/Engine/Source/Editor/UnrealEd/Public/MRUFavoritesList.h b/Engine/Source/Editor/UnrealEd/Public/MRUFavoritesList.h
--- a/Engine/Source/Editor/UnrealEd/Public/MRUFavoritesList.h
+++ b/Engine/Source/Editor/UnrealEd/Public/MRUFavoritesList.h
@@ -15,22 +15,22 @@
 #include "MRUList.h"
 
 /** Simple class to represent a combined MRU and favorite map list */
-class UNREALED_API FMainMRUFavoritesList : public FMRUList
+class FMainMRUFavoritesList : public FMRUList
 {
 public:
 	/** Constructor */
-	FMainMRUFavoritesList();
+	UNREALED_API FMainMRUFavoritesList();
 
-	FMainMRUFavoritesList(const FString& IniSectionOverride, const int32 InitMaxItems = 12);
+	UNREALED_API FMainMRUFavoritesList(const FString& IniSectionOverride, const int32 InitMaxItems = 12);
 	
 	/** Destructor */
-	~FMainMRUFavoritesList();
+	UNREALED_API ~FMainMRUFavoritesList();
 	
 	/** Populate MRU/Favorites list by reading saved values from the relevant INI file */
-	virtual void ReadFromINI();
+	UNREALED_API virtual void ReadFromINI();
 
 	/** Save off the state of the MRU and favorites lists to the relevant INI file */
-	virtual void WriteToINI() const;
+	UNREALED_API virtual void WriteToINI() const;
 
 	/**
 	 * Returns the number of favorites items
@@ -47,14 +47,14 @@ public:
 	 *
 	 * @param	Item	Filename of the item to add to the favorites list
 	 */
-	void AddFavoritesItem( const FString& Item );
+	UNREALED_API void AddFavoritesItem( const FString& Item );
 	
 	/**
 	 * Remove a file from the favorites list
 	 *
 	 * @param	Item	Filename of the item to remove from the favorites list
 	 */
-	void RemoveFavoritesItem( const FString& Item );
+	UNREALED_API void RemoveFavoritesItem( const FString& Item );
 
 	/**
 	 * Returns whether a filename is favorited or not
@@ -63,7 +63,7 @@ public:
 	 *
 	 * @return	true if the provided item is in the favorite's list; false if it is not
 	 */
-	bool ContainsFavoritesItem( const FString& Item ) const;
+	UNREALED_API bool ContainsFavoritesItem( const FString& Item ) const;
 
 	/**
 	 * Return the favorites item specified by the provided index
@@ -72,7 +72,7 @@ public:
 	 *
 	 * @return	The favorites item specified by the provided index
 	 */
-	FString GetFavoritesItem( int32 ItemIndex ) const;
+	UNREALED_API FString GetFavoritesItem( int32 ItemIndex ) const;
 
 	/**
 	 * Verifies that the favorites item specified by the provided index still exists. If it does not, the item
@@ -82,14 +82,14 @@ public:
 	 *
 	 * @return	true if the item specified by the index was verified and still exists; false if it does not
 	 */
-	bool VerifyFavoritesFile( int32 ItemIndex );
+	UNREALED_API bool VerifyFavoritesFile( int32 ItemIndex );
 
 	/**
 	 * Moves the specified favorites item to the head of the list
 	 *
 	 * @param	Item	Filename of the item to move
 	 */
-	void MoveFavoritesItemToHead( const FString& Item );
+	UNREALED_API void MoveFavoritesItemToHead( const FString& Item );
 
 	DECLARE_DELEGATE_RetVal_OneParam(bool, FDoesMRUFavoritesItemPassFilter, const FString& MRUFavoritesItem);
 	/**
@@ -98,12 +98,12 @@ public:
 	*
 	* @param DoesMRUFavoritesItemPassFilterDelegate The delegate to use
 	*/
-	void RegisterDoesMRUFavoritesItemPassFilterDelegate(FDoesMRUFavoritesItemPassFilter DoesMRUFavoritesItemPassFilterDelegate);
+	UNREALED_API void RegisterDoesMRUFavoritesItemPassFilterDelegate(FDoesMRUFavoritesItemPassFilter DoesMRUFavoritesItemPassFilterDelegate);
 
 	/**
 	* Unregisters the optional filter delegate
 	*/
-	void UnregisterDoesMRUFavoritesItemPassFilterDelegate();
+	UNREALED_API void UnregisterDoesMRUFavoritesItemPassFilterDelegate();
 
 	/**
 	* Checks the favorites item specified by the provided index against the optional 'DoesMRUFavoritesItemPassFilterDelegate'.
@@ -112,7 +112,7 @@ public:
 	*
 	* @return true if the item specified by the index passes the filter or if no filter has been provided; false if it does not pass the filter
 	*/
-	bool FavoritesItemPassesCurrentFilter(int32 ItemIndex) const;
+	UNREALED_API bool FavoritesItemPassesCurrentFilter(int32 ItemIndex) const;
 
 	/**
 	* Checks the MRU item specified by the provided index against the optional 'DoesMRUFavoritesItemPassFilterDelegate'.
@@ -121,7 +121,7 @@ public:
 	*
 	* @return true if the item specified by the index passes the filter or if no filter has been provided; false if it does not pass the filter
 	*/
-	bool MRUItemPassesCurrentFilter(int32 ItemIndex) const;
+	UNREALED_API bool MRUItemPassesCurrentFilter(int32 ItemIndex) const;
 
 private:
 
@@ -132,5 +132,5 @@ private:
 	TArray<FString> FavoriteItems;
 
 	/** INI section to read/write favorite items to */
-	static const FString FAVORITES_INI_SECTION;
+	static UNREALED_API const FString FAVORITES_INI_SECTION;
 };
diff --git a/Engine/Source/Editor/UnrealEd/Public/MRUList.h b/Engine/Source/Editor/UnrealEd/Public/MRUList.h
--- a/Engine/Source/Editor/UnrealEd/Public/MRUList.h
+++ b/Engine/Source/Editor/UnrealEd/Public/MRUList.h
@@ -15,14 +15,14 @@
 /**
  * An MRU list of files.
  */
-class UNREALED_API FMRUList
+class FMRUList
 {
 public:
-	FMRUList(const FString& InINISection, const int32 InitMaxItems = 12 );
-	virtual ~FMRUList();
+	UNREALED_API FMRUList(const FString& InINISection, const int32 InitMaxItems = 12 );
+	UNREALED_API virtual ~FMRUList();
 
-	virtual void ReadFromINI();
-	virtual void WriteToINI() const;
+	UNREALED_API virtual void ReadFromINI();
+	UNREALED_API virtual void WriteToINI() const;
 	
 	/**
 	 * Returns the number of MRU items
@@ -39,7 +39,7 @@ public:
 	 *
 	 * @param	Item		The long package name of the item to add.
 	 */
-	void AddMRUItem(const FString& Item);
+	UNREALED_API void AddMRUItem(const FString& Item);
 	
 	/**
 	 * Returns the index of the specified item, or INDEX_NONE if not found.
@@ -47,26 +47,26 @@ public:
 	 * @param	Item		The item to query.
 	 * @return				The index of the specified item.
 	 */
-	int32 FindMRUItemIdx(const FString& Item) const;
+	UNREALED_API int32 FindMRUItemIdx(const FString& Item) const;
 
 	/**
 	 * Removes the specified item from the list if it exists.
 	 *
 	 * @param	Item		The item to remove.
 	 */
-	void RemoveMRUItem(const FString& Item);
+	UNREALED_API void RemoveMRUItem(const FString& Item);
 
 	/**
 	 * Removes the item at the specified index.
 	 *
 	 * @param	ItemIndex	Index of the item to remove.  Must be in [0, GetMaxItems()-1].
 	 */
-	void RemoveMRUItem(int32 ItemIndex);
+	UNREALED_API void RemoveMRUItem(int32 ItemIndex);
 
 	/**
 	 * Clears the list of MRU items. 
 	 */
-	void ClearMRUItems();
+	UNREALED_API void ClearMRUItems();
 
 	/**
 	 * Checks to make sure the file specified still exists.  If it does, it is
@@ -77,7 +77,7 @@ public:
 	 * @param	OutPackageName	Long package name after following redirector
 	 * @return					true if the item exists, false if it doesn't
 	 */
-	bool VerifyMRUFile(int32 ItemIndex, FString& OutPackageName);
+	UNREALED_API bool VerifyMRUFile(int32 ItemIndex, FString& OutPackageName);
 
 	/**
 	 * Accessor.
@@ -96,14 +96,14 @@ private:
 	/**
 	 * Make sure we don't have more than GetMaxItems() in the list.
 	 */
-	void Cull();
+	UNREALED_API void Cull();
 
 	/**
 	 * Moves the specified item to the top of the MRU list.
 	 *
 	 * @param	ItemIndex		Index of the item to bring to top.
 	 */
-	void MoveToTop(int32 ItemIndex);
+	UNREALED_API void MoveToTop(int32 ItemIndex);
 
 protected:
 
@@ -116,7 +116,7 @@ protected:
 	 * @param	INIKeyBase	Base string to use to search for individual items; this value will have a number appended to it while searching
 	 * @param	NumElements	Maximum number of elements to search in the INI for
 	 */
-	static void InternalReadINI( TArray<FString>& OutItems, const FString& INISection, const FString& INIKeyBase, int32 NumElements );
+	static UNREALED_API void InternalReadINI( TArray<FString>& OutItems, const FString& INISection, const FString& INIKeyBase, int32 NumElements );
 	
 	/**
 	 * Internal helper function to write out the provided array of data to the provided INI section, using the provided key base as a prefix
@@ -126,7 +126,7 @@ protected:
 	 * @param	INISection	Section of the INI file to write to
 	 * @param	INIKeyBase	Base prefix to use for each item written to the INI file
 	 */
-	static void InternalWriteINI( const TArray<FString>& InItems, const FString& INISection, const FString& INIKeyBase );
+	static UNREALED_API void InternalWriteINI( const TArray<FString>& InItems, const FString& INISection, const FString& INIKeyBase );
 
 	/** Maximum MRU items */
 	const int32 MaxItems;
diff --git a/Engine/Source/Editor/UnrealEd/Public/MiniCurveEditor.h b/Engine/Source/Editor/UnrealEd/Public/MiniCurveEditor.h
--- a/Engine/Source/Editor/UnrealEd/Public/MiniCurveEditor.h
+++ b/Engine/Source/Editor/UnrealEd/Public/MiniCurveEditor.h
@@ -12,7 +12,7 @@
 class FCurveOwnerInterface;
 class SCurveEditor;
 
-class UNREALED_API SMiniCurveEditor :  public SCompoundWidget,public IAssetEditorInstance
+class SMiniCurveEditor :  public SCompoundWidget,public IAssetEditorInstance
 {
 public:
 	SLATE_BEGIN_ARGS( SMiniCurveEditor )
@@ -25,19 +25,19 @@ public:
 	SLATE_ARGUMENT( TWeakPtr<SWindow>, ParentWindow )
 	SLATE_END_ARGS()
 
-	void Construct(const FArguments& InArgs);
-	virtual ~SMiniCurveEditor();
+	UNREALED_API void Construct(const FArguments& InArgs);
+	UNREALED_API virtual ~SMiniCurveEditor();
 
 	// IAssetEditorInstance interface
-	virtual FName GetEditorName() const override;
-	virtual void FocusWindow(UObject* ObjectToFocusOn) override;
-	virtual bool CloseWindow() override;
-	virtual bool CloseWindow(EAssetEditorCloseReason InCloseReason) override;
+	UNREALED_API virtual FName GetEditorName() const override;
+	UNREALED_API virtual void FocusWindow(UObject* ObjectToFocusOn) override;
+	UNREALED_API virtual bool CloseWindow() override;
+	UNREALED_API virtual bool CloseWindow(EAssetEditorCloseReason InCloseReason) override;
 	virtual bool IsPrimaryEditor() const override { return true; };
 	virtual void InvokeTab(const struct FTabId& TabId) override {}
-	virtual TSharedPtr<class FTabManager> GetAssociatedTabManager() override;
-	virtual double GetLastActivationTime() override;
-	virtual void RemoveEditingAsset(UObject* Asset) override;
+	UNREALED_API virtual TSharedPtr<class FTabManager> GetAssociatedTabManager() override;
+	UNREALED_API virtual double GetLastActivationTime() override;
+	UNREALED_API virtual void RemoveEditingAsset(UObject* Asset) override;
 
 private:
 
@@ -49,9 +49,9 @@ private:
 	float GetViewMinInput() const { return ViewMinInput; }
 	float GetViewMaxInput() const { return ViewMaxInput; }
 	/** Return length of timeline */
-	float GetTimelineLength() const;
+	UNREALED_API float GetTimelineLength() const;
 
-	void SetInputViewRange(float InViewMinInput, float InViewMaxInput);
+	UNREALED_API void SetInputViewRange(float InViewMinInput, float InViewMaxInput);
 
 
 protected:
diff --git a/Engine/Source/Editor/UnrealEd/Public/MouseDeltaTracker.h b/Engine/Source/Editor/UnrealEd/Public/MouseDeltaTracker.h
--- a/Engine/Source/Editor/UnrealEd/Public/MouseDeltaTracker.h
+++ b/Engine/Source/Editor/UnrealEd/Public/MouseDeltaTracker.h
@@ -16,92 +16,92 @@ struct FInputEventState;
 /**
  * Keeps track of mouse movement deltas in the viewports.
  */
-class UNREALED_API FMouseDeltaTracker
+class FMouseDeltaTracker
 {
 public:
 
-	FMouseDeltaTracker();
-	~FMouseDeltaTracker();
+	UNREALED_API FMouseDeltaTracker();
+	UNREALED_API ~FMouseDeltaTracker();
 
 	/**
 	 * Begin tracking at the specified location for the specified viewport.
 	 */
-	void StartTracking(FEditorViewportClient* InViewportClient, const int32 InX, const int32 InY, const FInputEventState& InInputState, bool bNudge = false, bool bResetDragToolState = true);
+	UNREALED_API void StartTracking(FEditorViewportClient* InViewportClient, const int32 InX, const int32 InY, const FInputEventState& InInputState, bool bNudge = false, bool bResetDragToolState = true);
 
 	/**
 	 * Called when a mouse button has been released.  If there are no other
 	 * mouse buttons being held down, the internal information is reset.
 	 */
-	bool EndTracking(FEditorViewportClient* InViewportClient);
+	UNREALED_API bool EndTracking(FEditorViewportClient* InViewportClient);
 
 	/**
 	 * Adds delta movement into the tracker.
 	 */
-	void AddDelta(FEditorViewportClient* InViewportClient, FKey InKey, const int32 InDelta, bool InNudge);
+	UNREALED_API void AddDelta(FEditorViewportClient* InViewportClient, FKey InKey, const int32 InDelta, bool InNudge);
 
 	/**
 	* Returns the current delta.
 	*/
-	const FVector GetDelta() const;
+	UNREALED_API const FVector GetDelta() const;
 
 	/**
 	 * Returns the current snapped delta.
 	 */
-	const FVector GetDeltaSnapped() const;
+	UNREALED_API const FVector GetDeltaSnapped() const;
 
 	/**
 	* Returns the absolute delta since dragging started.
 	*/
-	const FVector GetAbsoluteDelta() const;
+	UNREALED_API const FVector GetAbsoluteDelta() const;
 
 	/**
 	* Returns the absolute snapped delta since dragging started. 
 	*/
-	const FVector GetAbsoluteDeltaSnapped() const;
+	UNREALED_API const FVector GetAbsoluteDeltaSnapped() const;
 
 	/**
 	 * Returns the screen space delta since dragging started.
 	 */
-	const FVector GetScreenDelta() const;
+	UNREALED_API const FVector GetScreenDelta() const;
 
 	/**
 	 * Returns the raw mouse delta in pixels since dragging started.
 	 */
-	const FVector GetRawDelta() const;
+	UNREALED_API const FVector GetRawDelta() const;
 
 	/**
 	 * Returns the unsnapped start position of the current mouse drag. (This will be zero if there is no drag in progress)
 	 */
-	const FVector GetDragStartPos() const;
+	UNREALED_API const FVector GetDragStartPos() const;
 
 	/**
 	 * Returns if the user used a modifier to drag a selected item. (Rather than using a widget handle)
 	 */
-	const bool GetUsedDragModifier() const;
+	UNREALED_API const bool GetUsedDragModifier() const;
 
 	/**
 	 * Resets the UsedDragModifier flag
 	 */
-	void ResetUsedDragModifier();
+	UNREALED_API void ResetUsedDragModifier();
 
 	/**
 	 * Converts the delta movement to drag/rotation/scale based on the viewport type or widget axis.
 	 */
-	void ConvertMovementDeltaToDragRot(FSceneView* InView, FEditorViewportClient* InViewportClient, FVector& InDragDelta, FVector& OutDrag, FRotator& OutRotation, FVector& OutScale) const;
+	UNREALED_API void ConvertMovementDeltaToDragRot(FSceneView* InView, FEditorViewportClient* InViewportClient, FVector& InDragDelta, FVector& OutDrag, FRotator& OutRotation, FVector& OutScale) const;
 	/**
 	 * Absolute Translation conversion from mouse position on the screen to widget axis movement/rotation.
 	 */
-	void AbsoluteTranslationConvertMouseToDragRot(FSceneView* InView, FEditorViewportClient* InViewportClient, FVector& OutDrag, FRotator& OutRotation, FVector& OutScale ) const;
+	UNREALED_API void AbsoluteTranslationConvertMouseToDragRot(FSceneView* InView, FEditorViewportClient* InViewportClient, FVector& OutDrag, FRotator& OutRotation, FVector& OutScale ) const;
 
 	/**
 	 * Subtracts the specified value from End and EndSnapped.
 	 */
-	void ReduceBy(const FVector& In);
+	UNREALED_API void ReduceBy(const FVector& In);
 
 	/**
 	 * @return		true if a drag tool is being used by the tracker, false otherwise.
 	 */
-	bool UsingDragTool() const;
+	UNREALED_API bool UsingDragTool() const;
 
 	/**
 	 * @return True if any mouse movement has happened since tracking started
@@ -120,12 +120,12 @@ public:
 	/**
 	 * Renders the drag tool.  Does nothing if no drag tool exists.
 	 */
-	void Render3DDragTool(const FSceneView* View, FPrimitiveDrawInterface* PDI);
+	UNREALED_API void Render3DDragTool(const FSceneView* View, FPrimitiveDrawInterface* PDI);
 
 	/**
 	 * Renders the drag tool.  Does nothing if no drag tool exists.
 	 */
-	void RenderDragTool(const FSceneView* View, FCanvas* Canvas);
+	UNREALED_API void RenderDragTool(const FSceneView* View, FCanvas* Canvas);
 
 	/** 
 	 * @return 		the widget mode last time this tracker object started tracking
@@ -135,7 +135,7 @@ private:
 	/**
 	 * Starts using a drag tool if needed
 	 */
-	void ConditionalBeginUsingDragTool( FEditorViewportClient* InViewportClient );
+	UNREALED_API void ConditionalBeginUsingDragTool( FEditorViewportClient* InViewportClient );
 private:
 	/** The unsnapped start position of the current mouse drag. */
 	FVector Start;
@@ -185,7 +185,7 @@ private:
 	 *
 	 * @param	InViewportClient		The viewport whose widget axis is to be set.
 	 */
-	void DetermineCurrentAxis(FEditorViewportClient* InViewportClient);
+	UNREALED_API void DetermineCurrentAxis(FEditorViewportClient* InViewportClient);
 
 };
 
diff --git a/Engine/Source/Editor/UnrealEd/Public/NormalMapPreview.h b/Engine/Source/Editor/UnrealEd/Public/NormalMapPreview.h
--- a/Engine/Source/Editor/UnrealEd/Public/NormalMapPreview.h
+++ b/Engine/Source/Editor/UnrealEd/Public/NormalMapPreview.h
@@ -12,10 +12,10 @@
 /**
  * Batched element parameters for previewing normal maps.
  */
-class UNREALED_API FNormalMapBatchedElementParameters : public FBatchedElementParameters
+class FNormalMapBatchedElementParameters : public FBatchedElementParameters
 {
 	/** Binds vertex and pixel shaders for this element */
-	virtual void BindShaders(
+	UNREALED_API virtual void BindShaders(
 		FRHICommandList& RHICmdList,
 		FGraphicsPipelineStateInitializer& GraphicsPSOInit,
 		ERHIFeatureLevel::Type InFeatureLevel,
diff --git a/Engine/Source/Editor/UnrealEd/Public/ObjectTools.h b/Engine/Source/Editor/UnrealEd/Public/ObjectTools.h
--- a/Engine/Source/Editor/UnrealEd/Public/ObjectTools.h
+++ b/Engine/Source/Editor/UnrealEd/Public/ObjectTools.h
@@ -173,7 +173,7 @@ namespace ObjectTools
 	UNREALED_API UObject* DuplicateSingleObject(UObject* Object, const FPackageGroupName& PGN, TSet<UPackage*>& InOutPackagesUserRefusedToFullyLoad, bool bPromptToOverwrite = true, TMap<TSoftObjectPtr<UObject>, TSoftObjectPtr<UObject>>* DuplicatedObjects = nullptr);
 
 	/** Helper struct to detail the results of a consolidation operation */
-	struct UNREALED_API FConsolidationResults : public FGCObject
+	struct FConsolidationResults : public FGCObject
 	{
 		/** FGCObject interface; Serialize any object references */
 		virtual void AddReferencedObjects( FReferenceCollector& Collector ) override
diff --git a/Engine/Source/Editor/UnrealEd/Public/PackageSourceControlHelper.h b/Engine/Source/Editor/UnrealEd/Public/PackageSourceControlHelper.h
--- a/Engine/Source/Editor/UnrealEd/Public/PackageSourceControlHelper.h
+++ b/Engine/Source/Editor/UnrealEd/Public/PackageSourceControlHelper.h
@@ -8,21 +8,21 @@
 
 //This class should not be used as a base class because it can force a connection to the SCM through FScopedSourceControl. Connection errors can result in failing automation processes that are unrelated. 
 //The same can happen if the class is used as a member of another class. It should only be used as a local variable in a method so the connection is only attempted if required. 
-class UNREALED_API FPackageSourceControlHelper final
+class FPackageSourceControlHelper final
 {
 public:
-	bool UseSourceControl() const;
-	bool Delete(const FString& PackageName) const;
-	bool Delete(const TArray<FString>& PackageNames, bool bErrorsAsWarnings = false) const;
-	bool Delete(UPackage* Package) const;
-	bool Delete(const TArray<UPackage*>& Packages) const;
-	bool AddToSourceControl(UPackage* Package) const;
-	bool AddToSourceControl(const TArray<UPackage*>& Packages, bool bErrorsAsWarnings = false) const;
-	bool AddToSourceControl(const TArray<FString>& PackageNames, bool bErrorsAsWarnings = false) const;
-	bool Checkout(UPackage* Package) const;
-	bool Checkout(const TArray<FString>& PackageNames, bool bErrorsAsWarnings = false) const;
-	bool GetDesiredStatesForModification(const TArray<FString>& PackageNames, TArray<FString>& OutPackagesToCheckout, TArray<FString>& OutPackagesToAdd, bool bErrorsAsWarnings = false) const;
-	bool GetMarkedForDeleteFiles(const TArray<FString>& PackageFilenames, TArray<FString>& OutPackageFilenames, bool bErrorsAsWarnings = false) const;
+	UNREALED_API bool UseSourceControl() const;
+	UNREALED_API bool Delete(const FString& PackageName) const;
+	UNREALED_API bool Delete(const TArray<FString>& PackageNames, bool bErrorsAsWarnings = false) const;
+	UNREALED_API bool Delete(UPackage* Package) const;
+	UNREALED_API bool Delete(const TArray<UPackage*>& Packages) const;
+	UNREALED_API bool AddToSourceControl(UPackage* Package) const;
+	UNREALED_API bool AddToSourceControl(const TArray<UPackage*>& Packages, bool bErrorsAsWarnings = false) const;
+	UNREALED_API bool AddToSourceControl(const TArray<FString>& PackageNames, bool bErrorsAsWarnings = false) const;
+	UNREALED_API bool Checkout(UPackage* Package) const;
+	UNREALED_API bool Checkout(const TArray<FString>& PackageNames, bool bErrorsAsWarnings = false) const;
+	UNREALED_API bool GetDesiredStatesForModification(const TArray<FString>& PackageNames, TArray<FString>& OutPackagesToCheckout, TArray<FString>& OutPackagesToAdd, bool bErrorsAsWarnings = false) const;
+	UNREALED_API bool GetMarkedForDeleteFiles(const TArray<FString>& PackageFilenames, TArray<FString>& OutPackageFilenames, bool bErrorsAsWarnings = false) const;
 private:
 	ISourceControlProvider& GetSourceControlProvider() const;
 	FScopedSourceControl SourceControl;	
diff --git a/Engine/Source/Editor/UnrealEd/Public/PackageTools.h b/Engine/Source/Editor/UnrealEd/Public/PackageTools.h
--- a/Engine/Source/Editor/UnrealEd/Public/PackageTools.h
+++ b/Engine/Source/Editor/UnrealEd/Public/PackageTools.h
@@ -28,8 +28,8 @@ enum class EReloadPackagesInteractionMode : uint8
 	AssumeNegative,
 };
 
-UCLASS(Abstract)
-class UNREALED_API UPackageTools : public UObject
+UCLASS(Abstract, MinimalAPI)
+class UPackageTools : public UObject
 {
 	GENERATED_UCLASS_BODY()
 
@@ -42,7 +42,7 @@ public:
 	 * @param	OutGroupPackages			The map that receives the filtered list of group packages.
 	 * @param	OutPackageList				The array that will contain the list of filtered packages.
 	 */
-	static void GetFilteredPackageList(TSet<UPackage*>& OutFilteredPackageMap);
+	static UNREALED_API void GetFilteredPackageList(TSet<UPackage*>& OutFilteredPackageMap);
 	
 	/**
 	 * Fills the OutObjects list with all valid objects that are supported by the current
@@ -51,7 +51,7 @@ public:
 	 * @param	InPackages			Filters objects based on package.
 	 * @param	OutObjects			[out] Receives the list of objects
 	 */
-	static void GetObjectsInPackages( const TArray<UPackage*>* InPackages, TArray<UObject*>& OutObjects );
+	static UNREALED_API void GetObjectsInPackages( const TArray<UPackage*>* InPackages, TArray<UObject*>& OutObjects );
 
 	/**
 	 * Handles fully loading passed in packages.
@@ -61,7 +61,7 @@ public:
 	 * 
 	 * @return true if all packages where fully loaded, false otherwise
 	 */
-	static bool HandleFullyLoadingPackages( const TArray<UPackage*>& TopLevelPackages, const FText& OperationText );
+	static UNREALED_API bool HandleFullyLoadingPackages( const TArray<UPackage*>& TopLevelPackages, const FText& OperationText );
 
 
 	/**
@@ -71,7 +71,7 @@ public:
 	 *
 	 * @return	The loaded package (or NULL if something went wrong.)
 	 */
-	static UPackage* LoadPackage( FString InFilename );
+	static UNREALED_API UPackage* LoadPackage( FString InFilename );
 
 	/**
 	 * Helper function that attempts to unload the specified top-level packages.
@@ -80,7 +80,7 @@ public:
 	 *
 	 * @return	true if the set of loaded packages was changed
 	 */
-	static bool UnloadPackages( const TArray<UPackage*>& PackagesToUnload );
+	static UNREALED_API bool UnloadPackages( const TArray<UPackage*>& PackagesToUnload );
 
 	/**
 	 * Helper function that attempts to unload the specified top-level packages.
@@ -91,7 +91,7 @@ public:
 	 *
 	 * @return	true if the set of loaded packages was changed
 	 */
-	static bool UnloadPackages( const TArray<UPackage*>& PackagesToUnload, FText& OutErrorMessage, bool bUnloadDirtyPackages = false);
+	static UNREALED_API bool UnloadPackages( const TArray<UPackage*>& PackagesToUnload, FText& OutErrorMessage, bool bUnloadDirtyPackages = false);
 
 	/**
 	 * struct containing list of packages to unload and other params.
@@ -118,7 +118,7 @@ public:
 	 * 
 	 * @return true if the set of loaded packages was changed
 	 */
-	static bool UnloadPackages(FUnloadPackageParams& Params);
+	static UNREALED_API bool UnloadPackages(FUnloadPackageParams& Params);
 
 	/**
 	 * Helper function that attempts to reload the specified top-level packages.
@@ -127,7 +127,7 @@ public:
 	 *
 	 * @return	true if the set of loaded packages was changed
 	 */
-	static bool ReloadPackages( const TArray<UPackage*>& PackagesToReload );
+	static UNREALED_API bool ReloadPackages( const TArray<UPackage*>& PackagesToReload );
 
 	/**
 	 * Helper function that attempts to reload the specified top-level packages.
@@ -139,7 +139,7 @@ public:
 	 * @return	true if the set of loaded packages was changed
 	 */
 	UE_DEPRECATED(4.21, "ReloadPackages taking bInteractive is deprecated. Use the version taking EReloadPackagesInteractionMode instead.")
-	static bool ReloadPackages( const TArray<UPackage*>& PackagesToReload, FText& OutErrorMessage, const bool bInteractive = true );
+	static UNREALED_API bool ReloadPackages( const TArray<UPackage*>& PackagesToReload, FText& OutErrorMessage, const bool bInteractive = true );
 
 	/**
 	 * Helper function that attempts to reload the specified top-level packages.
@@ -150,7 +150,7 @@ public:
 	 *
 	 * @return	true if the set of loaded packages was changed
 	 */
-	static bool ReloadPackages( const TArray<UPackage*>& PackagesToReload, FText& OutErrorMessage, const EReloadPackagesInteractionMode InteractionMode = EReloadPackagesInteractionMode::Interactive );
+	static UNREALED_API bool ReloadPackages( const TArray<UPackage*>& PackagesToReload, FText& OutErrorMessage, const EReloadPackagesInteractionMode InteractionMode = EReloadPackagesInteractionMode::Interactive );
 
 	/**
 	 *	Exports the given packages to files.
@@ -159,7 +159,7 @@ public:
 	 * @param	ExportPath				receives the value of the path the user chose for exporting.
 	 * @param	bUseProvidedExportPath	If true and ExportPath is specified, use ExportPath as the user's export path w/o prompting for a directory, where applicable
 	 */
-	static void ExportPackages( const TArray<UPackage*>& PackagesToExport, FString* ExportPath=NULL, bool bUseProvidedExportPath = false );
+	static UNREALED_API void ExportPackages( const TArray<UPackage*>& PackagesToExport, FString* ExportPath=NULL, bool bUseProvidedExportPath = false );
 
 	/**
 	 * Wrapper method for multiple objects at once.
@@ -171,10 +171,10 @@ public:
 	 *
 	 * @return	the path that the user chose for the export.
 	 */
-	static FString DoBulkExport(const TArray<UPackage*>& TopLevelPackages, FString LastExportPath, const TSet<UClass*>* FilteredClasses = NULL, bool bUseProvidedExportPath = false );
+	static UNREALED_API FString DoBulkExport(const TArray<UPackage*>& TopLevelPackages, FString LastExportPath, const TSet<UClass*>* FilteredClasses = NULL, bool bUseProvidedExportPath = false );
 
 	/** Helper function that attempts to check out the specified top-level packages. */
-	static void CheckOutRootPackages( const TArray<UPackage*>& Packages );
+	static UNREALED_API void CheckOutRootPackages( const TArray<UPackage*>& Packages );
 
 
 	/**
@@ -183,7 +183,7 @@ public:
 	 * @param	PackagePath	Path of the package to check, relative or absolute
 	 * @return	true if PackagePath points to an external location
 	 */
-	static bool IsPackagePathExternal(const FString& PackagePath);
+	static UNREALED_API bool IsPackagePathExternal(const FString& PackagePath);
 
 	/**
 	 * Checks if the passed in package's filename is in an external directory. I.E Ones not found automatically in the content directory
@@ -191,10 +191,10 @@ public:
 	 * @param	Package	The package to check
 	 * @return	true if the package points to an external filename
 	 */
-	static bool IsPackageExternal(const UPackage& Package);
+	static UNREALED_API bool IsPackageExternal(const UPackage& Package);
 
 	/** Saves all the dirty packages for the specified objects*/
-	static bool SavePackagesForObjects(const TArray<UObject*>& ObjectsToSave);
+	static UNREALED_API bool SavePackagesForObjects(const TArray<UObject*>& ObjectsToSave);
 
 	/**
 	 * Checks if the package has only one asset which shares its name
@@ -202,11 +202,11 @@ public:
 	 * @param Package The package to check
 	 * @return true if the package has only one asset which shares the name of the package
 	 */
-	static bool IsSingleAssetPackage (const FString& Package);
+	static UNREALED_API bool IsSingleAssetPackage (const FString& Package);
 
 	/** Replaces all invalid package name characters with _ */
 	UFUNCTION(BlueprintCallable, Category = "Editor Scripting | Package Tools")
-	static FString SanitizePackageName(const FString& InPackageName);
+	static UNREALED_API FString SanitizePackageName(const FString& InPackageName);
 
 	/** 
 	 * Converts a long package name to a file name.
@@ -217,7 +217,7 @@ public:
 	 * @return Package filename, or empty if it failed.
 	 */
 	UFUNCTION(BlueprintCallable, Category = "Editor Scripting | Package Tools", meta=(AdvancedDisplay=1))
-	static FString PackageNameToFilename(const FString& PackageName, const FString& Extension = TEXT(""));
+	static UNREALED_API FString PackageNameToFilename(const FString& PackageName, const FString& Extension = TEXT(""));
 	
 	/** 
 	 * Tries to convert a given relative or absolute filename to a long package name or path starting with a root like /Game
@@ -226,7 +226,7 @@ public:
 	 * @return Resulting long package name if the supplied filename properly maps to a long package root, empty string otherwise.
 	 */
 	UFUNCTION(BlueprintCallable, Category = "Editor Scripting | Package Tools")
-	static FString FilenameToPackageName(const FString& Filename);
+	static UNREALED_API FString FilenameToPackageName(const FString& Filename);
 
 	/**
 	 * Find or create a package at the desired path
@@ -235,7 +235,7 @@ public:
 	 * @param LongPackageName The package path in unreal asset path
 	 * @param AssetClass The class of asset the package should be for.
 	 */
-	static UPackage* FindOrCreatePackageForAssetType(const FName LongPackageName, UClass* AssetClass);
+	static UNREALED_API UPackage* FindOrCreatePackageForAssetType(const FName LongPackageName, UClass* AssetClass);
 
 	/**
 	 * Utility function that gathers all async compilable objects from given packages
@@ -243,7 +243,7 @@ public:
 	 * from said packages.
 	 * @param Packages List of packages from which to gather compilable assets to flush.
 	 */
-	static void FlushAsyncCompilation(TArrayView<UPackage* const> Packages);
+	static UNREALED_API void FlushAsyncCompilation(TArrayView<UPackage* const> Packages);
 
 private:
 	static void RestoreStandaloneOnReachableObjects();
diff --git a/Engine/Source/Editor/UnrealEd/Public/PropertyInfoViewStyle.h b/Engine/Source/Editor/UnrealEd/Public/PropertyInfoViewStyle.h
--- a/Engine/Source/Editor/UnrealEd/Public/PropertyInfoViewStyle.h
+++ b/Engine/Source/Editor/UnrealEd/Public/PropertyInfoViewStyle.h
@@ -61,24 +61,24 @@ namespace PropertyInfoViewStyle
 	 * see SDebugLineItem::GenerateWidgetForColumn (SKismetDebuggingView.cpp) for
 	 * use example
 	 */
-	class UNREALED_API SIndent : public SCompoundWidget
+	class SIndent : public SCompoundWidget
 	{ 
 	public:
 			
 		SLATE_BEGIN_ARGS(SIndent) {}
 		SLATE_END_ARGS()
 
-		void Construct(const FArguments& InArgs, TSharedRef<ITableRow> DetailsRow);
+		UNREALED_API void Construct(const FArguments& InArgs, TSharedRef<ITableRow> DetailsRow);
 
 	private:
-		virtual int32 OnPaint(const FPaintArgs& Args, const FGeometry& AllottedGeometry,
+		UNREALED_API virtual int32 OnPaint(const FPaintArgs& Args, const FGeometry& AllottedGeometry,
 		                      const FSlateRect& MyCullingRect, FSlateWindowElementList& OutDrawElements,
 		                      int32 LayerId, const FWidgetStyle& InWidgetStyle,
 		                      bool bParentEnabled) const override;
 
-		virtual FOptionalSize GetIndentWidth() const;
+		UNREALED_API virtual FOptionalSize GetIndentWidth() const;
 
-		virtual FSlateColor GetRowBackgroundColor(int32 IndentLevel) const;
+		UNREALED_API virtual FSlateColor GetRowBackgroundColor(int32 IndentLevel) const;
 
 	private:
 		TWeakPtr<ITableRow> Row;
@@ -93,14 +93,14 @@ namespace PropertyInfoViewStyle
 	 * see SDebugLineItem::GenerateWidgetForColumn (SKismetDebuggingView.cpp) for
 	 * use example
 	 */
-	class UNREALED_API SExpanderArrow : public SCompoundWidget
+	class SExpanderArrow : public SCompoundWidget
 	{
 	public:
 		SLATE_BEGIN_ARGS(SExpanderArrow) {}
 			SLATE_ATTRIBUTE(bool, HasChildren)
 		SLATE_END_ARGS()
 
-		void Construct(const FArguments& InArgs, TSharedRef<ITableRow> DetailsRow);
+		UNREALED_API void Construct(const FArguments& InArgs, TSharedRef<ITableRow> DetailsRow);
 
 	private:
 		EVisibility GetExpanderVisibility() const;
@@ -124,7 +124,7 @@ namespace PropertyInfoViewStyle
 	 * the same as a standard STextBlock would be in order for the highlight
 	 * to match up properly
 	 */
-	class UNREALED_API STextHighlightOverlay : public SCompoundWidget
+	class STextHighlightOverlay : public SCompoundWidget
 	{
 	public:
 		SLATE_BEGIN_ARGS(STextHighlightOverlay) {}
@@ -133,6 +133,6 @@ namespace PropertyInfoViewStyle
 			SLATE_ATTRIBUTE(FText, HighlightText)
 		SLATE_END_ARGS()
 
-		void Construct(const FArguments& InArgs);
+		UNREALED_API void Construct(const FArguments& InArgs);
 	};
 }
diff --git a/Engine/Source/Editor/UnrealEd/Public/RayTracingDebugVisualizationMenuCommands.h b/Engine/Source/Editor/UnrealEd/Public/RayTracingDebugVisualizationMenuCommands.h
--- a/Engine/Source/Editor/UnrealEd/Public/RayTracingDebugVisualizationMenuCommands.h
+++ b/Engine/Source/Editor/UnrealEd/Public/RayTracingDebugVisualizationMenuCommands.h
@@ -11,7 +11,7 @@
 
 class FEditorViewportClient;
 
-class UNREALED_API FRayTracingDebugVisualizationMenuCommands : public TCommands<FRayTracingDebugVisualizationMenuCommands>
+class FRayTracingDebugVisualizationMenuCommands : public TCommands<FRayTracingDebugVisualizationMenuCommands>
 {
 public:
 	struct FRayTracingDebugVisualizationRecord
@@ -28,15 +28,15 @@ public:
 		}
 	};
 
-	FRayTracingDebugVisualizationMenuCommands();
+	UNREALED_API FRayTracingDebugVisualizationMenuCommands();
 
-	static void BuildVisualisationSubMenu(FMenuBuilder& Menu);
+	static UNREALED_API void BuildVisualisationSubMenu(FMenuBuilder& Menu);
 
-	virtual void RegisterCommands() override;
+	UNREALED_API virtual void RegisterCommands() override;
 
-	void BindCommands(FUICommandList& CommandList, const TSharedPtr<FEditorViewportClient>& Client) const;
+	UNREALED_API void BindCommands(FUICommandList& CommandList, const TSharedPtr<FEditorViewportClient>& Client) const;
 
-	static bool DebugModeShouldBeTonemapped(const FName& RayTracingDebugModeName);
+	static UNREALED_API bool DebugModeShouldBeTonemapped(const FName& RayTracingDebugModeName);
 
 private:
 	void BuildCommandMap();
@@ -51,4 +51,4 @@ private:
 	TArray<FRayTracingDebugVisualizationRecord> RayTracingDebugVisualizationCommands;
 
 	static TArray<FText> RayTracingDebugModeNames;
-};
\ No newline at end of file
+};
diff --git a/Engine/Source/Editor/UnrealEd/Public/ReferencedAssetsUtils.h b/Engine/Source/Editor/UnrealEd/Public/ReferencedAssetsUtils.h
--- a/Engine/Source/Editor/UnrealEd/Public/ReferencedAssetsUtils.h
+++ b/Engine/Source/Editor/UnrealEd/Public/ReferencedAssetsUtils.h
@@ -139,7 +139,7 @@ public:
 
 	/* === FGCObject interface === */
 	UNREALED_API virtual void AddReferencedObjects( FReferenceCollector& Collector );
-	UNREALED_API virtual FString GetReferencerName() const override
+	virtual FString GetReferencerName() const override
 	{
 		return TEXT("FFindReferencedAssets");
 	}
diff --git a/Engine/Source/Editor/UnrealEd/Public/SActorEditorContext.h b/Engine/Source/Editor/UnrealEd/Public/SActorEditorContext.h
--- a/Engine/Source/Editor/UnrealEd/Public/SActorEditorContext.h
+++ b/Engine/Source/Editor/UnrealEd/Public/SActorEditorContext.h
@@ -4,17 +4,17 @@
 
 #include "Widgets/SCompoundWidget.h"
 
-class UNREALED_API SActorEditorContext : public SCompoundWidget
+class SActorEditorContext : public SCompoundWidget
 {
 public:
 	SLATE_BEGIN_ARGS(SActorEditorContext) {}
 	SLATE_ARGUMENT(UWorld*, World)
 	SLATE_END_ARGS()
 
-	void Construct(const FArguments& InArgs);
-	~SActorEditorContext();
+	UNREALED_API void Construct(const FArguments& InArgs);
+	UNREALED_API ~SActorEditorContext();
 
-	static bool IsVisible(UWorld* InWorld);
+	static UNREALED_API bool IsVisible(UWorld* InWorld);
 private:
 
 	void OnEditorMapChange(uint32 MapChangeFlags = 0) { Rebuild(); }
@@ -22,4 +22,4 @@ private:
 	
 	bool bIsContextExpanded;
 	UWorld* World;
-};
\ No newline at end of file
+};
diff --git a/Engine/Source/Editor/UnrealEd/Public/SAssetEditorViewport.h b/Engine/Source/Editor/UnrealEd/Public/SAssetEditorViewport.h
--- a/Engine/Source/Editor/UnrealEd/Public/SAssetEditorViewport.h
+++ b/Engine/Source/Editor/UnrealEd/Public/SAssetEditorViewport.h
@@ -11,7 +11,7 @@ class FViewportTabContent;
 class FMenuBuilder;
 struct FAssetEditorViewportConstructionArgs;
 
-class UNREALED_API SAssetEditorViewport : public SEditorViewport
+class SAssetEditorViewport : public SEditorViewport
 {
 public:
 
@@ -23,16 +23,16 @@ public:
 		SLATE_ARGUMENT(TSharedPtr<FEditorViewportClient>, EditorViewportClient)
 	SLATE_END_ARGS()
 
-	void Construct(const FArguments& InArgs, const FAssetEditorViewportConstructionArgs& InViewportConstructionArgs);
-	void OnSetViewportConfiguration(FName ConfigurationName);
-	bool IsViewportConfigurationSet(FName ConfigurationName) const;
+	UNREALED_API void Construct(const FArguments& InArgs, const FAssetEditorViewportConstructionArgs& InViewportConstructionArgs);
+	UNREALED_API void OnSetViewportConfiguration(FName ConfigurationName);
+	UNREALED_API bool IsViewportConfigurationSet(FName ConfigurationName) const;
 
-	void GenerateLayoutMenu(FMenuBuilder& MenuBuilder) const;
+	UNREALED_API void GenerateLayoutMenu(FMenuBuilder& MenuBuilder) const;
 
 protected:
-	virtual void BindCommands() override;
+	UNREALED_API virtual void BindCommands() override;
 
-	virtual TSharedRef<FEditorViewportClient> MakeEditorViewportClient() override;
+	UNREALED_API virtual TSharedRef<FEditorViewportClient> MakeEditorViewportClient() override;
 	FName ConfigKey;
 	TWeakPtr<FViewportTabContent> ParentTabContent;
 };
diff --git a/Engine/Source/Editor/UnrealEd/Public/SCSVImportOptions.h b/Engine/Source/Editor/UnrealEd/Public/SCSVImportOptions.h
--- a/Engine/Source/Editor/UnrealEd/Public/SCSVImportOptions.h
+++ b/Engine/Source/Editor/UnrealEd/Public/SCSVImportOptions.h
@@ -45,7 +45,7 @@ enum class ECSVImportOptionDlgResponse : uint8
 	Cancel
 };
 
-class UNREALED_API SCSVImportOptions : public SCompoundWidget
+class SCSVImportOptions : public SCompoundWidget
 {
 private:
 	/** Typedef for curve enum pointers */
@@ -70,64 +70,64 @@ public:
 		, TempImportDataTable(nullptr)
 		{}
 
-	void Construct(const FArguments& InArgs);
+	UNREALED_API void Construct(const FArguments& InArgs);
 
 	/** If we should import */
-	bool ShouldImport();
+	UNREALED_API bool ShouldImport();
 
 	/** If the current settings should be applied to all items being imported */
-	bool ShouldImportAll();
+	UNREALED_API bool ShouldImportAll();
 
 	/** Get the row struct we selected */
-	UScriptStruct* GetSelectedRowStruct();
+	UNREALED_API UScriptStruct* GetSelectedRowStruct();
 
 	/** Get the import type we selected */
-	ECSVImportType GetSelectedImportType();
+	UNREALED_API ECSVImportType GetSelectedImportType();
 
 	/** Get the interpolation mode we selected */
-	ERichCurveInterpMode GetSelectedCurveIterpMode();
+	UNREALED_API ERichCurveInterpMode GetSelectedCurveIterpMode();
 
 	/** Whether to show table row options */
-	EVisibility GetTableRowOptionVis() const;
+	UNREALED_API EVisibility GetTableRowOptionVis() const;
 
 	/** Whether to show curve type options */
-	EVisibility GetCurveTypeVis() const;
+	UNREALED_API EVisibility GetCurveTypeVis() const;
 
 	/** Whether to show details panel */
-	EVisibility GetDetailsPanelVis() const;
+	UNREALED_API EVisibility GetDetailsPanelVis() const;
 
-	FString GetImportTypeText(TSharedPtr<ECSVImportType> Type) const;
+	UNREALED_API FString GetImportTypeText(TSharedPtr<ECSVImportType> Type) const;
 
 	/** Called to create a widget for each struct */
-	TSharedRef<SWidget> MakeImportTypeItemWidget(TSharedPtr<ECSVImportType> Type);
+	UNREALED_API TSharedRef<SWidget> MakeImportTypeItemWidget(TSharedPtr<ECSVImportType> Type);
 
 	/** Called when import type changes */
-	void OnImportTypeSelected(TSharedPtr<ECSVImportType> Selection, ESelectInfo::Type SelectionType);
+	UNREALED_API void OnImportTypeSelected(TSharedPtr<ECSVImportType> Selection, ESelectInfo::Type SelectionType);
 
 	/** Called when datatable row is selected */
-	void OnStructSelected(UScriptStruct* NewStruct);
+	UNREALED_API void OnStructSelected(UScriptStruct* NewStruct);
 
-	FString GetCurveTypeText(CurveInterpModePtr InterpMode) const;
+	UNREALED_API FString GetCurveTypeText(CurveInterpModePtr InterpMode) const;
 
 	/** Called to create a widget for each curve interpolation enum */
-	TSharedRef<SWidget> MakeCurveTypeWidget(CurveInterpModePtr InterpMode);
+	UNREALED_API TSharedRef<SWidget> MakeCurveTypeWidget(CurveInterpModePtr InterpMode);
 
 	/** Called when 'Apply' button is pressed */
-	FReply OnImport();
+	UNREALED_API FReply OnImport();
 
 	/** Do we have all of the data we need to import this asset? */
-	bool CanImport() const;
+	UNREALED_API bool CanImport() const;
 
 	/** Called when 'Cancel' button is pressed */
-	FReply OnCancel();
+	UNREALED_API FReply OnCancel();
 
-	FText GetSelectedItemText() const;
+	UNREALED_API FText GetSelectedItemText() const;
 
-	FText GetSelectedCurveTypeText() const;
+	UNREALED_API FText GetSelectedCurveTypeText() const;
 
 private:
-	FReply HandleImport();
-	FReply OnImportAll();
+	UNREALED_API FReply HandleImport();
+	UNREALED_API FReply OnImportAll();
 
 	/** Whether we should go ahead with import */
 	ECSVImportOptionDlgResponse					UserDlgResponse;
diff --git a/Engine/Source/Editor/UnrealEd/Public/SColorGradientEditor.h b/Engine/Source/Editor/UnrealEd/Public/SColorGradientEditor.h
--- a/Engine/Source/Editor/UnrealEd/Public/SColorGradientEditor.h
+++ b/Engine/Source/Editor/UnrealEd/Public/SColorGradientEditor.h
@@ -17,21 +17,21 @@ class FCurveOwnerInterface;
 class FPaintArgs;
 class FSlateWindowElementList;
 
-struct UNREALED_API FGradientStopMark
+struct FGradientStopMark
 {
 public:
 
-	FGradientStopMark();
-	FGradientStopMark( float InTime, FKeyHandle InRedKeyHandle, FKeyHandle InGreenKeyHandle, FKeyHandle InBlueKeyHandle, FKeyHandle InAlphaKeyHandle = FKeyHandle() );
+	UNREALED_API FGradientStopMark();
+	UNREALED_API FGradientStopMark( float InTime, FKeyHandle InRedKeyHandle, FKeyHandle InGreenKeyHandle, FKeyHandle InBlueKeyHandle, FKeyHandle InAlphaKeyHandle = FKeyHandle() );
 
 
-	bool operator==( const FGradientStopMark& Other ) const;
-	bool IsValid( FCurveOwnerInterface& CurveOwner );
-	bool IsValidAlphaMark( const TArray<FRichCurveEditInfo>& Curves ) const;
-	bool IsValidColorMark( const TArray<FRichCurveEditInfo>& Curves ) const;
-	FLinearColor GetColor( FCurveOwnerInterface& CurveOwner ) const;
-	void SetColor( const FLinearColor& InColor, FCurveOwnerInterface& CurveOwner );
-	void SetTime( float NewTime, FCurveOwnerInterface& CurveOwner );
+	UNREALED_API bool operator==( const FGradientStopMark& Other ) const;
+	UNREALED_API bool IsValid( FCurveOwnerInterface& CurveOwner );
+	UNREALED_API bool IsValidAlphaMark( const TArray<FRichCurveEditInfo>& Curves ) const;
+	UNREALED_API bool IsValidColorMark( const TArray<FRichCurveEditInfo>& Curves ) const;
+	UNREALED_API FLinearColor GetColor( FCurveOwnerInterface& CurveOwner ) const;
+	UNREALED_API void SetColor( const FLinearColor& InColor, FCurveOwnerInterface& CurveOwner );
+	UNREALED_API void SetTime( float NewTime, FCurveOwnerInterface& CurveOwner );
 
 public:
 	float Time;
@@ -42,7 +42,7 @@ public:
 
 };
 
-class UNREALED_API SColorGradientEditor : public SLeafWidget
+class SColorGradientEditor : public SLeafWidget
 {
 public:
 	SLATE_BEGIN_ARGS( SColorGradientEditor ) 
@@ -56,25 +56,25 @@ public:
 	SLATE_END_ARGS()
 
 
-	void Construct( const FArguments& InArgs );
+	UNREALED_API void Construct( const FArguments& InArgs );
 
 	/** SWidget Interface */
 	virtual bool SupportsKeyboardFocus() const override { return true; }
-	virtual int32 OnPaint( const FPaintArgs& Args, const FGeometry& AllottedGeometry, const FSlateRect& MyCullingRect, FSlateWindowElementList& OutDrawElements, int32 LayerId, const FWidgetStyle& InWidgetStyle, bool bParentEnabled ) const override;
-	virtual FReply OnMouseButtonDown( const FGeometry& MyGeometry, const FPointerEvent& MouseEvent ) override;
-	virtual FReply OnMouseButtonDoubleClick( const FGeometry& InMyGeometry, const FPointerEvent& InMouseEvent ) override;
-	virtual FReply OnMouseMove( const FGeometry& MyGeometry, const FPointerEvent& MouseEvent ) override;
-	virtual FReply OnMouseButtonUp( const FGeometry& MyGeometry, const FPointerEvent& MouseEvent ) override;
-	virtual FReply OnKeyDown( const FGeometry& MyGeometry, const FKeyEvent& InKeyEvent ) override;
-	virtual void OnMouseLeave(const FPointerEvent& MouseEvent) override;
-	virtual FVector2D ComputeDesiredSize(float) const override;
+	UNREALED_API virtual int32 OnPaint( const FPaintArgs& Args, const FGeometry& AllottedGeometry, const FSlateRect& MyCullingRect, FSlateWindowElementList& OutDrawElements, int32 LayerId, const FWidgetStyle& InWidgetStyle, bool bParentEnabled ) const override;
+	UNREALED_API virtual FReply OnMouseButtonDown( const FGeometry& MyGeometry, const FPointerEvent& MouseEvent ) override;
+	UNREALED_API virtual FReply OnMouseButtonDoubleClick( const FGeometry& InMyGeometry, const FPointerEvent& InMouseEvent ) override;
+	UNREALED_API virtual FReply OnMouseMove( const FGeometry& MyGeometry, const FPointerEvent& MouseEvent ) override;
+	UNREALED_API virtual FReply OnMouseButtonUp( const FGeometry& MyGeometry, const FPointerEvent& MouseEvent ) override;
+	UNREALED_API virtual FReply OnKeyDown( const FGeometry& MyGeometry, const FKeyEvent& InKeyEvent ) override;
+	UNREALED_API virtual void OnMouseLeave(const FPointerEvent& MouseEvent) override;
+	UNREALED_API virtual FVector2D ComputeDesiredSize(float) const override;
 
 	/**
 	 * Sets the curve that is being edited by this editor
 	 *
 	 * @param InCurveOwner Interface to the curves to edit
 	 */
-	void SetCurveOwner( FCurveOwnerInterface* InCurveOwner );
+	UNREALED_API void SetCurveOwner( FCurveOwnerInterface* InCurveOwner );
 
 private:
 	/**
@@ -82,51 +82,51 @@ private:
 	 *
 	 * @param MouseEvent	The mouse event that raised this handler (containing the location and widget path of the click)
 	 */
-	void OpenGradientStopContextMenu(const FPointerEvent& MouseEvent);
+	UNREALED_API void OpenGradientStopContextMenu(const FPointerEvent& MouseEvent);
 
 	/**
 	 * Opens a color picker to change the color of the selected stop
 	 */
-	void OpenGradientStopColorPicker();
+	UNREALED_API void OpenGradientStopColorPicker();
 
 	/**
 	 * Called when the selected stop color changes from the color picker
 	 *
 	 * @param InNewColor	The new color
 	 */
-	void OnSelectedStopColorChanged( FLinearColor InNewColor );
+	UNREALED_API void OnSelectedStopColorChanged( FLinearColor InNewColor );
 
 	/**
 	 * Called when canceling out of the color picker
 	 *
 	 * @param PreviousColor	The color before anything was changed
 	 */
-	void OnCancelSelectedStopColorChange( FLinearColor PreviousColor );
+	UNREALED_API void OnCancelSelectedStopColorChange( FLinearColor PreviousColor );
 
 	/**
 	 * Called right before a user begins using the slider to change the alpha value of a stop
 	 */
-	void OnBeginChangeAlphaValue();
+	UNREALED_API void OnBeginChangeAlphaValue();
 
 	/**
 	 * Called right after a user ends using the slider to change the alpha value of a stop
 	 */
-	void OnEndChangeAlphaValue( float NewValue );
+	UNREALED_API void OnEndChangeAlphaValue( float NewValue );
 
 	/**
 	 * Called when the alpha value of a stop changes
 	 */
-	void OnAlphaValueChanged( float NewValue );
+	UNREALED_API void OnAlphaValueChanged( float NewValue );
 
 	/**
 	 * Called when the alpha value of a stop changes by typing into the sliders text box
 	 */
-	void OnAlphaValueCommitted( float NewValue, ETextCommit::Type );
+	UNREALED_API void OnAlphaValueCommitted( float NewValue, ETextCommit::Type );
 
 	/**
 	 * Called to remove the selected gradient stop
 	 */
-	void OnRemoveSelectedGradientStop();
+	UNREALED_API void OnRemoveSelectedGradientStop();
 
 	/**
 	 * Called when the gradient stop time is changed from the text entry popup
@@ -134,7 +134,7 @@ private:
 	 * @param NewText	The new time as text
 	 * @param Type		How the value was committed 
 	 */
-	void OnSetGradientStopTimeFromPopup( const FText& NewText, ETextCommit::Type Type );
+	UNREALED_API void OnSetGradientStopTimeFromPopup( const FText& NewText, ETextCommit::Type Type );
 
 	/**
 	 * Draws a single gradient stop
@@ -149,21 +149,21 @@ private:
 	 * @param DrawEffects		Any draw effects to apply
 	 * @param bColor			If true RGB of the Color param will be used, otherwise the A value will be used
 	 */
-	void DrawGradientStopMark( const FGradientStopMark& Mark, const FGeometry& Geometry, float XPos, const FLinearColor& Color, FSlateWindowElementList& OutDrawElements, int32 LayerId, const FSlateRect& InClippingRect, ESlateDrawEffect DrawEffects, bool bColor, const FWidgetStyle& InWidgetStyle ) const;
+	UNREALED_API void DrawGradientStopMark( const FGradientStopMark& Mark, const FGeometry& Geometry, float XPos, const FLinearColor& Color, FSlateWindowElementList& OutDrawElements, int32 LayerId, const FSlateRect& InClippingRect, ESlateDrawEffect DrawEffects, bool bColor, const FWidgetStyle& InWidgetStyle ) const;
 
 	/**
 	 * Calculates the geometry of the gradient stop color mark area
 	 *
 	 * @param InGeometry The parent geometry of the area
 	 */
-	FGeometry GetColorMarkAreaGeometry( const FGeometry& InGeometry ) const;
+	UNREALED_API FGeometry GetColorMarkAreaGeometry( const FGeometry& InGeometry ) const;
 
 	/**
 	 * Calculates the geometry of the gradient stop alpha mark area
 	 *
 	 * @param InGeometry The parent geometry of the area
 	 */
-	FGeometry GetAlphaMarkAreaGeometry( const FGeometry& InGeometry ) const;
+	UNREALED_API FGeometry GetAlphaMarkAreaGeometry( const FGeometry& InGeometry ) const;
 
 	/**
 	 * Gets the gradient stop (if any) of at the current mouse position
@@ -172,19 +172,19 @@ private:
 	 * @param MyGeometry	The geometry that was clicked in
 	 * @return				The stop mark at MousePos or invalid if none was found
 	 */
-	FGradientStopMark GetGradientStopAtPoint( const FVector2D& MousePos, const FGeometry& MyGeometry );
+	UNREALED_API FGradientStopMark GetGradientStopAtPoint( const FVector2D& MousePos, const FGeometry& MyGeometry );
 
 	/**
 	 * Get all gradient stop marks on the curve
 	 */
-	void GetGradientStopMarks( TArray<FGradientStopMark>& OutColorMarks, TArray<FGradientStopMark>& OutAlphaMarks ) const;
+	UNREALED_API void GetGradientStopMarks( TArray<FGradientStopMark>& OutColorMarks, TArray<FGradientStopMark>& OutAlphaMarks ) const;
 
 	/**
 	 * Removes a gradient stop
 	 *
 	 * @param InMark	The stop to remove
 	 */
-	void DeleteStop( const FGradientStopMark& InMark );
+	UNREALED_API void DeleteStop( const FGradientStopMark& InMark );
 
 	/**
 	 * Adds a gradient stop
@@ -193,7 +193,7 @@ private:
 	 * @param MyGeometry	The geometry of the stop area
 	 * @param bColorStop	If true a color stop is added, otherwise an alpha stop is added
 	 */
-	FGradientStopMark AddStop( const FVector2D& Position, const FGeometry& MyGeometry, bool bColorStop );
+	UNREALED_API FGradientStopMark AddStop( const FVector2D& Position, const FGeometry& MyGeometry, bool bColorStop );
 
 	/**
 	 * Moves a gradient stop to a new time
@@ -201,11 +201,11 @@ private:
 	 * @param Mark		The gradient stop to move
 	 * @param NewTime	The new time to move to
 	 */
-	void MoveStop( FGradientStopMark& Mark, float NewTime );
+	UNREALED_API void MoveStop( FGradientStopMark& Mark, float NewTime );
 
 private:
 	/** Local space rectangle of each gradient stop handle */
-	static const FSlateRect HandleRect;
+	static UNREALED_API const FSlateRect HandleRect;
 
 	/** The currently selected stop */
 	FGradientStopMark SelectedStop;
diff --git a/Engine/Source/Editor/UnrealEd/Public/SCommonEditorViewportToolbarBase.h b/Engine/Source/Editor/UnrealEd/Public/SCommonEditorViewportToolbarBase.h
--- a/Engine/Source/Editor/UnrealEd/Public/SCommonEditorViewportToolbarBase.h
+++ b/Engine/Source/Editor/UnrealEd/Public/SCommonEditorViewportToolbarBase.h
@@ -63,7 +63,7 @@ namespace CommonEditorViewportUtils
 /**
  * A viewport toolbar widget for an asset or level editor that is placed in a viewport
  */
-class UNREALED_API SCommonEditorViewportToolbarBase : public SViewportToolBar
+class SCommonEditorViewportToolbarBase : public SViewportToolBar
 {
 public:
 	SLATE_BEGIN_ARGS(SCommonEditorViewportToolbarBase)
@@ -75,12 +75,12 @@ public:
 		SLATE_ARGUMENT(TSharedPtr<IPreviewProfileController>, PreviewProfileController) // Should be null if the Preview doesn't require profile.
 	SLATE_END_ARGS()
 
-	virtual ~SCommonEditorViewportToolbarBase();
+	UNREALED_API virtual ~SCommonEditorViewportToolbarBase();
 
-	void Construct(const FArguments& InArgs, TSharedPtr<class ICommonEditorViewportToolbarInfoProvider> InInfoProvider);
+	UNREALED_API void Construct(const FArguments& InArgs, TSharedPtr<class ICommonEditorViewportToolbarInfoProvider> InInfoProvider);
 
 	/** Build the screen percentage menu. */
-	static void ConstructScreenPercentageMenu(FMenuBuilder& MenuBuilder, class FEditorViewportClient* ViewportClient);
+	static UNREALED_API void ConstructScreenPercentageMenu(FMenuBuilder& MenuBuilder, class FEditorViewportClient* ViewportClient);
 
 private:
 	/**
@@ -88,7 +88,7 @@ private:
 	 *
 	 * @return	Label to use for this menu label
 	 */
-	FText GetCameraMenuLabel() const;
+	UNREALED_API FText GetCameraMenuLabel() const;
 
 
 	/**
@@ -96,82 +96,82 @@ private:
 	 *
 	 * @return	Label to use for this menu label
 	 */
-	FText GetViewMenuLabel() const;
+	UNREALED_API FText GetViewMenuLabel() const;
 
 	/**
 	 * Generates the toolbar options menu content 
 	 *
 	 * @return The widget containing the options menu content
 	 */
-	TSharedRef<SWidget> GenerateOptionsMenu() const;
+	UNREALED_API TSharedRef<SWidget> GenerateOptionsMenu() const;
 
 	/**
 	 * Generates the toolbar camera menu content 
 	 *
 	 * @return The widget containing the view menu content
 	 */
-	TSharedRef<SWidget> GenerateCameraMenu() const;
+	UNREALED_API TSharedRef<SWidget> GenerateCameraMenu() const;
 
 	/**
 	 * Generates the toolbar view menu content 
 	 *
 	 * @return The widget containing the view menu content
 	 */
-	TSharedRef<SWidget> GenerateViewMenu() const;
+	UNREALED_API TSharedRef<SWidget> GenerateViewMenu() const;
 
 	/**
 	 * Generates the toolbar show menu content 
 	 *
 	 * @return The widget containing the show menu content
 	 */
-	virtual TSharedRef<SWidget> GenerateShowMenu() const;
+	UNREALED_API virtual TSharedRef<SWidget> GenerateShowMenu() const;
 
 	/**
 	 * Returns the initial visibility of the view mode options widget 
 	 *
 	 * @return The visibility value
 	 */
-	EVisibility GetViewModeOptionsVisibility() const;
+	UNREALED_API EVisibility GetViewModeOptionsVisibility() const;
 
 	/**
 	 * Generates the toolbar view param menu content 
 	 *
 	 * @return The widget containing the show menu content
 	 */
-	TSharedRef<SWidget> GenerateViewModeOptionsMenu() const;
+	UNREALED_API TSharedRef<SWidget> GenerateViewModeOptionsMenu() const;
 
 	/**
 	 * @return The widget containing the perspective only FOV window.
 	 */
-	TSharedRef<SWidget> GenerateFOVMenu() const;
+	UNREALED_API TSharedRef<SWidget> GenerateFOVMenu() const;
 
 	/** Called by the FOV slider in the perspective viewport to get the FOV value */
-	float OnGetFOVValue() const;
+	UNREALED_API float OnGetFOVValue() const;
 
 	/**
 	 * @return The widget containing the far view plane slider.
 	 */
-	TSharedRef<SWidget> GenerateFarViewPlaneMenu() const;
+	UNREALED_API TSharedRef<SWidget> GenerateFarViewPlaneMenu() const;
 
 	/** Called by the far view plane slider in the perspective viewport to get the far view plane value */
-	float OnGetFarViewPlaneValue() const;
+	UNREALED_API float OnGetFarViewPlaneValue() const;
 
 	/** Called when the far view plane slider is adjusted in the perspective viewport */
-	void OnFarViewPlaneValueChanged( float NewValue );
+	UNREALED_API void OnFarViewPlaneValueChanged( float NewValue );
 
 	/** Called when we click the realtime warning */
-	FReply OnRealtimeWarningClicked();
+	UNREALED_API FReply OnRealtimeWarningClicked();
 	/** Called to determine if we should show the realtime warning */
-	EVisibility GetRealtimeWarningVisibility() const;
+	UNREALED_API EVisibility GetRealtimeWarningVisibility() const;
 
 protected:
 	// Merges the extender list from the host with the specified extender and returns the results
-	TSharedPtr<FExtender> GetCombinedExtenderList(TSharedRef<FExtender> MenuExtender) const;
+	UNREALED_API TSharedPtr<FExtender> GetCombinedExtenderList(TSharedRef<FExtender> MenuExtender) const;
 
 	/** Gets the extender for the view menu */
-	virtual TSharedPtr<FExtender> GetViewMenuExtender() const;
+	UNREALED_API virtual TSharedPtr<FExtender> GetViewMenuExtender() const;
 
-	void CreateViewMenuExtensions(FMenuBuilder& MenuBuilder);
+	UNREALED_API void CreateViewMenuExtensions(FMenuBuilder& MenuBuilder);
 
 	/** Extension allowing derived classes to add to the options menu.*/	
 	virtual void ExtendOptionsMenu(FMenuBuilder& OptionsMenuBuilder) const {}
@@ -181,37 +181,37 @@ protected:
 
 protected:
 	// Returns the info provider for this viewport
-	ICommonEditorViewportToolbarInfoProvider& GetInfoProvider() const;
+	UNREALED_API ICommonEditorViewportToolbarInfoProvider& GetInfoProvider() const;
 
 	// Get the viewport client
-	class FEditorViewportClient& GetViewportClient() const;
+	UNREALED_API class FEditorViewportClient& GetViewportClient() const;
 
 protected:
 	// Creates the view menu widget (override point for children)
-	virtual TSharedRef<class SEditorViewportViewMenu> MakeViewMenu();
+	UNREALED_API virtual TSharedRef<class SEditorViewportViewMenu> MakeViewMenu();
 
-	FText GetScalabilityWarningLabel() const;
-	EVisibility GetScalabilityWarningVisibility() const;
-	TSharedRef<SWidget> GetScalabilityWarningMenuContent() const;
+	UNREALED_API FText GetScalabilityWarningLabel() const;
+	UNREALED_API EVisibility GetScalabilityWarningVisibility() const;
+	UNREALED_API TSharedRef<SWidget> GetScalabilityWarningMenuContent() const;
 	virtual bool GetShowScalabilityMenu() const
 	{
 		return false;
 	}
 	/** Called when the FOV slider is adjusted in the perspective viewport */
-	virtual void OnFOVValueChanged(float NewValue) const;
+	UNREALED_API virtual void OnFOVValueChanged(float NewValue) const;
 
 	/** Called when the ScreenPercentage slider is adjusted in the viewport */
-	void OnScreenPercentageValueChanged(int32 NewValue);
+	UNREALED_API void OnScreenPercentageValueChanged(int32 NewValue);
 
 	/** Update the list of asset viewer profiles displayed by the combo box. */
-	void UpdateAssetViewerProfileList();
-	void UpdateAssetViewerProfileSelection();
+	UNREALED_API void UpdateAssetViewerProfileList();
+	UNREALED_API void UpdateAssetViewerProfileSelection();
 
 	/** Invoked when the asset viewer profile combo box selection changes. */
-	void OnAssetViewerProfileComboBoxSelectionChanged(TSharedPtr<FString> NewSelection, ESelectInfo::Type /*SelectInfo*/);
+	UNREALED_API void OnAssetViewerProfileComboBoxSelectionChanged(TSharedPtr<FString> NewSelection, ESelectInfo::Type /*SelectInfo*/);
 
 	/** Creates and returns the asset viewr profile combo box.*/
-	TSharedRef<SWidget> MakeAssetViewerProfileComboBox();
+	UNREALED_API TSharedRef<SWidget> MakeAssetViewerProfileComboBox();
 
 private:
 	/** The viewport that we are in */
diff --git a/Engine/Source/Editor/UnrealEd/Public/SComponentClassCombo.h b/Engine/Source/Editor/UnrealEd/Public/SComponentClassCombo.h
--- a/Engine/Source/Editor/UnrealEd/Public/SComponentClassCombo.h
+++ b/Engine/Source/Editor/UnrealEd/Public/SComponentClassCombo.h
@@ -187,7 +187,7 @@ private:
 
 //////////////////////////////////////////////////////////////////////////
 
-class UNREALED_API SComponentClassCombo : public SCompoundWidget
+class SComponentClassCombo : public SCompoundWidget
 {
 public:
 	SLATE_BEGIN_ARGS( SComponentClassCombo )
@@ -201,33 +201,33 @@ public:
 
 	SLATE_END_ARGS()
 
-	void Construct(const FArguments& InArgs);
+	UNREALED_API void Construct(const FArguments& InArgs);
 
-	~SComponentClassCombo();
+	UNREALED_API ~SComponentClassCombo();
 
 	/** Clear the current combo list selection */
-	void ClearSelection();
+	UNREALED_API void ClearSelection();
 
-	FText GetCurrentSearchString() const;
+	UNREALED_API FText GetCurrentSearchString() const;
 
 	/**
 	 * Called when the user changes the text in the search box.
 	 * @param InSearchText The new search string.
 	 */
-	void OnSearchBoxTextChanged( const FText& InSearchText );
+	UNREALED_API void OnSearchBoxTextChanged( const FText& InSearchText );
 
 	/** Callback when the user commits the text in the searchbox */
-	void OnSearchBoxTextCommitted(const FText& NewText, ETextCommit::Type CommitInfo);
+	UNREALED_API void OnSearchBoxTextCommitted(const FText& NewText, ETextCommit::Type CommitInfo);
 
-	void OnAddComponentSelectionChanged( FComponentClassComboEntryPtr InItem, ESelectInfo::Type SelectInfo );
+	UNREALED_API void OnAddComponentSelectionChanged( FComponentClassComboEntryPtr InItem, ESelectInfo::Type SelectInfo );
 
-	TSharedRef<ITableRow> GenerateAddComponentRow( FComponentClassComboEntryPtr Entry, const TSharedRef<STableViewBase> &OwnerTable ) const;
+	UNREALED_API TSharedRef<ITableRow> GenerateAddComponentRow( FComponentClassComboEntryPtr Entry, const TSharedRef<STableViewBase> &OwnerTable ) const;
 
 	/** Update list of component classes */
-	void UpdateComponentClassList();
+	UNREALED_API void UpdateComponentClassList();
 
 	/** Returns a component name without the substring "Component" and sanitized for display */
-	static FString GetSanitizedComponentName(FComponentClassComboEntryPtr Entry);
+	static UNREALED_API FString GetSanitizedComponentName(FComponentClassComboEntryPtr Entry);
 
 protected:
 	/** Internal data used to facilitate component class filtering */
@@ -241,7 +241,7 @@ protected:
 	/**
 	 * Updates the filtered list of component names.
 	 */
-	void GenerateFilteredComponentList();
+	UNREALED_API void GenerateFilteredComponentList();
 
 private:
 	
diff --git a/Engine/Source/Editor/UnrealEd/Public/SCurveEditor.h b/Engine/Source/Editor/UnrealEd/Public/SCurveEditor.h
--- a/Engine/Source/Editor/UnrealEd/Public/SCurveEditor.h
+++ b/Engine/Source/Editor/UnrealEd/Public/SCurveEditor.h
@@ -267,7 +267,7 @@ public:
 	 * @param Ar The archive to serialize with
 	 */
 	UNREALED_API virtual void AddReferencedObjects( FReferenceCollector& Collector ) override;
-	UNREALED_API virtual FString GetReferencerName() const override
+	virtual FString GetReferencerName() const override
 	{
 		return TEXT("SCurveEditor");
 	}
@@ -276,8 +276,8 @@ public:
 	UNREALED_API TSharedPtr<FUICommandList> GetCommands();
 
 	/** Gets or sets whether autoframing is allowed */
-	UNREALED_API bool GetAllowAutoFrame() const { return bAllowAutoFrame; }
-	UNREALED_API void SetAllowAutoFrame(bool bInAllowAutoFrame) { bAllowAutoFrame = bInAllowAutoFrame; }
+	bool GetAllowAutoFrame() const { return bAllowAutoFrame; }
+	void SetAllowAutoFrame(bool bInAllowAutoFrame) { bAllowAutoFrame = bInAllowAutoFrame; }
 
 	/** Gets whether autoframe will be invoked (combination of allow auto frame and curve editor auto frame setting) */
 	UNREALED_API bool GetAutoFrame() const;
@@ -598,7 +598,7 @@ private:
 
 	//~ Begin FEditorUndoClient Interface
 	UNREALED_API virtual void PostUndo(bool bSuccess) override;
-	UNREALED_API virtual void PostRedo(bool bSuccess) override { PostUndo(bSuccess); }
+	virtual void PostRedo(bool bSuccess) override { PostUndo(bSuccess); }
 	// End of FEditorUndoClient
 
 	bool AreCurvesVisible() const { return bAlwaysDisplayColorCurves || bAreCurvesVisible.Get(); }
diff --git a/Engine/Source/Editor/UnrealEd/Public/SEditorViewport.h b/Engine/Source/Editor/UnrealEd/Public/SEditorViewport.h
--- a/Engine/Source/Editor/UnrealEd/Public/SEditorViewport.h
+++ b/Engine/Source/Editor/UnrealEd/Public/SEditorViewport.h
@@ -18,7 +18,7 @@ class SViewport;
 struct FSlateBrush;
 struct FToolMenuContext;
 
-class UNREALED_API SEditorViewport
+class SEditorViewport
 	: public SCompoundWidget
 {
 public:
@@ -28,56 +28,56 @@ public:
 	SLATE_ATTRIBUTE(FVector2D, ViewportSize);
 	SLATE_END_ARGS()
 	
-	SEditorViewport();
-	virtual ~SEditorViewport();
+	UNREALED_API SEditorViewport();
+	UNREALED_API virtual ~SEditorViewport();
 
-	void Construct( const FArguments& InArgs );
+	UNREALED_API void Construct( const FArguments& InArgs );
 
-	virtual FReply OnKeyDown( const FGeometry& MyGeometry, const FKeyEvent& InKeyEvent ) override;
-	virtual bool SupportsKeyboardFocus() const override;
-	virtual FReply OnFocusReceived( const FGeometry& MyGeometry, const FFocusEvent& InFocusEvent ) override;
-	virtual void Tick( const FGeometry& AllottedGeometry, const double InCurrentTime, const float InDeltaTime ) override;
+	UNREALED_API virtual FReply OnKeyDown( const FGeometry& MyGeometry, const FKeyEvent& InKeyEvent ) override;
+	UNREALED_API virtual bool SupportsKeyboardFocus() const override;
+	UNREALED_API virtual FReply OnFocusReceived( const FGeometry& MyGeometry, const FFocusEvent& InFocusEvent ) override;
+	UNREALED_API virtual void Tick( const FGeometry& AllottedGeometry, const double InCurrentTime, const float InDeltaTime ) override;
 	
 	/**
 	 * @return True if the viewport is being updated in realtime
 	 */
-	bool IsRealtime() const;
+	UNREALED_API bool IsRealtime() const;
 
 	/** @return True if the viewport is currently visible */
-	virtual bool IsVisible() const;
+	UNREALED_API virtual bool IsVisible() const;
 
 	/** 
 	 * Invalidates the viewport to ensure it is redrawn during the next tick. 
 	 * This is implied every frame while the viewport IsRealtime().
 	 */
-	void Invalidate();
+	UNREALED_API void Invalidate();
 
 	/** Toggles realtime on/off for the viewport. Slate tick/paint is ensured when realtime is on. */
-	void OnToggleRealtime();
+	UNREALED_API void OnToggleRealtime();
 
 	/**
 	 * Sets whether this viewport can render directly to the back buffer.  Advanced use only
 	 * 
 	 * @param	bInRenderDirectlyToWindow	Whether we should be able to render to the back buffer
 	 */
-	void SetRenderDirectlyToWindow( const bool bInRenderDirectlyToWindow );
+	UNREALED_API void SetRenderDirectlyToWindow( const bool bInRenderDirectlyToWindow );
 
 	/**
 	 * Sets whether stereo rendering is allowed for this viewport.  Advanced use only
 	 * 
 	 * @param	bInEnableStereoRendering	Whether stereo rendering should be allowed for this viewport
 	 */
-	void EnableStereoRendering( const bool bInEnableStereoRendering );
+	UNREALED_API void EnableStereoRendering( const bool bInEnableStereoRendering );
 
 	/**
 	 * @return true if the specified coordinate system the active one active
 	 */
-	virtual bool IsCoordSystemActive( ECoordSystem CoordSystem ) const;
+	UNREALED_API virtual bool IsCoordSystemActive( ECoordSystem CoordSystem ) const;
 
 	/**
 	 * Cycles between world and local coordinate systems
 	 */
-	virtual void OnCycleCoordinateSystem();
+	UNREALED_API virtual void OnCycleCoordinateSystem();
 	
 	/** @return The viewport command list */
 	const TSharedPtr<FUICommandList> GetCommandList() const { return CommandList; }
@@ -92,55 +92,55 @@ public:
 	/**
 	 * Controls the visibility of the widget transform toolbar, if there is an associated toolbar
 	 */
-	virtual EVisibility GetTransformToolbarVisibility() const;
+	UNREALED_API virtual EVisibility GetTransformToolbarVisibility() const;
 
 	/** Build the exposure menu using EV100 settings */
-	TSharedRef<SWidget> BuildFixedEV100Menu()  const;
+	UNREALED_API TSharedRef<SWidget> BuildFixedEV100Menu()  const;
 
 	/**
  * Called when the user wants to show the in-viewport context menu
  */
 	virtual void ToggleInViewportContextMenu() {}
 	virtual void HideInViewportContextMenu() {}
-	virtual void UpdateInViewportMenuLocation(const FVector2D InLocation);
+	UNREALED_API virtual void UpdateInViewportMenuLocation(const FVector2D InLocation);
 	virtual bool CanToggleInViewportContextMenu() { return false; }
 
-	bool IsPreviewingScreenPercentage() const;
-	void TogglePreviewingScreenPercentage();
-	void OnOpenViewportPerformanceProjectSettings();
-	void OnOpenViewportPerformanceEditorPreferences();
+	UNREALED_API bool IsPreviewingScreenPercentage() const;
+	UNREALED_API void TogglePreviewingScreenPercentage();
+	UNREALED_API void OnOpenViewportPerformanceProjectSettings();
+	UNREALED_API void OnOpenViewportPerformanceEditorPreferences();
 
 ///////////////////////////////////////////////////////////////////////////////
 // begin feature level control functions block
 ///////////////////////////////////////////////////////////////////////////////
 private:
 	/** Called to get the feature level preview text */
-	FText GetCurrentFeatureLevelPreviewText(bool bDrawOnlyLabel) const;
+	UNREALED_API FText GetCurrentFeatureLevelPreviewText(bool bDrawOnlyLabel) const;
 
 	/** Helper function that, for some FeatureLevel argument, will retrieve the required shader platform */
-	EShaderPlatform GetShaderPlatformHelper(const ERHIFeatureLevel::Type InFeatureLevel) const;
+	UNREALED_API EShaderPlatform GetShaderPlatformHelper(const ERHIFeatureLevel::Type InFeatureLevel) const;
 
 protected:
 	/** @return The visibility of the current feature level preview text display */
-	EVisibility GetCurrentFeatureLevelPreviewTextVisibility() const;
+	UNREALED_API EVisibility GetCurrentFeatureLevelPreviewTextVisibility() const;
 
 	/** @return true if realtime can be toggled (it cannot be toggled directly if there is an override in place) */
-	bool CanToggleRealtime() const;
+	UNREALED_API bool CanToggleRealtime() const;
 
 	/** call this function to build a 'text' widget that can display the present feature level */
-	TSharedRef<SWidget> BuildFeatureLevelWidget() const;
+	UNREALED_API TSharedRef<SWidget> BuildFeatureLevelWidget() const;
 ///////////////////////////////////////////////////////////////////////////////
 // end feature level control functions block
 ///////////////////////////////////////////////////////////////////////////////
 
 	/** Called by the fixed EV100 slider to get the fixed EV100 value */
-	float OnGetFixedEV100Value() const;
+	UNREALED_API float OnGetFixedEV100Value() const;
 
 	/** Called when fixed EV100 slider is adjusted */
-	void OnFixedEV100ValueChanged( float NewValue );
+	UNREALED_API void OnFixedEV100ValueChanged( float NewValue );
 
 	/** Called to know whether the fixed EV100 slider is enabled. */
-	bool IsFixedEV100Enabled() const;
+	UNREALED_API bool IsFixedEV100Enabled() const;
 
 
 	virtual TSharedRef<FEditorViewportClient> MakeEditorViewportClient() = 0;
@@ -151,14 +151,14 @@ protected:
 	// Implement this to add an arbitrary set of toolbars or other overlays to the inside of the viewport
 	virtual void PopulateViewportOverlays(TSharedRef<SOverlay> Overlay) { }
 
-	virtual void BindCommands();
+	UNREALED_API virtual void BindCommands();
 	virtual const FSlateBrush* OnGetViewportBorderBrush() const { return NULL; }
 	virtual FSlateColor OnGetViewportBorderColorAndOpacity() const { return FLinearColor::Black; }
 	
 	/**
 	 * @return The visibility of widgets in the viewport (e.g, menus).  Note this is not the visibility of the scene rendered in the viewport                                                              
 	 */
-	virtual EVisibility OnGetViewportContentVisibility() const;
+	UNREALED_API virtual EVisibility OnGetViewportContentVisibility() const;
 
 	/**
 	 * @return The visibility of the viewport focus indicator.                                                             
@@ -166,28 +166,28 @@ protected:
 	virtual EVisibility OnGetFocusedViewportIndicatorVisibility() const { return EVisibility::Collapsed; }
 
 	/** UI Command delegate bindings */
-	void OnToggleStats();
+	UNREALED_API void OnToggleStats();
 
 	/**
 	 * Toggles Stat command visibility in this viewport
 	 *
 	 * @param CommandName				Name of the command
 	 */
-	virtual void ToggleStatCommand(FString CommandName);
+	UNREALED_API virtual void ToggleStatCommand(FString CommandName);
 
 	/**
 	 * Checks if Stat command is visible in this viewport
 	 *
 	 * @param CommandName				Name of the command
 	 */
-	virtual bool IsStatCommandVisible(FString CommandName) const;
+	UNREALED_API virtual bool IsStatCommandVisible(FString CommandName) const;
 
 	/**
 	 * Toggles a show flag in this viewport
 	 *
 	 * @param EngineShowFlagIndex	the ID to toggle
 	 */
-	void ToggleShowFlag( uint32 EngineShowFlagIndex );
+	UNREALED_API void ToggleShowFlag( uint32 EngineShowFlagIndex );
 
 	/**
 	 * Checks if a show flag is enabled in this viewport
@@ -195,16 +195,16 @@ protected:
 	 * @param EngineShowFlagIndex	the ID to check
 	 * @return true if the show flag is enabled, false otherwise
 	 */
-	bool IsShowFlagEnabled( uint32 EngineShowFlagIndex ) const;
+	UNREALED_API bool IsShowFlagEnabled( uint32 EngineShowFlagIndex ) const;
 
 	/** Changes the auto exposure setting for this viewport */
-	void ChangeExposureSetting();
+	UNREALED_API void ChangeExposureSetting();
 
 	/** Checks if auto exposure setting is selected */
-	bool IsExposureSettingSelected() const;
+	UNREALED_API bool IsExposureSettingSelected() const;
 	
-	virtual void OnScreenCapture();
-	virtual void OnScreenCaptureForProjectThumbnail();
+	UNREALED_API virtual void OnScreenCapture();
+	UNREALED_API virtual void OnScreenCaptureForProjectThumbnail();
 	virtual bool DoesAllowScreenCapture() { return true; }
 	
 	/**
@@ -218,22 +218,22 @@ protected:
 	/**
 	 * @return true if the specified widget mode is active
 	 */
-	virtual bool IsWidgetModeActive( UE::Widget::EWidgetMode Mode ) const;
+	UNREALED_API virtual bool IsWidgetModeActive( UE::Widget::EWidgetMode Mode ) const;
 
 	/**
 	 * @return true if the translate/rotate mode widget is visible 
 	 */
-	virtual bool IsTranslateRotateModeVisible() const;
+	UNREALED_API virtual bool IsTranslateRotateModeVisible() const;
 
 	/**
 	* @return true if the 2d mode widget is visible
 	*/
-	virtual bool Is2DModeVisible() const;
+	UNREALED_API virtual bool Is2DModeVisible() const;
 
 	/**
 	 * Moves between widget modes
 	 */
-	virtual void OnCycleWidgetMode();
+	UNREALED_API virtual void OnCycleWidgetMode();
 
 	/**
 	 * Called when the user wants to focus the viewport to the current selection
@@ -241,17 +241,17 @@ protected:
 	virtual void OnFocusViewportToSelection(){}
 
 	/** Gets the world this viewport is for */
-	virtual UWorld* GetWorld() const;
+	UNREALED_API virtual UWorld* GetWorld() const;
 
 	/**
 	 * Called when surface snapping has been enabled/disabled
 	 */
-	static void OnToggleSurfaceSnap();
+	static UNREALED_API void OnToggleSurfaceSnap();
 
 	/**
 	 * Called to test whether surface snapping is enabled or not
 	 */
-	static bool OnIsSurfaceSnapEnabled();
+	static UNREALED_API bool OnIsSurfaceSnapEnabled();
 
 
 protected:
@@ -276,10 +276,10 @@ protected:
 
 private:
 	/** Ensures a Slate tick/paint pass when the viewport is realtime or was invalidated this frame */
-	EActiveTimerReturnType EnsureTick( double InCurrentTime, float InDeltaTime );
+	UNREALED_API EActiveTimerReturnType EnsureTick( double InCurrentTime, float InDeltaTime );
 
 	/** Gets the visibility of the active viewport border */
-	EVisibility GetActiveBorderVisibility() const;
+	UNREALED_API EVisibility GetActiveBorderVisibility() const;
 private:
 	/** The handle to the active timer */
 	TWeakPtr<FActiveTimerHandle> ActiveTimerHandle;
diff --git a/Engine/Source/Editor/UnrealEd/Public/SEditorViewportToolBarButton.h b/Engine/Source/Editor/UnrealEd/Public/SEditorViewportToolBarButton.h
--- a/Engine/Source/Editor/UnrealEd/Public/SEditorViewportToolBarButton.h
+++ b/Engine/Source/Editor/UnrealEd/Public/SEditorViewportToolBarButton.h
@@ -28,7 +28,7 @@ enum class ECheckBoxState : uint8;
 /**
  * A simple class that represents a toolbar button in an editor viewport toolbar
  */
-class UNREALED_API SEditorViewportToolBarButton : public SCompoundWidget
+class SEditorViewportToolBarButton : public SCompoundWidget
 {
 public:
 	SLATE_BEGIN_ARGS(SEditorViewportToolBarButton)
@@ -59,7 +59,7 @@ public:
 		SLATE_DEFAULT_SLOT( FArguments, Content )
 	SLATE_END_ARGS()
 
-	void Construct( const FArguments& Declaration );
+	UNREALED_API void Construct( const FArguments& Declaration );
 private:
 	/** 
 	 * Called when the button check state changes
diff --git a/Engine/Source/Editor/UnrealEd/Public/SEditorViewportToolBarMenu.h b/Engine/Source/Editor/UnrealEd/Public/SEditorViewportToolBarMenu.h
--- a/Engine/Source/Editor/UnrealEd/Public/SEditorViewportToolBarMenu.h
+++ b/Engine/Source/Editor/UnrealEd/Public/SEditorViewportToolBarMenu.h
@@ -39,7 +39,7 @@ namespace EMenuItemType
 /**
  * Widget that opens a menu when clicked
  */
-class UNREALED_API SEditorViewportToolbarMenu : public SCompoundWidget
+class SEditorViewportToolbarMenu : public SCompoundWidget
 {
 public:
 	SLATE_BEGIN_ARGS(SEditorViewportToolbarMenu)
@@ -65,17 +65,17 @@ public:
 	/**
 	 * Constructs the menu
 	 */
-	void Construct( const FArguments& Declaration );
+	UNREALED_API void Construct( const FArguments& Declaration );
 
 	/**
 	 * Returns parent tool bar
 	 */
-	TWeakPtr<class SViewportToolBar> GetParentToolBar() const;
+	UNREALED_API TWeakPtr<class SViewportToolBar> GetParentToolBar() const;
 
 	/**
 	 * @return true if this menu is open
 	 */
-	bool IsMenuOpen() const;
+	UNREALED_API bool IsMenuOpen() const;
 
 private:
 	/**
@@ -87,7 +87,7 @@ private:
 	 * Called when the mouse enters a menu button.  If there was a menu previously opened
 	 * we open this menu automatically
 	 */
-	void OnMouseEnter( const FGeometry& MyGeometry, const FPointerEvent& MouseEvent );
+	UNREALED_API void OnMouseEnter( const FGeometry& MyGeometry, const FPointerEvent& MouseEvent );
 
 	EVisibility GetLabelIconVisibility() const;
 
@@ -106,7 +106,7 @@ protected:
 	 *
 	 * @return	Tool tip text, or an empty text if filtered out
 	 */
-	FText GetFilteredToolTipText(TAttribute<FText> ToolTipText) const;
+	UNREALED_API FText GetFilteredToolTipText(TAttribute<FText> ToolTipText) const;
 
 private:
 	/** Our menus anchor */
diff --git a/Engine/Source/Editor/UnrealEd/Public/SEditorViewportViewMenu.h b/Engine/Source/Editor/UnrealEd/Public/SEditorViewportViewMenu.h
--- a/Engine/Source/Editor/UnrealEd/Public/SEditorViewportViewMenu.h
+++ b/Engine/Source/Editor/UnrealEd/Public/SEditorViewportViewMenu.h
@@ -14,14 +14,14 @@
 struct FSlateBrush;
 class UToolMenu;
 
-class UNREALED_API SEditorViewportViewMenu : public SEditorViewportToolbarMenu
+class SEditorViewportViewMenu : public SEditorViewportToolbarMenu
 {
 public:
 	SLATE_BEGIN_ARGS( SEditorViewportViewMenu ){}
 		SLATE_ARGUMENT( TSharedPtr<class FExtender>, MenuExtenders )
 	SLATE_END_ARGS()
 
-	void Construct( const FArguments& InArgs, TSharedRef<SEditorViewport> InViewport, TSharedRef<class SViewportToolBar> InParentToolBar );
+	UNREALED_API void Construct( const FArguments& InArgs, TSharedRef<SEditorViewport> InViewport, TSharedRef<class SViewportToolBar> InParentToolBar );
 
 private:
 	FText GetViewMenuLabel() const;
@@ -29,11 +29,11 @@ private:
 	void FillViewMenu(UToolMenu* Menu) const;
 
 protected:
-	virtual TSharedRef<SWidget> GenerateViewMenuContent() const;
-	virtual void RegisterMenus() const;
+	UNREALED_API virtual TSharedRef<SWidget> GenerateViewMenuContent() const;
+	UNREALED_API virtual void RegisterMenus() const;
 
 	TWeakPtr<SEditorViewport> Viewport;
 	TSharedPtr<class FExtender> MenuExtenders;
 
-	static const FName BaseMenuName;
+	static UNREALED_API const FName BaseMenuName;
 };
diff --git a/Engine/Source/Editor/UnrealEd/Public/SEditorViewportViewMenuContext.h b/Engine/Source/Editor/UnrealEd/Public/SEditorViewportViewMenuContext.h
--- a/Engine/Source/Editor/UnrealEd/Public/SEditorViewportViewMenuContext.h
+++ b/Engine/Source/Editor/UnrealEd/Public/SEditorViewportViewMenuContext.h
@@ -8,8 +8,8 @@
 
 class SEditorViewportViewMenu;
 
-UCLASS()
-class UNREALED_API UEditorViewportViewMenuContext : public UObject
+UCLASS(MinimalAPI)
+class UEditorViewportViewMenuContext : public UObject
 {
 	GENERATED_BODY()
 public:
diff --git a/Engine/Source/Editor/UnrealEd/Public/SKeySelector.h b/Engine/Source/Editor/UnrealEd/Public/SKeySelector.h
--- a/Engine/Source/Editor/UnrealEd/Public/SKeySelector.h
+++ b/Engine/Source/Editor/UnrealEd/Public/SKeySelector.h
@@ -53,7 +53,7 @@ typedef TSharedPtr<FKeyTreeInfo> FKeyTreeItem;
 typedef STreeView<FKeyTreeItem> SKeyTreeView;
 
 /** Widget for selecting an input key */
-class UNREALED_API SKeySelector : public SCompoundWidget
+class SKeySelector : public SCompoundWidget
 {
 public:
 	SLATE_BEGIN_ARGS( SKeySelector )
@@ -73,7 +73,7 @@ public:
 		SLATE_ARGUMENT( bool, AllowClear )
 	SLATE_END_ARGS()
 public:
-	void Construct(const FArguments& InArgs);
+	UNREALED_API void Construct(const FArguments& InArgs);
 
 	/** Sets bool to produce tooltip notifying this key selector it was disabled from KeyStructCustomization */
 	void SetEnabledFromKeyStructCustomization(bool bIsEnabled)
@@ -101,49 +101,49 @@ public:
 	
 protected:
 	/** Gets the icon for the key being manipulated */
-	const FSlateBrush* GetKeyIconImage() const;
+	UNREALED_API const FSlateBrush* GetKeyIconImage() const;
 	/** Toggles the icon's color when in listen mode */
-	FSlateColor GetKeyIconColor() const;
+	UNREALED_API FSlateColor GetKeyIconColor() const;
 
 	/** Gets a succinct description for the key being manipulated */
-	FText GetKeyDescription() const;
+	UNREALED_API FText GetKeyDescription() const;
 	/** Gets a description tooltip for the key being manipulated */
-	FText GetKeyDescriptionToolTip() const;
+	UNREALED_API FText GetKeyDescriptionToolTip() const;
 	/** Gets a tooltip for the selected key */
-	FText GetKeyTooltip() const;
+	UNREALED_API FText GetKeyTooltip() const;
 	
 	/** Tooltip to display on the selector when the selector is disabled*/
 	FText DisabledSelectorToolTip = FText::FromString(TEXT("Key Selector Disabled"));
 
 	/** Treeview support functions */
-	virtual TSharedRef<ITableRow> GenerateKeyTreeRow(FKeyTreeItem InItem, const TSharedRef<STableViewBase>& OwnerTree);
-	void OnKeySelectionChanged(FKeyTreeItem Selection, ESelectInfo::Type SelectInfo);
-	void GetKeyChildren(FKeyTreeItem InItem, TArray<FKeyTreeItem>& OutChildren);
+	UNREALED_API virtual TSharedRef<ITableRow> GenerateKeyTreeRow(FKeyTreeItem InItem, const TSharedRef<STableViewBase>& OwnerTree);
+	UNREALED_API void OnKeySelectionChanged(FKeyTreeItem Selection, ESelectInfo::Type SelectInfo);
+	UNREALED_API void GetKeyChildren(FKeyTreeItem InItem, TArray<FKeyTreeItem>& OutChildren);
 
 	/** Gets the Menu Content, setting it up if necessary */
-	virtual TSharedRef<SWidget>	GetMenuContent();
+	UNREALED_API virtual TSharedRef<SWidget>	GetMenuContent();
 
 	/** Key searching support */
-	void OnFilterTextChanged(const FText& NewText);
-	void OnFilterTextCommitted(const FText& NewText, ETextCommit::Type CommitInfo);
-	void GetSearchTokens(const FString& SearchString, TArray<FString>& OutTokens) const;
+	UNREALED_API void OnFilterTextChanged(const FText& NewText);
+	UNREALED_API void OnFilterTextCommitted(const FText& NewText, ETextCommit::Type CommitInfo);
+	UNREALED_API void GetSearchTokens(const FString& SearchString, TArray<FString>& OutTokens) const;
 
 	/** Helper to generate the filtered list of keys, based on the search string matching */
-	bool GetChildrenMatchingSearch(const TArray<FString>& SearchTokens, const TArray<FKeyTreeItem>& UnfilteredList, TArray<FKeyTreeItem>& OutFilteredList);
+	UNREALED_API bool GetChildrenMatchingSearch(const TArray<FString>& SearchTokens, const TArray<FKeyTreeItem>& UnfilteredList, TArray<FKeyTreeItem>& OutFilteredList);
 
 	/** Start listening for the next key press */
-	FReply ListenForInput();
+	UNREALED_API FReply ListenForInput();
 	/** Assigns the heard input as the current key */
-	FReply ProcessHeardInput(FKey KeyHeard);
+	UNREALED_API FReply ProcessHeardInput(FKey KeyHeard);
 
 	virtual bool SupportsKeyboardFocus() const override { return bListenForNextInput; }
 
 	/** Input listeners */
-	virtual FReply OnKeyDown(const FGeometry& MyGeometry, const FKeyEvent& InKeyEvent) override;
-	virtual FReply OnMouseButtonDown(const FGeometry& MyGeometry, const FPointerEvent& MouseEvent) override;
-	virtual FReply OnMouseWheel(const FGeometry& MyGeometry, const FPointerEvent& MouseEvent) override;
-	virtual FReply OnMouseMove(const FGeometry& MyGeometry, const FPointerEvent& MouseEvent) override;
-	virtual FReply OnAnalogValueChanged(const FGeometry& MyGeometry, const FAnalogInputEvent& InAnalogInputEvent) override;
+	UNREALED_API virtual FReply OnKeyDown(const FGeometry& MyGeometry, const FKeyEvent& InKeyEvent) override;
+	UNREALED_API virtual FReply OnMouseButtonDown(const FGeometry& MyGeometry, const FPointerEvent& MouseEvent) override;
+	UNREALED_API virtual FReply OnMouseWheel(const FGeometry& MyGeometry, const FPointerEvent& MouseEvent) override;
+	UNREALED_API virtual FReply OnMouseMove(const FGeometry& MyGeometry, const FPointerEvent& MouseEvent) override;
+	UNREALED_API virtual FReply OnAnalogValueChanged(const FGeometry& MyGeometry, const FAnalogInputEvent& InAnalogInputEvent) override;
 
 	/** 
 	 * Determine the best icon to represent the given key.
@@ -151,7 +151,7 @@ protected:
 	 * @param Key		The key to get the icon for.
 	 * @param returns a brush that best represents the icon
 	 */
-	const FSlateBrush* GetIconFromKey(FKey Key) const;
+	UNREALED_API const FSlateBrush* GetIconFromKey(FKey Key) const;
 
 protected:
 	/** Combo Button that shows current key and icon */
diff --git a/Engine/Source/Editor/UnrealEd/Public/STransformViewportToolbar.h b/Engine/Source/Editor/UnrealEd/Public/STransformViewportToolbar.h
--- a/Engine/Source/Editor/UnrealEd/Public/STransformViewportToolbar.h
+++ b/Engine/Source/Editor/UnrealEd/Public/STransformViewportToolbar.h
@@ -25,7 +25,7 @@ DECLARE_DELEGATE_OneParam(FOnCamSpeedScalarChanged, float);
 /**
  * Viewport toolbar containing transform, grid snapping, local to world and camera speed controls.
  */
-class UNREALED_API STransformViewportToolBar  : public SViewportToolBar
+class STransformViewportToolBar  : public SViewportToolBar
 {
 
 public:
@@ -38,13 +38,13 @@ public:
 		SLATE_EVENT( FOnCamSpeedScalarChanged, OnCamSpeedScalarChanged )
 	SLATE_END_ARGS()
 
-	void Construct( const FArguments& InArgs );
+	UNREALED_API void Construct( const FArguments& InArgs );
 	/**
 	 * Static: Creates a widget for the main tool bar
 	 *
 	 * @return	New widget
 	 */
-	TSharedRef<SWidget> MakeTransformToolBar(const TSharedPtr< FExtender > InExtenders);
+	UNREALED_API TSharedRef<SWidget> MakeTransformToolBar(const TSharedPtr< FExtender > InExtenders);
 
 private:
 	FSlateIcon GetLocalToWorldIcon() const;
diff --git a/Engine/Source/Editor/UnrealEd/Public/SViewportToolBar.h b/Engine/Source/Editor/UnrealEd/Public/SViewportToolBar.h
--- a/Engine/Source/Editor/UnrealEd/Public/SViewportToolBar.h
+++ b/Engine/Source/Editor/UnrealEd/Public/SViewportToolBar.h
@@ -18,29 +18,29 @@ struct FSlateBrush;
 /**
  * A level viewport toolbar widget that is placed in a viewport
  */
-class UNREALED_API SViewportToolBar : public SCompoundWidget
+class SViewportToolBar : public SCompoundWidget
 {
 public:
 	SLATE_BEGIN_ARGS( SViewportToolBar ){}
 	SLATE_END_ARGS()
 
-	void Construct( const FArguments& InArgs );
+	UNREALED_API void Construct( const FArguments& InArgs );
 
 	/**
 	 * @return The currently open pull down menu if there is one
 	 */
-	TWeakPtr<SMenuAnchor> GetOpenMenu() const;
+	UNREALED_API TWeakPtr<SMenuAnchor> GetOpenMenu() const;
 
 	/**
 	 * Sets the open menu to a new menu and closes any currently opened one
 	 *
 	 * @param NewMenu The new menu that is opened
 	 */
-	void SetOpenMenu( TSharedPtr< SMenuAnchor >& NewMenu );
+	UNREALED_API void SetOpenMenu( TSharedPtr< SMenuAnchor >& NewMenu );
 
 
 	/** @return Whether the given viewmode is supported. */ 
-	virtual bool IsViewModeSupported(EViewModeIndex ViewModeIndex) const;
+	UNREALED_API virtual bool IsViewModeSupported(EViewModeIndex ViewModeIndex) const;
 
 protected:
 	/**
@@ -50,7 +50,7 @@ protected:
 	*
 	* @return	Label to use for this menu label
 	*/
-	virtual FText GetCameraMenuLabelFromViewportType(const ELevelViewportType ViewportType) const;
+	UNREALED_API virtual FText GetCameraMenuLabelFromViewportType(const ELevelViewportType ViewportType) const;
 
 	/**
 	* Returns the label icon for the "Camera" tool bar menu, which changes depending on the viewport type
@@ -59,7 +59,7 @@ protected:
 	*
 	* @return	Label icon to use for this menu label
 	*/
-	virtual const FSlateBrush* GetCameraMenuLabelIconFromViewportType(const ELevelViewportType ViewportType) const;
+	UNREALED_API virtual const FSlateBrush* GetCameraMenuLabelIconFromViewportType(const ELevelViewportType ViewportType) const;
 
 private:
 	/** The pulldown menu that is open if any */
diff --git a/Engine/Source/Editor/UnrealEd/Public/SViewportToolBarComboMenu.h b/Engine/Source/Editor/UnrealEd/Public/SViewportToolBarComboMenu.h
--- a/Engine/Source/Editor/UnrealEd/Public/SViewportToolBarComboMenu.h
+++ b/Engine/Source/Editor/UnrealEd/Public/SViewportToolBarComboMenu.h
@@ -29,7 +29,7 @@ struct FSlateIcon;
  * | Toggle | Menu button |
  * +--------+-------------+
  */
-class UNREALED_API SViewportToolBarComboMenu : public SCompoundWidget 
+class SViewportToolBarComboMenu : public SCompoundWidget 
 {
 public:
 	SLATE_BEGIN_ARGS(SViewportToolBarComboMenu) : _BlockLocation(EMultiBlockLocation::None), _MinDesiredButtonWidth(-1.0f) {}
@@ -74,7 +74,7 @@ public:
 	 *
 	 * @param	InArgs	The declaration data for this widget
 	 */
-	void Construct( const FArguments& InArgs );
+	UNREALED_API void Construct( const FArguments& InArgs );
 
 protected:
 	/**
@@ -85,7 +85,7 @@ protected:
 	 *
 	 * @return	Tool tip text, or an empty text if filtered out
 	 */
-	FText GetFilteredToolTipText(TAttribute<FText> ToolTipText) const;
+	UNREALED_API FText GetFilteredToolTipText(TAttribute<FText> ToolTipText) const;
 
 private:
 	/**
@@ -97,7 +97,7 @@ private:
 	 * Called when the mouse enters a menu button.  If there was a menu previously opened
 	 * we open this menu automatically
 	 */
-	void OnMouseEnter( const FGeometry& MyGeometry, const FPointerEvent& MouseEvent );
+	UNREALED_API void OnMouseEnter( const FGeometry& MyGeometry, const FPointerEvent& MouseEvent );
 
 private:
 	/** Our menus anchor */
diff --git a/Engine/Source/Editor/UnrealEd/Public/SViewportToolBarIconMenu.h b/Engine/Source/Editor/UnrealEd/Public/SViewportToolBarIconMenu.h
--- a/Engine/Source/Editor/UnrealEd/Public/SViewportToolBarIconMenu.h
+++ b/Engine/Source/Editor/UnrealEd/Public/SViewportToolBarIconMenu.h
@@ -26,7 +26,7 @@ struct FSlateIcon;
  * | Icon     Menu button |
  * +--------+-------------+
  */
-class UNREALED_API SViewportToolBarIconMenu : public SCompoundWidget 
+class SViewportToolBarIconMenu : public SCompoundWidget 
 {
 public:
 	SLATE_BEGIN_ARGS( SViewportToolBarIconMenu ){}
@@ -53,7 +53,7 @@ public:
 	 *
 	 * @param	InArgs	The declaration data for this widget
 	 */
-	void Construct( const FArguments& InArgs );
+	UNREALED_API void Construct( const FArguments& InArgs );
 	
 private:
 	/**
@@ -65,7 +65,7 @@ private:
 	 * Called when the mouse enters a menu button.  If there was a menu previously opened
 	 * we open this menu automatically
 	 */
-	void OnMouseEnter(const FGeometry& MyGeometry, const FPointerEvent& MouseEvent);
+	UNREALED_API void OnMouseEnter(const FGeometry& MyGeometry, const FPointerEvent& MouseEvent);
 
 private:
 	/** Our menus anchor */
diff --git a/Engine/Source/Editor/UnrealEd/Public/SWorldPartitionViewportWidget.h b/Engine/Source/Editor/UnrealEd/Public/SWorldPartitionViewportWidget.h
--- a/Engine/Source/Editor/UnrealEd/Public/SWorldPartitionViewportWidget.h
+++ b/Engine/Source/Editor/UnrealEd/Public/SWorldPartitionViewportWidget.h
@@ -6,17 +6,17 @@
 
 class IWorldPartitionEditorModule;
 
-class UNREALED_API SWorldPartitionViewportWidget : public SCompoundWidget
+class SWorldPartitionViewportWidget : public SCompoundWidget
 {
 public:
 	SLATE_BEGIN_ARGS(SWorldPartitionViewportWidget) {}
 	SLATE_ARGUMENT(bool, Clickable)
 	SLATE_END_ARGS()
 
-	void Construct(const FArguments& InArgs);
-	~SWorldPartitionViewportWidget();
+	UNREALED_API void Construct(const FArguments& InArgs);
+	UNREALED_API ~SWorldPartitionViewportWidget();
 
-	EVisibility GetVisibility(UWorld* InWorld);
+	UNREALED_API EVisibility GetVisibility(UWorld* InWorld);
 
 private:
 	bool bClickable;
@@ -24,4 +24,4 @@ private:
 	FText Tooltip;
 	FName InvokeTab;
 	IWorldPartitionEditorModule* WorldPartitionEditorModule;
-};
\ No newline at end of file
+};
