[UnrealEd]

diff --git a/Engine/Source/Editor/UnrealEd/Public/Selection.h b/Engine/Source/Editor/UnrealEd/Public/Selection.h
--- a/Engine/Source/Editor/UnrealEd/Public/Selection.h
+++ b/Engine/Source/Editor/UnrealEd/Public/Selection.h
@@ -35,8 +35,8 @@ struct FTypedElementHandle;
  * Manages selections of objects.
  * Used in the editor for selecting objects in the various browser windows.
  */
-UCLASS(transient)
-class UNREALED_API USelection : public UObject
+UCLASS(transient, MinimalAPI)
+class USelection : public UObject
 {
 	GENERATED_BODY()
 
@@ -45,9 +45,9 @@ private:
 	friend class TSelectionIterator;
 
 public:
-	static USelection* CreateObjectSelection(UObject* InOuter = GetTransientPackage(), FName InName = NAME_None, EObjectFlags InFlags = RF_NoFlags);
-	static USelection* CreateActorSelection(UObject* InOuter = GetTransientPackage(), FName InName = NAME_None, EObjectFlags InFlags = RF_NoFlags);
-	static USelection* CreateComponentSelection(UObject* InOuter = GetTransientPackage(), FName InName = NAME_None, EObjectFlags InFlags = RF_NoFlags);
+	static UNREALED_API USelection* CreateObjectSelection(UObject* InOuter = GetTransientPackage(), FName InName = NAME_None, EObjectFlags InFlags = RF_NoFlags);
+	static UNREALED_API USelection* CreateActorSelection(UObject* InOuter = GetTransientPackage(), FName InName = NAME_None, EObjectFlags InFlags = RF_NoFlags);
+	static UNREALED_API USelection* CreateComponentSelection(UObject* InOuter = GetTransientPackage(), FName InName = NAME_None, EObjectFlags InFlags = RF_NoFlags);
 
 	/** Params: UObject* NewSelection */
 	DECLARE_EVENT_OneParam(USelection, FOnSelectionChanged, UObject*);
@@ -56,24 +56,24 @@ public:
 	DECLARE_EVENT_ThreeParams(USelection, FOnSelectionElementSelectionPtrChanged, USelection* /*Selection*/, UTypedElementSelectionSet* /*OldSelectionSet*/, UTypedElementSelectionSet* /*NewSelectionSet*/);
 
 	/** Called when selection in editor has changed */
-	static FOnSelectionChanged SelectionChangedEvent;
+	static UNREALED_API FOnSelectionChanged SelectionChangedEvent;
 	/** Called when an object has been selected (generally an actor) */
-	static FOnSelectionChanged SelectObjectEvent;
+	static UNREALED_API FOnSelectionChanged SelectObjectEvent;
 	/** Called to deselect everything */
-	static FSimpleMulticastDelegate SelectNoneEvent;
+	static UNREALED_API FSimpleMulticastDelegate SelectNoneEvent;
 	/** Called when the assigned typed element selection pointer set for a selection is changed */
-	static FOnSelectionElementSelectionPtrChanged SelectionElementSelectionPtrChanged;
+	static UNREALED_API FOnSelectionElementSelectionPtrChanged SelectionElementSelectionPtrChanged;
 
 	/**
 	 * Set the element selection set instance for this selection set.
 	 * @note Also sets the element list instance.
 	 */
-	void SetElementSelectionSet(UTypedElementSelectionSet* InElementSelectionSet);
+	UNREALED_API void SetElementSelectionSet(UTypedElementSelectionSet* InElementSelectionSet);
 
 	/**
 	 * Get the element selection set instance for this selection set, if any.
 	 */
-	UTypedElementSelectionSet* GetElementSelectionSet() const;
+	UNREALED_API UTypedElementSelectionSet* GetElementSelectionSet() const;
 
 	/**
 	 * Returns the number of objects in the selection set.  This function is used by clients in
@@ -83,42 +83,42 @@ public:
 	 * 
 	 * @return		Number of objects in the selection set.
 	 */
-	int32 Num() const;
+	UNREALED_API int32 Num() const;
 
 	/**
 	 * @return	The Index'th selected objects.  May be NULL.
 	 */
-	UObject* GetSelectedObject(const int32 InIndex) const;
+	UNREALED_API UObject* GetSelectedObject(const int32 InIndex) const;
 
 	/**
 	 * Call before beginning selection operations
 	 */
-	void BeginBatchSelectOperation();
+	UNREALED_API void BeginBatchSelectOperation();
 
 	/**
 	 * Should be called when selection operations are complete.  If all selection operations are complete, notifies all listeners
 	 * that the selection has been changed.
 	 */
-	void EndBatchSelectOperation(bool bNotify = true);
+	UNREALED_API void EndBatchSelectOperation(bool bNotify = true);
 
 	/**
 	 * @return	Returns whether or not the selection object is currently in the middle of a batch select block.
 	 */
-	bool IsBatchSelecting() const;
+	UNREALED_API bool IsBatchSelecting() const;
 
 	/**
 	 * Selects the specified object.
 	 *
 	 * @param	InObject	The object to select/deselect.  Must be non-NULL.
 	 */
-	void Select(UObject* InObject);
+	UNREALED_API void Select(UObject* InObject);
 
 	/**
 	 * Deselects the specified object.
 	 *
 	 * @param	InObject	The object to deselect.  Must be non-NULL.
 	 */
-	void Deselect(UObject* InObject);
+	UNREALED_API void Deselect(UObject* InObject);
 
 	/**
 	 * Selects or deselects the specified object, depending on the value of the bSelect flag.
@@ -126,37 +126,37 @@ public:
 	 * @param	InObject	The object to select/deselect.  Must be non-NULL.
 	 * @param	bSelect		true selects the object, false deselects.
 	 */
-	void Select(UObject* InObject, bool bSelect);
+	UNREALED_API void Select(UObject* InObject, bool bSelect);
 
 	/**
 	 * Toggles the selection state of the specified object.
 	 *
 	 * @param	InObject	The object to select/deselect.  Must be non-NULL.
 	 */
-	void ToggleSelect(UObject* InObject);
+	UNREALED_API void ToggleSelect(UObject* InObject);
 
 	/**
 	 * Deselects all objects of the specified class, if no class is specified it deselects all objects.
 	 *
 	 * @param	InClass		The type of object to deselect.  Can be NULL.
 	 */
-	void DeselectAll( UClass* InClass = NULL );
+	UNREALED_API void DeselectAll( UClass* InClass = NULL );
 
 	/**
 	 * If batch selection is active, sets flag indicating something actually changed.
 	 */
-	void ForceBatchDirty();
+	UNREALED_API void ForceBatchDirty();
 
 	/**
 	 * Manually invoke a selection changed notification for this set.
 	 */
-	void NoteSelectionChanged();
+	UNREALED_API void NoteSelectionChanged();
 
 	/**
 	 * Manually invoke a selection changed notification for no specific set.
 	 * @note Legacy BSP code only!
 	 */
-	static void NoteUnknownSelectionChanged();
+	static UNREALED_API void NoteUnknownSelectionChanged();
 
 	/**
 	 * Returns the first selected object of the specified class.
@@ -249,7 +249,7 @@ public:
 	 * @param	InObject	The object to query.  Can be NULL.
 	 * @return				true if the object is selected, or false if InObject is unselected or NULL.
 	 */
-	bool IsSelected(const UObject* InObject) const;
+	UNREALED_API bool IsSelected(const UObject* InObject) const;
 
 	/**
 	 * Returns the number of selected objects of the specified type.
@@ -280,11 +280,11 @@ public:
 		return Count;
 	}
 
-	bool IsClassSelected(UClass* Class) const;
+	UNREALED_API bool IsClassSelected(UClass* Class) const;
 
 	//~ Begin UObject Interface
-	virtual void Serialize(FArchive& Ar) override;
-	virtual bool Modify( bool bAlwaysMarkDirty=true) override;
+	UNREALED_API virtual void Serialize(FArchive& Ar) override;
+	UNREALED_API virtual bool Modify( bool bAlwaysMarkDirty=true) override;
 	//~ End UObject Interface
 
 
@@ -339,12 +339,12 @@ public:
 
 private:
 	/** Initializes the selection set with its typed element bridge */
-	void Initialize(TSharedRef<ISelectionElementBridge>&& InSelectionElementBridge);
+	UNREALED_API void Initialize(TSharedRef<ISelectionElementBridge>&& InSelectionElementBridge);
 
-	bool IsValidObjectToSelect(const UObject* InObject) const;
-	UObject* GetObjectForElementHandle(const FTypedElementHandle& InElementHandle) const;
+	UNREALED_API bool IsValidObjectToSelect(const UObject* InObject) const;
+	UNREALED_API UObject* GetObjectForElementHandle(const FTypedElementHandle& InElementHandle) const;
 
-	void OnElementListSyncEvent(const FTypedElementList& InElementList, FTypedElementList::FLegacySync::ESyncType InSyncType, const FTypedElementHandle& InElementHandle, bool bIsWithinBatchOperation);
+	UNREALED_API void OnElementListSyncEvent(const FTypedElementList& InElementList, FTypedElementList::FLegacySync::ESyncType InSyncType, const FTypedElementHandle& InElementHandle, bool bIsWithinBatchOperation);
 
 	/** Bridge from UObjects to their corresponding typed elements. */
 	TSharedPtr<ISelectionElementBridge> SelectionElementBridge;
@@ -495,15 +495,15 @@ public:
 	{}
 };
 
-class UNREALED_API FDeselectedActorsEvent
+class FDeselectedActorsEvent
 {
 public:
 	FDeselectedActorsEvent(const TArray<AActor*>& InDeselectedActors)
 		: DeselectedActors(InDeselectedActors)
 	{}
 
-	~FDeselectedActorsEvent();
+	UNREALED_API ~FDeselectedActorsEvent();
 
 private:
 	const TArray<AActor*>& DeselectedActors;
-};
\ No newline at end of file
+};
diff --git a/Engine/Source/Editor/UnrealEd/Public/Settings/BlueprintEditorProjectSettings.h b/Engine/Source/Editor/UnrealEd/Public/Settings/BlueprintEditorProjectSettings.h
--- a/Engine/Source/Editor/UnrealEd/Public/Settings/BlueprintEditorProjectSettings.h
+++ b/Engine/Source/Editor/UnrealEd/Public/Settings/BlueprintEditorProjectSettings.h
@@ -10,8 +10,8 @@
 #include "BlueprintEditorProjectSettings.generated.h"
 
 
-UCLASS(config=Editor, meta=(DisplayName="Blueprint Project Settings"), defaultconfig)
-class UNREALED_API UBlueprintEditorProjectSettings : public UDeveloperSettings
+UCLASS(config=Editor, meta=(DisplayName="Blueprint Project Settings"), defaultconfig, MinimalAPI)
+class UBlueprintEditorProjectSettings : public UDeveloperSettings
 {
 	GENERATED_UCLASS_BODY()
 
@@ -80,6 +80,6 @@ public:
 
 public:
 	// UObject interface
-	virtual void PostEditChangeProperty(struct FPropertyChangedEvent& PropertyChangedEvent) override;
+	UNREALED_API virtual void PostEditChangeProperty(struct FPropertyChangedEvent& PropertyChangedEvent) override;
 	// End of UObject interface
 };
diff --git a/Engine/Source/Editor/UnrealEd/Public/Settings/EditorProjectSettings.h b/Engine/Source/Editor/UnrealEd/Public/Settings/EditorProjectSettings.h
--- a/Engine/Source/Editor/UnrealEd/Public/Settings/EditorProjectSettings.h
+++ b/Engine/Source/Editor/UnrealEd/Public/Settings/EditorProjectSettings.h
@@ -54,17 +54,17 @@ enum class EReferenceViewerSettingMode : uint8
 /**
  * Editor project appearance settings. Stored in default config, per-project
  */
-UCLASS(config=Editor, defaultconfig, meta=(DisplayName="Appearance"))
-class UNREALED_API UEditorProjectAppearanceSettings : public UDeveloperSettings
+UCLASS(config=Editor, defaultconfig, meta=(DisplayName="Appearance"), MinimalAPI)
+class UEditorProjectAppearanceSettings : public UDeveloperSettings
 {
 public:
 	GENERATED_BODY()
-	UEditorProjectAppearanceSettings(const FObjectInitializer&);
+	UNREALED_API UEditorProjectAppearanceSettings(const FObjectInitializer&);
 
 protected:
 	/** Called when a property on this object is changed */
-	virtual void PostEditChangeProperty( struct FPropertyChangedEvent& PropertyChangedEvent ) override;
-	virtual void PostInitProperties() override;
+	UNREALED_API virtual void PostEditChangeProperty( struct FPropertyChangedEvent& PropertyChangedEvent ) override;
+	UNREALED_API virtual void PostInitProperties() override;
 
 public:
 
@@ -119,7 +119,7 @@ public:
 * 2D layer settings
 */
 USTRUCT()
-struct UNREALED_API FMode2DLayer
+struct FMode2DLayer
 {
 	GENERATED_USTRUCT_BODY()
 
@@ -153,8 +153,8 @@ enum class ELevelEditor2DAxis : uint8
 /**
  * Configure settings for the 2D Level Editor
  */
-UCLASS(config=Editor, meta=(DisplayName="2D"), defaultconfig)
-class UNREALED_API ULevelEditor2DSettings : public UDeveloperSettings
+UCLASS(config=Editor, meta=(DisplayName="2D"), defaultconfig, MinimalAPI)
+class ULevelEditor2DSettings : public UDeveloperSettings
 {
 	GENERATED_UCLASS_BODY()
 
@@ -177,7 +177,7 @@ public:
 
 public:
 	// UObject interface
-	virtual void PostEditChangeProperty(struct FPropertyChangedEvent& PropertyChangedEvent) override;
+	UNREALED_API virtual void PostEditChangeProperty(struct FPropertyChangedEvent& PropertyChangedEvent) override;
 	// End of UObject interface
 
 };
@@ -185,8 +185,8 @@ public:
 /**
  * Configure per-project settings for the Level Editor
  */
-UCLASS(config=Editor, meta=(DisplayName="Level Editor"), defaultconfig)
-class UNREALED_API ULevelEditorProjectSettings : public UDeveloperSettings
+UCLASS(config=Editor, meta=(DisplayName="Level Editor"), defaultconfig, MinimalAPI)
+class ULevelEditorProjectSettings : public UDeveloperSettings
 {
 	GENERATED_BODY()
 
@@ -197,18 +197,18 @@ public:
 	bool bEnableViewportSMInstanceSelection;
 
 public:
-	ULevelEditorProjectSettings(const class FObjectInitializer& ObjectInitializer);
+	UNREALED_API ULevelEditorProjectSettings(const class FObjectInitializer& ObjectInitializer);
 	// UObject interface
-	virtual void PostInitProperties() override;
-	virtual void PostEditChangeProperty(struct FPropertyChangedEvent& PropertyChangedEvent) override;
+	UNREALED_API virtual void PostInitProperties() override;
+	UNREALED_API virtual void PostEditChangeProperty(struct FPropertyChangedEvent& PropertyChangedEvent) override;
 	// End of UObject interface
 };
 
 /**
  * Configure per-project performance settings for the Editor
  */
-UCLASS(config=Editor, meta=(DisplayName="Performance"), defaultconfig)
-class UNREALED_API UEditorPerformanceProjectSettings : public UDeveloperSettings
+UCLASS(config=Editor, meta=(DisplayName="Performance"), defaultconfig, MinimalAPI)
+class UEditorPerformanceProjectSettings : public UDeveloperSettings
 {
 	GENERATED_UCLASS_BODY()
 
@@ -240,19 +240,19 @@ public:
 	int32 MaxViewportRenderingResolution;
 
 
-	static void ExportResolutionValuesToConsoleVariables();
+	static UNREALED_API void ExportResolutionValuesToConsoleVariables();
 
 public:
 	// UObject interface
-	virtual void PostInitProperties() override;
-	virtual void PostEditChangeProperty(struct FPropertyChangedEvent& PropertyChangedEvent) override;
+	UNREALED_API virtual void PostInitProperties() override;
+	UNREALED_API virtual void PostEditChangeProperty(struct FPropertyChangedEvent& PropertyChangedEvent) override;
 	// End of UObject interface
 
 };
 
 
-UCLASS(config=Editor, meta=(DisplayName="Derived Data"), defaultconfig)
-class UNREALED_API UDDCProjectSettings : public UDeveloperSettings
+UCLASS(config=Editor, meta=(DisplayName="Derived Data"), defaultconfig, MinimalAPI)
+class UDDCProjectSettings : public UDeveloperSettings
 {
 	GENERATED_BODY()
 
diff --git a/Engine/Source/Editor/UnrealEd/Public/Settings/EditorSettings.h b/Engine/Source/Editor/UnrealEd/Public/Settings/EditorSettings.h
--- a/Engine/Source/Editor/UnrealEd/Public/Settings/EditorSettings.h
+++ b/Engine/Source/Editor/UnrealEd/Public/Settings/EditorSettings.h
@@ -52,8 +52,8 @@ struct FRecentProjectFile
 };
 
 
-UCLASS(config=EditorSettings)
-class UNREALED_API UEditorSettings : public UObject
+UCLASS(config=EditorSettings, MinimalAPI)
+class UEditorSettings : public UObject
 {
 	GENERATED_UCLASS_BODY()
 
@@ -151,16 +151,16 @@ class UNREALED_API UEditorSettings : public UObject
 	Scalability::FQualityLevels EngineBenchmarkResult;
 
 	/** Load the engine scalability benchmark results. Performs a benchmark if not yet valid. */
-	void LoadScalabilityBenchmark();
+	UNREALED_API void LoadScalabilityBenchmark();
 
 	/** Auto detects and applies the scalability benchmark */
-	void AutoApplyScalabilityBenchmark();
+	UNREALED_API void AutoApplyScalabilityBenchmark();
 
 	/** @return true if the scalability benchmark is valid */
-	bool IsScalabilityBenchmarkValid() const;
+	UNREALED_API bool IsScalabilityBenchmarkValid() const;
 
 	//~ Begin UObject Interface
-	virtual bool CanEditChange(const FProperty* InProperty) const override;
-	virtual void PostEditChangeProperty( struct FPropertyChangedEvent& PropertyChangedEvent) override;
+	UNREALED_API virtual bool CanEditChange(const FProperty* InProperty) const override;
+	UNREALED_API virtual void PostEditChangeProperty( struct FPropertyChangedEvent& PropertyChangedEvent) override;
 	//~ End UObject Interface
 };
diff --git a/Engine/Source/Editor/UnrealEd/Public/SkeletalMeshToolMenuContext.h b/Engine/Source/Editor/UnrealEd/Public/SkeletalMeshToolMenuContext.h
--- a/Engine/Source/Editor/UnrealEd/Public/SkeletalMeshToolMenuContext.h
+++ b/Engine/Source/Editor/UnrealEd/Public/SkeletalMeshToolMenuContext.h
@@ -8,8 +8,8 @@
 
 class ISkeletalMeshEditor;
 
-UCLASS()
-class UNREALED_API USkeletalMeshToolMenuContext : public UObject
+UCLASS(MinimalAPI)
+class USkeletalMeshToolMenuContext : public UObject
 {
 	GENERATED_BODY()
 public:
diff --git a/Engine/Source/Editor/UnrealEd/Public/SkinWeightsUtilities.h b/Engine/Source/Editor/UnrealEd/Public/SkinWeightsUtilities.h
--- a/Engine/Source/Editor/UnrealEd/Public/SkinWeightsUtilities.h
+++ b/Engine/Source/Editor/UnrealEd/Public/SkinWeightsUtilities.h
@@ -14,7 +14,7 @@ namespace UnFbx
 //////////////////////////////////////////////////////////////////////////
 // FSkinWeightsUtilities
 
-class UNREALED_API FSkinWeightsUtilities
+class FSkinWeightsUtilities
 {
 public:
 	/**
@@ -27,7 +27,7 @@ public:
 	* @Param ProfileName - The name of the profile to associate the imported skin weight
 	* @param bReregisterComponent - if true the component using the skeletal mesh will all be re register.
 	*/
-	static bool ImportAlternateSkinWeight(USkeletalMesh* SkeletalMesh, const FString& Path, int32 TargetLODIndex, const FName& ProfileName, const bool bIsReimport);
+	static UNREALED_API bool ImportAlternateSkinWeight(USkeletalMesh* SkeletalMesh, const FString& Path, int32 TargetLODIndex, const FName& ProfileName, const bool bIsReimport);
 	
 	/**
 	* This function reimport all skin weights profile for a specified LOD. Return true if the weights are successfully updates.
@@ -37,14 +37,14 @@ public:
 	* @Param TargetLODIndex - The LODIndex to imp[ort the skin weight
 	* @param bReregisterComponent - if true the component using the skeletal mesh will all be re register.
 	*/
-	static bool ReimportAlternateSkinWeight(USkeletalMesh* SkeletalMesh, int32 TargetLODIndex);
+	static UNREALED_API bool ReimportAlternateSkinWeight(USkeletalMesh* SkeletalMesh, int32 TargetLODIndex);
 
-	static bool RemoveSkinnedWeightProfileData(USkeletalMesh* SkeletalMesh, const FName& ProfileName, int32 LODIndex);
+	static UNREALED_API bool RemoveSkinnedWeightProfileData(USkeletalMesh* SkeletalMesh, const FName& ProfileName, int32 LODIndex);
 
 	/*
 	 * Ask user a file path for a particular LOD
 	 */
-	static FString PickSkinWeightPath(int32 LODIndex, USkeletalMesh* SkeletalMesh);
+	static UNREALED_API FString PickSkinWeightPath(int32 LODIndex, USkeletalMesh* SkeletalMesh);
 private:
 	FSkinWeightsUtilities() {}
 };
diff --git a/Engine/Source/Editor/UnrealEd/Public/SnappingUtils.h b/Engine/Source/Editor/UnrealEd/Public/SnappingUtils.h
--- a/Engine/Source/Editor/UnrealEd/Public/SnappingUtils.h
+++ b/Engine/Source/Editor/UnrealEd/Public/SnappingUtils.h
@@ -9,37 +9,37 @@ class FLevelEditorViewportClient;
 class FPrimitiveDrawInterface;
 class FSceneView;
 
-class UNREALED_API FSnappingUtils
+class FSnappingUtils
 {
 public:
 	/**
 	 * @return true if snapping (translation) to the grid is enabled
 	 */
-	static bool IsSnapToGridEnabled();
+	static UNREALED_API bool IsSnapToGridEnabled();
 
 	/**
 	 * @return true if orientation snapping is enabled
 	 */
-	static bool IsRotationSnapEnabled();
+	static UNREALED_API bool IsRotationSnapEnabled();
 
 	/**
 	 * @return true if orientation snapping is enabled
 	 */
-	static bool IsScaleSnapEnabled();
+	static UNREALED_API bool IsScaleSnapEnabled();
 
 	/**
 	 * @return true if snapping actors to other actors is enabled
 	 */
-	static bool IsSnapToActorEnabled();
+	static UNREALED_API bool IsSnapToActorEnabled();
 
 	/** Set user setting for actor snap. */
-	static void EnableActorSnap(bool bEnable);
+	static UNREALED_API void EnableActorSnap(bool bEnable);
 
 	/** Access user setting for distance. Fractional 0.0->100.0 */
-	static float GetActorSnapDistance(bool bScalar = false);
+	static UNREALED_API float GetActorSnapDistance(bool bScalar = false);
 
 	/** Set user setting for distance. Fractional 0.0->100.0 */
-	static void SetActorSnapDistance(float Distance);
+	static UNREALED_API void SetActorSnapDistance(float Distance);
 
 	/**
 	 * Attempts to snap the selected actors to the nearest other actor
@@ -47,7 +47,7 @@ public:
 	 * @param DragDelta			The current world space drag amount
 	 * @param ViewportClient	The viewport client the user is dragging in
 	 */
-	static bool SnapActorsToNearestActor( FVector& DragDelta, FLevelEditorViewportClient* ViewportClient );
+	static UNREALED_API bool SnapActorsToNearestActor( FVector& DragDelta, FLevelEditorViewportClient* ViewportClient );
 
 	/**
 	 * Snaps actors to the nearest vertex on another actor
@@ -56,7 +56,7 @@ public:
 	 * @param ViewportClient	The viewport client the user is dragging in
 	 * @return true if anything was snapped
 	 */
-	static bool SnapDraggedActorsToNearestVertex( FVector& DragDelta, FLevelEditorViewportClient* ViewportClient );
+	static UNREALED_API bool SnapDraggedActorsToNearestVertex( FVector& DragDelta, FLevelEditorViewportClient* ViewportClient );
 
 	/**
 	 * Snaps a delta drag movement to the nearest vertex
@@ -66,7 +66,7 @@ public:
 	 * @param ViewportClient	The viewport client being dragged in.
 	 * @return true if anything was snapped
 	 */
-	static bool SnapDragLocationToNearestVertex( const FVector& BaseLocation, FVector& DragDelta, FLevelEditorViewportClient* ViewportClient, bool bIsPivot = false );
+	static UNREALED_API bool SnapDragLocationToNearestVertex( const FVector& BaseLocation, FVector& DragDelta, FLevelEditorViewportClient* ViewportClient, bool bIsPivot = false );
 
 	/**
 	 * Snaps a location to the nearest vertex
@@ -77,7 +77,7 @@ public:
 	 * @param OutVertexNormal	The normal at the closest vertex
 	 * @return true if anything was snapped
 	 */
-	static bool SnapLocationToNearestVertex( FVector& Location, const FVector2D& MouseLocation, FLevelEditorViewportClient* ViewportClient, FVector& OutVertexNormal, bool bDrawVertHelpers );
+	static UNREALED_API bool SnapLocationToNearestVertex( FVector& Location, const FVector2D& MouseLocation, FLevelEditorViewportClient* ViewportClient, FVector& OutVertexNormal, bool bDrawVertHelpers );
 
 	/**
 	 * Snaps a scale value to the scale grid
@@ -85,7 +85,7 @@ public:
 	 * @param Point	The point to snap.  This value will be modified to account for snapping
 	 * @param GridBase	Base grid offset
 	 */
-	static void SnapScale( FVector& Point, const FVector& GridBase );
+	static UNREALED_API void SnapScale( FVector& Point, const FVector& GridBase );
 
 	/**
 	 * Snaps a point value to the positional grid
@@ -93,23 +93,23 @@ public:
 	 * @param Point	The point to snap.  This value will be modified to account for snapping
 	 * @param GridBase	Base grid offset
 	 */
-	static void SnapPointToGrid( FVector& Point, const FVector& GridBase );
+	static UNREALED_API void SnapPointToGrid( FVector& Point, const FVector& GridBase );
 
 	/**
 	 * Snaps a rotator to the rotational grid
 	 *
 	 * @param Rotation	The rotator to snap.  This value will be modified to account for snapping
 	 */
-	static void SnapRotatorToGrid( FRotator& Rotation );
+	static UNREALED_API void SnapRotatorToGrid( FRotator& Rotation );
 
-	static bool SnapToBSPVertex( FVector& Location, FVector GridBase, FRotator& Rotation );
+	static UNREALED_API bool SnapToBSPVertex( FVector& Location, FVector GridBase, FRotator& Rotation );
 
 	/**
 	 * Clears all vertices being drawn to help a user snap
 	 *
 	 * @param bClearImmediatley	true to clear helpers without fading them out
 	 */
-	static void ClearSnappingHelpers( bool bClearImmediately = false);
+	static UNREALED_API void ClearSnappingHelpers( bool bClearImmediately = false);
 
 	/**
 	 * Draws snapping helpers 
@@ -117,12 +117,12 @@ public:
 	 * @param View	The current view of the scene
 	 * @param PDI	Drawing interface
 	 */
-	static void DrawSnappingHelpers(const FSceneView* View,FPrimitiveDrawInterface* PDI); 
+	static UNREALED_API void DrawSnappingHelpers(const FSceneView* View,FPrimitiveDrawInterface* PDI); 
 
 	/**
 	 * Initialize the snapping system
 	 */
-	static void InitEditorSnappingTools();
+	static UNREALED_API void InitEditorSnappingTools();
 
 private:
 	/** Built-in editor snap implementation */
diff --git a/Engine/Source/Editor/UnrealEd/Public/SourceCodeNavigation.h b/Engine/Source/Editor/UnrealEd/Public/SourceCodeNavigation.h
--- a/Engine/Source/Editor/UnrealEd/Public/SourceCodeNavigation.h
+++ b/Engine/Source/Editor/UnrealEd/Public/SourceCodeNavigation.h
@@ -337,7 +337,7 @@ public:
 	UNREALED_API static void RefreshCompilerAvailability();
 
 	/** Returns true if the compiler for the current platform is available for use */
-	UNREALED_API static bool IsCompilerAvailable() { return bCachedIsCompilerAvailable; }
+	static bool IsCompilerAvailable() { return bCachedIsCompilerAvailable; }
 
 	/** Finds the base directory for a given module name. Does not rely on symbols; finds matching .build.cs files. */
 	UNREALED_API static bool FindModulePath( const FString& InModuleName, FString& OutModulePath );
diff --git a/Engine/Source/Editor/UnrealEd/Public/Subsystems/ActorEditorContextSubsystem.h b/Engine/Source/Editor/UnrealEd/Public/Subsystems/ActorEditorContextSubsystem.h
--- a/Engine/Source/Editor/UnrealEd/Public/Subsystems/ActorEditorContextSubsystem.h
+++ b/Engine/Source/Editor/UnrealEd/Public/Subsystems/ActorEditorContextSubsystem.h
@@ -15,33 +15,33 @@ DECLARE_MULTICAST_DELEGATE(FOnActorEditorContextSubsystemChanged);
 /**
 * UActorEditorContextSubsystem 
 */
-UCLASS()
-class UNREALED_API UActorEditorContextSubsystem : public UEditorSubsystem
+UCLASS(MinimalAPI)
+class UActorEditorContextSubsystem : public UEditorSubsystem
 {
 public:
 
 	GENERATED_BODY()
 
-	static UActorEditorContextSubsystem* Get();
+	static UNREALED_API UActorEditorContextSubsystem* Get();
 
-	virtual void Initialize(FSubsystemCollectionBase& Collection) override;
-	virtual void Deinitialize() override;
+	UNREALED_API virtual void Initialize(FSubsystemCollectionBase& Collection) override;
+	UNREALED_API virtual void Deinitialize() override;
 
-	void RegisterClient(IActorEditorContextClient* Client);
-	void UnregisterClient(IActorEditorContextClient* Client);
-	void ResetContext();
-	void ResetContext(IActorEditorContextClient* Client);
-	void PushContext();
-	void PopContext();
-	TArray<IActorEditorContextClient*> GetDisplayableClients() const;
+	UNREALED_API void RegisterClient(IActorEditorContextClient* Client);
+	UNREALED_API void UnregisterClient(IActorEditorContextClient* Client);
+	UNREALED_API void ResetContext();
+	UNREALED_API void ResetContext(IActorEditorContextClient* Client);
+	UNREALED_API void PushContext();
+	UNREALED_API void PopContext();
+	UNREALED_API TArray<IActorEditorContextClient*> GetDisplayableClients() const;
 	FOnActorEditorContextSubsystemChanged& OnActorEditorContextSubsystemChanged() { return ActorEditorContextSubsystemChanged; }
 
 private:
 
-	UWorld* GetWorld() const;
-	void OnActorEditorContextClientChanged(IActorEditorContextClient* Client);
-	void ApplyContext(AActor* InActor);
+	UNREALED_API UWorld* GetWorld() const;
+	UNREALED_API void OnActorEditorContextClientChanged(IActorEditorContextClient* Client);
+	UNREALED_API void ApplyContext(AActor* InActor);
 
 	FOnActorEditorContextSubsystemChanged ActorEditorContextSubsystemChanged;
 	TArray<IActorEditorContextClient*> Clients;
-};
\ No newline at end of file
+};
diff --git a/Engine/Source/Editor/UnrealEd/Public/Subsystems/AssetEditorSubsystem.h b/Engine/Source/Editor/UnrealEd/Public/Subsystems/AssetEditorSubsystem.h
--- a/Engine/Source/Editor/UnrealEd/Public/Subsystems/AssetEditorSubsystem.h
+++ b/Engine/Source/Editor/UnrealEd/Public/Subsystems/AssetEditorSubsystem.h
@@ -50,7 +50,7 @@ enum class EAssetEditorCloseReason : uint8
  * brought into focus, closed, etc..., without concern for how the editor was
  * implemented.
  */
-class UNREALED_API IAssetEditorInstance
+class IAssetEditorInstance
 {
 public:
 
@@ -79,7 +79,7 @@ public:
 	virtual void RemoveEditingAsset(UObject* Asset) = 0;
 };
 
-struct UNREALED_API FRegisteredModeInfo
+struct FRegisteredModeInfo
 {
 	TWeakObjectPtr<UClass> ModeClass;
 	FEditorModeInfo ModeInfo;
@@ -89,28 +89,28 @@ using RegisteredModeInfoMap = TMap<FEditorModeID, FRegisteredModeInfo>;
 /**
  * UAssetEditorSubsystem
  */
-UCLASS()
-class UNREALED_API UAssetEditorSubsystem : public UEditorSubsystem
+UCLASS(MinimalAPI)
+class UAssetEditorSubsystem : public UEditorSubsystem
 {
 	GENERATED_BODY()
 
 public:
-	UAssetEditorSubsystem();
+	UNREALED_API UAssetEditorSubsystem();
 
-	virtual void Initialize(FSubsystemCollectionBase& Collection) override;
-	virtual void Deinitialize() override;
+	UNREALED_API virtual void Initialize(FSubsystemCollectionBase& Collection) override;
+	UNREALED_API virtual void Deinitialize() override;
 
-	bool IsAssetEditable(const UObject* Asset);
+	UNREALED_API bool IsAssetEditable(const UObject* Asset);
 
 	/** Opens an asset by path */
-	void OpenEditorForAsset(const FString& AssetPathName, const EAssetTypeActivationOpenedMethod OpenedMethod = EAssetTypeActivationOpenedMethod::Edit);
-	void OpenEditorForAsset(const FSoftObjectPath& AssetPath);
+	UNREALED_API void OpenEditorForAsset(const FString& AssetPathName, const EAssetTypeActivationOpenedMethod OpenedMethod = EAssetTypeActivationOpenedMethod::Edit);
+	UNREALED_API void OpenEditorForAsset(const FSoftObjectPath& AssetPath);
 
 	/**
 	 * Tries to open an editor for the specified asset.  Returns true if the asset is open in an editor.
 	 * If the file is already open in an editor, it will not create another editor window but instead bring it to front
 	 */
-	bool OpenEditorForAsset(UObject* Asset, const EToolkitMode::Type ToolkitMode = EToolkitMode::Standalone, TSharedPtr<IToolkitHost> OpenedFromLevelEditor = TSharedPtr<IToolkitHost>(), const bool bShowProgressWindow = true, EAssetTypeActivationOpenedMethod OpenedMethod = EAssetTypeActivationOpenedMethod::Edit);
+	UNREALED_API bool OpenEditorForAsset(UObject* Asset, const EToolkitMode::Type ToolkitMode = EToolkitMode::Standalone, TSharedPtr<IToolkitHost> OpenedFromLevelEditor = TSharedPtr<IToolkitHost>(), const bool bShowProgressWindow = true, EAssetTypeActivationOpenedMethod OpenedMethod = EAssetTypeActivationOpenedMethod::Edit);
 	
 	template<typename ObjectType>
 	bool OpenEditorForAsset(TObjectPtr<ObjectType> Asset, const EToolkitMode::Type ToolkitMode = EToolkitMode::Standalone, TSharedPtr<IToolkitHost> OpenedFromLevelEditor = TSharedPtr<IToolkitHost>(), const bool bShowProgressWindow = true)
@@ -124,65 +124,65 @@ public:
 	 * If all assets are of the same type, the supporting AssetTypeAction (if it exists) is responsible for the details of how to handle opening multiple assets at once.
 	 */
 	UFUNCTION(BlueprintCallable, Category = "Editor Scripting | Asset Tools")
-	bool OpenEditorForAssets(const TArray<UObject*>& Assets, const EAssetTypeActivationOpenedMethod OpenedMethod = EAssetTypeActivationOpenedMethod::Edit);
-	bool OpenEditorForAssets_Advanced(const TArray<UObject*>& InAssets, const EToolkitMode::Type ToolkitMode = EToolkitMode::Standalone, TSharedPtr<IToolkitHost> OpenedFromLevelEditor = TSharedPtr<IToolkitHost>(), const EAssetTypeActivationOpenedMethod OpenedMethod = EAssetTypeActivationOpenedMethod::Edit);
+	UNREALED_API bool OpenEditorForAssets(const TArray<UObject*>& Assets, const EAssetTypeActivationOpenedMethod OpenedMethod = EAssetTypeActivationOpenedMethod::Edit);
+	UNREALED_API bool OpenEditorForAssets_Advanced(const TArray<UObject*>& InAssets, const EToolkitMode::Type ToolkitMode = EToolkitMode::Standalone, TSharedPtr<IToolkitHost> OpenedFromLevelEditor = TSharedPtr<IToolkitHost>(), const EAssetTypeActivationOpenedMethod OpenedMethod = EAssetTypeActivationOpenedMethod::Edit);
 
 	/** Opens editors for the supplied assets (via OpenEditorForAsset) */
-	void OpenEditorsForAssets(const TArray<FString>& AssetsToOpen, const EAssetTypeActivationOpenedMethod OpenedMethod = EAssetTypeActivationOpenedMethod::Edit);
-	void OpenEditorsForAssets(const TArray<FName>& AssetsToOpen, const EAssetTypeActivationOpenedMethod OpenedMethod = EAssetTypeActivationOpenedMethod::Edit);
-	void OpenEditorsForAssets(const TArray<FSoftObjectPath>& AssetsToOpen);
+	UNREALED_API void OpenEditorsForAssets(const TArray<FString>& AssetsToOpen, const EAssetTypeActivationOpenedMethod OpenedMethod = EAssetTypeActivationOpenedMethod::Edit);
+	UNREALED_API void OpenEditorsForAssets(const TArray<FName>& AssetsToOpen, const EAssetTypeActivationOpenedMethod OpenedMethod = EAssetTypeActivationOpenedMethod::Edit);
+	UNREALED_API void OpenEditorsForAssets(const TArray<FSoftObjectPath>& AssetsToOpen);
 
 	/** Returns the primary editor if one is already open for the specified asset.
 	 * If there is one open and bFocusIfOpen is true, that editor will be brought to the foreground and focused if possible.
 	 */
-	IAssetEditorInstance* FindEditorForAsset(UObject* Asset, bool bFocusIfOpen);
+	UNREALED_API IAssetEditorInstance* FindEditorForAsset(UObject* Asset, bool bFocusIfOpen);
 
 	/** Returns all editors currently opened for the specified asset */
-	TArray<IAssetEditorInstance*> FindEditorsForAsset(UObject* Asset);
+	UNREALED_API TArray<IAssetEditorInstance*> FindEditorsForAsset(UObject* Asset);
 
 	/** Returns all editors currently opened for the specified asset or any of its subobjects */
-	TArray<IAssetEditorInstance*> FindEditorsForAssetAndSubObjects(UObject* Asset);
+	UNREALED_API TArray<IAssetEditorInstance*> FindEditorsForAssetAndSubObjects(UObject* Asset);
 
 	/** Close all active editors for the supplied asset and return the number of asset editors that were closed */
 	UFUNCTION(BlueprintCallable, Category = "Editor Scripting | Asset Tools")
-	int32 CloseAllEditorsForAsset(UObject* Asset);
+	UNREALED_API int32 CloseAllEditorsForAsset(UObject* Asset);
 
 	/** Close any editor which is not this one */
-	void CloseOtherEditors(UObject* Asset, IAssetEditorInstance* OnlyEditor);
+	UNREALED_API void CloseOtherEditors(UObject* Asset, IAssetEditorInstance* OnlyEditor);
 
 	/** Remove given asset from all open editors */
-	void RemoveAssetFromAllEditors(UObject* Asset);
+	UNREALED_API void RemoveAssetFromAllEditors(UObject* Asset);
 
 	/** Event called when CloseAllEditorsForAsset/RemoveAssetFromAllEditors is called */
 	DECLARE_EVENT_TwoParams(UAssetEditorSubsystem, FAssetEditorRequestCloseEvent, UObject*, EAssetEditorCloseReason);
 	virtual FAssetEditorRequestCloseEvent& OnAssetEditorRequestClose() { return AssetEditorRequestCloseEvent; }
 
 	/** Get all assets currently being tracked with open editors */
-	TArray<UObject*> GetAllEditedAssets();
+	UNREALED_API TArray<UObject*> GetAllEditedAssets();
 
 	/** Notify the asset editor manager that an asset was opened */
-	void NotifyAssetOpened(UObject* Asset, IAssetEditorInstance* Instance);
-	void NotifyAssetsOpened(const TArray< UObject* >& Assets, IAssetEditorInstance* Instance);
+	UNREALED_API void NotifyAssetOpened(UObject* Asset, IAssetEditorInstance* Instance);
+	UNREALED_API void NotifyAssetsOpened(const TArray< UObject* >& Assets, IAssetEditorInstance* Instance);
 
 	/** Called when an asset has been opened in an editor */
 	DECLARE_EVENT_TwoParams(UAssetEditorSubsystem, FOnAssetOpenedInEditorEvent, UObject*, IAssetEditorInstance*);
 	virtual FOnAssetOpenedInEditorEvent& OnAssetOpenedInEditor() { return AssetOpenedInEditorEvent; }
 
 	/** Notify the asset editor manager that an asset editor is being opened and before widgets are constructed */
-	void NotifyEditorOpeningPreWidgets(const TArray< UObject* >& Assets, IAssetEditorInstance* Instance);
+	UNREALED_API void NotifyEditorOpeningPreWidgets(const TArray< UObject* >& Assets, IAssetEditorInstance* Instance);
 
 	/** Called when an asset editor is opening and before widgets are constructed */
 	DECLARE_EVENT_TwoParams(UAssetEditorSubsystem, FOnAssetsOpenedInEditorEvent, const TArray<UObject*>&, IAssetEditorInstance*);
 	virtual FOnAssetsOpenedInEditorEvent& OnEditorOpeningPreWidgets() { return EditorOpeningPreWidgetsEvent; }
 
 	/** Notify the asset editor manager that an asset editor is done editing an asset */
-	void NotifyAssetClosed(UObject* Asset, IAssetEditorInstance* Instance);
+	UNREALED_API void NotifyAssetClosed(UObject* Asset, IAssetEditorInstance* Instance);
 
 	/** Notify the asset editor manager that an asset was closed */
-	void NotifyEditorClosed(IAssetEditorInstance* Instance);
+	UNREALED_API void NotifyEditorClosed(IAssetEditorInstance* Instance);
 
 	/** Close all open asset editors */
-	bool CloseAllAssetEditors();
+	UNREALED_API bool CloseAllAssetEditors();
 
 	/** Called when an asset editor is requested to be opened */
 	DECLARE_EVENT_OneParam(UAssetEditorSubsystem, FAssetEditorRequestOpenEvent, UObject*);
@@ -193,10 +193,10 @@ public:
 	FAssetEditorOpenEvent& OnAssetEditorOpened() { return AssetEditorOpenedEvent; }
 
 	/** Request notification to restore the assets that were previously open when the editor was last closed */
-	void RequestRestorePreviouslyOpenAssets();
+	UNREALED_API void RequestRestorePreviouslyOpenAssets();
 
-	void RegisterUAssetEditor(UAssetEditor* NewAssetEditor);
-	void UnregisterUAssetEditor(UAssetEditor* RemovedAssetEditor);
+	UNREALED_API void RegisterUAssetEditor(UAssetEditor* NewAssetEditor);
+	UNREALED_API void UnregisterUAssetEditor(UAssetEditor* RemovedAssetEditor);
 	
 	/**
 	 * Creates a scriptable editor mode based on ID name, which will be owned by the given Owner, if that name exists in the map of editor modes found at system startup.
@@ -205,7 +205,7 @@ public:
 	 *
 	 * @return 			A pointer to the created UEdMode or nullptr, if the given ModeID does not exist in the set of known modes.
 	 */
-	UEdMode* CreateEditorModeWithToolsOwner(FEditorModeID ModeID, FEditorModeTools& Owner);
+	UNREALED_API UEdMode* CreateEditorModeWithToolsOwner(FEditorModeID ModeID, FEditorModeTools& Owner);
 	
 	/**
 	 * Returns information about an editor mode, based on the given ID.
@@ -214,35 +214,35 @@ public:
 	 *
 	 * @return 				True if OutModeInfo was filled out successfully, otherwise false.
 	 */
-	bool FindEditorModeInfo(const FEditorModeID& InModeID, FEditorModeInfo& OutModeInfo) const;
+	UNREALED_API bool FindEditorModeInfo(const FEditorModeID& InModeID, FEditorModeInfo& OutModeInfo) const;
 	
 	/**
 	 * Creates an array of all known FEditorModeInfos, sorted by their priority, from greatest to least.
 	 *
 	 * @return 			The sorted array of FEditorModeInfos.
 	 */
-	TArray<FEditorModeInfo> GetEditorModeInfoOrderedByPriority() const;
+	UNREALED_API TArray<FEditorModeInfo> GetEditorModeInfoOrderedByPriority() const;
 
 	/**
 	 * Event that is triggered whenever a mode is registered or unregistered
 	 */
-	FRegisteredModesChangedEvent& OnEditorModesChanged();
+	UNREALED_API FRegisteredModesChangedEvent& OnEditorModesChanged();
 
 	/**
 	 * Event that is triggered whenever a mode is registered
 	 */
-	FOnModeRegistered& OnEditorModeRegistered();
+	UNREALED_API FOnModeRegistered& OnEditorModeRegistered();
 
 	/**
 	 * Event that is triggered whenever a mode is unregistered
 	 */
-	FOnModeUnregistered& OnEditorModeUnregistered();
+	UNREALED_API FOnModeUnregistered& OnEditorModeUnregistered();
 
 	/** Get the permission list that controls which editor modes are exposed */
-	FNamePermissionList& GetAllowedEditorModes();
+	UNREALED_API FNamePermissionList& GetAllowedEditorModes();
 
 	/** Fill in the given section with the recent assets for the given asset editor instance */
-	void CreateRecentAssetsMenuForEditor(const IAssetEditorInstance* InAssetEditorInstance, FToolMenuSection& InSection);
+	UNREALED_API void CreateRecentAssetsMenuForEditor(const IAssetEditorInstance* InAssetEditorInstance, FToolMenuSection& InSection);
 
 	FMainMRUFavoritesList* GetRecentlyOpenedAssets() const
 	{
@@ -252,33 +252,33 @@ public:
 private:
 
 	/** Handles FAssetEditorRequestOpenAsset messages. */
-	void HandleRequestOpenAssetMessage(const FAssetEditorRequestOpenAsset& Message, const TSharedRef<IMessageContext, ESPMode::ThreadSafe>& Context);
+	UNREALED_API void HandleRequestOpenAssetMessage(const FAssetEditorRequestOpenAsset& Message, const TSharedRef<IMessageContext, ESPMode::ThreadSafe>& Context);
 
 	/** Handles ticks from the ticker. */
-	bool HandleTicker(float DeltaTime);
+	UNREALED_API bool HandleTicker(float DeltaTime);
 
 	/** Spawn a notification asking the user if they want to restore their previously open assets */
-	void SpawnRestorePreviouslyOpenAssetsNotification(const bool bCleanShutdown, const TArray<FString>& AssetsToOpen);
+	UNREALED_API void SpawnRestorePreviouslyOpenAssetsNotification(const bool bCleanShutdown, const TArray<FString>& AssetsToOpen);
 
 	/** Handler for when the "Restore Now" button is clicked on the RestorePreviouslyOpenAssets notification */
-	void OnConfirmRestorePreviouslyOpenAssets(TArray<FString> AssetsToOpen);
+	UNREALED_API void OnConfirmRestorePreviouslyOpenAssets(TArray<FString> AssetsToOpen);
 
 	/** Handler for when the "Don't Restore" button is clicked on the RestorePreviouslyOpenAssets notification */
-	void OnCancelRestorePreviouslyOpenAssets();
+	UNREALED_API void OnCancelRestorePreviouslyOpenAssets();
 
 	/** Register extensions to the Level Editor "File" Menu */
-	void RegisterLevelEditorMenuExtensions();
+	UNREALED_API void RegisterLevelEditorMenuExtensions();
 
 	/**
 	 * Fill in the given menu with entries from the recent assets list, clicking on which will open up the asset
 	 * @param InMenu The menu to fill in
 	 * @param InAssetEditorName If provided, only assets in the list that were opened in this asset editor will be added to the menu
 	 */
-	void CreateRecentAssetsMenu(UToolMenu* InMenu, const FName InAssetEditorName = FName());
+	UNREALED_API void CreateRecentAssetsMenu(UToolMenu* InMenu, const FName InAssetEditorName = FName());
 
-	void OnAssetRemoved(const FAssetData& AssetData);
+	UNREALED_API void OnAssetRemoved(const FAssetData& AssetData);
 
-	void OnAssetRenamed(const FAssetData& AssetData, const FString& AssetOldName);
+	UNREALED_API void OnAssetRenamed(const FAssetData& AssetData, const FString& AssetOldName);
 public:
 
 	/**
@@ -286,7 +286,7 @@ public:
 	 * @param bOnShutdown If true, this is handled as if the engine is shutting down right now.
 	 */
 
-	void SaveOpenAssetEditors(const bool bOnShutdown);
+	UNREALED_API void SaveOpenAssetEditors(const bool bOnShutdown);
 	
 	/**
 	 * Saves a list of open asset editors so they can be restored on editor restart.
@@ -294,43 +294,43 @@ public:
 	 * @param bCancelIfDebugger If true, don't save a list of assets to restore if we are running under a debugger.
 	 */
 	UE_DEPRECATED(5.0, "Please use the version of SaveOpenAssetEditors with only one argument, bOnShutdown.")
-	void SaveOpenAssetEditors(const bool bOnShutdown, const bool bCancelIfDebugger);
+	UNREALED_API void SaveOpenAssetEditors(const bool bOnShutdown, const bool bCancelIfDebugger);
 
 	/** Restore the assets that were previously open when the editor was last closed. */
-	void RestorePreviouslyOpenAssets();
+	UNREALED_API void RestorePreviouslyOpenAssets();
 
 	/** Sets bAutoRestoreAndDisableSaving and sets bRequestRestorePreviouslyOpenAssets to false to avoid running RestorePreviouslyOpenAssets() twice. */
 	UE_DEPRECATED(5.0, "Please use the SetAutoRestoreAndDisableSavingOverride instead to set/unset the override")
-	void SetAutoRestoreAndDisableSaving(const bool bInAutoRestoreAndDisableSaving);
+	UNREALED_API void SetAutoRestoreAndDisableSaving(const bool bInAutoRestoreAndDisableSaving);
 
 	/** Set/Unset the bAutoRestoreAndDisableSaving override, along with setting bRequestRestorePreviouslyOpenAssets to false to avoid running RestorePreviouslyOpenAssets() twice */
-	void SetAutoRestoreAndDisableSavingOverride(TOptional<bool> bInAutoRestoreAndDisableSaving);
+	UNREALED_API void SetAutoRestoreAndDisableSavingOverride(TOptional<bool> bInAutoRestoreAndDisableSaving);
 
 	/** Get the current value of bAutoRestoreAndDisableSaving */
-	TOptional<bool> GetAutoRestoreAndDisableSavingOverride() const;
+	UNREALED_API TOptional<bool> GetAutoRestoreAndDisableSavingOverride() const;
 
-	void SetRecentAssetsFilter(const FMainMRUFavoritesList::FDoesMRUFavoritesItemPassFilter& InFilter);
+	UNREALED_API void SetRecentAssetsFilter(const FMainMRUFavoritesList::FDoesMRUFavoritesItemPassFilter& InFilter);
 
 private:
 
 	/** Handles a package being reloaded */
-	void HandlePackageReloaded(const EPackageReloadPhase InPackageReloadPhase, FPackageReloadedEvent* InPackageReloadedEvent);
+	UNREALED_API void HandlePackageReloaded(const EPackageReloadPhase InPackageReloadPhase, FPackageReloadedEvent* InPackageReloadedEvent);
 
 	/** Callback for when the Editor closes, before Slate shuts down all the windows. */
-	void OnEditorClose();
+	UNREALED_API void OnEditorClose();
 
-	void RegisterEditorModes();
-	void UnregisterEditorModes();
-	void OnSMInstanceElementsEnabled();
+	UNREALED_API void RegisterEditorModes();
+	UNREALED_API void UnregisterEditorModes();
+	UNREALED_API void OnSMInstanceElementsEnabled();
 
-	bool IsEditorModeAllowed(const FName ModeId) const;
+	UNREALED_API bool IsEditorModeAllowed(const FName ModeId) const;
 
-	void InitializeRecentAssets();
-	void SaveRecentAssets(const bool bOnShutdown = false);
-	void CullRecentAssetEditorsMap();
+	UNREALED_API void InitializeRecentAssets();
+	UNREALED_API void SaveRecentAssets(const bool bOnShutdown = false);
+	UNREALED_API void CullRecentAssetEditorsMap();
 
-	bool ShouldShowRecentAsset(const FString& AssetName, int32 RecentAssetIndex,  const FName& InAssetEditorName) const;
-	bool ShouldShowRecentAssetsMenu(const FName& InAssetEditorName) const;
+	UNREALED_API bool ShouldShowRecentAsset(const FString& AssetName, int32 RecentAssetIndex,  const FName& InAssetEditorName) const;
+	UNREALED_API bool ShouldShowRecentAssetsMenu(const FName& InAssetEditorName) const;
 
 private:
 
diff --git a/Engine/Source/Editor/UnrealEd/Public/Subsystems/BrushEditingSubsystem.h b/Engine/Source/Editor/UnrealEd/Public/Subsystems/BrushEditingSubsystem.h
--- a/Engine/Source/Editor/UnrealEd/Public/Subsystems/BrushEditingSubsystem.h
+++ b/Engine/Source/Editor/UnrealEd/Public/Subsystems/BrushEditingSubsystem.h
@@ -12,13 +12,13 @@ class HHitProxy;
 class ABrush;
 struct FViewportClick;
 
-UCLASS(abstract)
-class UNREALED_API UBrushEditingSubsystem : public UEditorSubsystem
+UCLASS(abstract, MinimalAPI)
+class UBrushEditingSubsystem : public UEditorSubsystem
 {
 	GENERATED_BODY()
 
 public:
-	UBrushEditingSubsystem();
+	UNREALED_API UBrushEditingSubsystem();
 
 	virtual bool ProcessClickOnBrushGeometry(FLevelEditorViewportClient* ViewportClient, HHitProxy* InHitProxy, const FViewportClick& Click) { return false; }
 
diff --git a/Engine/Source/Editor/UnrealEd/Public/Subsystems/EditorActorSubsystem.h b/Engine/Source/Editor/UnrealEd/Public/Subsystems/EditorActorSubsystem.h
--- a/Engine/Source/Editor/UnrealEd/Public/Subsystems/EditorActorSubsystem.h
+++ b/Engine/Source/Editor/UnrealEd/Public/Subsystems/EditorActorSubsystem.h
@@ -40,15 +40,15 @@ DECLARE_DYNAMIC_MULTICAST_DELEGATE(FOnDeleteActorsEnd);
 * UEditorActorUtilitiesSubsystem
 * Subsystem for exposing actor related utilities to scripts,
 */
-UCLASS()
-class UNREALED_API UEditorActorSubsystem : public UEditorSubsystem
+UCLASS(MinimalAPI)
+class UEditorActorSubsystem : public UEditorSubsystem
 {
 	GENERATED_BODY()
 
 public:
 
-	virtual void Initialize(FSubsystemCollectionBase& Collection) override;
-	virtual void Deinitialize() override;
+	UNREALED_API virtual void Initialize(FSubsystemCollectionBase& Collection) override;
+	UNREALED_API virtual void Deinitialize() override;
 
 	UPROPERTY(BlueprintAssignable, Category = "Editor Scripting | Level Utility")
 	FOnEditNewActorsDropped OnNewActorsDropped;
@@ -94,7 +94,7 @@ public:
 	 * @return	The duplicated actor, or none if it didn't succeed
 	 */
 	UFUNCTION(BlueprintCallable, Category = "Editor Scripting | Level Utility", meta = (AdvancedDisplay = 1))
-	AActor* DuplicateActor(AActor* ActorToDuplicate, UWorld* ToWorld = nullptr, FVector Offset = FVector::ZeroVector);
+	UNREALED_API AActor* DuplicateActor(AActor* ActorToDuplicate, UWorld* ToWorld = nullptr, FVector Offset = FVector::ZeroVector);
 
 	/**
 	 * Duplicate actors from the world editor.
@@ -104,35 +104,35 @@ public:
 	 * @return	The duplicated actors, or empty if it didn't succeed
 	 */
 	UFUNCTION(BlueprintCallable, Category = "Editor Scripting | Level Utility", meta = (AdvancedDisplay = 1))
-	TArray<AActor*> DuplicateActors(const TArray<AActor*>& ActorsToDuplicate, UWorld* ToWorld = nullptr, FVector Offset = FVector::ZeroVector);
+	UNREALED_API TArray<AActor*> DuplicateActors(const TArray<AActor*>& ActorsToDuplicate, UWorld* ToWorld = nullptr, FVector Offset = FVector::ZeroVector);
 
 	/**
 	 * Duplicate all the selected actors in the given world
 	 * @param	InWorld 	The world the actors are selected in.
 	 */
 	UFUNCTION(BlueprintCallable, Category = "Editor Scripting | Level Utility")
-	void DuplicateSelectedActors(UWorld* InWorld);
+	UNREALED_API void DuplicateSelectedActors(UWorld* InWorld);
 
 	/**
 	 * Delete all the selected actors in the given world
 	 * @param	InWorld 	The world the actors are selected in.
 	 */
 	UFUNCTION(BlueprintCallable, Category = "Editor Scripting | Level Utility")
-	void DeleteSelectedActors(UWorld* InWorld);
+	UNREALED_API void DeleteSelectedActors(UWorld* InWorld);
 
 	/**
 	 * Invert the selection in the given world
 	 * @param	InWorld 	The world the selection is in.
 	 */
 	UFUNCTION(BlueprintCallable, Category = "Editor Scripting | Level Utility")
-	void InvertSelection(UWorld* InWorld);
+	UNREALED_API void InvertSelection(UWorld* InWorld);
 
 	/**
 	* Select all actors and BSP models in the given world, except those which are hidden
 	*  @param	InWorld 	The world the actors are to be selected in.
 	*/
 	UFUNCTION(BlueprintCallable, Category = "Editor Scripting | Level Utility")
-	void SelectAll(UWorld* InWorld);
+	UNREALED_API void SelectAll(UWorld* InWorld);
 
 	/**
 	 * Select all children actors of the current selection.
@@ -140,21 +140,21 @@ public:
 	 * @param   bRecurseChildren	True to recurse through all descendants of the children
 	 */
 	UFUNCTION(BlueprintCallable, Category = "Editor Scripting | Level Utility")
-	void SelectAllChildren(bool bRecurseChildren);
+	UNREALED_API void SelectAllChildren(bool bRecurseChildren);
 
 	/**
 	 * Find all loaded Actors in the world editor. Exclude actor that are pending kill, in PIE, PreviewEditor, ...
 	 * @return	List of found Actors
 	 */
 	UFUNCTION(BlueprintCallable, Category = "Editor Scripting | Level Utility")
-	TArray<class AActor*> GetAllLevelActors();
+	UNREALED_API TArray<class AActor*> GetAllLevelActors();
 
 	/**
 	 * Find all loaded ActorComponent own by an actor in the world editor. Exclude actor that are pending kill, in PIE, PreviewEditor, ...
 	 * @return	List of found ActorComponent
 	 */
 	UFUNCTION(BlueprintCallable, Category = "Editor Scripting | Level Utility")
-	TArray<class UActorComponent*> GetAllLevelActorsComponents();
+	UNREALED_API TArray<class UActorComponent*> GetAllLevelActorsComponents();
 
 	/**
 	 * Find all loaded Actors that are selected in the world editor. Exclude actor that are pending kill, in PIE, PreviewEditor, ...
@@ -162,26 +162,26 @@ public:
 	 * @return	List of found Actors
 	 */
 	UFUNCTION(BlueprintCallable, Category = "Editor Scripting | Level Utility")
-	TArray<class AActor*> GetSelectedLevelActors();
+	UNREALED_API TArray<class AActor*> GetSelectedLevelActors();
 
 	/**
 	 * Clear the current world editor selection and select the provided actors. Exclude actor that are pending kill, in PIE, PreviewEditor, ...
 	 * @param	ActorsToSelect	Actor that should be selected in the world editor.
 	 */
 	UFUNCTION(BlueprintCallable, Category = "Editor Scripting | Level Utility")
-	void SetSelectedLevelActors(const TArray<class AActor*>& ActorsToSelect);
+	UNREALED_API void SetSelectedLevelActors(const TArray<class AActor*>& ActorsToSelect);
 
 	// Remove all actors from the selection set
 	UFUNCTION(BlueprintCallable, Category = "Development|Editor")
-	void ClearActorSelectionSet();
+	UNREALED_API void ClearActorSelectionSet();
 
 	// Selects nothing in the editor (another way to clear the selection)
 	UFUNCTION(BlueprintCallable, Category = "Development|Editor")
-	void SelectNothing();
+	UNREALED_API void SelectNothing();
 
 	// Set the selection state for the selected actor
 	UFUNCTION(BlueprintCallable, Category = "Development|Editor")
-	void SetActorSelectionState(AActor* Actor, bool bShouldBeSelected);
+	UNREALED_API void SetActorSelectionState(AActor* Actor, bool bShouldBeSelected);
 
 	/**
 	* Attempts to find the actor specified by PathToActor in the current editor world
@@ -189,7 +189,7 @@ public:
 	* @return	A reference to the actor, or none if it wasn't found
 	*/
 	UFUNCTION(BlueprintPure, Category = "Development|Editor")
-	AActor* GetActorReference(FString PathToActor);
+	UNREALED_API AActor* GetActorReference(FString PathToActor);
 
 	/**
 	 * Create an actor and place it in the world editor. The Actor can be created from a Factory, Archetype, Blueprint, Class or an Asset.
@@ -199,7 +199,7 @@ public:
 	 * @return	The created actor.
 	 */
 	UFUNCTION(BlueprintCallable, Category = "Editor Scripting | Level Utility")
-	AActor* SpawnActorFromObject(class UObject* ObjectToUse, FVector Location, FRotator Rotation = FRotator::ZeroRotator, bool bTransient = false);
+	UNREALED_API AActor* SpawnActorFromObject(class UObject* ObjectToUse, FVector Location, FRotator Rotation = FRotator::ZeroRotator, bool bTransient = false);
 
 	/**
 	 * Create an actor and place it in the world editor. Can be created from a Blueprint or a Class.
@@ -209,7 +209,7 @@ public:
 	 * @return	The created actor.
 	 */
 	UFUNCTION(BlueprintCallable, Category = "Editor Scripting | Level Utility", meta = (DeterminesOutputType = "ActorClass"))
-	AActor* SpawnActorFromClass(TSubclassOf<class AActor> ActorClass, FVector Location, FRotator Rotation = FRotator::ZeroRotator, bool bTransient = false);
+	UNREALED_API AActor* SpawnActorFromClass(TSubclassOf<class AActor> ActorClass, FVector Location, FRotator Rotation = FRotator::ZeroRotator, bool bTransient = false);
 
 	/**
 	 * Destroy the actor from the world editor. Notify the Editor that the actor got destroyed.
@@ -217,7 +217,7 @@ public:
 	 * @return	True if the operation succeeds.
 	 */
 	UFUNCTION(BlueprintCallable, Category = "Editor Scripting | Level Utility")
-	bool DestroyActor(AActor* ActorToDestroy);
+	UNREALED_API bool DestroyActor(AActor* ActorToDestroy);
 
 	/**
 	* Destroy the actors from the world editor. Notify the Editor that the actor got destroyed.
@@ -225,7 +225,7 @@ public:
 	* @return	True if the operation succeeds.
 	*/
 	UFUNCTION(BlueprintCallable, Category = "Editor Scripting | Level Utility")
-	bool DestroyActors(const TArray<AActor*>& ActorsToDestroy);
+	UNREALED_API bool DestroyActors(const TArray<AActor*>& ActorsToDestroy);
 
 	/**
 	 * Replace in the level all Actors provided with a new actor of type ActorClass. Destroy all Actors provided.
@@ -234,21 +234,21 @@ public:
 	 * @param	StaticMeshPackagePath	If the list contains Brushes and it is requested to change them to StaticMesh, StaticMeshPackagePath is the package path to where the StaticMesh will be created. ie. /Game/MyFolder/
 	 */
 	UFUNCTION(BlueprintCallable, Category = "Editor Scripting | Dataprep", meta = (DeterminesOutputType = "ActorClass"))
-	TArray<class AActor*> ConvertActors(const TArray<class AActor*>& Actors, TSubclassOf<class AActor> ActorClass, const FString& StaticMeshPackagePath);
+	UNREALED_API TArray<class AActor*> ConvertActors(const TArray<class AActor*>& Actors, TSubclassOf<class AActor> ActorClass, const FString& StaticMeshPackagePath);
 
 	/**
 	 * Sets the world transform of the given actor, if possible.
 	 * @returns false if the world transform could not be set.
 	 */
 	UFUNCTION(BlueprintCallable, Category = "Editor Scripting | Level Utility")
-	bool SetActorTransform(AActor* InActor, const FTransform& InWorldTransform);
+	UNREALED_API bool SetActorTransform(AActor* InActor, const FTransform& InWorldTransform);
 
 	/**
 	 * Sets the world transform of the given component, if possible.
 	 * @returns false if the world transform could not be set.
 	 */
 	UFUNCTION(BlueprintCallable, Category = "Editor Scripting | Level Utility")
-	bool SetComponentTransform(USceneComponent* InSceneComponent, const FTransform& InWorldTransform);
+	UNREALED_API bool SetComponentTransform(USceneComponent* InSceneComponent, const FTransform& InWorldTransform);
 
 private:
 
diff --git a/Engine/Source/Editor/UnrealEd/Public/Subsystems/EditorAssetSubsystem.h b/Engine/Source/Editor/UnrealEd/Public/Subsystems/EditorAssetSubsystem.h
--- a/Engine/Source/Editor/UnrealEd/Public/Subsystems/EditorAssetSubsystem.h
+++ b/Engine/Source/Editor/UnrealEd/Public/Subsystems/EditorAssetSubsystem.h
@@ -20,16 +20,16 @@
 *		/Game/MyNewFolder/
 *		/Game/MyNewFolder
 */
-UCLASS()
-class UNREALED_API UEditorAssetSubsystem : public UEditorSubsystem
+UCLASS(MinimalAPI)
+class UEditorAssetSubsystem : public UEditorSubsystem
 {
 	GENERATED_BODY()
 	
 public:
 
-	UEditorAssetSubsystem();
-	virtual void Initialize(FSubsystemCollectionBase& Collection) override;
-	virtual void Deinitialize() override;
+	UNREALED_API UEditorAssetSubsystem();
+	UNREALED_API virtual void Initialize(FSubsystemCollectionBase& Collection) override;
+	UNREALED_API virtual void Deinitialize() override;
 	
 	/**
 	 * Load an asset. It will verify if the object is already loaded and only load it if it's necessary.
@@ -37,7 +37,7 @@ public:
 	 * @return	Found or loaded asset.
 	 */
 	UFUNCTION(BlueprintCallable, Category = "Editor Scripting | Asset")
-	UObject* LoadAsset(const FString& AssetPath);
+	UNREALED_API UObject* LoadAsset(const FString& AssetPath);
 
 	/**
 	 * Load a Blueprint asset and return its generated class. It will verify if the object is already loaded and only load it if it's necessary.
@@ -45,7 +45,7 @@ public:
 	 * @return	Found or loaded class.
 	 */
 	UFUNCTION(BlueprintCallable, Category = "Editor Scripting | Asset")
-	UClass* LoadBlueprintClass(const FString& AssetPath);
+	UNREALED_API UClass* LoadBlueprintClass(const FString& AssetPath);
 
 	/**
 	 * Return a valid AssetPath for a loaded asset.
@@ -54,7 +54,7 @@ public:
 	 * @return	If valid, the asset Path of the loaded asset.
 	 */
 	UFUNCTION(BlueprintCallable, Category = "Editor Scripting | Asset")
-	FString GetPathNameForLoadedAsset(UObject* LoadedAsset);
+	UNREALED_API FString GetPathNameForLoadedAsset(UObject* LoadedAsset);
 
 	/**
 	 * Return the AssetData for the Asset that can then be used with AssetRegistryHelpers.
@@ -62,7 +62,7 @@ public:
 	 * @return	The AssetData found.
 	 */
 	UFUNCTION(BlueprintCallable, Category = "Editor Scripting | Asset")
-	FAssetData FindAssetData(const FString& AssetPath);
+	UNREALED_API FAssetData FindAssetData(const FString& AssetPath);
 
 	/**
 	 * Check if an asset exists in the Asset Registry.
@@ -70,7 +70,7 @@ public:
 	 * @return	True if the asset exists and is valid.
 	 */
 	UFUNCTION(BlueprintCallable, Category = "Editor Scripting | Asset")
-	bool DoesAssetExist(const FString& AssetPath);
+	UNREALED_API bool DoesAssetExist(const FString& AssetPath);
 
 	/**
 	 * Check if assets exist in the Asset Registry.
@@ -78,7 +78,7 @@ public:
 	 * @return	True if all assets exist and are valid.
 	 */
 	UFUNCTION(BlueprintCallable, Category = "Editor Scripting | Asset")
-	bool DoAssetsExist(const TArray<FString>& AssetPaths);
+	UNREALED_API bool DoAssetsExist(const TArray<FString>& AssetPaths);
 	
 	/**
 	 * Find Package Referencers for an asset. Only Soft and Hard dependencies will be looked for.
@@ -91,7 +91,7 @@ public:
 	 * @return	The package paths of the referencers.
 	 */
 	UFUNCTION(BlueprintCallable, Category = "Editor Scripting | Asset", meta=(AdvancedDisplay=1))
-	TArray<FString> FindPackageReferencersForAsset(const FString& AssetPath, bool bLoadAssetsToConfirm = false);
+	UNREALED_API TArray<FString> FindPackageReferencersForAsset(const FString& AssetPath, bool bLoadAssetsToConfirm = false);
 
 	/**
 	 * Consolidates assets by replacing all references/uses of the provided AssetsToConsolidate with references to AssetToConsolidateTo.
@@ -105,7 +105,7 @@ public:
 	 * @return	True if the operation succeeds.
 	 */
 	UFUNCTION(BlueprintCallable, Category = "Editor Scripting | Asset")
-	bool ConsolidateAssets(UObject* AssetToConsolidateTo, const TArray<UObject*>& AssetsToConsolidate);
+	UNREALED_API bool ConsolidateAssets(UObject* AssetToConsolidateTo, const TArray<UObject*>& AssetsToConsolidate);
 	
 	/**
 	 * Delete an asset that is already loaded.
@@ -116,7 +116,7 @@ public:
 	 * @return	True if the operation succeeds.
 	 */
 	UFUNCTION(BlueprintCallable, Category = "Editor Scripting | Asset")
-	bool DeleteLoadedAsset(UObject* AssetToDelete);
+	UNREALED_API bool DeleteLoadedAsset(UObject* AssetToDelete);
 
 	/**
 	 * Delete assets that are already loaded.
@@ -127,7 +127,7 @@ public:
 	 * @return	True if the operation succeeds.
 	 */
 	UFUNCTION(BlueprintCallable, Category = "Editor Scripting | Asset")
-	bool DeleteLoadedAssets(const TArray<UObject*>& AssetsToDelete);
+	UNREALED_API bool DeleteLoadedAssets(const TArray<UObject*>& AssetsToDelete);
 
 	/**
 	 * Delete the package an asset is in. All objects in the package will be deleted.
@@ -138,7 +138,7 @@ public:
 	 * @return	True if the operation succeeds.
 	 */
 	UFUNCTION(BlueprintCallable, Category = "Editor Scripting | Asset")
-	bool DeleteAsset(const FString& AssetPathToDelete);
+	UNREALED_API bool DeleteAsset(const FString& AssetPathToDelete);
 
 	/**
 	 * Delete the packages inside a directory. If the directory is then empty, delete the directory.
@@ -150,7 +150,7 @@ public:
 	 * @return	True if the operation succeeds.
 	 */
 	UFUNCTION(BlueprintCallable, Category = "Editor Scripting | Asset")
-	bool DeleteDirectory(const FString& DirectoryPath);
+	UNREALED_API bool DeleteDirectory(const FString& DirectoryPath);
 	
 	/**
 	 * Duplicate an asset that is already loaded. Will try to checkout the file.
@@ -159,7 +159,7 @@ public:
 	 * @return	The duplicated object if the operation succeeds
 	 */
 	UFUNCTION(BlueprintCallable, Category = "Editor Scripting | Asset")
-	UObject* DuplicateLoadedAsset(UObject* SourceAsset, const FString& DestinationAssetPath);
+	UNREALED_API UObject* DuplicateLoadedAsset(UObject* SourceAsset, const FString& DestinationAssetPath);
 
 	/**
 	 * Duplicate an asset. Will try to checkout the file. The Asset will be loaded before being duplicated.
@@ -168,7 +168,7 @@ public:
 	 * @return	The duplicated object if the operation succeeds.
 	 */
 	UFUNCTION(BlueprintCallable, Category = "Editor Scripting | Asset")
-	UObject* DuplicateAsset(const FString& SourceAssetPath, const FString& DestinationAssetPath);
+	UNREALED_API UObject* DuplicateAsset(const FString& SourceAssetPath, const FString& DestinationAssetPath);
 
 	/**
 	 * Duplicate a directory and the assets in it.
@@ -178,7 +178,7 @@ public:
 	 * @return	The duplicated object if the operation succeeds.
 	 */
 	UFUNCTION(BlueprintCallable, Category = "Editor Scripting | Asset")
-	bool DuplicateDirectory(const FString& SourceDirectoryPath, const FString& DestinationDirectoryPath);
+	UNREALED_API bool DuplicateDirectory(const FString& SourceDirectoryPath, const FString& DestinationDirectoryPath);
 
 	/**
 	 * Rename an asset that is already loaded. Equivalent to a Move operation.
@@ -188,7 +188,7 @@ public:
 	 * @return	True if the operation succeeds.
 	 */
 	UFUNCTION(BlueprintCallable, Category = "Editor Scripting | Asset")
-	bool RenameLoadedAsset(UObject* SourceAsset, const FString& DestinationAssetPath);
+	UNREALED_API bool RenameLoadedAsset(UObject* SourceAsset, const FString& DestinationAssetPath);
 
 	/**
 	 * Rename an asset. Equivalent to a Move operation.
@@ -198,7 +198,7 @@ public:
 	 * @return	True if the operation succeeds.
 	 */
 	UFUNCTION(BlueprintCallable, Category = "Editor Scripting | Asset")
-	bool RenameAsset(const FString& SourceAssetPath, const FString& DestinationAssetPath);
+	UNREALED_API bool RenameAsset(const FString& SourceAssetPath, const FString& DestinationAssetPath);
 
 	/**
 	 * Rename a directory. Equivalent to a Move operation moving all contained assets.
@@ -208,7 +208,7 @@ public:
 	 * @return	True if the operation succeeds.
 	 */
 	UFUNCTION(BlueprintCallable, Category = "Editor Scripting | Asset")
-	bool RenameDirectory(const FString& SourceDirectoryPath, const FString& DestinationDirectoryPath);
+	UNREALED_API bool RenameDirectory(const FString& SourceDirectoryPath, const FString& DestinationDirectoryPath);
 	
 	/**
 	 * Checkout the asset corresponding to an object.
@@ -216,7 +216,7 @@ public:
 	 * @return	True if the operation succeeds.
 	 */
 	UFUNCTION(BlueprintCallable, Category = "Editor Scripting | Asset | Revision Control", meta = (Keywords = "Source Control"))
-	bool CheckoutLoadedAsset(UObject* AssetToCheckout);
+	UNREALED_API bool CheckoutLoadedAsset(UObject* AssetToCheckout);
 
 	/**
 	 * Checkout the assets.
@@ -224,7 +224,7 @@ public:
 	 * @return	True if the operation succeeds.
 	 */
 	UFUNCTION(BlueprintCallable, Category = "Editor Scripting | Asset | Revision Control", meta = (Keywords = "Source Control"))
-	bool CheckoutLoadedAssets(const TArray<UObject*>& AssetsToCheckout);
+	UNREALED_API bool CheckoutLoadedAssets(const TArray<UObject*>& AssetsToCheckout);
 
 	/**
 	 * Checkout an asset.
@@ -232,7 +232,7 @@ public:
 	 * @return	True if the operation succeeds.
 	 */
 	UFUNCTION(BlueprintCallable, Category = "Editor Scripting | Asset | Revision Control", meta = (Keywords = "Source Control"))
-	bool CheckoutAsset(const FString& AssetToCheckout);
+	UNREALED_API bool CheckoutAsset(const FString& AssetToCheckout);
 
 	/**
 	 * Checkout all assets in a directory. It will load the assets if needed.
@@ -242,7 +242,7 @@ public:
 	 * @return	True if the operation succeeds.
 	 */
 	UFUNCTION(BlueprintCallable, Category = "Editor Scripting | Asset | Revision Control", meta = (Keywords = "Source Control"))
-	bool CheckoutDirectory(const FString& DirectoryPath, bool bRecursive = true);
+	UNREALED_API bool CheckoutDirectory(const FString& DirectoryPath, bool bRecursive = true);
 	
 	/**
 	 * Save the package the asset lives in. All objects that live in the package will be saved. Will try to checkout the file.
@@ -251,7 +251,7 @@ public:
 	 * @return	True if the operation succeeds.
 	 */
 	UFUNCTION(BlueprintCallable, Category = "Editor Scripting | Asset")
-	bool SaveLoadedAsset(UObject* AssetToSave, bool bOnlyIfIsDirty = true);
+	UNREALED_API bool SaveLoadedAsset(UObject* AssetToSave, bool bOnlyIfIsDirty = true);
 
 	/**
 	 * Save the packages the assets live in. All objects that live in the packages will be saved. Will try to checkout the files.
@@ -260,7 +260,7 @@ public:
 	 * @return	True if the operation succeeds.
 	 */
 	UFUNCTION(BlueprintCallable, Category = "Editor Scripting | Asset")
-	bool SaveLoadedAssets(const TArray<UObject*>& AssetsToSave, bool bOnlyIfIsDirty = true);
+	UNREALED_API bool SaveLoadedAssets(const TArray<UObject*>& AssetsToSave, bool bOnlyIfIsDirty = true);
 
 	/**
 	 * Save the packages the assets live in. All objects that live in the package will be saved.
@@ -270,7 +270,7 @@ public:
 	 * @return	True if the operation succeeds.
 	 */
 	UFUNCTION(BlueprintCallable, Category = "Editor Scripting | Asset")
-	bool SaveAsset(const FString& AssetToSave, bool bOnlyIfIsDirty = true);
+	UNREALED_API bool SaveAsset(const FString& AssetToSave, bool bOnlyIfIsDirty = true);
 
 	/**
 	 * Save the packages the assets live in inside the directory. All objects that are in the directory will be saved.
@@ -281,7 +281,7 @@ public:
 	 * @return	True if the operation succeeds.
 	 */
 	UFUNCTION(BlueprintCallable, Category = "Editor Scripting | Asset")
-	bool SaveDirectory(const FString& DirectoryPath, bool bOnlyIfIsDirty = true, bool bRecursive = true);
+	UNREALED_API bool SaveDirectory(const FString& DirectoryPath, bool bOnlyIfIsDirty = true, bool bRecursive = true);
 	
 	/**
 	* Check if a directory exists.
@@ -289,7 +289,7 @@ public:
 	* @return	True if it does exist and it is valid.
 	*/
 	UFUNCTION(BlueprintCallable, Category = "Editor Scripting | Asset")
-	bool DoesDirectoryExist(const FString& DirectoryPath);
+	UNREALED_API bool DoesDirectoryExist(const FString& DirectoryPath);
 
 	/**
 	 * Check if a directory contains any assets.
@@ -297,7 +297,7 @@ public:
 	 * @return	True if there is any assets.
 	 */
 	UFUNCTION(BlueprintCallable, Category = "Editor Scripting | Asset")
-	bool DoesDirectoryContainAssets(const FString& DirectoryPath, bool bRecursive = true);
+	UNREALED_API bool DoesDirectoryContainAssets(const FString& DirectoryPath, bool bRecursive = true);
 
 	/**
 	 * Create a directory on disk.
@@ -305,7 +305,7 @@ public:
 	 * @return	True if the operation succeeds.
 	 */
 	UFUNCTION(BlueprintCallable, Category = "Editor Scripting | Asset")
-	bool MakeDirectory(const FString& DirectoryPath);
+	UNREALED_API bool MakeDirectory(const FString& DirectoryPath);
 	
 	/**
 	 * Return the list of all the assets found in the DirectoryPath.
@@ -315,7 +315,7 @@ public:
 	 * @return	The list of assets found.
 	 */
 	UFUNCTION(BlueprintCallable, Category = "Editor Scripting | Asset")
-	TArray<FString> ListAssets(const FString& DirectoryPath, bool bRecursive = true, bool bIncludeFolder = false);
+	UNREALED_API TArray<FString> ListAssets(const FString& DirectoryPath, bool bRecursive = true, bool bIncludeFolder = false);
 
 	/**
 	 * Return the list of all the assets that have the pair of Tag/Value.
@@ -324,7 +324,7 @@ public:
 	 * @return	The list of assets found.
 	 */
 	UFUNCTION(BlueprintCallable, Category = "Editor Scripting | Asset")
-	TArray<FString> ListAssetsByTagValue(FName TagName, const FString& TagValue);
+	UNREALED_API TArray<FString> ListAssetsByTagValue(FName TagName, const FString& TagValue);
 
 	/**
 	 * Gets all TagValues (from Asset Registry) associated with an (unloaded) asset as strings value.
@@ -332,7 +332,7 @@ public:
 	 * @return	The list of all TagNames & TagValues.
 	 */
 	UFUNCTION(BlueprintCallable, Category = "Editor Scripting | Asset")
-	TMap<FName, FString> GetTagValues(const FString& AssetPath);
+	UNREALED_API TMap<FName, FString> GetTagValues(const FString& AssetPath);
 	
 	/**
      * Get all tags/values of a loaded asset's metadata.
@@ -340,7 +340,7 @@ public:
      * @return	The list of all Tags and Values.
      */
     UFUNCTION(BlueprintCallable, Category = "Editor Scripting | Asset | Metadata")
-    TMap<FName, FString> GetMetadataTagValues(UObject* Object);
+    UNREALED_API TMap<FName, FString> GetMetadataTagValues(UObject* Object);
    
     /**
      * Get the value associated with the given tag of a loaded asset's metadata.
@@ -349,7 +349,7 @@ public:
      * @return	The string value associated with the tag.
      */
     UFUNCTION(BlueprintCallable, Category = "Editor Scripting | Asset | Metadata")
-    FString GetMetadataTag(UObject* Object, FName Tag);
+    UNREALED_API FString GetMetadataTag(UObject* Object, FName Tag);
    
     /**
      * Set the value associated with a given tag of a loaded asset's metadata.
@@ -358,7 +358,7 @@ public:
      * @param	Value		The string value to associate with the tag.
      */
     UFUNCTION(BlueprintCallable, Category = "Editor Scripting | Asset | Metadata")
-    void SetMetadataTag(UObject* Object, FName Tag, const FString& Value);
+    UNREALED_API void SetMetadataTag(UObject* Object, FName Tag, const FString& Value);
    
     /**
      * Remove the given tag from a loaded asset's metadata.
@@ -366,7 +366,7 @@ public:
      * @param	Tag			The tag to remove from the metadata.
      */
     UFUNCTION(BlueprintCallable, Category = "Editor Scripting | Asset | Metadata")
-    void RemoveMetadataTag(UObject* Object, FName Tag);
+    UNREALED_API void RemoveMetadataTag(UObject* Object, FName Tag);
 
 	/**
 	 * Delegate for extracting an asset from a file, for example from a drag and drop operation.
@@ -386,13 +386,13 @@ public:
 	 * for example from a drag and drop operation.
 	 */
 	UFUNCTION(BlueprintCallable, Category = "Editor Scripting | Asset")
-		void AddOnExtractAssetFromFile(FOnExtractAssetFromFileDynamic Delegate);
+		UNREALED_API void AddOnExtractAssetFromFile(FOnExtractAssetFromFileDynamic Delegate);
 
 	/**
 	 * Call this to remove a callback added with AddOnExtractAssetFromFile.
 	 */
 	UFUNCTION(BlueprintCallable, Category = "Editor Scripting | Asset")
-		void RemoveOnExtractAssetFromFile(FOnExtractAssetFromFileDynamic Delegate);
+		UNREALED_API void RemoveOnExtractAssetFromFile(FOnExtractAssetFromFileDynamic Delegate);
 
 	/**
 	 * Get the delegate for extracting an asset from a file,
@@ -409,7 +409,7 @@ private:
 	 * Calls all the delegates in OnExtractAssetFromFileDynamicArray when OnExtractAssetFromFile
 	 * is broadcast.
 	 */
-	void CallOnExtractAssetFromFileDynamicArray(const TArray<FString>& Files,
+	UNREALED_API void CallOnExtractAssetFromFileDynamicArray(const TArray<FString>& Files,
 		TArray<FAssetData>& InAssetDataArray);
 
 	/** Delegate for extracting an asset from a file, for example from a drag and drop operation. */
diff --git a/Engine/Source/Editor/UnrealEd/Public/Subsystems/EditorSubsystemBlueprintLibrary.h b/Engine/Source/Editor/UnrealEd/Public/Subsystems/EditorSubsystemBlueprintLibrary.h
--- a/Engine/Source/Editor/UnrealEd/Public/Subsystems/EditorSubsystemBlueprintLibrary.h
+++ b/Engine/Source/Editor/UnrealEd/Public/Subsystems/EditorSubsystemBlueprintLibrary.h
@@ -8,8 +8,8 @@
 
 #include "EditorSubsystemBlueprintLibrary.generated.h"
 
-UCLASS()
-class UNREALED_API UEditorSubsystemBlueprintLibrary : public UBlueprintFunctionLibrary
+UCLASS(MinimalAPI)
+class UEditorSubsystemBlueprintLibrary : public UBlueprintFunctionLibrary
 {
 	GENERATED_BODY()
 
@@ -17,6 +17,6 @@ public:
 
 	/** Get a Local Player Subsystem from the Local Player associated with the provided context */
 	UFUNCTION(BlueprintPure, Category = "Editor Subsystems", meta = (BlueprintInternalUseOnly = "true"))
-	static UEditorSubsystem* GetEditorSubsystem(TSubclassOf<UEditorSubsystem> Class);
+	static UNREALED_API UEditorSubsystem* GetEditorSubsystem(TSubclassOf<UEditorSubsystem> Class);
 
-};
\ No newline at end of file
+};
diff --git a/Engine/Source/Editor/UnrealEd/Public/Subsystems/ImportSubsystem.h b/Engine/Source/Editor/UnrealEd/Public/Subsystems/ImportSubsystem.h
--- a/Engine/Source/Editor/UnrealEd/Public/Subsystems/ImportSubsystem.h
+++ b/Engine/Source/Editor/UnrealEd/Public/Subsystems/ImportSubsystem.h
@@ -29,19 +29,19 @@ public:
  * Subsystem for importing assets in the editor, 
  * Contains utility functions and callbacks for hooking into importing.
  */
-UCLASS()
-class UNREALED_API UImportSubsystem : public UEditorSubsystem
+UCLASS(MinimalAPI)
+class UImportSubsystem : public UEditorSubsystem
 {
 	GENERATED_BODY()
 
 public:
-	UImportSubsystem();
+	UNREALED_API UImportSubsystem();
 
-	virtual void Initialize(FSubsystemCollectionBase& Collection) override;
-	virtual void Deinitialize() override;
+	UNREALED_API virtual void Initialize(FSubsystemCollectionBase& Collection) override;
+	UNREALED_API virtual void Deinitialize() override;
 
 	/* Import files next tick */
-	void ImportNextTick(const TArray<FString>& Files, const FString& DestinationPath);
+	UNREALED_API void ImportNextTick(const TArray<FString>& Files, const FString& DestinationPath);
 
 	/** delegate type fired when new assets are being (re-)imported. Params: UFactory* InFactory, UClass* InClass, UObject* InParent, const FName& Name, const TCHAR* Type */
 	DECLARE_MULTICAST_DELEGATE_FiveParams(FOnAssetPreImport, UFactory*, UClass*, UObject*, const FName&, const TCHAR*);
@@ -53,13 +53,13 @@ public:
 	DECLARE_MULTICAST_DELEGATE_TwoParams(FOnAssetPostLODImport, UObject*, int32);
 
 	// Broadcast AssetPreImport, do not broadcast with OnAssetPostImport directly.
-	void BroadcastAssetPreImport(UFactory* InFactory, UClass* InClass, UObject* InParent, const FName& Name, const TCHAR* Type);
+	UNREALED_API void BroadcastAssetPreImport(UFactory* InFactory, UClass* InClass, UObject* InParent, const FName& Name, const TCHAR* Type);
 	// Broadcast AssetPostImport, do not broadcast with OnAssetPostImport directly.
-	void BroadcastAssetPostImport(UFactory* InFactory, UObject* InCreatedObject);
+	UNREALED_API void BroadcastAssetPostImport(UFactory* InFactory, UObject* InCreatedObject);
 	// Broadcast AssetReimport, do not broadcast with OnAssetReimport directly.
-	void BroadcastAssetReimport(UObject* InCreatedObject);
+	UNREALED_API void BroadcastAssetReimport(UObject* InCreatedObject);
 	// Broadcast AssetPostLODImport, do not broadcast with OnAssetPostLODImport directly.
-	void BroadcastAssetPostLODImport(UObject* InObject, int32 inLODIndex);
+	UNREALED_API void BroadcastAssetPostLODImport(UObject* InObject, int32 inLODIndex);
 
 	// Used to register and unregister ONLY use Broadcast functions to execute the delegate
 	FOnAssetPreImport OnAssetPreImport;
diff --git a/Engine/Source/Editor/UnrealEd/Public/Subsystems/PanelExtensionSubsystem.h b/Engine/Source/Editor/UnrealEd/Public/Subsystems/PanelExtensionSubsystem.h
--- a/Engine/Source/Editor/UnrealEd/Public/Subsystems/PanelExtensionSubsystem.h
+++ b/Engine/Source/Editor/UnrealEd/Public/Subsystems/PanelExtensionSubsystem.h
@@ -78,27 +78,27 @@ private:
  * UPanelExtensionSubsystem
  * Subsystem for creating extensible panels in the Editor
  */
-UCLASS()
-class UNREALED_API UPanelExtensionSubsystem : public UEditorSubsystem
+UCLASS(MinimalAPI)
+class UPanelExtensionSubsystem : public UEditorSubsystem
 {
 	GENERATED_BODY()
 
 public:
-	UPanelExtensionSubsystem();
+	UNREALED_API UPanelExtensionSubsystem();
 
-	virtual void Initialize(FSubsystemCollectionBase& Collection) override;
-	virtual void Deinitialize() override;
+	UNREALED_API virtual void Initialize(FSubsystemCollectionBase& Collection) override;
+	UNREALED_API virtual void Deinitialize() override;
 
-	void RegisterPanelFactory(FName ExtensionPanelID, const FPanelExtensionFactory& InPanelExtensionFactory);
-	void UnregisterPanelFactory(FName Identifier, FName ExtensionPanelID = NAME_None);
-	bool IsPanelFactoryRegistered(FName Identifier, FName ExtensionPanelID = NAME_None) const;
+	UNREALED_API void RegisterPanelFactory(FName ExtensionPanelID, const FPanelExtensionFactory& InPanelExtensionFactory);
+	UNREALED_API void UnregisterPanelFactory(FName Identifier, FName ExtensionPanelID = NAME_None);
+	UNREALED_API bool IsPanelFactoryRegistered(FName Identifier, FName ExtensionPanelID = NAME_None) const;
 
 protected:
 	friend class SExtensionPanel;
-	TSharedRef<SWidget> CreateWidget(FName ExtensionPanelID, FWeakObjectPtr ExtensionContext);
+	UNREALED_API TSharedRef<SWidget> CreateWidget(FName ExtensionPanelID, FWeakObjectPtr ExtensionContext);
 
 	DECLARE_MULTICAST_DELEGATE(FPanelFactoryRegistryChanged);
-	FPanelFactoryRegistryChanged& OnPanelFactoryRegistryChanged(FName ExtensionPanelID);
+	UNREALED_API FPanelFactoryRegistryChanged& OnPanelFactoryRegistryChanged(FName ExtensionPanelID);
 
 private:
 	TMap<FName, TArray<FPanelExtensionFactory>> ExtensionPointMap;
diff --git a/Engine/Source/Editor/UnrealEd/Public/Subsystems/UnrealEditorSubsystem.h b/Engine/Source/Editor/UnrealEd/Public/Subsystems/UnrealEditorSubsystem.h
--- a/Engine/Source/Editor/UnrealEd/Public/Subsystems/UnrealEditorSubsystem.h
+++ b/Engine/Source/Editor/UnrealEd/Public/Subsystems/UnrealEditorSubsystem.h
@@ -11,8 +11,8 @@
 * UUnrealEditorSubsystem 
 * Subsystem for exposing editor functionality to scripts
 */
-UCLASS()
-class UNREALED_API UUnrealEditorSubsystem  : public UEditorSubsystem
+UCLASS(MinimalAPI)
+class UUnrealEditorSubsystem  : public UEditorSubsystem
 {
 	GENERATED_BODY()
 
@@ -27,7 +27,7 @@ public:
 	 * @return	Whether or not we were able to get a camera for a level editing viewport
 	 */
 	UFUNCTION(BlueprintPure, Category = "Development|Editor")
-	bool GetLevelViewportCameraInfo(FVector& CameraLocation, FRotator& CameraRotation);
+	UNREALED_API bool GetLevelViewportCameraInfo(FVector& CameraLocation, FRotator& CameraRotation);
 
 	/**
 	* Sets information about the camera position for the primary level editor viewport.
@@ -37,15 +37,15 @@ public:
 	* @param	CameraRotation	Rotation the camera will be set to.
 	*/
 	UFUNCTION(BlueprintCallable, Category = "Development|Editor")
-	void SetLevelViewportCameraInfo(FVector CameraLocation, FRotator CameraRotation);
+	UNREALED_API void SetLevelViewportCameraInfo(FVector CameraLocation, FRotator CameraRotation);
 
 	/**
 	 * Find the World in the world editor. It can then be used as WorldContext by other libraries like GameplayStatics.
 	 * @return	The World used by the world editor.
 	 */
 	UFUNCTION(BlueprintCallable, Category = "Editor Scripting | Level Utility")
-	UWorld* GetEditorWorld();
+	UNREALED_API UWorld* GetEditorWorld();
 
 	UFUNCTION(BlueprintCallable, Category = "Editor Scripting | Level Utility")
-	UWorld* GetGameWorld();
+	UNREALED_API UWorld* GetGameWorld();
 };
diff --git a/Engine/Source/Editor/UnrealEd/Public/Tests/AutomationEditorCommon.h b/Engine/Source/Editor/UnrealEd/Public/Tests/AutomationEditorCommon.h
--- a/Engine/Source/Editor/UnrealEd/Public/Tests/AutomationEditorCommon.h
+++ b/Engine/Source/Editor/UnrealEd/Public/Tests/AutomationEditorCommon.h
@@ -14,7 +14,7 @@ class UFactory;
 
 //////////////////////////////////////////////////////////////////////////
 // FAutomationEditorCommonUtils
-class UNREALED_API FAutomationEditorCommonUtils
+class FAutomationEditorCommonUtils
 {
 public:
 	
@@ -23,7 +23,7 @@ public:
 	*
 	* @return - The UWorld for the new map
 	*/
-	static UWorld* CreateNewMap();
+	static UNREALED_API UWorld* CreateNewMap();
 
 	/**
 	* Imports an object using a given factory
@@ -33,21 +33,21 @@ public:
 	* @param PackagePath - The full path of the package file to create
 	* @param ImportPath - The path to the object to import
 	*/
-	static UObject* ImportAssetUsingFactory(UFactory* ImportFactory, const FString& ObjectName, const FString& PackageName, const FString& ImportPath);
+	static UNREALED_API UObject* ImportAssetUsingFactory(UFactory* ImportFactory, const FString& ObjectName, const FString& PackageName, const FString& ImportPath);
 
 	/**
 	* Nulls out references to a given object
 	*
 	* @param InObject - Object to null references to
 	*/
-	static void NullReferencesToObject(UObject* InObject);
+	static UNREALED_API void NullReferencesToObject(UObject* InObject);
 
 	/**
 	* gets a factory class based off an asset file extension
 	*
 	* @param AssetExtension - The file extension to use to find a supporting UFactory
 	*/
-	static UClass* GetFactoryClassForType(const FString& AssetExtension);
+	static UNREALED_API UClass* GetFactoryClassForType(const FString& AssetExtension);
 
 	/**
 	* Applies settings to an object by finding UProperties by name and calling ImportText
@@ -56,7 +56,7 @@ public:
 	* @param PropertyChain - The list FProperty names recursively to search through
 	* @param Value - The value to import on the found property
 	*/
-	static void ApplyCustomFactorySetting(UObject* InObject, TArray<FString>& PropertyChain, const FString& Value);
+	static UNREALED_API void ApplyCustomFactorySetting(UObject* InObject, TArray<FString>& PropertyChain, const FString& Value);
 
 	/**
 	* Applies the custom factory settings
@@ -64,7 +64,7 @@ public:
 	* @param InFactory - The factory to apply custom settings to
 	* @param FactorySettings - An array of custom settings to apply to the factory
 	*/
-	static void ApplyCustomFactorySettings(UFactory* InFactory, const TArray<FImportFactorySettingValues>& FactorySettings);
+	static UNREALED_API void ApplyCustomFactorySettings(UFactory* InFactory, const TArray<FImportFactorySettingValues>& FactorySettings);
 
 	/**
 	* Writes a number to a text file.
@@ -75,7 +75,7 @@ public:
 	* @param InEntry - is the double-precision number that is expected to be written to the file.
 	* @param Delimiter - is the delimiter to be used. TEXT(",")
 	*/
-	static void WriteToTextFile(const FString& InTestName, const FString& InTestItem, const FString& InFileName, const double& InEntry, const FString& Delimiter);
+	static UNREALED_API void WriteToTextFile(const FString& InTestName, const FString& InTestItem, const FString& InFileName, const double& InEntry, const FString& Delimiter);
 
 	/**
 	* Returns the sum of the numbers available in an array of float.
@@ -83,14 +83,14 @@ public:
 	* @param InFloatArray - is the name of the array intended to be used.
 	* @param bisAveragedInstead - will return the average of the available numbers instead of the sum.
 	*/
-	static float TotalFromFloatArray(const TArray<float>& InFloatArray, bool bisAveragedInstead);
+	static UNREALED_API float TotalFromFloatArray(const TArray<float>& InFloatArray, bool bisAveragedInstead);
 
 	/**
 	* Returns the largest value from an array of float numbers.
 
 	* @param InFloatArray - is the name of the array intended to be used.
 	*/
-	static float LargestValueInFloatArray(const TArray<float>& InFloatArray);
+	static UNREALED_API float LargestValueInFloatArray(const TArray<float>& InFloatArray);
 
 	/**
 	* Returns the contents of a text file as an array of FString.
@@ -98,7 +98,7 @@ public:
 	* @param InFileLocation - is the location of the file.
 	* @param OutArray - The name of the array that will store the data.
 	*/
-	static void CreateArrayFromFile(const FString& InFileLocation, TArray<FString>& OutArray);
+	static UNREALED_API void CreateArrayFromFile(const FString& InFileLocation, TArray<FString>& OutArray);
 
 	/**
 	* Returns true if the archive/file can be written to otherwise false.
@@ -106,7 +106,7 @@ public:
 	* @param InFilePath - is the location of the file.
 	* @param InArchiveName - is the name of the archive to be used.
 	*/
-	static bool IsArchiveWriteable(const FString& InFilePath, const FArchive* InArchiveName);
+	static UNREALED_API bool IsArchiveWriteable(const FString& InFilePath, const FArchive* InArchiveName);
 
 	/**
 	* Returns the first DeviceID 'Platform@Device'
@@ -115,7 +115,7 @@ public:
 	* @param OutDeviceID - The variable that will hold the device ID.
 	* @param InMapName - The map name to check against in the automation preferences.
 	*/
-	static void GetLaunchOnDeviceID(FString& OutDeviceID, const FString& InMapName);
+	static UNREALED_API void GetLaunchOnDeviceID(FString& OutDeviceID, const FString& InMapName);
 
 	/**
 	* Returns the DeviceID 'Platform@Device'
@@ -125,7 +125,7 @@ public:
 	* @param InMapName - The map name to check against in the automation preferences.
 	* @param InDeviceName - Device Name
 	*/
-	static void GetLaunchOnDeviceID(FString& OutDeviceID, const FString& InMapName, const FString& InDeviceName);
+	static UNREALED_API void GetLaunchOnDeviceID(FString& OutDeviceID, const FString& InMapName, const FString& InDeviceName);
 
 	/**
 	* Sets the first found ortho viewport camera to the desired location and rotation.
@@ -133,58 +133,58 @@ public:
 	* @param ViewLocation - Desired location for the viewport view.
 	* @param ViewRotation - Desired rotation of the viewport view.
 	*/
-	static bool SetOrthoViewportView(const FVector& ViewLocation, const FRotator& ViewRotation);
+	static UNREALED_API bool SetOrthoViewportView(const FVector& ViewLocation, const FRotator& ViewRotation);
 
 	/**
 	* Set Play Session parameters to Active Viewport location
 	* 
 	* @param OutParams - FRequestPlaySessionParams object to set
 	*/
-	static bool SetPlaySessionStartToActiveViewport(FRequestPlaySessionParams& OutParams);
+	static UNREALED_API bool SetPlaySessionStartToActiveViewport(FRequestPlaySessionParams& OutParams);
 
 	/**
 	* Converts a package path to an asset path
 	*
 	* @param PackagePath - The package path to convert
 	*/
-	static FString ConvertPackagePathToAssetPath(const FString& PackagePath);
+	static UNREALED_API FString ConvertPackagePathToAssetPath(const FString& PackagePath);
 
 	/**
 	* Gets the asset data from a package path
 	*
 	* @param PackagePath - The package path used to look up the asset data
 	*/
-	static FAssetData GetAssetDataFromPackagePath(const FString& PackagePath);
+	static UNREALED_API FAssetData GetAssetDataFromPackagePath(const FString& PackagePath);
 
 	/**
 	* Loads the map specified by an automation test
 	*
 	* @param MapName - Map to load
 	*/
-	static void LoadMap(const FString& MapName);
+	static UNREALED_API void LoadMap(const FString& MapName);
 
 	/**
 	* Run PIE
 	*/
-	static void RunPIE(float PIEDuration = 3.f);
+	static UNREALED_API void RunPIE(float PIEDuration = 3.f);
 	
 	/**
 	* Generates a list of assets from the ENGINE and the GAME by a specific type.
 	* This is to be used by the GetTest() function.
 	*/
-	static void CollectTestsByClass(UClass * Class, TArray<FString>& OutBeautifiedNames, TArray <FString>& OutTestCommands);
+	static UNREALED_API void CollectTestsByClass(UClass * Class, TArray<FString>& OutBeautifiedNames, TArray <FString>& OutTestCommands);
 	
 	/**
 	* Generates a list of assets from the GAME by a specific type.
 	* This is to be used by the GetTest() function.
 	*/
-	static void CollectGameContentTestsByClass(UClass * Class, bool bRecursiveClass, TArray<FString>& OutBeautifiedNames, TArray <FString>& OutTestCommands);
+	static UNREALED_API void CollectGameContentTestsByClass(UClass * Class, bool bRecursiveClass, TArray<FString>& OutBeautifiedNames, TArray <FString>& OutTestCommands);
 		
 	/**
 	* Generates a list of assets from the GAME by a specific type.
 	* This is to be used by the GetTest() function.
 	*/
-	static void CollectGameContentTests(TArray<FString>& OutBeautifiedNames, TArray <FString>& OutTestCommands);
+	static UNREALED_API void CollectGameContentTests(TArray<FString>& OutBeautifiedNames, TArray <FString>& OutTestCommands);
 };
 
 
@@ -267,4 +267,4 @@ DEFINE_EXPORTED_LATENT_AUTOMATION_COMMAND(UNREALED_API, FWaitToFinishBuildDeploy
 */
 DEFINE_EXPORTED_LATENT_AUTOMATION_COMMAND_ONE_PARAMETER(UNREALED_API, FDeleteDirCommand, FString, InFolderLocation);
 
-#endif
\ No newline at end of file
+#endif
diff --git a/Engine/Source/Editor/UnrealEd/Public/Text/HLSLSyntaxHighlighterMarshaller.h b/Engine/Source/Editor/UnrealEd/Public/Text/HLSLSyntaxHighlighterMarshaller.h
--- a/Engine/Source/Editor/UnrealEd/Public/Text/HLSLSyntaxHighlighterMarshaller.h
+++ b/Engine/Source/Editor/UnrealEd/Public/Text/HLSLSyntaxHighlighterMarshaller.h
@@ -14,7 +14,7 @@
 /**
  * Syntax highlighting for hlsl text
  */
-class UNREALED_API FHLSLSyntaxHighlighterMarshaller : public FSyntaxHighlighterTextLayoutMarshaller
+class FHLSLSyntaxHighlighterMarshaller : public FSyntaxHighlighterTextLayoutMarshaller
 {
 public:
 
@@ -51,7 +51,7 @@ public:
 		FTextBlockStyle ErrorTextStyle;
 	};
 
-	static TSharedRef<FHLSLSyntaxHighlighterMarshaller> Create(const FSyntaxTextStyle& InSyntaxTextStyle);
+	static UNREALED_API TSharedRef<FHLSLSyntaxHighlighterMarshaller> Create(const FSyntaxTextStyle& InSyntaxTextStyle);
 
 protected:
 
@@ -64,14 +64,14 @@ protected:
 		LookingForMultiLineComment,
 	};
 
-	static TSharedPtr<ISyntaxTokenizer> CreateTokenizer();
+	static UNREALED_API TSharedPtr<ISyntaxTokenizer> CreateTokenizer();
 	
-	virtual void ParseTokens(const FString& SourceString, FTextLayout& TargetTextLayout, TArray<ISyntaxTokenizer::FTokenizedLine> TokenizedLines) override;
+	UNREALED_API virtual void ParseTokens(const FString& SourceString, FTextLayout& TargetTextLayout, TArray<ISyntaxTokenizer::FTokenizedLine> TokenizedLines) override;
 
-	virtual FTextLayout::FNewLineData ProcessTokenizedLine(const ISyntaxTokenizer::FTokenizedLine& TokenizedLine, const int32& LineNumber, const FString& SourceString, EParseState& CurrentParseState);
+	UNREALED_API virtual FTextLayout::FNewLineData ProcessTokenizedLine(const ISyntaxTokenizer::FTokenizedLine& TokenizedLine, const int32& LineNumber, const FString& SourceString, EParseState& CurrentParseState);
 
-	FHLSLSyntaxHighlighterMarshaller(TSharedPtr<ISyntaxTokenizer> InTokenizer, const FSyntaxTextStyle& InSyntaxTextStyle);
+	UNREALED_API FHLSLSyntaxHighlighterMarshaller(TSharedPtr<ISyntaxTokenizer> InTokenizer, const FSyntaxTextStyle& InSyntaxTextStyle);
 
 	/** Styles used to display the text */
 	FSyntaxTextStyle SyntaxTextStyle;
-};
\ No newline at end of file
+};
diff --git a/Engine/Source/Editor/UnrealEd/Public/Texture2DPreview.h b/Engine/Source/Editor/UnrealEd/Public/Texture2DPreview.h
--- a/Engine/Source/Editor/UnrealEd/Public/Texture2DPreview.h
+++ b/Engine/Source/Editor/UnrealEd/Public/Texture2DPreview.h
@@ -12,7 +12,7 @@
 /**
  * Batched element parameters for previewing 2d textures.
  */
-class UNREALED_API FBatchedElementTexture2DPreviewParameters : public FBatchedElementParameters
+class FBatchedElementTexture2DPreviewParameters : public FBatchedElementParameters
 {
 public:
 	FBatchedElementTexture2DPreviewParameters(float InMipLevel, float InLayerIndex, float InSliceIndex, bool bInIsNormalMap, bool bInIsSingleChannel, bool bInIsSingleVTPhysicalSpace, bool bInIsVirtualTexture, bool bInIsTextureArray, bool bInUsePointSampling)
@@ -29,7 +29,7 @@ public:
 	}
 
 	/** Binds vertex and pixel shaders for this element */
-	virtual void BindShaders(FRHICommandList& RHICmdList, FGraphicsPipelineStateInitializer& GraphicsPSOInit, ERHIFeatureLevel::Type InFeatureLevel, const FMatrix& InTransform, const float InGamma, const FMatrix& ColorWeights, const FTexture* Texture) override;
+	UNREALED_API virtual void BindShaders(FRHICommandList& RHICmdList, FGraphicsPipelineStateInitializer& GraphicsPSOInit, ERHIFeatureLevel::Type InFeatureLevel, const FMatrix& InTransform, const float InGamma, const FMatrix& ColorWeights, const FTexture* Texture) override;
 
 private:
 
diff --git a/Engine/Source/Editor/UnrealEd/Public/ThumbnailExternalCache.h b/Engine/Source/Editor/UnrealEd/Public/ThumbnailExternalCache.h
--- a/Engine/Source/Editor/UnrealEd/Public/ThumbnailExternalCache.h
+++ b/Engine/Source/Editor/UnrealEd/Public/ThumbnailExternalCache.h
@@ -9,7 +9,7 @@
 
 class IPlugin;
 
-struct UNREALED_API FThumbnailExternalCacheSettings
+struct FThumbnailExternalCacheSettings
 {
 	/** Recompress any lossless thumbnails */
 	bool bRecompressLossless = false;
