[UnrealEd]

diff --git a/Engine/Source/Editor/UnrealEd/Public/ThumbnailHelpers.h b/Engine/Source/Editor/UnrealEd/Public/ThumbnailHelpers.h
--- a/Engine/Source/Editor/UnrealEd/Public/ThumbnailHelpers.h
+++ b/Engine/Source/Editor/UnrealEd/Public/ThumbnailHelpers.h
@@ -15,27 +15,27 @@
 class FSceneViewFamily;
 class USceneThumbnailInfo;
 
-class UNREALED_API FThumbnailPreviewScene : public FPreviewScene, public FTickableEditorObject
+class FThumbnailPreviewScene : public FPreviewScene, public FTickableEditorObject
 {
 public:
 	/** Constructor */
-	FThumbnailPreviewScene();
+	UNREALED_API FThumbnailPreviewScene();
 
 	UE_DEPRECATED(5.0, "Use CreateView")
-	void GetView(FSceneViewFamily* ViewFamily, int32 X, int32 Y, uint32 SizeX, uint32 SizeY) const;
+	UNREALED_API void GetView(FSceneViewFamily* ViewFamily, int32 X, int32 Y, uint32 SizeX, uint32 SizeY) const;
 
 	/** Allocates then adds an FSceneView to the ViewFamily. */
-	[[nodiscard]] FSceneView* CreateView(FSceneViewFamily* ViewFamily, int32 X, int32 Y, uint32 SizeX, uint32 SizeY) const;
+	[[nodiscard]] UNREALED_API FSceneView* CreateView(FSceneViewFamily* ViewFamily, int32 X, int32 Y, uint32 SizeX, uint32 SizeY) const;
 
 	/* Begin FTickableEditorObject */
-	virtual void Tick(float DeltaTime) override;
+	UNREALED_API virtual void Tick(float DeltaTime) override;
 	virtual ETickableTickType GetTickableTickType() const override { return ETickableTickType::Always; }
-	virtual TStatId GetStatId() const override;
+	UNREALED_API virtual TStatId GetStatId() const override;
 	/* End FTickableEditorObject */
 
 protected:
 	/** Helper function to get the bounds offset to display an asset */
-	float GetBoundsZOffset(const FBoxSphereBounds& Bounds) const;
+	UNREALED_API float GetBoundsZOffset(const FBoxSphereBounds& Bounds) const;
 
 	/**
 	  * Gets parameters to create a view matrix to be used by CreateView(). Implemented in children classes.
@@ -48,19 +48,19 @@ protected:
 	virtual void GetViewMatrixParameters(const float InFOVDegrees, FVector& OutOrigin, float& OutOrbitPitch, float& OutOrbitYaw, float& OutOrbitZoom) const = 0;
 };
 
-class UNREALED_API FParticleSystemThumbnailScene : public FThumbnailPreviewScene
+class FParticleSystemThumbnailScene : public FThumbnailPreviewScene
 {
 public:
 	/** Constructor/Destructor */
-	FParticleSystemThumbnailScene();
-	virtual ~FParticleSystemThumbnailScene();
+	UNREALED_API FParticleSystemThumbnailScene();
+	UNREALED_API virtual ~FParticleSystemThumbnailScene();
 
 	/** Sets the particle system to use in the next CreateView() */
-	void SetParticleSystem(class UParticleSystem* ParticleSystem);
+	UNREALED_API void SetParticleSystem(class UParticleSystem* ParticleSystem);
 
 protected:
 	// FThumbnailPreviewScene implementation
-	virtual void GetViewMatrixParameters(const float InFOVDegrees, FVector& OutOrigin, float& OutOrbitPitch, float& OutOrbitYaw, float& OutOrbitZoom) const override;
+	UNREALED_API virtual void GetViewMatrixParameters(const float InFOVDegrees, FVector& OutOrigin, float& OutOrbitPitch, float& OutOrbitYaw, float& OutOrbitZoom) const override;
 
 protected:
 	/** The particle system component used to display all particle system thumbnails */
@@ -70,20 +70,20 @@ protected:
 	class FFXSystemInterface* ThumbnailFXSystem;
 };
 
-class UNREALED_API FMaterialThumbnailScene : public FThumbnailPreviewScene
+class FMaterialThumbnailScene : public FThumbnailPreviewScene
 {
 public:	
 	/** Constructor */
-	FMaterialThumbnailScene();
+	UNREALED_API FMaterialThumbnailScene();
 
 	/** Sets the material to use in the next CreateView() */
-	void SetMaterialInterface(class UMaterialInterface* InMaterial);
+	UNREALED_API void SetMaterialInterface(class UMaterialInterface* InMaterial);
 
-	bool ShouldSetSeparateTranslucency(class UMaterialInterface* InMaterial) const;
+	UNREALED_API bool ShouldSetSeparateTranslucency(class UMaterialInterface* InMaterial) const;
 
 protected:
 	// FThumbnailPreviewScene implementation
-	virtual void GetViewMatrixParameters(const float InFOVDegrees, FVector& OutOrigin, float& OutOrbitPitch, float& OutOrbitYaw, float& OutOrbitZoom) const override;
+	UNREALED_API virtual void GetViewMatrixParameters(const float InFOVDegrees, FVector& OutOrigin, float& OutOrbitPitch, float& OutOrbitYaw, float& OutOrbitZoom) const override;
 
 protected:
 	/** The static mesh actor used to display all material thumbnails */
@@ -92,24 +92,24 @@ protected:
 	bool bForcePlaneThumbnail;
 };
 
-class UNREALED_API FSkeletalMeshThumbnailScene : public FThumbnailPreviewScene
+class FSkeletalMeshThumbnailScene : public FThumbnailPreviewScene
 {
 public:
 	/** Constructor */
-	FSkeletalMeshThumbnailScene();
+	UNREALED_API FSkeletalMeshThumbnailScene();
 
 	/** Sets the skeletal mesh to use in the next CreateView() */
-	void SetSkeletalMesh(class USkeletalMesh* InSkeletalMesh);
+	UNREALED_API void SetSkeletalMesh(class USkeletalMesh* InSkeletalMesh);
 
 	/** Set whether to draw debug skeleton */
-	void SetDrawDebugSkeleton(bool bInDrawDebugSkeleton, const FLinearColor& InSkeletonColor);
+	UNREALED_API void SetDrawDebugSkeleton(bool bInDrawDebugSkeleton, const FLinearColor& InSkeletonColor);
 	
 	/** Returns the preview actor within the scene */
 	class ASkeletalMeshActor* GetPreviewActor() { return PreviewActor; }
 
 protected:
 	// FThumbnailPreviewScene implementation
-	virtual void GetViewMatrixParameters(const float InFOVDegrees, FVector& OutOrigin, float& OutOrbitPitch, float& OutOrbitYaw, float& OutOrbitZoom) const override;
+	UNREALED_API virtual void GetViewMatrixParameters(const float InFOVDegrees, FVector& OutOrigin, float& OutOrbitPitch, float& OutOrbitYaw, float& OutOrbitZoom) const override;
 
 private:
 	/** The skeletal mesh actor used to display all skeletal mesh thumbnails */
@@ -122,21 +122,21 @@ private:
 	bool bDrawDebugSkeleton = false;
 };
 
-class UNREALED_API FStaticMeshThumbnailScene : public FThumbnailPreviewScene
+class FStaticMeshThumbnailScene : public FThumbnailPreviewScene
 {
 public:
 	/** Constructor */
-	FStaticMeshThumbnailScene();
+	UNREALED_API FStaticMeshThumbnailScene();
 
 	/** Sets the static mesh to use in the next CreateView() */
-	void SetStaticMesh(class UStaticMesh* StaticMesh);
+	UNREALED_API void SetStaticMesh(class UStaticMesh* StaticMesh);
 
 	/** Sets override materials for the static mesh  */
-	void SetOverrideMaterials(const TArray<class UMaterialInterface*>& OverrideMaterials);
+	UNREALED_API void SetOverrideMaterials(const TArray<class UMaterialInterface*>& OverrideMaterials);
 
 protected:
 	// FThumbnailPreviewScene implementation
-	virtual void GetViewMatrixParameters(const float InFOVDegrees, FVector& OutOrigin, float& OutOrbitPitch, float& OutOrbitYaw, float& OutOrbitZoom) const override;
+	UNREALED_API virtual void GetViewMatrixParameters(const float InFOVDegrees, FVector& OutOrigin, float& OutOrbitPitch, float& OutOrbitYaw, float& OutOrbitZoom) const override;
 
 private:
 	/** The static mesh actor used to display all static mesh thumbnails */
@@ -149,21 +149,21 @@ class AAnimationThumbnailSkeletalMeshActor : public ASkeletalMeshActor
 	GENERATED_UCLASS_BODY()
 };
 
-class UNREALED_API FAnimationSequenceThumbnailScene : public FThumbnailPreviewScene
+class FAnimationSequenceThumbnailScene : public FThumbnailPreviewScene
 {
 public:
 	/** Constructor */
-	FAnimationSequenceThumbnailScene();
+	UNREALED_API FAnimationSequenceThumbnailScene();
 
 	/** Sets the animation to use in the next CreateView() */
-	bool SetAnimation(class UAnimSequenceBase* InAnimation);
+	UNREALED_API bool SetAnimation(class UAnimSequenceBase* InAnimation);
 
 protected:
 	// FThumbnailPreviewScene implementation
-	virtual void GetViewMatrixParameters(const float InFOVDegrees, FVector& OutOrigin, float& OutOrbitPitch, float& OutOrbitYaw, float& OutOrbitZoom) const override;
+	UNREALED_API virtual void GetViewMatrixParameters(const float InFOVDegrees, FVector& OutOrigin, float& OutOrbitPitch, float& OutOrbitYaw, float& OutOrbitZoom) const override;
 	
 	//Clean up the children of this component
-	void CleanupComponentChildren(USceneComponent* Component);
+	UNREALED_API void CleanupComponentChildren(USceneComponent* Component);
 
 private:
 	/** The skeletal mesh actor used to display all animation thumbnails */
@@ -173,21 +173,21 @@ private:
 	class UAnimSequenceBase* PreviewAnimation;
 };
 
-class UNREALED_API FBlendSpaceThumbnailScene : public FThumbnailPreviewScene
+class FBlendSpaceThumbnailScene : public FThumbnailPreviewScene
 {
 public:
 	/** Constructor */
-	FBlendSpaceThumbnailScene();
+	UNREALED_API FBlendSpaceThumbnailScene();
 
 	/** Sets the animation to use in the next CreateView() */
-	bool SetBlendSpace(class UBlendSpace* InBlendSpace);
+	UNREALED_API bool SetBlendSpace(class UBlendSpace* InBlendSpace);
 
 protected:
 	// FThumbnailPreviewScene implementation
-	virtual void GetViewMatrixParameters(const float InFOVDegrees, FVector& OutOrigin, float& OutOrbitPitch, float& OutOrbitYaw, float& OutOrbitZoom) const override;
+	UNREALED_API virtual void GetViewMatrixParameters(const float InFOVDegrees, FVector& OutOrigin, float& OutOrbitPitch, float& OutOrbitYaw, float& OutOrbitZoom) const override;
 
 	//Clean up the children of this component
-	void CleanupComponentChildren(USceneComponent* Component);
+	UNREALED_API void CleanupComponentChildren(USceneComponent* Component);
 
 private:
 	/** The skeletal mesh actor used to display all animation thumbnails */
@@ -197,21 +197,21 @@ private:
 	class UBlendSpace* PreviewAnimation;
 };
 
-class UNREALED_API FAnimBlueprintThumbnailScene : public FThumbnailPreviewScene
+class FAnimBlueprintThumbnailScene : public FThumbnailPreviewScene
 {
 public:
 	/** Constructor */
-	FAnimBlueprintThumbnailScene();
+	UNREALED_API FAnimBlueprintThumbnailScene();
 
 	/** Sets the animation blueprint to use in the next CreateView() */
-	bool SetAnimBlueprint(class UAnimBlueprint* InBlueprint);
+	UNREALED_API bool SetAnimBlueprint(class UAnimBlueprint* InBlueprint);
 
 protected:
 	// FThumbnailPreviewScene implementation
-	virtual void GetViewMatrixParameters(const float InFOVDegrees, FVector& OutOrigin, float& OutOrbitPitch, float& OutOrbitYaw, float& OutOrbitZoom) const override;
+	UNREALED_API virtual void GetViewMatrixParameters(const float InFOVDegrees, FVector& OutOrigin, float& OutOrbitPitch, float& OutOrbitYaw, float& OutOrbitZoom) const override;
 
 	//Clean up the children of this component
-	void CleanupComponentChildren(USceneComponent* Component);
+	UNREALED_API void CleanupComponentChildren(USceneComponent* Component);
 
 private:
 	/** The skeletal mesh actor used to display all animation thumbnails */
@@ -221,18 +221,18 @@ private:
 	class UAnimBlueprint* PreviewBlueprint;
 };
 
-class UNREALED_API FPhysicsAssetThumbnailScene : public FThumbnailPreviewScene
+class FPhysicsAssetThumbnailScene : public FThumbnailPreviewScene
 {
 public:
 	/** Constructor */
-	FPhysicsAssetThumbnailScene();
+	UNREALED_API FPhysicsAssetThumbnailScene();
 
 	/** Sets the skeletal mesh to use in the next CreateView() */
-	void SetPhysicsAsset(class UPhysicsAsset* InPhysicsAsset);
+	UNREALED_API void SetPhysicsAsset(class UPhysicsAsset* InPhysicsAsset);
 
 protected:
 	// FThumbnailPreviewScene implementation
-	virtual void GetViewMatrixParameters(const float InFOVDegrees, FVector& OutOrigin, float& OutOrbitPitch, float& OutOrbitYaw, float& OutOrbitZoom) const override;
+	UNREALED_API virtual void GetViewMatrixParameters(const float InFOVDegrees, FVector& OutOrigin, float& OutOrbitPitch, float& OutOrbitYaw, float& OutOrbitZoom) const override;
 
 private:
 	/** The skeletal mesh actor used to display all physics asset thumbnails */
@@ -241,70 +241,70 @@ private:
 
 class UActorComponent;
 
-class UNREALED_API FClassActorThumbnailScene : public FThumbnailPreviewScene
+class FClassActorThumbnailScene : public FThumbnailPreviewScene
 {
 public:
 
-	FClassActorThumbnailScene();
+	UNREALED_API FClassActorThumbnailScene();
 
 	/** Returns true if this component can be visualized */
-	static bool IsValidComponentForVisualization(const UActorComponent* Component);
+	static UNREALED_API bool IsValidComponentForVisualization(const UActorComponent* Component);
 
 protected:
 	// FThumbnailPreviewScene implementation
-	virtual void GetViewMatrixParameters(const float InFOVDegrees, FVector& OutOrigin, float& OutOrbitPitch, float& OutOrbitYaw, float& OutOrbitZoom) const override;
+	UNREALED_API virtual void GetViewMatrixParameters(const float InFOVDegrees, FVector& OutOrigin, float& OutOrbitPitch, float& OutOrbitYaw, float& OutOrbitZoom) const override;
 
 	/** Sets the object (class or blueprint) used in the next CreateView() */
-	void SpawnPreviewActor(class UClass* Obj);
+	UNREALED_API void SpawnPreviewActor(class UClass* Obj);
 
 	/** Get the scene thumbnail info to use for the object currently being rendered */
 	virtual USceneThumbnailInfo* GetSceneThumbnailInfo(const float TargetDistance) const = 0;
 
-	FBoxSphereBounds GetPreviewActorBounds() const;
+	UNREALED_API FBoxSphereBounds GetPreviewActorBounds() const;
 
 private:
 
 	/** Clears out any stale actors in this scene if PreviewActor enters a stale state */
-	void ClearStaleActors();
+	UNREALED_API void ClearStaleActors();
 
 	int32 NumStartingActors;
 	TWeakObjectPtr<class AActor> PreviewActor;
 };
 
-class UNREALED_API FBlueprintThumbnailScene : public FClassActorThumbnailScene
+class FBlueprintThumbnailScene : public FClassActorThumbnailScene
 {
 public:
 	/** Constructor/Destructor */
-	FBlueprintThumbnailScene();
+	UNREALED_API FBlueprintThumbnailScene();
 
 	/** Sets the static mesh to use in the next CreateView() */
-	void SetBlueprint(class UBlueprint* Blueprint);
+	UNREALED_API void SetBlueprint(class UBlueprint* Blueprint);
 
 	/** Refreshes components for the specified blueprint */
-	void BlueprintChanged(class UBlueprint* Blueprint);
+	UNREALED_API void BlueprintChanged(class UBlueprint* Blueprint);
 
 protected:
 
 	/** Get the scene thumbnail info to use for the object currently being rendered */
-	virtual USceneThumbnailInfo* GetSceneThumbnailInfo(const float TargetDistance) const override;
+	UNREALED_API virtual USceneThumbnailInfo* GetSceneThumbnailInfo(const float TargetDistance) const override;
 
 private:
 	/** The blueprint that is currently being rendered. NULL when not rendering. */
 	TWeakObjectPtr<class UBlueprint> CurrentBlueprint;
 };
 
-class UNREALED_API FClassThumbnailScene : public FClassActorThumbnailScene
+class FClassThumbnailScene : public FClassActorThumbnailScene
 {
 public:
 	/** Constructor/Destructor */
-	FClassThumbnailScene();
+	UNREALED_API FClassThumbnailScene();
 
 	/** Sets the class use in the next CreateView() */
-	void SetClass(class UClass* Class);
+	UNREALED_API void SetClass(class UClass* Class);
 
 protected:
 	/** Get the scene thumbnail info to use for the object currently being rendered */
-	virtual USceneThumbnailInfo* GetSceneThumbnailInfo(const float TargetDistance) const override;
+	UNREALED_API virtual USceneThumbnailInfo* GetSceneThumbnailInfo(const float TargetDistance) const override;
 
 private:
 	/** The class that is currently being rendered. NULL when not rendering. */
diff --git a/Engine/Source/Editor/UnrealEd/Public/TickableEditorObject.h b/Engine/Source/Editor/UnrealEd/Public/TickableEditorObject.h
--- a/Engine/Source/Editor/UnrealEd/Public/TickableEditorObject.h
+++ b/Engine/Source/Editor/UnrealEd/Public/TickableEditorObject.h
@@ -17,7 +17,7 @@
  * If you need a class that can tick in both the Editor and at Runtime then use FTickableGameObject
  * instead, overriding the IsTickableInEditor() function instead.
  */
-class UNREALED_API FTickableEditorObject : public FTickableObjectBase
+class FTickableEditorObject : public FTickableObjectBase
 {
 public:
 
@@ -112,17 +112,17 @@ private:
 	friend class TTickableObjectsCollection;
 
 	/** True if collection of tickable objects is still intact. */
-	static bool bCollectionIntact;
+	static UNREALED_API bool bCollectionIntact;
 	/** True if currently ticking of tickable editor objects. */
-	static bool bIsTickingObjects;
+	static UNREALED_API bool bIsTickingObjects;
 
 	/** Set if we are in the Tick function for an editor tickable object */
-	static FTickableObjectBase* ObjectBeingTicked;
+	static UNREALED_API FTickableObjectBase* ObjectBeingTicked;
 
 
-	static TArray<FTickableObjectEntry>& GetTickableObjects();
+	static UNREALED_API TArray<FTickableObjectEntry>& GetTickableObjects();
 
-	static TArray<FTickableEditorObject*>& GetPendingTickableObjects();
+	static UNREALED_API TArray<FTickableEditorObject*>& GetPendingTickableObjects();
 };
 
 /**
@@ -132,7 +132,7 @@ private:
  * it should dual-inherit from both FTickableCookObject and FTickableEditorObject.
  * TODO: Reduce duplication between  FTickableCookObject and FTickableEditorObject.
  */
-class UNREALED_API FTickableCookObject : public FTickableObjectBase
+class FTickableCookObject : public FTickableObjectBase
 {
 public:
 
@@ -230,15 +230,15 @@ private:
 	friend class TTickableObjectsCollection;
 
 	/** True if collection of tickable objects is still intact. */
-	static bool bCollectionIntact;
+	static UNREALED_API bool bCollectionIntact;
 	/** True if currently ticking of tickable editor objects. */
-	static bool bIsTickingObjects;
+	static UNREALED_API bool bIsTickingObjects;
 
 	/** Set if we are in the Tick function for an editor tickable object */
-	static FTickableCookObject* ObjectBeingTicked;
+	static UNREALED_API FTickableCookObject* ObjectBeingTicked;
 
 
-	static TArray<FTickableObjectEntry>& GetTickableObjects();
+	static UNREALED_API TArray<FTickableObjectEntry>& GetTickableObjects();
 
-	static TArray<FTickableCookObject*>& GetPendingTickableObjects();
+	static UNREALED_API TArray<FTickableCookObject*>& GetPendingTickableObjects();
 };
diff --git a/Engine/Source/Editor/UnrealEd/Public/Toolkits/AssetEditorToolkit.h b/Engine/Source/Editor/UnrealEd/Public/Toolkits/AssetEditorToolkit.h
--- a/Engine/Source/Editor/UnrealEd/Public/Toolkits/AssetEditorToolkit.h
+++ b/Engine/Source/Editor/UnrealEd/Public/Toolkits/AssetEditorToolkit.h
@@ -48,7 +48,7 @@ DECLARE_DELEGATE_RetVal_TwoParams(TSharedRef<FExtender>, FAssetEditorExtender, c
 * Extensibility managers simply keep a series of FExtenders for a single menu/toolbar/anything
 * It is here to keep a standardized approach to editor extensibility among modules
 */
-class UNREALED_API FExtensibilityManager
+class FExtensibilityManager
 {
 public:
 	virtual ~FExtensibilityManager() {}
@@ -61,9 +61,9 @@ public:
 	virtual TArray<FAssetEditorExtender>& GetExtenderDelegates() { return ExtenderDelegates; }
 
 	/** Gets all extenders, consolidated, for use by the editor to be extended */
-	virtual TSharedPtr<FExtender> GetAllExtenders();
+	UNREALED_API virtual TSharedPtr<FExtender> GetAllExtenders();
 	/** Gets all extenders and asset editor extenders from delegates consolidated */
-	virtual TSharedPtr<FExtender> GetAllExtenders(const TSharedRef<FUICommandList>& CommandList, const TArray<UObject*>& ContextSensitiveObjects);
+	UNREALED_API virtual TSharedPtr<FExtender> GetAllExtenders(const TSharedRef<FUICommandList>& CommandList, const TArray<UObject*>& ContextSensitiveObjects);
 
 private:
 	/** A list of extenders the editor will use */
@@ -89,7 +89,7 @@ public:
 /**
  * Base class for toolkits that are used for asset editing (abstract)
  */
-class UNREALED_API FAssetEditorToolkit
+class FAssetEditorToolkit
 	: public IAssetEditorInstance
 	, public FBaseToolkit
 	, public TSharedFromThis<FAssetEditorToolkit>
@@ -99,7 +99,7 @@ public:
 
 
 	/** Default constructor */
-	FAssetEditorToolkit();
+	UNREALED_API FAssetEditorToolkit();
 
 	/**
 	 * Initializes this asset editor.  Called immediately after construction.
@@ -115,86 +115,86 @@ public:
 	 * @param	bInIsToolbarFocusable	Whether the buttons on the default toolbar can receive keyboard focus
 	 * @param	bUseSmallToolbarIcons	Whether the buttons on the default toolbar use the small icons
 	 */
-	void InitAssetEditor(const EToolkitMode::Type Mode, const TSharedPtr<IToolkitHost>& InitToolkitHost, const FName AppIdentifier, const TSharedRef<FTabManager::FLayout>& StandaloneDefaultLayout, const bool bCreateDefaultStandaloneMenu, const bool bCreateDefaultToolbar, const TArray<UObject*>& ObjectsToEdit, const bool bInIsToolbarFocusable = false, const bool bInUseSmallToolbarIcons = false);
-	void InitAssetEditor(const EToolkitMode::Type Mode, const TSharedPtr<IToolkitHost>& InitToolkitHost, const FName AppIdentifier, const TSharedRef<FTabManager::FLayout>& StandaloneDefaultLayout, const bool bCreateDefaultStandaloneMenu, const bool bCreateDefaultToolbar, UObject* ObjectToEdit, const bool bInIsToolbarFocusable = false, const bool bInUseSmallToolbarIcons = false);
+	UNREALED_API void InitAssetEditor(const EToolkitMode::Type Mode, const TSharedPtr<IToolkitHost>& InitToolkitHost, const FName AppIdentifier, const TSharedRef<FTabManager::FLayout>& StandaloneDefaultLayout, const bool bCreateDefaultStandaloneMenu, const bool bCreateDefaultToolbar, const TArray<UObject*>& ObjectsToEdit, const bool bInIsToolbarFocusable = false, const bool bInUseSmallToolbarIcons = false);
+	UNREALED_API void InitAssetEditor(const EToolkitMode::Type Mode, const TSharedPtr<IToolkitHost>& InitToolkitHost, const FName AppIdentifier, const TSharedRef<FTabManager::FLayout>& StandaloneDefaultLayout, const bool bCreateDefaultStandaloneMenu, const bool bCreateDefaultToolbar, UObject* ObjectToEdit, const bool bInIsToolbarFocusable = false, const bool bInUseSmallToolbarIcons = false);
 
 	FAssetEditorToolkit(const FAssetEditorToolkit&) = delete;
 	FAssetEditorToolkit& operator=(const FAssetEditorToolkit&) = delete;
 	
 	/** Virtual destructor, so that we can clean up our app when destroyed */
-	virtual ~FAssetEditorToolkit();
+	UNREALED_API virtual ~FAssetEditorToolkit();
 
 	/** IToolkit interface */
-	virtual void RegisterTabSpawners(const TSharedRef<FTabManager>& TabManager) override;
-	virtual void UnregisterTabSpawners(const TSharedRef<FTabManager>& TabManager) override;
-	virtual bool IsAssetEditor() const override;
-	virtual const TArray< UObject* >* GetObjectsCurrentlyBeingEdited() const override;
+	UNREALED_API virtual void RegisterTabSpawners(const TSharedRef<FTabManager>& TabManager) override;
+	UNREALED_API virtual void UnregisterTabSpawners(const TSharedRef<FTabManager>& TabManager) override;
+	UNREALED_API virtual bool IsAssetEditor() const override;
+	UNREALED_API virtual const TArray< UObject* >* GetObjectsCurrentlyBeingEdited() const override;
 	virtual FName GetToolkitFName() const override = 0;				// Must implement in derived class!
 	virtual FText GetBaseToolkitName() const override = 0;			// Must implement in derived class!
-	virtual FText GetToolkitName() const override;
-	virtual FText GetTabSuffix() const override;
-	virtual FText GetToolkitToolTipText() const override;
+	UNREALED_API virtual FText GetToolkitName() const override;
+	UNREALED_API virtual FText GetTabSuffix() const override;
+	UNREALED_API virtual FText GetToolkitToolTipText() const override;
 	virtual FString GetWorldCentricTabPrefix() const override = 0;	// Must implement in derived class!
-	virtual FEditorModeTools& GetEditorModeManager() const final;
+	UNREALED_API virtual FEditorModeTools& GetEditorModeManager() const final;
 
 	/** IAssetEditorInstance interface */
-	virtual FName GetEditorName() const override;
-	virtual void FocusWindow(UObject* ObjectToFocusOn = nullptr) override;
+	UNREALED_API virtual FName GetEditorName() const override;
+	UNREALED_API virtual void FocusWindow(UObject* ObjectToFocusOn = nullptr) override;
 	UE_DEPRECATED(5.3, "Use CloseWindow that takes in an EAssetEditorCloseReason instead")
-	virtual bool CloseWindow() override;
-	virtual bool CloseWindow(EAssetEditorCloseReason InCloseReason) override;
+	UNREALED_API virtual bool CloseWindow() override;
+	UNREALED_API virtual bool CloseWindow(EAssetEditorCloseReason InCloseReason) override;
 	virtual bool IsPrimaryEditor() const override { return true; };
-	virtual void InvokeTab(const FTabId& TabId) override;
-	virtual FName GetEditingAssetTypeName() const override;
-	virtual TSharedPtr<FTabManager> GetAssociatedTabManager() override;
-	virtual double GetLastActivationTime() override;
-	virtual void RemoveEditingAsset(UObject* Asset) override;
+	UNREALED_API virtual void InvokeTab(const FTabId& TabId) override;
+	UNREALED_API virtual FName GetEditingAssetTypeName() const override;
+	UNREALED_API virtual TSharedPtr<FTabManager> GetAssociatedTabManager() override;
+	UNREALED_API virtual double GetLastActivationTime() override;
+	UNREALED_API virtual void RemoveEditingAsset(UObject* Asset) override;
 
 	/**
 	 * Fills in the supplied menu with commands for working with this asset file
 	 *
 	 * @param	MenuBuilder		The menu to add commands to
 	 */
-	void FillDefaultFileMenuOpenCommands(FToolMenuSection& InSection);
+	UNREALED_API void FillDefaultFileMenuOpenCommands(FToolMenuSection& InSection);
 
 	/**
 	 * Fills in the supplied menu with commands for working with this asset file
 	 *
 	 * @param	MenuBuilder		The menu to add commands to
 	 */
-	void FillDefaultFileMenuCommands(FToolMenuSection& InSection);
+	UNREALED_API void FillDefaultFileMenuCommands(FToolMenuSection& InSection);
 
 	/**
 	 * Fills in the supplied menu with commands for modifying this asset that are generally common to most asset editors
 	 *
 	 * @param	MenuBuilder		The menu to add commands to
 	 */
-	void FillDefaultAssetMenuCommands(FToolMenuSection& InSection);
+	UNREALED_API void FillDefaultAssetMenuCommands(FToolMenuSection& InSection);
 
 	/**
 	 * Fills in the supplied menu with commands for the help menu
 	 *
 	 * @param	MenuBuilder		The menu to add commands to
 	 */
-	void FillDefaultHelpMenuCommands(FToolMenuSection& InSection);
+	UNREALED_API void FillDefaultHelpMenuCommands(FToolMenuSection& InSection);
 
 	/** @return	For standalone asset editing tool-kits, returns the toolkit host that was last hosting this asset editor before it was switched to standalone mode (if it's still valid.)  Returns null if these conditions aren't met. */
-	TSharedPtr< IToolkitHost > GetPreviousWorldCentricToolkitHost();
+	UNREALED_API TSharedPtr< IToolkitHost > GetPreviousWorldCentricToolkitHost();
 
 	/**
 	 * Static: Used internally to set the world-centric toolkit host for a newly-created standalone asset editing toolkit
 	 *
 	 * @param ToolkitHost	The tool-kit to use if/when this toolkit is switched back to world-centric mode
 	 */
-	static void SetPreviousWorldCentricToolkitHostForNewAssetEditor(TSharedRef< IToolkitHost > ToolkitHost);
+	static UNREALED_API void SetPreviousWorldCentricToolkitHostForNewAssetEditor(TSharedRef< IToolkitHost > ToolkitHost);
 
 	/**
 	 * Registers a drawer for the asset editor status bar
 	 */
-	void RegisterDrawer(struct FWidgetDrawerConfig&& Drawer, int32 SlotIndex = INDEX_NONE);
+	UNREALED_API void RegisterDrawer(struct FWidgetDrawerConfig&& Drawer, int32 SlotIndex = INDEX_NONE);
 
 	/** Applies the passed in layout (or the saved user-modified version if available).  Must be called after InitAssetEditor. */
-	void RestoreFromLayout(const TSharedRef<FTabManager::FLayout>& NewLayout);
+	UNREALED_API void RestoreFromLayout(const TSharedRef<FTabManager::FLayout>& NewLayout);
 
 	/** @return Returns this asset editor's tab manager object.  May be nullptr for non-standalone toolkits */
 	TSharedPtr<FTabManager> GetTabManager()
@@ -203,20 +203,20 @@ public:
 	}
 
 	/** Registers default tool bar */
-	static void RegisterDefaultToolBar();
+	static UNREALED_API void RegisterDefaultToolBar();
 
 	/** Makes a default asset editing toolbar */
-	void GenerateToolbar();
+	UNREALED_API void GenerateToolbar();
 
 	/** Regenerates the menubar and toolbar widgets */
-	void RegenerateMenusAndToolbars();
+	UNREALED_API void RegenerateMenusAndToolbars();
 
 	/** Get name used by tool menu */
-	virtual FName GetToolMenuToolbarName(FName& OutParentName) const;
-	virtual void InitToolMenuContext(FToolMenuContext& MenuContext);
-	FName GetToolMenuToolbarName() const;
-	FName GetToolMenuAppName() const;
-	FName GetToolMenuName() const;
+	UNREALED_API virtual FName GetToolMenuToolbarName(FName& OutParentName) const;
+	UNREALED_API virtual void InitToolMenuContext(FToolMenuContext& MenuContext);
+	UNREALED_API FName GetToolMenuToolbarName() const;
+	UNREALED_API FName GetToolMenuAppName() const;
+	UNREALED_API FName GetToolMenuName() const;
 
 	/** Called at the end of RegenerateMenusAndToolbars() */
 	virtual void PostRegenerateMenusAndToolbars() { }
@@ -236,14 +236,14 @@ public:
 	virtual void RemoveViewportOverlayWidget(TSharedRef<SWidget> InViewportOverlayWidget) {}
 	
 	/** Adds or removes extenders to the default menu or the toolbar menu this asset editor */
-	void AddMenuExtender(TSharedPtr<FExtender> Extender);
-	void RemoveMenuExtender(TSharedPtr<FExtender> Extender);
-	void AddToolbarExtender(TSharedPtr<FExtender> Extender);
-	void RemoveToolbarExtender(TSharedPtr<FExtender> Extender);
+	UNREALED_API void AddMenuExtender(TSharedPtr<FExtender> Extender);
+	UNREALED_API void RemoveMenuExtender(TSharedPtr<FExtender> Extender);
+	UNREALED_API void AddToolbarExtender(TSharedPtr<FExtender> Extender);
+	UNREALED_API void RemoveToolbarExtender(TSharedPtr<FExtender> Extender);
 
 	/** Returns the default extensibility managers, these are applied for all asset types */
-	static TSharedPtr<FExtensibilityManager> GetSharedMenuExtensibilityManager();
-	static TSharedPtr<FExtensibilityManager> GetSharedToolBarExtensibilityManager();
+	static UNREALED_API TSharedPtr<FExtensibilityManager> GetSharedMenuExtensibilityManager();
+	static UNREALED_API TSharedPtr<FExtensibilityManager> GetSharedToolBarExtensibilityManager();
 
 	/** Override this to manually specify the tab that cannot be closed for your asset editor, it is the first tab
 	 *  registered otherwise
@@ -255,36 +255,36 @@ public:
 	 *
 	 * @param Widget The widget to use as the overlay
 	 */
-	void SetMenuOverlay(TSharedRef<SWidget> Widget);
+	UNREALED_API void SetMenuOverlay(TSharedRef<SWidget> Widget);
 
 	/** Adds or removes widgets from the default toolbar in this asset editor */
-	void AddToolbarWidget(TSharedRef<SWidget> Widget);
-	void RemoveAllToolbarWidgets();
+	UNREALED_API void AddToolbarWidget(TSharedRef<SWidget> Widget);
+	UNREALED_API void RemoveAllToolbarWidgets();
 
 	/** True if this actually is editing an asset */
-	bool IsActuallyAnAsset() const;
+	UNREALED_API bool IsActuallyAnAsset() const;
 
 	/**
 	 * Gets the text to display in a toolkit titlebar for an object
 	 * @param	InObject	The object we want a description of
 	 * @return a formatted description of the object state (e.g. "MyObject*")
 	 */
-	static FText GetLabelForObject(const UObject* InObject);
+	static UNREALED_API FText GetLabelForObject(const UObject* InObject);
 
 	/**
 	 * Gets the text to display in a toolkit tooltip for an object
 	 * @param	InObject	The object we want a description of
 	 * @return a formatted description of the object
 	 */
-	static FText GetToolTipTextForObject(const UObject* InObject);
+	static UNREALED_API FText GetToolTipTextForObject(const UObject* InObject);
 
 	/** Get the asset editor mode manager we are using */
 	UE_DEPRECATED(4.26, "Use GetEditorModeManager instead.")
-	class FAssetEditorModeManager* GetAssetEditorModeManager() const;
+	UNREALED_API class FAssetEditorModeManager* GetAssetEditorModeManager() const;
 
 	/** Set the asset editor mode manager we are using */
 	UE_DEPRECATED(4.26, "Override CreateEditorModeManager on the toolkit class instead.")
-	void SetAssetEditorModeManager(class FAssetEditorModeManager* InModeManager);
+	UNREALED_API void SetAssetEditorModeManager(class FAssetEditorModeManager* InModeManager);
 
 	virtual void AddGraphEditorPinActionsToContextMenu(FToolMenuSection& InSection) const {};
 
@@ -301,35 +301,35 @@ protected:
 	bool HasEditingObject() const { return !EditingObjects.IsEmpty(); }
 
 	/**	Returns the single object currently being edited. Asserts if currently editing no object or multiple objects */
-	UObject* GetEditingObject() const;
+	UNREALED_API UObject* GetEditingObject() const;
 
 	/**	Returns an array of all the objects currently being edited. Asserts if editing no objects */
-	const TArray< UObject* >& GetEditingObjects() const;
-	TArray<TObjectPtr<UObject>>& GetEditingObjectPtrs();
+	UNREALED_API const TArray< UObject* >& GetEditingObjects() const;
+	UNREALED_API TArray<TObjectPtr<UObject>>& GetEditingObjectPtrs();
 
 	/** Generate the toolbar for common asset actions like Save*/
-	UToolMenu* GenerateCommonActionsToolbar(FToolMenuContext& MenuContext);
+	UNREALED_API UToolMenu* GenerateCommonActionsToolbar(FToolMenuContext& MenuContext);
 
 	/** Get the collection of edited objects that can be saved. */
-	virtual void GetSaveableObjects(TArray<UObject*>& OutObjects) const;
+	UNREALED_API virtual void GetSaveableObjects(TArray<UObject*>& OutObjects) const;
 
 	/** Adds an item to the Editing Objects list */
-	virtual void AddEditingObject(UObject* Object);
+	UNREALED_API virtual void AddEditingObject(UObject* Object);
 
 	/** Removes an item from the Editing Objects list */
-	virtual void RemoveEditingObject(UObject* Object);
+	UNREALED_API virtual void RemoveEditingObject(UObject* Object);
 
 	/** Called to test if "Save" should be enabled for this asset */
 	virtual bool CanSaveAsset() const { return true; }
 
 	/** Called when "Save" is clicked for this asset */
-	virtual void SaveAsset_Execute();
+	UNREALED_API virtual void SaveAsset_Execute();
 
 	/** Called to test if "Save As" should be enabled for this asset */
 	virtual bool CanSaveAssetAs() const { return true; }
 
 	/** Called when "Save As" is clicked for this asset */
-	virtual void SaveAssetAs_Execute();
+	UNREALED_API virtual void SaveAssetAs_Execute();
 
 	/** Called to test if "Find in Content Browser" should be enabled for this asset */
 	virtual bool CanFindInContentBrowser() const { return true; }
@@ -338,31 +338,31 @@ protected:
 	virtual bool IsFindInContentBrowserButtonVisible() const { return true; }
 
 	/** Called when "Find in Content Browser" is clicked for this asset */
-	virtual void FindInContentBrowser_Execute();
+	UNREALED_API virtual void FindInContentBrowser_Execute();
 
 	/** Called when "Browse Documentation" is clicked for this asset */
-	virtual void BrowseDocumentation_Execute() const;
+	UNREALED_API virtual void BrowseDocumentation_Execute() const;
 
 	/** @return the documentation location for this editor */
-	virtual FString GetDocumentationLink() const;
+	UNREALED_API virtual FString GetDocumentationLink() const;
 
 	/** Called to check to see if there's an asset capable of being reimported */
-	virtual bool CanReimport() const;
-	virtual bool CanReimport(UObject* EditingObject) const;
+	UNREALED_API virtual bool CanReimport() const;
+	UNREALED_API virtual bool CanReimport(UObject* EditingObject) const;
 
 	/** Called when "Reimport" is clicked for this asset */
-	virtual void Reimport_Execute();
-	virtual void Reimport_Execute(UObject* EditingObject);
+	UNREALED_API virtual void Reimport_Execute();
+	UNREALED_API virtual void Reimport_Execute(UObject* EditingObject);
 
 	/** Called to determine if the user should be prompted for a new file if one is missing during an asset reload */
-	virtual bool ShouldPromptForNewFilesOnReload(const UObject& object) const;
+	UNREALED_API virtual bool ShouldPromptForNewFilesOnReload(const UObject& object) const;
 
 	/** Called when this toolkit is requested to close. Returns false if closing should be prevented. */
 	UE_DEPRECATED(5.3, "Use OnRequestClose that takes in an EAssetEditorCloseReason instead")
 	virtual bool OnRequestClose() { return true; }
 
 	/** Called when this toolkit is requested to close. Returns false if closing should be prevented. */
-	virtual bool OnRequestClose(EAssetEditorCloseReason InCloseReason);
+	UNREALED_API virtual bool OnRequestClose(EAssetEditorCloseReason InCloseReason);
 
 	/** Called when this toolkit is being closed */
 	virtual void OnClose() {}
@@ -371,22 +371,22 @@ protected:
 	  *
 	  * @param ThisToolkitWeakRef	The toolkit that we want to restart in standalone mode
 	  */
-	static void SwitchToStandaloneEditor_Execute(TWeakPtr< FAssetEditorToolkit > ThisToolkitWeakRef);
+	static UNREALED_API void SwitchToStandaloneEditor_Execute(TWeakPtr< FAssetEditorToolkit > ThisToolkitWeakRef);
 
 	/**
 	  * Static: Called when "Switch to World-Centric Editor" is clicked for the asset editor
 	  *
 	  * @param	ThisToolkitWeakRef			The toolkit that we want to restart in world-centric mode
 	  */
-	static void SwitchToWorldCentricEditor_Execute(TWeakPtr< FAssetEditorToolkit > ThisToolkitWeakRef);
+	static UNREALED_API void SwitchToWorldCentricEditor_Execute(TWeakPtr< FAssetEditorToolkit > ThisToolkitWeakRef);
 
 	/** @return a pointer to the brush to use for the tab icon */
-	virtual const FSlateBrush* GetDefaultTabIcon() const;
+	UNREALED_API virtual const FSlateBrush* GetDefaultTabIcon() const;
 
 	/** @return the color to use for the tab color */
-	virtual FLinearColor GetDefaultTabColor() const;
+	UNREALED_API virtual FLinearColor GetDefaultTabColor() const;
 
-	virtual void CreateEditorModeManager() override;
+	UNREALED_API virtual void CreateEditorModeManager() override;
 
 private:
 	// Callback for persisting the Asset Editor's layout.
@@ -404,14 +404,14 @@ private:
 	 * so that we don't break compatibility with all the asset editors out there that individually 
 	 * implement FGCObject.
 	 */
-	class UNREALED_API FGCEditingObjects : public FGCObject
+	class FGCEditingObjects : public FGCObject
 	{
 	public:
 		FGCEditingObjects(FAssetEditorToolkit& InOwnerToolkit) : OwnerToolkit(InOwnerToolkit) {}
 
 		/** FGCObject interface */
-		virtual void AddReferencedObjects(FReferenceCollector& Collector) override;
-		virtual FString GetReferencerName() const override;
+		UNREALED_API virtual void AddReferencedObjects(FReferenceCollector& Collector) override;
+		UNREALED_API virtual FString GetReferencerName() const override;
 
 	private:
 		FAssetEditorToolkit& OwnerToolkit;
@@ -442,16 +442,16 @@ protected:
 	TSharedPtr<FWorkspaceItem> AssetEditorTabsCategory;
 
 private:
-	static const FName DefaultAssetEditorToolBarName;
+	static UNREALED_API const FName DefaultAssetEditorToolBarName;
 	/** The toolkit standalone host; may be nullptr */
 	TWeakPtr<SStandaloneAssetEditorToolkitHost> StandaloneHost;
 
 	/** Static: World centric toolkit host to use for the next created asset editing toolkit */
-	static TWeakPtr< IToolkitHost > PreviousWorldCentricToolkitHostForNewAssetEditor;
+	static UNREALED_API TWeakPtr< IToolkitHost > PreviousWorldCentricToolkitHostForNewAssetEditor;
 
 	/** The extensibility managers shared by all asset types */
-	static TSharedPtr<FExtensibilityManager> SharedMenuExtensibilityManager;
-	static TSharedPtr<FExtensibilityManager> SharedToolBarExtensibilityManager;
+	static UNREALED_API TSharedPtr<FExtensibilityManager> SharedMenuExtensibilityManager;
+	static UNREALED_API TSharedPtr<FExtensibilityManager> SharedToolBarExtensibilityManager;
 
 	/** The object we're currently editing */
 	// @todo toolkit minor: Currently we don't need to serialize this object reference because the AssetEditorSubsystem is kept in sync (and will always serialize it.)
diff --git a/Engine/Source/Editor/UnrealEd/Public/Toolkits/AssetEditorToolkitMenuContext.h b/Engine/Source/Editor/UnrealEd/Public/Toolkits/AssetEditorToolkitMenuContext.h
--- a/Engine/Source/Editor/UnrealEd/Public/Toolkits/AssetEditorToolkitMenuContext.h
+++ b/Engine/Source/Editor/UnrealEd/Public/Toolkits/AssetEditorToolkitMenuContext.h
@@ -9,15 +9,15 @@
 
 class FAssetEditorToolkit;
 
-UCLASS()
-class UNREALED_API UAssetEditorToolkitMenuContext : public UObject
+UCLASS(MinimalAPI)
+class UAssetEditorToolkitMenuContext : public UObject
 {
 	GENERATED_BODY()
 
 public:
 
 	UFUNCTION(BlueprintCallable, Category = "Tool Menus")
-	TArray<UObject*> GetEditingObjects() const;
+	UNREALED_API TArray<UObject*> GetEditingObjects() const;
 
 	TWeakPtr<FAssetEditorToolkit> Toolkit;
 };
diff --git a/Engine/Source/Editor/UnrealEd/Public/Toolkits/BaseToolkit.h b/Engine/Source/Editor/UnrealEd/Public/Toolkits/BaseToolkit.h
--- a/Engine/Source/Editor/UnrealEd/Public/Toolkits/BaseToolkit.h
+++ b/Engine/Source/Editor/UnrealEd/Public/Toolkits/BaseToolkit.h
@@ -24,29 +24,29 @@ class UToolMenu;
 /**
  * Base class for all toolkits (abstract).
  */
-class UNREALED_API FBaseToolkit
+class FBaseToolkit
 	: public IToolkit
 {
 public:
 
 	/** FBaseToolkit constructor */
-	FBaseToolkit();
+	UNREALED_API FBaseToolkit();
 
 	/** Virtual destructor */
-	virtual ~FBaseToolkit();
+	UNREALED_API virtual ~FBaseToolkit();
 
 public:
 
 	// IToolkit interface
 
-	virtual FName GetToolkitContextFName() const override;
+	UNREALED_API virtual FName GetToolkitContextFName() const override;
 	virtual FText GetTabSuffix() const override { return FText::GetEmpty(); }
-	virtual bool ProcessCommandBindings(const FKeyEvent& InKeyEvent) const override;
-	virtual bool IsHosted() const override;
-	virtual const TSharedRef<IToolkitHost> GetToolkitHost() const override;
-	virtual void BringToolkitToFront() override;
-	virtual TSharedPtr<SWidget> GetInlineContent() const override;
-	virtual bool IsBlueprintEditor() const override;
+	UNREALED_API virtual bool ProcessCommandBindings(const FKeyEvent& InKeyEvent) const override;
+	UNREALED_API virtual bool IsHosted() const override;
+	UNREALED_API virtual const TSharedRef<IToolkitHost> GetToolkitHost() const override;
+	UNREALED_API virtual void BringToolkitToFront() override;
+	UNREALED_API virtual TSharedPtr<SWidget> GetInlineContent() const override;
+	UNREALED_API virtual bool IsBlueprintEditor() const override;
 	virtual TSharedRef<FWorkspaceItem> GetWorkspaceMenuCategory() const override { return WorkspaceMenuCategory.ToSharedRef(); }
 
 	virtual FEditorModeTools& GetEditorModeManager() const = 0;
@@ -54,7 +54,7 @@ public:
 public:
 
 	/** @return	Returns true if this is a world-centric asset editor.  That is, the user is editing the asset inline in a Level Editor app. */
-	bool IsWorldCentricAssetEditor() const;	
+	UNREALED_API bool IsWorldCentricAssetEditor() const;	
 
 	/** @returns Returns our toolkit command list */
 	const TSharedRef<FUICommandList> GetToolkitCommands() const
@@ -66,14 +66,14 @@ protected:
 
 	/** @return Returns the prefix string to use for tabs created for this toolkit.  In world-centric mode, tabs get a
 	    name prefix to make them distinguishable from other tabs */
-	FString GetTabPrefix() const;
+	UNREALED_API FString GetTabPrefix() const;
 
 	/** @return Returns the color to use for tabs created for this toolkit.  In world-centric mode, tabs may be colored to
 	    make them more easy to distinguish compared to other tabs. */
-	FLinearColor GetTabColorScale() const;
+	UNREALED_API FLinearColor GetTabColorScale() const;
 
 	// Creates the Editor mode manager for your class. Default is to create none, for legacy reasons.
-	virtual void CreateEditorModeManager();
+	UNREALED_API virtual void CreateEditorModeManager();
 
 protected:
 
@@ -95,41 +95,41 @@ protected:
  * This FModeToolkit just creates a basic UI panel that allows various InteractiveTools to
  * be initialized, and a DetailsView used to show properties of the active Tool.
  */
-class UNREALED_API FModeToolkit
+class FModeToolkit
 	: public FBaseToolkit
 	, public TSharedFromThis<FModeToolkit>
 {
 public:
 
 	/** Initializes the mode toolkit */
-	virtual void Init(const TSharedPtr<IToolkitHost>& InitToolkitHost);
-	virtual void Init(const TSharedPtr<IToolkitHost>& InitToolkitHost, TWeakObjectPtr<UEdMode> InOwningMode);
-	~FModeToolkit();
+	UNREALED_API virtual void Init(const TSharedPtr<IToolkitHost>& InitToolkitHost);
+	UNREALED_API virtual void Init(const TSharedPtr<IToolkitHost>& InitToolkitHost, TWeakObjectPtr<UEdMode> InOwningMode);
+	UNREALED_API ~FModeToolkit();
 
 public:
-	virtual void SetModeUILayer(const TSharedPtr<FAssetEditorModeUILayer> InLayer) override;
+	UNREALED_API virtual void SetModeUILayer(const TSharedPtr<FAssetEditorModeUILayer> InLayer) override;
 
 	// FBaseToolkit overrides
 
-	virtual void RegisterTabSpawners(const TSharedRef<FTabManager>& TabManager) final;
-	virtual void UnregisterTabSpawners(const TSharedRef<FTabManager>& TabManager) final;
+	UNREALED_API virtual void RegisterTabSpawners(const TSharedRef<FTabManager>& TabManager) final;
+	UNREALED_API virtual void UnregisterTabSpawners(const TSharedRef<FTabManager>& TabManager) final;
 
 	// IToolkit interface
-	virtual FName GetToolkitFName() const override;
-	virtual FText GetBaseToolkitName() const override;
+	UNREALED_API virtual FName GetToolkitFName() const override;
+	UNREALED_API virtual FText GetBaseToolkitName() const override;
 	virtual FText GetToolkitName() const override { return GetBaseToolkitName(); }
 	virtual FText GetToolkitToolTipText() const override { return GetBaseToolkitName(); }
-	virtual FString GetWorldCentricTabPrefix() const override;
-	virtual bool IsAssetEditor() const override;
-	virtual const TArray<UObject*>* GetObjectsCurrentlyBeingEdited() const override;
-	virtual FLinearColor GetWorldCentricTabColorScale() const override;
-	virtual FEditorModeTools& GetEditorModeManager() const final;
-
-	virtual TWeakObjectPtr<UEdMode> GetScriptableEditorMode() const final;
-	virtual TSharedPtr<SWidget> GetInlineContent() const;
-	virtual FEdMode* GetEditorMode() const;
-	virtual	FText GetEditorModeDisplayName() const;
-	virtual FSlateIcon GetEditorModeIcon() const;
+	UNREALED_API virtual FString GetWorldCentricTabPrefix() const override;
+	UNREALED_API virtual bool IsAssetEditor() const override;
+	UNREALED_API virtual const TArray<UObject*>* GetObjectsCurrentlyBeingEdited() const override;
+	UNREALED_API virtual FLinearColor GetWorldCentricTabColorScale() const override;
+	UNREALED_API virtual FEditorModeTools& GetEditorModeManager() const final;
+
+	UNREALED_API virtual TWeakObjectPtr<UEdMode> GetScriptableEditorMode() const final;
+	UNREALED_API virtual TSharedPtr<SWidget> GetInlineContent() const;
+	UNREALED_API virtual FEdMode* GetEditorMode() const;
+	UNREALED_API virtual	FText GetEditorModeDisplayName() const;
+	UNREALED_API virtual FSlateIcon GetEditorModeIcon() const;
 	/** Returns the number of Mode specific tabs in the mode toolbar **/
 	virtual void GetToolPaletteNames(TArray<FName>& PaletteNames) const {}
 
@@ -149,7 +149,7 @@ public:
 	 * @param PaletteIndex      The index of the ToolPalette to build
 	 * @param ToolbarBuilder    The builder to use for given PaletteIndex
 	**/
-	virtual void BuildToolPalette(FName Palette, class FToolBarBuilder& ToolbarBuilder);
+	UNREALED_API virtual void BuildToolPalette(FName Palette, class FToolBarBuilder& ToolbarBuilder);
 
 	virtual FText GetActiveToolDisplayName() const { return FText(); }
 	virtual FText GetActiveToolMessage() const { return FText(); }
@@ -158,10 +158,10 @@ public:
 	FOnPaletteChanged& OnPaletteChanged() { return OnPaletteChangedDelegate; }
 
 	virtual void OnToolPaletteChanged(FName PaletteName) {};
-	void SetCurrentPalette(FName InName);
-	FName GetCurrentPalette() const;
-	void SetModeSettingsObject(UObject* InSettingsObject);
-	virtual void InvokeUI();
+	UNREALED_API void SetCurrentPalette(FName InName);
+	UNREALED_API FName GetCurrentPalette() const;
+	UNREALED_API void SetModeSettingsObject(UObject* InSettingsObject);
+	UNREALED_API virtual void InvokeUI();
 
 	/**
 	 * Override this function to extend the secondary mode toolbar (that appears below the main toolbar) for your mode by
@@ -172,41 +172,41 @@ public:
 	/**
 	 * Force the Mode Toolbar/Palette to be repopulated with the current ToolPaletteNames
 	 */
-	virtual void RebuildModeToolPalette();
+	UNREALED_API virtual void RebuildModeToolPalette();
 
 protected:
-	virtual void OnToolStarted(UInteractiveToolManager* Manager, UInteractiveTool* Tool);
-	virtual void OnToolEnded(UInteractiveToolManager* Manager, UInteractiveTool* Tool);
+	UNREALED_API virtual void OnToolStarted(UInteractiveToolManager* Manager, UInteractiveTool* Tool);
+	UNREALED_API virtual void OnToolEnded(UInteractiveToolManager* Manager, UInteractiveTool* Tool);
 
 	virtual void CustomizeModeDetailsViewArgs(FDetailsViewArgs& ArgsInOut) {}
 	virtual void CustomizeDetailsViewArgs(FDetailsViewArgs& ArgsInOut) {}
 
-	virtual void RequestModeUITabs();
+	UNREALED_API virtual void RequestModeUITabs();
 
-	void OnModeIDChanged(const FEditorModeID& InID, bool bIsEntering);
-	const FEditorModeInfo* GetEditorModeInfo() const;
+	UNREALED_API void OnModeIDChanged(const FEditorModeID& InID, bool bIsEntering);
+	UNREALED_API const FEditorModeInfo* GetEditorModeInfo() const;
 
 	/**
 	 * Whether or not the mode toolbar should be shown.  If any active modes generated a toolbar this method will return true
 	 */
-	bool ShouldShowModeToolbar() const;
-	TSharedRef<SDockTab> CreatePrimaryModePanel(const FSpawnTabArgs& Args);
-	void UpdatePrimaryModePanel();
-	EVisibility GetInlineContentHolderVisibility() const;
-	EVisibility GetNoToolSelectedTextVisibility() const;
+	UNREALED_API bool ShouldShowModeToolbar() const;
+	UNREALED_API TSharedRef<SDockTab> CreatePrimaryModePanel(const FSpawnTabArgs& Args);
+	UNREALED_API void UpdatePrimaryModePanel();
+	UNREALED_API EVisibility GetInlineContentHolderVisibility() const;
+	UNREALED_API EVisibility GetNoToolSelectedTextVisibility() const;
 
 	/**
 	 * Creates the mode toolbar tab if needed
 	 */
-	TSharedRef<SDockTab> MakeModeToolbarTab(const FSpawnTabArgs& Args);
+	UNREALED_API TSharedRef<SDockTab> MakeModeToolbarTab(const FSpawnTabArgs& Args);
 
 	/**
 	 * Creates the entire tool palette widget, override to specify toolbar style
 	 */
-	virtual TSharedRef<SWidget> CreatePaletteWidget(TSharedPtr<FUICommandList> InCommandList, FName InToolbarCustomizationName, FName InPaletteName);
+	UNREALED_API virtual TSharedRef<SWidget> CreatePaletteWidget(TSharedPtr<FUICommandList> InCommandList, FName InToolbarCustomizationName, FName InPaletteName);
 
-	void SpawnOrUpdateModeToolbar();
-	void RebuildModeToolBar();
+	UNREALED_API void SpawnOrUpdateModeToolbar();
+	UNREALED_API void RebuildModeToolBar();
 
 
 protected:
@@ -269,7 +269,7 @@ protected:
 	TSharedPtr<FToolkitBuilder> ToolkitBuilder;
 
 	/** returns whether or not this FModeToolkit has a ToolkitBuilder defined */
-	bool HasToolkitBuilder() const;
+	UNREALED_API bool HasToolkitBuilder() const;
 
 	/** The sections of the toolkit, defined for the ToolkitBuilder (if present) */
 	TSharedPtr<FToolkitSections> ToolkitSections;
diff --git a/Engine/Source/Editor/UnrealEd/Public/Toolkits/FConsoleCommandExecutor.h b/Engine/Source/Editor/UnrealEd/Public/Toolkits/FConsoleCommandExecutor.h
--- a/Engine/Source/Editor/UnrealEd/Public/Toolkits/FConsoleCommandExecutor.h
+++ b/Engine/Source/Editor/UnrealEd/Public/Toolkits/FConsoleCommandExecutor.h
@@ -10,19 +10,19 @@
 /**
 * Executor for Unreal console commands
 */
-class UNREALED_API FConsoleCommandExecutor : public IConsoleCommandExecutor
+class FConsoleCommandExecutor : public IConsoleCommandExecutor
 {
 public:
-	static FName StaticName();
-	virtual FName GetName() const override;
-	virtual FText GetDisplayName() const override;
-	virtual FText GetDescription() const override;
-	virtual FText GetHintText() const override;
-	virtual void GetAutoCompleteSuggestions(const TCHAR* Input, TArray<FString>& Out) override;
-	virtual void GetExecHistory(TArray<FString>& Out) override;
-	virtual bool Exec(const TCHAR* Input) override;
-	virtual bool AllowHotKeyClose() const override;
-	virtual bool AllowMultiLine() const override;
+	static UNREALED_API FName StaticName();
+	UNREALED_API virtual FName GetName() const override;
+	UNREALED_API virtual FText GetDisplayName() const override;
+	UNREALED_API virtual FText GetDescription() const override;
+	UNREALED_API virtual FText GetHintText() const override;
+	UNREALED_API virtual void GetAutoCompleteSuggestions(const TCHAR* Input, TArray<FString>& Out) override;
+	UNREALED_API virtual void GetExecHistory(TArray<FString>& Out) override;
+	UNREALED_API virtual bool Exec(const TCHAR* Input) override;
+	UNREALED_API virtual bool AllowHotKeyClose() const override;
+	UNREALED_API virtual bool AllowMultiLine() const override;
 
 	virtual FInputChord GetHotKey() const override
 	{
@@ -34,5 +34,5 @@ public:
 		return FGlobalEditorCommonCommands::Get().FGlobalEditorCommonCommands::Get().SelectNextConsoleExecutor->GetActiveChord(EMultipleKeyBindingIndex::Primary).Get();
 	}
 private:
-	bool ExecInternal(const TCHAR*) const;
+	UNREALED_API bool ExecInternal(const TCHAR*) const;
 };
diff --git a/Engine/Source/Editor/UnrealEd/Public/Toolkits/GlobalEditorCommonCommands.h b/Engine/Source/Editor/UnrealEd/Public/Toolkits/GlobalEditorCommonCommands.h
--- a/Engine/Source/Editor/UnrealEd/Public/Toolkits/GlobalEditorCommonCommands.h
+++ b/Engine/Source/Editor/UnrealEd/Public/Toolkits/GlobalEditorCommonCommands.h
@@ -13,26 +13,26 @@ class SWindow;
 
 // Global editor common commands
 // Note: There is no real global command concept, so these must still be registered in each editor
-class UNREALED_API FGlobalEditorCommonCommands : public TCommands< FGlobalEditorCommonCommands >
+class FGlobalEditorCommonCommands : public TCommands< FGlobalEditorCommonCommands >
 {
 public:
-	FGlobalEditorCommonCommands();
-	~FGlobalEditorCommonCommands();
+	UNREALED_API FGlobalEditorCommonCommands();
+	UNREALED_API ~FGlobalEditorCommonCommands();
 
-	virtual void RegisterCommands() override;
+	UNREALED_API virtual void RegisterCommands() override;
 
-	static void MapActions(TSharedRef<FUICommandList>& ToolkitCommands);
+	static UNREALED_API void MapActions(TSharedRef<FUICommandList>& ToolkitCommands);
 
 protected:
-	static void OnPressedCtrlTab(TSharedPtr<FUICommandInfo> TriggeringCommand);
-	static void OnSummonedAssetPicker();
-	static void OnSummonedConsoleCommandBox();
-	static void OnOpenContentBrowserDrawer();
-	static void OnOpenOutputLogDrawer();
+	static UNREALED_API void OnPressedCtrlTab(TSharedPtr<FUICommandInfo> TriggeringCommand);
+	static UNREALED_API void OnSummonedAssetPicker();
+	static UNREALED_API void OnSummonedConsoleCommandBox();
+	static UNREALED_API void OnOpenContentBrowserDrawer();
+	static UNREALED_API void OnOpenOutputLogDrawer();
 
-	static TSharedRef<SDockTab> SpawnAssetPicker(const FSpawnTabArgs& InArgs);
+	static UNREALED_API TSharedRef<SDockTab> SpawnAssetPicker(const FSpawnTabArgs& InArgs);
 
-	static TSharedPtr<IMenu> OpenPopupMenu(TSharedRef<SWidget> WindowContents, const FVector2D& PopupDesiredSize);
+	static UNREALED_API TSharedPtr<IMenu> OpenPopupMenu(TSharedRef<SWidget> WindowContents, const FVector2D& PopupDesiredSize);
 public:
 	TSharedPtr<FUICommandInfo> FindInContentBrowser;
 	TSharedPtr<FUICommandInfo> SummonControlTabNavigation;
diff --git a/Engine/Source/Editor/UnrealEd/Public/Toolkits/SimpleAssetEditor.h b/Engine/Source/Editor/UnrealEd/Public/Toolkits/SimpleAssetEditor.h
--- a/Engine/Source/Editor/UnrealEd/Public/Toolkits/SimpleAssetEditor.h
+++ b/Engine/Source/Editor/UnrealEd/Public/Toolkits/SimpleAssetEditor.h
@@ -9,14 +9,14 @@
 
 class IDetailsView;
 
-class UNREALED_API FSimpleAssetEditor : public FAssetEditorToolkit
+class FSimpleAssetEditor : public FAssetEditorToolkit
 {
 public:
 	/** Delegate that, given an array of assets, returns an array of objects to use in the details view of an FSimpleAssetEditor */
 	DECLARE_DELEGATE_RetVal_OneParam(TArray<UObject*>, FGetDetailsViewObjects, const TArray<UObject*>&);
 
-	virtual void RegisterTabSpawners(const TSharedRef<class FTabManager>& TabManager) override;
-	virtual void UnregisterTabSpawners(const TSharedRef<class FTabManager>& TabManager) override;
+	UNREALED_API virtual void RegisterTabSpawners(const TSharedRef<class FTabManager>& TabManager) override;
+	UNREALED_API virtual void UnregisterTabSpawners(const TSharedRef<class FTabManager>& TabManager) override;
 
 
 	/**
@@ -27,57 +27,57 @@ public:
 	 * @param	ObjectsToEdit			The object to edit
 	 * @param	GetDetailsViewObjects	If bound, a delegate to get the array of objects to use in the details view; uses ObjectsToEdit if not bound
 	 */
-	void InitEditor( const EToolkitMode::Type Mode, const TSharedPtr< class IToolkitHost >& InitToolkitHost, const TArray<UObject*>& ObjectsToEdit, FGetDetailsViewObjects GetDetailsViewObjects );
+	UNREALED_API void InitEditor( const EToolkitMode::Type Mode, const TSharedPtr< class IToolkitHost >& InitToolkitHost, const TArray<UObject*>& ObjectsToEdit, FGetDetailsViewObjects GetDetailsViewObjects );
 
 	/** Destructor */
-	virtual ~FSimpleAssetEditor();
+	UNREALED_API virtual ~FSimpleAssetEditor();
 
 	/** IToolkit interface */
-	virtual FName GetToolkitFName() const override;
-	virtual FText GetBaseToolkitName() const override;
-	virtual FText GetToolkitName() const override;
-	virtual FText GetToolkitToolTipText() const override;
-	virtual FString GetWorldCentricTabPrefix() const override;
-	virtual FLinearColor GetWorldCentricTabColorScale() const override;
+	UNREALED_API virtual FName GetToolkitFName() const override;
+	UNREALED_API virtual FText GetBaseToolkitName() const override;
+	UNREALED_API virtual FText GetToolkitName() const override;
+	UNREALED_API virtual FText GetToolkitToolTipText() const override;
+	UNREALED_API virtual FString GetWorldCentricTabPrefix() const override;
+	UNREALED_API virtual FLinearColor GetWorldCentricTabColorScale() const override;
 	virtual bool IsPrimaryEditor() const override { return true; }
 	virtual bool IsSimpleAssetEditor() const override { return true; }
-	virtual FName GetEditingAssetTypeName() const override;
+	UNREALED_API virtual FName GetEditingAssetTypeName() const override;
 
 	/** FAssetEditorToolkit interface */
-	virtual void PostRegenerateMenusAndToolbars() override;
+	UNREALED_API virtual void PostRegenerateMenusAndToolbars() override;
 	
 	/** Used to show or hide certain properties */
-	void SetPropertyVisibilityDelegate(FIsPropertyVisible InVisibilityDelegate);
+	UNREALED_API void SetPropertyVisibilityDelegate(FIsPropertyVisible InVisibilityDelegate);
 	/** Can be used to disable the details view making it read-only */
-	void SetPropertyEditingEnabledDelegate(FIsPropertyEditingEnabled InPropertyEditingDelegate);
+	UNREALED_API void SetPropertyEditingEnabledDelegate(FIsPropertyEditingEnabled InPropertyEditingDelegate);
 
 private:
 	/** Create the properties tab and its content */
-	TSharedRef<SDockTab> SpawnPropertiesTab( const FSpawnTabArgs& Args );
+	UNREALED_API TSharedRef<SDockTab> SpawnPropertiesTab( const FSpawnTabArgs& Args );
 
 	/** Handles when an asset is imported */
-	void HandleAssetPostImport(class UFactory* InFactory, UObject* InObject);
+	UNREALED_API void HandleAssetPostImport(class UFactory* InFactory, UObject* InObject);
 
 	/** Called when objects need to be swapped out for new versions, like after a blueprint recompile. */
-	void OnObjectsReplaced(const TMap<UObject*, UObject*>& ReplacementMap);
+	UNREALED_API void OnObjectsReplaced(const TMap<UObject*, UObject*>& ReplacementMap);
 
 	/** Details view */
 	TSharedPtr< class IDetailsView > DetailsView;
 
 	/** App Identifier. Technically, all simple editors are the same app, despite editing a variety of assets. */
-	static const FName SimpleEditorAppIdentifier;
+	static UNREALED_API const FName SimpleEditorAppIdentifier;
 
 	/**	The tab ids for all the tabs used */
-	static const FName PropertiesTabId;
+	static UNREALED_API const FName PropertiesTabId;
 
 	/** The objects open within this editor */
 	TArray<UObject*> EditingObjects;
 
 public:
 	/** The name given to all instances of this type of editor */
-	static const FName ToolkitFName;
+	static UNREALED_API const FName ToolkitFName;
 
-	static TSharedRef<FSimpleAssetEditor> CreateEditor( const EToolkitMode::Type Mode, const TSharedPtr< IToolkitHost >& InitToolkitHost, UObject* ObjectToEdit, FGetDetailsViewObjects GetDetailsViewObjects = FGetDetailsViewObjects() );
+	static UNREALED_API TSharedRef<FSimpleAssetEditor> CreateEditor( const EToolkitMode::Type Mode, const TSharedPtr< IToolkitHost >& InitToolkitHost, UObject* ObjectToEdit, FGetDetailsViewObjects GetDetailsViewObjects = FGetDetailsViewObjects() );
 
-	static TSharedRef<FSimpleAssetEditor> CreateEditor( const EToolkitMode::Type Mode, const TSharedPtr< IToolkitHost >& InitToolkitHost, const TArray<UObject*>& ObjectsToEdit, FGetDetailsViewObjects GetDetailsViewObjects = FGetDetailsViewObjects() );
+	static UNREALED_API TSharedRef<FSimpleAssetEditor> CreateEditor( const EToolkitMode::Type Mode, const TSharedPtr< IToolkitHost >& InitToolkitHost, const TArray<UObject*>& ObjectsToEdit, FGetDetailsViewObjects GetDetailsViewObjects = FGetDetailsViewObjects() );
 };
diff --git a/Engine/Source/Editor/UnrealEd/Public/Tools/BaseAssetToolkit.h b/Engine/Source/Editor/UnrealEd/Public/Tools/BaseAssetToolkit.h
--- a/Engine/Source/Editor/UnrealEd/Public/Tools/BaseAssetToolkit.h
+++ b/Engine/Source/Editor/UnrealEd/Public/Tools/BaseAssetToolkit.h
@@ -18,12 +18,12 @@ class FPreviewScene;
 class FEditorViewportClient;
 class UAssetEditor;
 
-class UNREALED_API FBaseAssetToolkit : public FAssetEditorToolkit
+class FBaseAssetToolkit : public FAssetEditorToolkit
 {
 public:
-	FBaseAssetToolkit(UAssetEditor* InOwningAssetEditor);
+	UNREALED_API FBaseAssetToolkit(UAssetEditor* InOwningAssetEditor);
 
-	~FBaseAssetToolkit();
+	UNREALED_API ~FBaseAssetToolkit();
 	virtual FName GetToolkitFName() const override
 	{
 		return NAME_None;
@@ -41,24 +41,24 @@ public:
 	{
 		return FLinearColor::White;
 	}
-	virtual void RegisterTabSpawners(const TSharedRef<FTabManager>& InTabManager) override;
+	UNREALED_API virtual void RegisterTabSpawners(const TSharedRef<FTabManager>& InTabManager) override;
 
 
-	virtual const TSharedRef<FTabManager::FLayout> GetDefaultLayout() const;
-	virtual void CreateWidgets();
-	virtual void SetEditingObject(class UObject* InObject);
-	virtual void CreateEditorModeManager() override;
+	UNREALED_API virtual const TSharedRef<FTabManager::FLayout> GetDefaultLayout() const;
+	UNREALED_API virtual void CreateWidgets();
+	UNREALED_API virtual void SetEditingObject(class UObject* InObject);
+	UNREALED_API virtual void CreateEditorModeManager() override;
 
 public:
-	static const FName ViewportTabID;
-	static const FName DetailsTabID;
+	static UNREALED_API const FName ViewportTabID;
+	static UNREALED_API const FName DetailsTabID;
 
 protected:
-	virtual void RegisterToolbar();
-	virtual AssetEditorViewportFactoryFunction GetViewportDelegate();
-	virtual TSharedPtr<FEditorViewportClient> CreateEditorViewportClient() const;
-	TSharedRef<SDockTab> SpawnTab_Viewport(const FSpawnTabArgs& Args);
-	virtual TSharedRef<SDockTab> SpawnTab_Details(const FSpawnTabArgs& Args);
+	UNREALED_API virtual void RegisterToolbar();
+	UNREALED_API virtual AssetEditorViewportFactoryFunction GetViewportDelegate();
+	UNREALED_API virtual TSharedPtr<FEditorViewportClient> CreateEditorViewportClient() const;
+	UNREALED_API TSharedRef<SDockTab> SpawnTab_Viewport(const FSpawnTabArgs& Args);
+	UNREALED_API virtual TSharedRef<SDockTab> SpawnTab_Details(const FSpawnTabArgs& Args);
 
 protected:
 	/** Property View */
@@ -73,4 +73,4 @@ protected:
 	TSharedPtr<FTabManager::FLayout> StandaloneDefaultLayout;
 	FString LayoutAppendix;
 	UAssetEditor* OwningAssetEditor;
-};
\ No newline at end of file
+};
diff --git a/Engine/Source/Editor/UnrealEd/Public/Tools/DefaultEdMode.h b/Engine/Source/Editor/UnrealEd/Public/Tools/DefaultEdMode.h
--- a/Engine/Source/Editor/UnrealEd/Public/Tools/DefaultEdMode.h
+++ b/Engine/Source/Editor/UnrealEd/Public/Tools/DefaultEdMode.h
@@ -6,17 +6,17 @@
 
 #include "DefaultEdMode.generated.h"
 
-UCLASS(Transient)
-class UNREALED_API UEdModeDefault : public UBaseLegacyWidgetEdMode
+UCLASS(Transient, MinimalAPI)
+class UEdModeDefault : public UBaseLegacyWidgetEdMode
 {
 	GENERATED_BODY()
 
 public:
-	UEdModeDefault();
+	UNREALED_API UEdModeDefault();
 
-	virtual bool UsesPropertyWidgets() const override;
-	virtual bool UsesToolkits() const override;
+	UNREALED_API virtual bool UsesPropertyWidgets() const override;
+	UNREALED_API virtual bool UsesToolkits() const override;
 
 protected:
-	virtual TSharedRef<FLegacyEdModeWidgetHelper> CreateWidgetHelper() override;
+	UNREALED_API virtual TSharedRef<FLegacyEdModeWidgetHelper> CreateWidgetHelper() override;
 };
diff --git a/Engine/Source/Editor/UnrealEd/Public/Tools/EditorComponentSourceFactory.h b/Engine/Source/Editor/UnrealEd/Public/Tools/EditorComponentSourceFactory.h
--- a/Engine/Source/Editor/UnrealEd/Public/Tools/EditorComponentSourceFactory.h
+++ b/Engine/Source/Editor/UnrealEd/Public/Tools/EditorComponentSourceFactory.h
@@ -10,35 +10,35 @@
 
 
 
-class UNREALED_API FStaticMeshComponentTargetFactory : public FComponentTargetFactory
+class FStaticMeshComponentTargetFactory : public FComponentTargetFactory
 {
 public:
 	// new FStaticMeshComponentTargets returned by Build() will be requested for this LOD
 	EMeshLODIdentifier CurrentEditingLOD = EMeshLODIdentifier::MaxQuality;
 
-	bool CanBuild( UActorComponent* Candidate ) override;
-	TUniquePtr<FPrimitiveComponentTarget> Build( UPrimitiveComponent* PrimitiveComponent ) override;
+	UNREALED_API bool CanBuild( UActorComponent* Candidate ) override;
+	UNREALED_API TUniquePtr<FPrimitiveComponentTarget> Build( UPrimitiveComponent* PrimitiveComponent ) override;
 };
 
 
 
-class UNREALED_API FStaticMeshComponentTarget : public FPrimitiveComponentTarget
+class FStaticMeshComponentTarget : public FPrimitiveComponentTarget
 {
 public:
 
-	FStaticMeshComponentTarget(UPrimitiveComponent* Component, EMeshLODIdentifier EditingLOD = EMeshLODIdentifier::LOD0);
+	UNREALED_API FStaticMeshComponentTarget(UPrimitiveComponent* Component, EMeshLODIdentifier EditingLOD = EMeshLODIdentifier::LOD0);
 
-	virtual bool IsValid() const override;
+	UNREALED_API virtual bool IsValid() const override;
 
-	virtual void GetMaterialSet(FComponentMaterialSet& MaterialSetOut, bool bAssetMaterials) const override;
+	UNREALED_API virtual void GetMaterialSet(FComponentMaterialSet& MaterialSetOut, bool bAssetMaterials) const override;
 
-	FMeshDescription* GetMesh() override;
+	UNREALED_API FMeshDescription* GetMesh() override;
 
-	void CommitMesh( const FCommitter& ) override;
+	UNREALED_API void CommitMesh( const FCommitter& ) override;
 
-	virtual void CommitMaterialSetUpdate(const FComponentMaterialSet& MaterialSet, bool bApplyToAsset) override;
+	UNREALED_API virtual void CommitMaterialSetUpdate(const FComponentMaterialSet& MaterialSet, bool bApplyToAsset) override;
 
-	virtual bool HasSameSourceData(const FPrimitiveComponentTarget& OtherTarget) const override;
+	UNREALED_API virtual bool HasSameSourceData(const FPrimitiveComponentTarget& OtherTarget) const override;
 
 
 protected:
diff --git a/Engine/Source/Editor/UnrealEd/Public/Tools/LegacyEdModeInterfaces.h b/Engine/Source/Editor/UnrealEd/Public/Tools/LegacyEdModeInterfaces.h
--- a/Engine/Source/Editor/UnrealEd/Public/Tools/LegacyEdModeInterfaces.h
+++ b/Engine/Source/Editor/UnrealEd/Public/Tools/LegacyEdModeInterfaces.h
@@ -30,7 +30,7 @@ class ULegacyEdModeSelectInterface : public UInterface
 	GENERATED_BODY()
 };
 
-class UNREALED_API ILegacyEdModeSelectInterface
+class ILegacyEdModeSelectInterface
 {
 	GENERATED_BODY()
 public:
@@ -57,7 +57,7 @@ class ULegacyEdModeWidgetInterface : public UInterface
 	GENERATED_BODY()
 };
 
-class UNREALED_API ILegacyEdModeWidgetInterface
+class ILegacyEdModeWidgetInterface
 {
 	GENERATED_BODY()
 
@@ -126,7 +126,7 @@ class ULegacyEdModeToolInterface : public UInterface
 	GENERATED_BODY()
 };
 
-class UNREALED_API ILegacyEdModeToolInterface
+class ILegacyEdModeToolInterface
 {
 	GENERATED_BODY()
 
@@ -149,7 +149,7 @@ class ULegacyEdModeDrawHelperInterface : public UInterface
 	GENERATED_BODY()
 };
 
-class UNREALED_API ILegacyEdModeDrawHelperInterface
+class ILegacyEdModeDrawHelperInterface
 {
 	GENERATED_BODY()
 
@@ -163,7 +163,7 @@ class ULegacyEdModeViewportInterface : public UInterface
 	GENERATED_BODY()
 };
 
-class UNREALED_API ILegacyEdModeViewportInterface
+class ILegacyEdModeViewportInterface
 {
 	GENERATED_BODY()
 
diff --git a/Engine/Source/Editor/UnrealEd/Public/Tools/LegacyEdModeWidgetHelpers.h b/Engine/Source/Editor/UnrealEd/Public/Tools/LegacyEdModeWidgetHelpers.h
--- a/Engine/Source/Editor/UnrealEd/Public/Tools/LegacyEdModeWidgetHelpers.h
+++ b/Engine/Source/Editor/UnrealEd/Public/Tools/LegacyEdModeWidgetHelpers.h
@@ -23,32 +23,32 @@ class HHitProxy;
 struct FViewportClick;
 class UObject;
 
-class UNREALED_API FLegacyEdModeWidgetHelper
+class FLegacyEdModeWidgetHelper
 {
 	friend class FEditorModeRegistry;
 	friend class UBaseLegacyWidgetEdMode;
 
 public:
-	FLegacyEdModeWidgetHelper();
+	UNREALED_API FLegacyEdModeWidgetHelper();
 	virtual ~FLegacyEdModeWidgetHelper() = default;
 
-	virtual bool AllowWidgetMove();
-	virtual bool CanCycleWidgetMode() const;
-	virtual bool ShowModeWidgets() const;
-	virtual EAxisList::Type GetWidgetAxisToDraw(UE::Widget::EWidgetMode InWidgetMode) const;
-	virtual FVector GetWidgetLocation() const;
-	virtual bool ShouldDrawWidget() const;
-	virtual bool UsesTransformWidget() const;
-	virtual bool UsesTransformWidget(UE::Widget::EWidgetMode CheckMode) const;
-	virtual FVector GetWidgetNormalFromCurrentAxis(void* InData);
-	virtual void SetCurrentWidgetAxis(EAxisList::Type InAxis) final;
-	virtual EAxisList::Type GetCurrentWidgetAxis() const final;
-	virtual bool UsesPropertyWidgets() const;
-	virtual bool GetCustomDrawingCoordinateSystem(FMatrix& InMatrix, void* InData);
-	virtual bool GetCustomInputCoordinateSystem(FMatrix& InMatrix, void* InData);
+	UNREALED_API virtual bool AllowWidgetMove();
+	UNREALED_API virtual bool CanCycleWidgetMode() const;
+	UNREALED_API virtual bool ShowModeWidgets() const;
+	UNREALED_API virtual EAxisList::Type GetWidgetAxisToDraw(UE::Widget::EWidgetMode InWidgetMode) const;
+	UNREALED_API virtual FVector GetWidgetLocation() const;
+	UNREALED_API virtual bool ShouldDrawWidget() const;
+	UNREALED_API virtual bool UsesTransformWidget() const;
+	UNREALED_API virtual bool UsesTransformWidget(UE::Widget::EWidgetMode CheckMode) const;
+	UNREALED_API virtual FVector GetWidgetNormalFromCurrentAxis(void* InData);
+	UNREALED_API virtual void SetCurrentWidgetAxis(EAxisList::Type InAxis) final;
+	UNREALED_API virtual EAxisList::Type GetCurrentWidgetAxis() const final;
+	UNREALED_API virtual bool UsesPropertyWidgets() const;
+	UNREALED_API virtual bool GetCustomDrawingCoordinateSystem(FMatrix& InMatrix, void* InData);
+	UNREALED_API virtual bool GetCustomInputCoordinateSystem(FMatrix& InMatrix, void* InData);
 	
-	virtual void ActorSelectionChangeNotify();
-	virtual bool AllowsViewportDragTool() const;
+	UNREALED_API virtual void ActorSelectionChangeNotify();
+	UNREALED_API virtual bool AllowsViewportDragTool() const;
 
 	// Property Widgets
 	/** Structure that holds info about our optional property widget */
@@ -69,28 +69,28 @@ public:
 		void GetTransformAndColor(UObject* BestSelectedItem, bool bIsSelected, FTransform& OutLocalTransform, FString& OutValidationMessage, FColor& OutDrawColor) const;
 	};
 
-	virtual bool InputDelta(FEditorViewportClient* InViewportClient, FViewport* InViewport, FVector& InDrag, FRotator& InRot, FVector& InScale);
-	virtual bool HandleClick(FEditorViewportClient* InViewportClient, HHitProxy* HitProxy, const FViewportClick& Click);
-	virtual void Render(const FSceneView* View, FViewport* Viewport, FPrimitiveDrawInterface* PDI);
-	virtual void DrawHUD(FEditorViewportClient* ViewportClient, FViewport* Viewport, const FSceneView* View, FCanvas* Canvas);
+	UNREALED_API virtual bool InputDelta(FEditorViewportClient* InViewportClient, FViewport* InViewport, FVector& InDrag, FRotator& InRot, FVector& InScale);
+	UNREALED_API virtual bool HandleClick(FEditorViewportClient* InViewportClient, HHitProxy* HitProxy, const FViewportClick& Click);
+	UNREALED_API virtual void Render(const FSceneView* View, FViewport* Viewport, FPrimitiveDrawInterface* PDI);
+	UNREALED_API virtual void DrawHUD(FEditorViewportClient* ViewportClient, FViewport* Viewport, const FSceneView* View, FCanvas* Canvas);
 
 	// @TODO: Find a better home for these?
 	/** Value of UPROPERTY can be edited with a widget in the editor (translation, rotation) */
-	static const FName MD_MakeEditWidget;
+	static UNREALED_API const FName MD_MakeEditWidget;
 	/** Specifies a function used for validation of the current value of a property.  The function returns a string that is empty if the value is valid, or contains an error description if the value is invalid */
-	static const FName MD_ValidateWidgetUsing;
+	static UNREALED_API const FName MD_ValidateWidgetUsing;
 	/** Returns true if this structure can support creating a widget in the editor */
-	static bool CanCreateWidgetForStructure(const UStruct* InPropStruct);
+	static UNREALED_API bool CanCreateWidgetForStructure(const UStruct* InPropStruct);
 	/** Returns true if this property can support creating a widget in the editor */
-	static bool CanCreateWidgetForProperty(FProperty* InProp);
+	static UNREALED_API bool CanCreateWidgetForProperty(FProperty* InProp);
 	/** See if we should create a widget for the supplied property when selecting an actor instance */
-	static bool ShouldCreateWidgetForProperty(FProperty* InProp);
+	static UNREALED_API bool ShouldCreateWidgetForProperty(FProperty* InProp);
 
 protected:
 	/**
 	 * Returns the first selected Actor, or NULL if there is no selection.
 	 */
-	AActor* GetFirstSelectedActorInstance() const;
+	UNREALED_API AActor* GetFirstSelectedActorInstance() const;
 
 	/**
 	 * Gets an array of property widget info structures for the given struct/class type for the given container.
@@ -99,10 +99,10 @@ protected:
 	 * @param InContainer The container of the given type.
 	 * @param OutInfos An array of widget info structures (output).
 	 */
-	void GetPropertyWidgetInfos(const UStruct* InStruct, const void* InContainer, TArray<FPropertyWidgetInfo>& OutInfos) const;
+	UNREALED_API void GetPropertyWidgetInfos(const UStruct* InStruct, const void* InContainer, TArray<FPropertyWidgetInfo>& OutInfos) const;
 
 	/** Finds the best item to display widgets for (preferring selected components over selected actors) */
-	virtual UObject* GetItemToTryDisplayingWidgetsFor(FTransform& OutWidgetToWorld) const;
+	UNREALED_API virtual UObject* GetItemToTryDisplayingWidgetsFor(FTransform& OutWidgetToWorld) const;
 
 	/** Name of the property currently being edited */
 	FString EditedPropertyName;
@@ -127,8 +127,8 @@ protected:
 //   1. Subclass FLegacyEdModeWidgetHelper, and override the methods needed for your native FEdMode implementation
 //   2. Transition your native FEdMode to a UObject, by inheriting from this class
 //   3. Override CreateWidgetHelper function to return a SharedRef to the class you created in step 1.
-UCLASS(Abstract)
-class UNREALED_API UBaseLegacyWidgetEdMode : public UEdMode, public ILegacyEdModeWidgetInterface, public ILegacyEdModeViewportInterface
+UCLASS(Abstract, MinimalAPI)
+class UBaseLegacyWidgetEdMode : public UEdMode, public ILegacyEdModeWidgetInterface, public ILegacyEdModeViewportInterface
 {
 	GENERATED_BODY()
 
@@ -136,32 +136,32 @@ public:
 	// UEdMode overrides
 	// If you need to do any initialization in your mode, be sure to still call through to this function.
 	// It creates the WidgetHelper and hooks up the mode manager pointer for you.
-	virtual void Initialize() override;
-	virtual bool InputDelta(FEditorViewportClient* InViewportClient, FViewport* InViewport, FVector& InDrag, FRotator& InRot, FVector& InScale) override;
-	virtual bool HandleClick(FEditorViewportClient* InViewportClient, HHitProxy* HitProxy, const FViewportClick& Click) override;
-	virtual void Render(const FSceneView* View, FViewport* Viewport, FPrimitiveDrawInterface* PDI) override;
-	virtual void DrawHUD(FEditorViewportClient* ViewportClient, FViewport* Viewport, const FSceneView* View, FCanvas* Canvas) override;
+	UNREALED_API virtual void Initialize() override;
+	UNREALED_API virtual bool InputDelta(FEditorViewportClient* InViewportClient, FViewport* InViewport, FVector& InDrag, FRotator& InRot, FVector& InScale) override;
+	UNREALED_API virtual bool HandleClick(FEditorViewportClient* InViewportClient, HHitProxy* HitProxy, const FViewportClick& Click) override;
+	UNREALED_API virtual void Render(const FSceneView* View, FViewport* Viewport, FPrimitiveDrawInterface* PDI) override;
+	UNREALED_API virtual void DrawHUD(FEditorViewportClient* ViewportClient, FViewport* Viewport, const FSceneView* View, FCanvas* Canvas) override;
 
 	// ILegacyEdModeWidgetInterface overrides are just a pass through to the FLegacyEdModeWidgetHelper
-	virtual bool AllowWidgetMove() override;
-	virtual bool CanCycleWidgetMode() const override;
-	virtual bool ShowModeWidgets() const override;
-	virtual EAxisList::Type GetWidgetAxisToDraw(UE::Widget::EWidgetMode InWidgetMode) const override;
-	virtual FVector GetWidgetLocation() const override;
-	virtual bool ShouldDrawWidget() const override;
-	virtual bool UsesTransformWidget() const override;
-	virtual bool UsesTransformWidget(UE::Widget::EWidgetMode CheckMode) const override;
-	virtual FVector GetWidgetNormalFromCurrentAxis(void* InData) override;
-	virtual void SetCurrentWidgetAxis(EAxisList::Type InAxis) override;
-	virtual EAxisList::Type GetCurrentWidgetAxis() const override;
-	virtual bool UsesPropertyWidgets() const override;
-	virtual bool GetCustomDrawingCoordinateSystem(FMatrix& InMatrix, void* InData) override;
-	virtual bool GetCustomInputCoordinateSystem(FMatrix& InMatrix, void* InData) override;
-	virtual void ActorSelectionChangeNotify() override;
-	virtual bool AllowsViewportDragTool() const override;
+	UNREALED_API virtual bool AllowWidgetMove() override;
+	UNREALED_API virtual bool CanCycleWidgetMode() const override;
+	UNREALED_API virtual bool ShowModeWidgets() const override;
+	UNREALED_API virtual EAxisList::Type GetWidgetAxisToDraw(UE::Widget::EWidgetMode InWidgetMode) const override;
+	UNREALED_API virtual FVector GetWidgetLocation() const override;
+	UNREALED_API virtual bool ShouldDrawWidget() const override;
+	UNREALED_API virtual bool UsesTransformWidget() const override;
+	UNREALED_API virtual bool UsesTransformWidget(UE::Widget::EWidgetMode CheckMode) const override;
+	UNREALED_API virtual FVector GetWidgetNormalFromCurrentAxis(void* InData) override;
+	UNREALED_API virtual void SetCurrentWidgetAxis(EAxisList::Type InAxis) override;
+	UNREALED_API virtual EAxisList::Type GetCurrentWidgetAxis() const override;
+	UNREALED_API virtual bool UsesPropertyWidgets() const override;
+	UNREALED_API virtual bool GetCustomDrawingCoordinateSystem(FMatrix& InMatrix, void* InData) override;
+	UNREALED_API virtual bool GetCustomInputCoordinateSystem(FMatrix& InMatrix, void* InData) override;
+	UNREALED_API virtual void ActorSelectionChangeNotify() override;
+	UNREALED_API virtual bool AllowsViewportDragTool() const override;
 
 protected:
-	virtual TSharedRef<FLegacyEdModeWidgetHelper> CreateWidgetHelper();
+	UNREALED_API virtual TSharedRef<FLegacyEdModeWidgetHelper> CreateWidgetHelper();
 
 	TSharedPtr<FLegacyEdModeWidgetHelper> WidgetHelper;
 };
diff --git a/Engine/Source/Editor/UnrealEd/Public/Tools/StandardToolModeCommands.h b/Engine/Source/Editor/UnrealEd/Public/Tools/StandardToolModeCommands.h
--- a/Engine/Source/Editor/UnrealEd/Public/Tools/StandardToolModeCommands.h
+++ b/Engine/Source/Editor/UnrealEd/Public/Tools/StandardToolModeCommands.h
@@ -36,20 +36,20 @@ enum class EStandardToolModeCommands
  * call FStandardToolModeCommands::Get().FindStandardCommand() to get the registered UICommandInfo.
  * 
  */
-class UNREALED_API FStandardToolModeCommands : public TCommands<FStandardToolModeCommands>
+class FStandardToolModeCommands : public TCommands<FStandardToolModeCommands>
 {
 public:
-	FStandardToolModeCommands();
+	UNREALED_API FStandardToolModeCommands();
 
 	/**
 	 * Registers the set of standard commands. Call on module startup.
 	 */
-	virtual void RegisterCommands() override;
+	UNREALED_API virtual void RegisterCommands() override;
 
 	/**
 	 * Look up the UICommandInfo for a standard command
 	 */
-	TSharedPtr<FUICommandInfo> FindStandardCommand(EStandardToolModeCommands Command) const;
+	UNREALED_API TSharedPtr<FUICommandInfo> FindStandardCommand(EStandardToolModeCommands Command) const;
 
 protected:
 	TMap<EStandardToolModeCommands, TSharedPtr<FUICommandInfo>> Commands;
diff --git a/Engine/Source/Editor/UnrealEd/Public/Tools/UAssetEditor.h b/Engine/Source/Editor/UnrealEd/Public/Tools/UAssetEditor.h
--- a/Engine/Source/Editor/UnrealEd/Public/Tools/UAssetEditor.h
+++ b/Engine/Source/Editor/UnrealEd/Public/Tools/UAssetEditor.h
@@ -18,32 +18,32 @@ struct FTabId;
 /**
  * Base class for all asset editors.
  */
-UCLASS(Abstract)
-class UNREALED_API UAssetEditor : public UObject, public IAssetEditorInstance
+UCLASS(Abstract, MinimalAPI)
+class UAssetEditor : public UObject, public IAssetEditorInstance
 {
 	GENERATED_BODY()
 
 public:
-	UAssetEditor();
+	UNREALED_API UAssetEditor();
 
-	virtual FName GetEditorName() const override;
-	virtual void FocusWindow(UObject* ObjectToFocusOn = nullptr) override;
+	UNREALED_API virtual FName GetEditorName() const override;
+	UNREALED_API virtual void FocusWindow(UObject* ObjectToFocusOn = nullptr) override;
 	virtual bool IsPrimaryEditor() const override
 	{
 		return true;
 	}
 	virtual void InvokeTab(const FTabId& TabId) override {}
-	virtual TSharedPtr<FTabManager> GetAssociatedTabManager() override;
+	UNREALED_API virtual TSharedPtr<FTabManager> GetAssociatedTabManager() override;
 	virtual double GetLastActivationTime() override
 	{
 		return 0.0;
 	}
 	virtual void RemoveEditingAsset(UObject* Asset) override {}
 
-	void Initialize();
-	virtual void GetObjectsToEdit(TArray<UObject*>& InObjectsToEdit);
-	virtual TSharedPtr<FBaseAssetToolkit> CreateToolkit();
-	void OnToolkitClosed();
+	UNREALED_API void Initialize();
+	UNREALED_API virtual void GetObjectsToEdit(TArray<UObject*>& InObjectsToEdit);
+	UNREALED_API virtual TSharedPtr<FBaseAssetToolkit> CreateToolkit();
+	UNREALED_API void OnToolkitClosed();
 
 protected:
 	FBaseAssetToolkit* ToolkitInstance;
diff --git a/Engine/Source/Editor/UnrealEd/Public/Tools/UEdMode.h b/Engine/Source/Editor/UnrealEd/Public/Tools/UEdMode.h
--- a/Engine/Source/Editor/UnrealEd/Public/Tools/UEdMode.h
+++ b/Engine/Source/Editor/UnrealEd/Public/Tools/UEdMode.h
@@ -67,8 +67,8 @@ enum class EToolsContextScope
 /**
  * Base class for all editor modes.
  */
-UCLASS(Abstract)
-class UNREALED_API UEdMode : public UObject
+UCLASS(Abstract, MinimalAPI)
+class UEdMode : public UObject
 {
 	GENERATED_BODY()
 
@@ -76,9 +76,9 @@ public:
 	/** Friends so it can access mode's internals on construction */
 	friend class UAssetEditorSubsystem;
 
-	UEdMode();
+	UNREALED_API UEdMode();
 
-	virtual void Initialize();
+	UNREALED_API virtual void Initialize();
 
 	// Added for handling EDIT Command...
 	virtual EEditAction::Type GetActionEditDuplicate() { return EEditAction::Skip; }
@@ -87,7 +87,7 @@ public:
 	virtual EEditAction::Type GetActionEditCopy() { return EEditAction::Skip; }
 	virtual EEditAction::Type GetActionEditPaste() { return EEditAction::Skip; }
 	virtual bool ProcessEditDuplicate() { return false; }
-	virtual bool ProcessEditDelete();
+	UNREALED_API virtual bool ProcessEditDelete();
 	virtual bool ProcessEditCut() { return false; }
 	virtual bool ProcessEditCopy() { return false; }
 	virtual bool ProcessEditPaste() { return false; }
@@ -118,7 +118,7 @@ public:
 	virtual bool ShouldDrawBrushWireframe(AActor* InActor) const { return true; }
 
 	/** If Rotation Snap should be enabled for this mode*/
-	virtual bool IsSnapRotationEnabled();
+	UNREALED_API virtual bool IsSnapRotationEnabled();
 
 	/** If this mode should override the snap rotation
 	* @param	Rotation		The Rotation Override
@@ -129,7 +129,7 @@ public:
 
 	virtual void UpdateInternalData() {}
 
-	virtual void Enter();
+	UNREALED_API virtual void Enter();
 
 	/**
 	 * Registers and maps the provided UI command to actions that start / stop a given tool.
@@ -140,14 +140,14 @@ public:
 	 * @param	Builder			Builder for tool to be used by actions
 	 * @param	ToolScope		Scope to determine lifetime of tool (Editor, Mode, etc)
 	 */
-	virtual void RegisterTool(TSharedPtr<FUICommandInfo> UICommand, FString ToolIdentifier, UInteractiveToolBuilder* Builder, EToolsContextScope ToolScope = EToolsContextScope::Default);
+	UNREALED_API virtual void RegisterTool(TSharedPtr<FUICommandInfo> UICommand, FString ToolIdentifier, UInteractiveToolBuilder* Builder, EToolsContextScope ToolScope = EToolsContextScope::Default);
 
 	/** 
 	 * Subclasses can override this to add additional checks on whether a tool should be allowed to start.
 	 * By default the check disallows starting tools during play/simulate in editor.
 	 */
-	virtual bool ShouldToolStartBeAllowed(const FString& ToolIdentifier) const;
-	virtual void Exit();
+	UNREALED_API virtual bool ShouldToolStartBeAllowed(const FString& ToolIdentifier) const;
+	UNREALED_API virtual void Exit();
 
 	virtual void PostUndo() {}
 
@@ -165,7 +165,7 @@ public:
 	 */
 	virtual bool IsOperationSupportedForCurrentAsset(EAssetOperation InOperation) const { return true; }
 
-	virtual void SelectNone();
+	UNREALED_API virtual void SelectNone();
 	virtual void SelectionChanged() {}
 
 	/**
@@ -194,16 +194,16 @@ public:
 	const FEditorModeInfo& GetModeInfo() const { return Info; }
 
 	/** True if this mode uses a toolkit mode (eventually they all should) */
-	virtual bool UsesToolkits() const;
+	UNREALED_API virtual bool UsesToolkits() const;
 
 	/** Gets the toolkit created by this mode */
 	TWeakPtr<FModeToolkit> GetToolkit() { return Toolkit; }
 
 	/** Returns the world this toolkit is editing */
-	UWorld* GetWorld() const;
+	UNREALED_API UWorld* GetWorld() const;
 
 	/** Returns the owning mode manager for this mode */
-	FEditorModeTools* GetModeManager() const;
+	UNREALED_API FEditorModeTools* GetModeManager() const;
 
 	/**
 	 * For use by the EditorModeTools class to get the legacy FEdMode type from a legacy FEdMode wrapper
@@ -231,7 +231,7 @@ protected:
 	/**
 	 * Returns the first selected Actor, or NULL if there is no selection.
 	 */
-	AActor* GetFirstSelectedActorInstance() const;
+	UNREALED_API AActor* GetFirstSelectedActorInstance() const;
 
 	bool bHaveSavedEditorState;
 	bool bSavedAntiAliasingState;
@@ -247,12 +247,12 @@ public:
 	/**
 	 * @return active ToolManager for the desired (or default) ToolsContext Scope
 	 */
-	UInteractiveToolManager* GetToolManager(EToolsContextScope ToolScope = EToolsContextScope::Default) const;
+	UNREALED_API UInteractiveToolManager* GetToolManager(EToolsContextScope ToolScope = EToolsContextScope::Default) const;
 
 	/**
 	 * @return active ToolsContext for the desired (or default) ToolsContext Scope
 	 */
-	UEditorInteractiveToolsContext* GetInteractiveToolsContext(EToolsContextScope ToolScope = EToolsContextScope::Default) const;
+	UNREALED_API UEditorInteractiveToolsContext* GetInteractiveToolsContext(EToolsContextScope ToolScope = EToolsContextScope::Default) const;
 
 
 	virtual TMap<FName, TArray<TSharedPtr<FUICommandInfo>>> GetModeCommands() const
@@ -261,12 +261,12 @@ public:
 	};
 
 protected:
-	virtual void CreateToolkit();
+	UNREALED_API virtual void CreateToolkit();
 	virtual void OnToolStarted(UInteractiveToolManager* Manager, UInteractiveTool* Tool) {}
 	virtual void OnToolEnded(UInteractiveToolManager* Manager, UInteractiveTool* Tool) {}
 	virtual void ActivateDefaultTool() {}
 	virtual void BindCommands() {}
-	void OnModeActivated(const FEditorModeID& InID, bool bIsActive);
+	UNREALED_API void OnModeActivated(const FEditorModeID& InID, bool bIsActive);
 
 private:
 
@@ -293,6 +293,6 @@ protected:
 
 protected:
 
-	void CreateInteractiveToolsContexts();
-	void DestroyInteractiveToolsContexts();
+	UNREALED_API void CreateInteractiveToolsContexts();
+	UNREALED_API void DestroyInteractiveToolsContexts();
 };
diff --git a/Engine/Source/Editor/UnrealEd/Public/UnrealEdMisc.h b/Engine/Source/Editor/UnrealEd/Public/UnrealEdMisc.h
--- a/Engine/Source/Editor/UnrealEd/Public/UnrealEdMisc.h
+++ b/Engine/Source/Editor/UnrealEd/Public/UnrealEdMisc.h
@@ -28,27 +28,27 @@ enum class EMapChangeType : uint8
 };
 
 /** The public interface for the unreal editor misc singleton. */
-class UNREALED_API FUnrealEdMisc
+class FUnrealEdMisc
 {
 public:
 
 	/** Singleton accessor */
-	static FUnrealEdMisc& Get();
+	static UNREALED_API FUnrealEdMisc& Get();
 
 	/** Destructor */
-	virtual ~FUnrealEdMisc();
+	UNREALED_API virtual ~FUnrealEdMisc();
 
 	/** Initalizes various systems */
-	virtual void OnInit();
+	UNREALED_API virtual void OnInit();
 
 	/* Check if this we are editing a template project, and if so mount any shared resource paths it uses */
-	void MountTemplateSharedPaths();
+	UNREALED_API void MountTemplateSharedPaths();
 
 	/* Cleans up various systems */
-	virtual void OnExit();
+	UNREALED_API virtual void OnExit();
 
 	/** Performs any required cleanup in the case of a fatal error. */
-	virtual void ShutdownAfterError();
+	UNREALED_API virtual void ShutdownAfterError();
 
 	/**
 	 *	Whether or not the map build in progressed was cancelled by the user. 
@@ -159,16 +159,16 @@ public:
 	}
 
 	/** Opens the specified project file or game. Restarts the editor */
-	void SwitchProject(const FString& GameOrProjectFileName, bool bWarn = true, const TOptional<FString>& NewCommandLine = TOptional<FString>());
+	UNREALED_API void SwitchProject(const FString& GameOrProjectFileName, bool bWarn = true, const TOptional<FString>& NewCommandLine = TOptional<FString>());
 
 	/** Restarts the editor, reopening the current project, if any */
-	void RestartEditor(bool bWarn = true, const TOptional<FString>& NewCommandLine = TOptional<FString>());
+	UNREALED_API void RestartEditor(bool bWarn = true, const TOptional<FString>& NewCommandLine = TOptional<FString>());
 
 	/** Ticks the performance analytics used by the analytics heartbeat */
-	void TickPerformanceAnalytics();
+	UNREALED_API void TickPerformanceAnalytics();
 
 	/** Triggers asset analytics if it hasn't been run yet */
-	void TickAssetAnalytics();
+	UNREALED_API void TickAssetAnalytics();
 
 	/**
 	 * Fetches a URL from the config and optionally switches it to rocket if required
@@ -179,113 +179,113 @@ public:
 	 *
 	 * @returns true if successful
 	 */
-	bool GetURL( const TCHAR* InKey, FString& OutURL, const bool bCheckRocket = false ) const;
+	UNREALED_API bool GetURL( const TCHAR* InKey, FString& OutURL, const bool bCheckRocket = false ) const;
 
-	void ReplaceDocumentationURLWildcards(FString& Url, const FCultureRef& Culture, const FString& PageId = FString());
+	UNREALED_API void ReplaceDocumentationURLWildcards(FString& Url, const FCultureRef& Culture, const FString& PageId = FString());
 
 	/** Returns the editor executable to use to execute commandlets */
-	FString GetExecutableForCommandlets() const;
+	UNREALED_API FString GetExecutableForCommandlets() const;
 
 	/** 
 	 * Opens the Unreal Engine Launcher marketplace page
 	 *
 	 * @param CustomLocation	Optional custom location within the marketplace to navigate to.  If not specified the launcher will open to the root marketplace page
 	 */
-	void OpenMarketplace(const FString& CustomLocation = TEXT(""));
+	UNREALED_API void OpenMarketplace(const FString& CustomLocation = TEXT(""));
 
 	/** Constructor, private - use Get() function */
-	FUnrealEdMisc();
+	UNREALED_API FUnrealEdMisc();
 
 	/** Displays a property dialog based upon what is currently selected. If any actors are selected, the actor property dialog is displayed. */
-	void CB_SelectedProps();
-	void CB_DisplayLoadErrors();
-	void CB_RefreshEditor();
+	UNREALED_API void CB_SelectedProps();
+	UNREALED_API void CB_DisplayLoadErrors();
+	UNREALED_API void CB_RefreshEditor();
 
 	/** Tells the editor that something has been done to change the map.  Can be anything from loading a whole new map to changing the BSP. */
-	void CB_MapChange( uint32 InFlags );
-	void CB_RedrawAllViewports();
-	void CB_EditorModeWindowClosed(const TSharedRef<SWindow>&);
-	void CB_LevelActorsAdded(class AActor* InActor);
+	UNREALED_API void CB_MapChange( uint32 InFlags );
+	UNREALED_API void CB_RedrawAllViewports();
+	UNREALED_API void CB_EditorModeWindowClosed(const TSharedRef<SWindow>&);
+	UNREALED_API void CB_LevelActorsAdded(class AActor* InActor);
 
 	/** Called right before unit testing is about to begin */
-	void CB_PreAutomationTesting();
+	UNREALED_API void CB_PreAutomationTesting();
 
 	/** Called right after unit testing concludes */
-	void CB_PostAutomationTesting();
+	UNREALED_API void CB_PostAutomationTesting();
 
-	void OnEditorChangeMode(FEditorModeID NewEditorMode);
-	void OnEditorPreModal();
-	void OnEditorPostModal();
+	UNREALED_API void OnEditorChangeMode(FEditorModeID NewEditorMode);
+	UNREALED_API void OnEditorPreModal();
+	UNREALED_API void OnEditorPostModal();
 
 	/** BP nativization settings */
 	UE_DEPRECATED(5.0, "Blueprint Nativization has been removed as a supported feature. This delegate method will eventually be removed.")
 	void OnNativizeBlueprintsSettingChanged(const FString& PackageName, bool bSelect) {}
 
 	/** Called from tab manager when the tab changes */
-	void OnActiveTabChanged(TSharedPtr<SDockTab> PreviouslyActive, TSharedPtr<SDockTab> NewlyActivated);
-	void OnTabForegrounded(TSharedPtr<SDockTab> ForegroundTab, TSharedPtr<SDockTab> BackgroundTab);
-	void OnUserActivityTabChanged(TSharedPtr<SDockTab> InTab);
+	UNREALED_API void OnActiveTabChanged(TSharedPtr<SDockTab> PreviouslyActive, TSharedPtr<SDockTab> NewlyActivated);
+	UNREALED_API void OnTabForegrounded(TSharedPtr<SDockTab> ForegroundTab, TSharedPtr<SDockTab> BackgroundTab);
+	UNREALED_API void OnUserActivityTabChanged(TSharedPtr<SDockTab> InTab);
 
 	/** Delegate that gets called by modules that can't directly access Engine */
-	void OnDeferCommand( const FString& DeferredCommand );
+	UNREALED_API void OnDeferCommand( const FString& DeferredCommand );
 
 	/** Start the performance survey that attempts to monitor editor performance in the default or simple startup maps */
-	void BeginPerformanceSurvey();
+	UNREALED_API void BeginPerformanceSurvey();
 
 	/** End the performance survey that attempts to monitor editor performance in the default or simple startup maps */
-	void CancelPerformanceSurvey();
+	UNREALED_API void CancelPerformanceSurvey();
 
 	/**
 	 * Called when a map is changed (loaded,saved,new map, etc)
 	 */
-	void OnMapChanged( UWorld* World, EMapChangeType MapChangeType );
+	UNREALED_API void OnMapChanged( UWorld* World, EMapChangeType MapChangeType );
 
 	/** Called when the input manager records a user-defined chord */
-	void OnUserDefinedChordChanged(const FUICommandInfo& CommandInfo);
+	UNREALED_API void OnUserDefinedChordChanged(const FUICommandInfo& CommandInfo);
 
 	/** Delegate for (default) message log UObject token activation - selects the object that the token refers to (if any) */
-	void OnMessageTokenActivated(const TSharedRef<class IMessageToken>& Token);
+	UNREALED_API void OnMessageTokenActivated(const TSharedRef<class IMessageToken>& Token);
 
 	/** Delegate for (default) display name of UObject tokens. Can display the name of the actor if an object is/is part of one */
-	FText OnGetDisplayName(const UObject* InObject, const bool bFullPath);
+	UNREALED_API FText OnGetDisplayName(const UObject* InObject, const bool bFullPath);
 
 	/** Delegate for (default) message log message selection - selects the objects that the tokens refer to (if any) */
-	void OnMessageSelectionChanged(TArray< TSharedRef<class FTokenizedMessage> >& Selection);
+	UNREALED_API void OnMessageSelectionChanged(TArray< TSharedRef<class FTokenizedMessage> >& Selection);
 
 	/** Delegate used to go to assets in the content browser */
-	void OnGotoAsset(const FString& InAssetPath) const;
+	UNREALED_API void OnGotoAsset(const FString& InAssetPath) const;
 
 	/** Delegate used on message log AActor token activation */
-	void OnActorTokenActivated(const TSharedRef<class IMessageToken>& Token);
+	UNREALED_API void OnActorTokenActivated(const TSharedRef<class IMessageToken>& Token);
 
 	/** Delegate used to update the map of asset update counts */
-	void OnObjectSaved(UObject* SavedObject, FObjectPreSaveContext SaveContext);
+	UNREALED_API void OnObjectSaved(UObject* SavedObject, FObjectPreSaveContext SaveContext);
 
 	/** Logs an update to an asset */
-	void LogAssetUpdate(UObject* UpdatedAsset, FObjectPreSaveContext SaveContext);
+	UNREALED_API void LogAssetUpdate(UObject* UpdatedAsset, FObjectPreSaveContext SaveContext);
 
 	/** Initialize engine analytics */
-	void InitEngineAnalytics();
+	UNREALED_API void InitEngineAnalytics();
 
 	/** Called when the heartbeat event should be sent to engine analytics */
-	void EditorAnalyticsHeartbeat();
+	UNREALED_API void EditorAnalyticsHeartbeat();
 
 	/** Handles "Enable World Composition" option in WorldSettings */
-	bool EnableWorldComposition(UWorld* InWorld, bool bEnable);
+	UNREALED_API bool EnableWorldComposition(UWorld* InWorld, bool bEnable);
 
 	/** Get the path to the executable that runs the editor */
-	FString GetProjectEditorBinaryPath();
+	UNREALED_API FString GetProjectEditorBinaryPath();
 
 	/** Finds a map using only the map name, no extension, no path, also caches it for faster lookup next time. */
-	FString FindMapFileFromPartialName(const FString& PartialMapName);
+	UNREALED_API FString FindMapFileFromPartialName(const FString& PartialMapName);
 
 	/** Launch an editor instance, passing ProjectName and the command line args from this running instance */
-	bool SpawnEditorInstance(const FString& ProjectName);
+	UNREALED_API bool SpawnEditorInstance(const FString& ProjectName);
 
 private:
-	void SelectActorFromMessageToken(AActor* InActor);
+	UNREALED_API void SelectActorFromMessageToken(AActor* InActor);
 
-	void PreSaveWorld(class UWorld* World, FObjectPreSaveContext ObjectSaveContext);
+	UNREALED_API void PreSaveWorld(class UWorld* World, FObjectPreSaveContext ObjectSaveContext);
 
 	/** Stores whether or not the current map build was cancelled. */
 	bool bCancelBuild;
diff --git a/Engine/Source/Editor/UnrealEd/Public/ViewportToolBarContext.h b/Engine/Source/Editor/UnrealEd/Public/ViewportToolBarContext.h
--- a/Engine/Source/Editor/UnrealEd/Public/ViewportToolBarContext.h
+++ b/Engine/Source/Editor/UnrealEd/Public/ViewportToolBarContext.h
@@ -7,10 +7,10 @@
 
 class SViewportToolBar;
 
-UCLASS()
-class UNREALED_API UViewportToolBarContext : public UObject
+UCLASS(MinimalAPI)
+class UViewportToolBarContext : public UObject
 {
 	GENERATED_BODY()
 public:
 	TWeakPtr<SViewportToolBar> ViewportToolBar;
-};
\ No newline at end of file
+};
