[UnrealEd]

diff --git a/Engine/Source/Editor/UnrealEd/Public/VolumeTexturePreview.h b/Engine/Source/Editor/UnrealEd/Public/VolumeTexturePreview.h
--- a/Engine/Source/Editor/UnrealEd/Public/VolumeTexturePreview.h
+++ b/Engine/Source/Editor/UnrealEd/Public/VolumeTexturePreview.h
@@ -12,7 +12,7 @@
 /**
  * Batched element parameters for previewing 2d textures.
  */
-class UNREALED_API FBatchedElementVolumeTexturePreviewParameters : public FBatchedElementParameters
+class FBatchedElementVolumeTexturePreviewParameters : public FBatchedElementParameters
 {
 public:
 	FBatchedElementVolumeTexturePreviewParameters(bool InViewModeAsDepthSlices, int32 InSizeZ, float InMipLevel, float InOpacity, bool InShowSlices, const FRotator& InTraceOrientation, bool bInUsePointSampling)
@@ -27,7 +27,7 @@ public:
 	}
 
 	/** Binds vertex and pixel shaders for this element */
-	virtual void BindShaders(FRHICommandList& RHICmdList, FGraphicsPipelineStateInitializer& GraphicsPSOInit, ERHIFeatureLevel::Type InFeatureLevel, const FMatrix& InTransform, const float InGamma, const FMatrix& ColorWeights, const FTexture* Texture) override;
+	UNREALED_API virtual void BindShaders(FRHICommandList& RHICmdList, FGraphicsPipelineStateInitializer& GraphicsPSOInit, ERHIFeatureLevel::Type InFeatureLevel, const FMatrix& InTransform, const float InGamma, const FMatrix& ColorWeights, const FTexture* Texture) override;
 
 private:
 
diff --git a/Engine/Source/Editor/UnrealEd/Public/WidgetModeManager.h b/Engine/Source/Editor/UnrealEd/Public/WidgetModeManager.h
--- a/Engine/Source/Editor/UnrealEd/Public/WidgetModeManager.h
+++ b/Engine/Source/Editor/UnrealEd/Public/WidgetModeManager.h
@@ -13,30 +13,30 @@
  * 
  * In all other cases (Tick, FocusLost, etc), the default editor mode tool behavior is used and sufficient.
  */
-class UNREALED_API FWidgetModeManager : public FEditorModeTools
+class FWidgetModeManager : public FEditorModeTools
 {
 public:
-	FWidgetModeManager();
+	UNREALED_API FWidgetModeManager();
 
 	// ~Begin SWidget "Proxy" Interface
-	virtual bool OnKeyChar(const FGeometry& MyGeometry, const FCharacterEvent& InCharacterEvent);
-	virtual bool OnKeyDown(const FGeometry& MyGeometry, const FKeyEvent& InKeyEvent);
-	virtual bool OnKeyUp(const FGeometry& MyGeometry, const FKeyEvent& InKeyEvent);
-	virtual bool OnMouseButtonDown(const FGeometry& MyGeometry, const FPointerEvent& MouseEvent);
-	virtual bool OnMouseButtonUp(const FGeometry& MyGeometry, const FPointerEvent& MouseEvent);
-	virtual bool OnMouseMove(const FGeometry& MyGeometry, const FPointerEvent& MouseEvent);
-	virtual bool OnMouseWheel(const FGeometry& MyGeometry, const FPointerEvent& MouseEvent);
-	virtual bool OnMouseButtonDoubleClick(const FGeometry& InMyGeometry, const FPointerEvent& InMouseEvent);
-	virtual bool OnDragDetected(const FGeometry& MyGeometry, const FPointerEvent& MouseEvent);
-	virtual void OnMouseCaptureLost(const FCaptureLostEvent& CaptureLostEvent);
-	virtual int32 OnPaint(const FPaintArgs& Args, const FGeometry& AllottedGeometry, const FSlateRect& MyCullingRect, FSlateWindowElementList& OutDrawElements, int32 LayerId, const FWidgetStyle& InWidgetStyle, bool bParentEnabled);
+	UNREALED_API virtual bool OnKeyChar(const FGeometry& MyGeometry, const FCharacterEvent& InCharacterEvent);
+	UNREALED_API virtual bool OnKeyDown(const FGeometry& MyGeometry, const FKeyEvent& InKeyEvent);
+	UNREALED_API virtual bool OnKeyUp(const FGeometry& MyGeometry, const FKeyEvent& InKeyEvent);
+	UNREALED_API virtual bool OnMouseButtonDown(const FGeometry& MyGeometry, const FPointerEvent& MouseEvent);
+	UNREALED_API virtual bool OnMouseButtonUp(const FGeometry& MyGeometry, const FPointerEvent& MouseEvent);
+	UNREALED_API virtual bool OnMouseMove(const FGeometry& MyGeometry, const FPointerEvent& MouseEvent);
+	UNREALED_API virtual bool OnMouseWheel(const FGeometry& MyGeometry, const FPointerEvent& MouseEvent);
+	UNREALED_API virtual bool OnMouseButtonDoubleClick(const FGeometry& InMyGeometry, const FPointerEvent& InMouseEvent);
+	UNREALED_API virtual bool OnDragDetected(const FGeometry& MyGeometry, const FPointerEvent& MouseEvent);
+	UNREALED_API virtual void OnMouseCaptureLost(const FCaptureLostEvent& CaptureLostEvent);
+	UNREALED_API virtual int32 OnPaint(const FPaintArgs& Args, const FGeometry& AllottedGeometry, const FSlateRect& MyCullingRect, FSlateWindowElementList& OutDrawElements, int32 LayerId, const FWidgetStyle& InWidgetStyle, bool bParentEnabled);
 	// ~End SWidget "Proxy" Interface
 
 	/** Not done in constructor as managed widget typically will not exist during mode construction for editors */
-	virtual void SetManagedWidget(TSharedPtr<SWidget> InManagedWidget);
+	UNREALED_API virtual void SetManagedWidget(TSharedPtr<SWidget> InManagedWidget);
 
 	/** Not done in constructor as managed widget typically will not exist during mode construction for editors */
-	TSharedPtr<SWidget> GetManagedWidget() const;
+	UNREALED_API TSharedPtr<SWidget> GetManagedWidget() const;
 	
 	/** Toolkit owning the widget of interest */
 	TWeakPtr<IToolkit> OwningToolkit;
diff --git a/Engine/Source/Editor/UnrealEd/Public/WorkflowOrientedApp/ApplicationMode.h b/Engine/Source/Editor/UnrealEd/Public/WorkflowOrientedApp/ApplicationMode.h
--- a/Engine/Source/Editor/UnrealEd/Public/WorkflowOrientedApp/ApplicationMode.h
+++ b/Engine/Source/Editor/UnrealEd/Public/WorkflowOrientedApp/ApplicationMode.h
@@ -15,7 +15,7 @@ class FExtender;
 /////////////////////////////////////////////////////
 // FApplicationMode
 
-class UNREALED_API FApplicationMode : public TSharedFromThis<FApplicationMode>
+class FApplicationMode : public TSharedFromThis<FApplicationMode>
 {
 protected:
 	// The layout to use in this mode
@@ -38,14 +38,14 @@ protected:
 	TSharedPtr<FWorkspaceItem> WorkspaceMenuCategory;
 
 public:
-	FApplicationMode(FName InModeName);
+	UNREALED_API FApplicationMode(FName InModeName);
 
-	FApplicationMode(FName InModeName, FText(*GetLocalizedMode)(const FName));
+	UNREALED_API FApplicationMode(FName InModeName, FText(*GetLocalizedMode)(const FName));
 
 	virtual ~FApplicationMode() {}
 
-	void DeactivateMode(TSharedPtr<FTabManager> InTabManager);
-	TSharedRef<FTabManager::FLayout> ActivateMode(TSharedPtr<FTabManager> InTabManager);
+	UNREALED_API void DeactivateMode(TSharedPtr<FTabManager> InTabManager);
+	UNREALED_API TSharedRef<FTabManager::FLayout> ActivateMode(TSharedPtr<FTabManager> InTabManager);
 	virtual void RegisterTabFactories(TSharedPtr<FTabManager> InTabManager) {}
 
 	virtual void AddTabFactory(FCreateWorkflowTabFactory FactoryCreator) {}
diff --git a/Engine/Source/Editor/UnrealEd/Public/WorkflowOrientedApp/SContentReference.h b/Engine/Source/Editor/UnrealEd/Public/WorkflowOrientedApp/SContentReference.h
--- a/Engine/Source/Editor/UnrealEd/Public/WorkflowOrientedApp/SContentReference.h
+++ b/Engine/Source/Editor/UnrealEd/Public/WorkflowOrientedApp/SContentReference.h
@@ -17,7 +17,7 @@ class SBorder;
 class SComboButton;
 
 // This widget is a content reference, which can optionally show buttons for interacting with the reference
-class UNREALED_API SContentReference : public SCompoundWidget
+class SContentReference : public SCompoundWidget
 {
 public:
 
@@ -80,34 +80,34 @@ public:
 
 public:
 	// Construct a SContentReference
-	void Construct(const FArguments& InArgs);
+	UNREALED_API void Construct(const FArguments& InArgs);
 
 	// Causes the asset picker to open, allowing the user to select a new asset
-	void OpenAssetPickerMenu();
+	UNREALED_API void OpenAssetPickerMenu();
 
 protected:
-	EVisibility GetUseButtonVisibility() const;
-	EVisibility GetPickButtonVisibility() const;
-	EVisibility GetFindButtonVisibility() const;
-	EVisibility GetClearButtonVisibility() const;
-	EVisibility GetToolsButtonVisibility() const;
+	UNREALED_API EVisibility GetUseButtonVisibility() const;
+	UNREALED_API EVisibility GetPickButtonVisibility() const;
+	UNREALED_API EVisibility GetFindButtonVisibility() const;
+	UNREALED_API EVisibility GetClearButtonVisibility() const;
+	UNREALED_API EVisibility GetToolsButtonVisibility() const;
 
-	FReply OnClickUseButton();
-	FReply OnClickFindButton();
-	FReply OnClickClearButton();
+	UNREALED_API FReply OnClickUseButton();
+	UNREALED_API FReply OnClickFindButton();
+	UNREALED_API FReply OnClickClearButton();
 
-	TSharedRef<SWidget> MakeAssetPickerMenu();
-	void OnAssetSelectedFromPicker(const FAssetData& AssetData);
+	UNREALED_API TSharedRef<SWidget> MakeAssetPickerMenu();
+	UNREALED_API void OnAssetSelectedFromPicker(const FAssetData& AssetData);
 
-	FText GetAssetShortName() const;
-	FText GetAssetFullName() const;
+	UNREALED_API FText GetAssetShortName() const;
+	UNREALED_API FText GetAssetFullName() const;
 
-	FReply OnDoubleClickedOnAssetName(const FGeometry& InMyGeometry, const FPointerEvent& InMouseEvent);
-	void OpenAssetToEdit();
+	UNREALED_API FReply OnDoubleClickedOnAssetName(const FGeometry& InMyGeometry, const FPointerEvent& InMouseEvent);
+	UNREALED_API void OpenAssetToEdit();
 
-	const FSlateBrush* GetBorderImage() const;
+	UNREALED_API const FSlateBrush* GetBorderImage() const;
 
-	static void FindObjectInContentBrowser(UObject* Object);
+	static UNREALED_API void FindObjectInContentBrowser(UObject* Object);
 protected:
 	// Attributes
 	TAttribute<UObject*> AssetReference;
diff --git a/Engine/Source/Editor/UnrealEd/Public/WorkflowOrientedApp/SModeWidget.h b/Engine/Source/Editor/UnrealEd/Public/WorkflowOrientedApp/SModeWidget.h
--- a/Engine/Source/Editor/UnrealEd/Public/WorkflowOrientedApp/SModeWidget.h
+++ b/Engine/Source/Editor/UnrealEd/Public/WorkflowOrientedApp/SModeWidget.h
@@ -22,7 +22,7 @@ DECLARE_DELEGATE_OneParam( FOnModeChangeRequested, FName );
 enum class ECheckBoxState : uint8;
 
 // This is the base class for a mode widget in a workflow oriented editor
-class UNREALED_API SModeWidget : public SCompoundWidget
+class SModeWidget : public SCompoundWidget
 {
 public:
 	SLATE_BEGIN_ARGS(SModeWidget)
@@ -49,7 +49,7 @@ public:
 
 public:
 	// Construct a SModeWidget
-	void Construct(const FArguments& InArgs, const FText& InText, const FName InMode);
+	UNREALED_API void Construct(const FArguments& InArgs, const FText& InText, const FName InMode);
 
 private:
 	bool IsActiveMode() const;
diff --git a/Engine/Source/Editor/UnrealEd/Public/WorkflowOrientedApp/WorkflowCentricApplication.h b/Engine/Source/Editor/UnrealEd/Public/WorkflowOrientedApp/WorkflowCentricApplication.h
--- a/Engine/Source/Editor/UnrealEd/Public/WorkflowOrientedApp/WorkflowCentricApplication.h
+++ b/Engine/Source/Editor/UnrealEd/Public/WorkflowOrientedApp/WorkflowCentricApplication.h
@@ -16,32 +16,32 @@ DECLARE_DELEGATE_RetVal_TwoParams(TSharedRef<FApplicationMode>, FWorkflowApplica
 /////////////////////////////////////////////////////
 // FWorkflowCentricApplication
 
-class UNREALED_API FWorkflowCentricApplication : public FAssetEditorToolkit
+class FWorkflowCentricApplication : public FAssetEditorToolkit
 {
 public:
 	// IToolkit interface
-	virtual void RegisterTabSpawners(const TSharedRef<class FTabManager>& TabManager) override;
-	virtual void UnregisterTabSpawners(const TSharedRef<class FTabManager>& TabManager) override;
+	UNREALED_API virtual void RegisterTabSpawners(const TSharedRef<class FTabManager>& TabManager) override;
+	UNREALED_API virtual void UnregisterTabSpawners(const TSharedRef<class FTabManager>& TabManager) override;
 	// End of IToolkit interface
 	
 	// FAssetEditorToolkit interface
-	virtual bool OnRequestClose(EAssetEditorCloseReason InCloseReason) override;
-	virtual void OnClose() override;
-	virtual FName GetToolMenuToolbarName(FName& OutParentName) const override;
+	UNREALED_API virtual bool OnRequestClose(EAssetEditorCloseReason InCloseReason) override;
+	UNREALED_API virtual void OnClose() override;
+	UNREALED_API virtual FName GetToolMenuToolbarName(FName& OutParentName) const override;
 	// End of FAssetEditorToolkit interface
 
 	// Returns the name of the toolbar for the given mode along with the name of the toolbar's parent.
-	FName GetToolMenuToolbarNameForMode(const FName InModeName, FName& OutParentName) const;
+	UNREALED_API FName GetToolMenuToolbarNameForMode(const FName InModeName, FName& OutParentName) const;
 
 	// Returns the current mode of this application
-	FName GetCurrentMode() const;
+	UNREALED_API FName GetCurrentMode() const;
 	TSharedPtr<FApplicationMode> GetCurrentModePtr() const { return CurrentAppModePtr; }
 	bool IsModeCurrent(FName ModeToCheck) const {return GetCurrentMode() == ModeToCheck;}
 
 	// Attempt to set the current mode.  If this mode is illegal or unknown, the mode will remain unchanged.
-	virtual void SetCurrentMode(FName NewMode);
+	UNREALED_API virtual void SetCurrentMode(FName NewMode);
 
-	void PushTabFactories(FWorkflowAllowedTabSet& FactorySetToPush);
+	UNREALED_API void PushTabFactories(FWorkflowAllowedTabSet& FactorySetToPush);
 
 	// Gets the mode extender list for all workflow applications (append to customize a specific mode)
 	static TArray<FWorkflowApplicationModeExtender>& GetModeExtenderList() { return ModeExtenderList; }
@@ -52,10 +52,10 @@ public:
 	* @param InModeName - Name of the mode to register toolbar for
 	* @return nullptr if toolbar already registered.
 	*/
-	virtual UToolMenu* RegisterModeToolbarIfUnregistered(const FName InModeName);
+	UNREALED_API virtual UToolMenu* RegisterModeToolbarIfUnregistered(const FName InModeName);
 protected:
 
-	virtual void AddApplicationMode(FName ModeName, TSharedRef<FApplicationMode> Mode);
+	UNREALED_API virtual void AddApplicationMode(FName ModeName, TSharedRef<FApplicationMode> Mode);
 protected:
 	TSharedPtr<FApplicationMode> CurrentAppModePtr;
 
@@ -63,7 +63,7 @@ private:
 	// List of modes; do not access directly, use AddApplicationMode and SetCurrentMode
 	TMap< FName, TSharedPtr<FApplicationMode> > ApplicationModeList;
 
-	static TArray<FWorkflowApplicationModeExtender> ModeExtenderList;
+	static UNREALED_API TArray<FWorkflowApplicationModeExtender> ModeExtenderList;
 };
 
 
diff --git a/Engine/Source/Editor/UnrealEd/Public/WorkflowOrientedApp/WorkflowTabFactory.h b/Engine/Source/Editor/UnrealEd/Public/WorkflowOrientedApp/WorkflowTabFactory.h
--- a/Engine/Source/Editor/UnrealEd/Public/WorkflowOrientedApp/WorkflowTabFactory.h
+++ b/Engine/Source/Editor/UnrealEd/Public/WorkflowOrientedApp/WorkflowTabFactory.h
@@ -68,7 +68,7 @@ public:
 /////////////////////////////////////////////////////
 // FWorkflowTabFactory
 
-class UNREALED_API FWorkflowTabFactory : public TSharedFromThis<FWorkflowTabFactory>
+class FWorkflowTabFactory : public TSharedFromThis<FWorkflowTabFactory>
 {
 protected:
 	FName TabIdentifier;
@@ -91,9 +91,9 @@ public:
 	virtual ~FWorkflowTabFactory() {}
 
 	// Creates a toolkit tab given the spawning information
-	virtual TSharedRef<SDockTab> SpawnTab(const FWorkflowTabSpawnInfo& Info) const;
+	UNREALED_API virtual TSharedRef<SDockTab> SpawnTab(const FWorkflowTabSpawnInfo& Info) const;
 	// Creates a toolkit tab without any content, given the spawning information
-	virtual TSharedRef<SDockTab> SpawnBlankTab() const;
+	UNREALED_API virtual TSharedRef<SDockTab> SpawnBlankTab() const;
 
 	/**
 	 * Updates a given tab with new content
@@ -102,17 +102,17 @@ public:
 	 * @param InInfo		The spawning info, used to generate required tab information
 	 * @param InContent		The content to update the tab with
 	 */
-	virtual void UpdateTab(TSharedPtr<SDockTab> InDockTab, const FWorkflowTabSpawnInfo& Info, TSharedPtr< SWidget > InContent) const;
+	UNREALED_API virtual void UpdateTab(TSharedPtr<SDockTab> InDockTab, const FWorkflowTabSpawnInfo& Info, TSharedPtr< SWidget > InContent) const;
 
 	// Creates a view menu entry for this factory
-	virtual void CreateViewMenuEntry(FMenuBuilder& MenuBuilder, const FUIAction& Action) const;
+	UNREALED_API virtual void CreateViewMenuEntry(FMenuBuilder& MenuBuilder, const FUIAction& Action) const;
 
 	/** Register a tab spawner for this factory */
-	virtual FTabSpawnerEntry& RegisterTabSpawner(TSharedRef<FTabManager> TabManager, const FApplicationMode* CurrentApplicationMode) const;
+	UNREALED_API virtual FTabSpawnerEntry& RegisterTabSpawner(TSharedRef<FTabManager> TabManager, const FApplicationMode* CurrentApplicationMode) const;
 
 public:
 	// Constructor
-	FWorkflowTabFactory(FName InIdentifier, TSharedPtr<class FAssetEditorToolkit> InHostingApp);
+	UNREALED_API FWorkflowTabFactory(FName InIdentifier, TSharedPtr<class FAssetEditorToolkit> InHostingApp);
 
 	// Get the identifier of this factory
 	FName GetIdentifier() const { return TabIdentifier; }
@@ -121,34 +121,34 @@ public:
 	bool IsSingleton() const { return bIsSingleton; }
 
 	// Creates the label for the tab
-	virtual TAttribute<FText> ConstructTabName(const FWorkflowTabSpawnInfo& Info) const;
+	UNREALED_API virtual TAttribute<FText> ConstructTabName(const FWorkflowTabSpawnInfo& Info) const;
 
 	/** Returns TRUE if the payload is valid */
 	virtual bool IsPayloadValid(TSharedRef<FTabPayload> Payload) const { return false; }
 
 protected:
 	// Call this to enable the default padding surrounding the tab contents
-	void EnableTabPadding();
+	UNREALED_API void EnableTabPadding();
 
 	/** Callback function for spawning the tab */
-	virtual TSharedRef<SDockTab> OnSpawnTab(const FSpawnTabArgs& SpawnArgs, TWeakPtr<FTabManager> WeakTabManager) const;
+	UNREALED_API virtual TSharedRef<SDockTab> OnSpawnTab(const FSpawnTabArgs& SpawnArgs, TWeakPtr<FTabManager> WeakTabManager) const;
 
 	/** Callback function for checking to see if a tab can/should be spawned */
-	bool CanSpawnTab(const FSpawnTabArgs& SpawnArgs, TWeakPtr<FTabManager> WeakTabManager) const;
+	UNREALED_API bool CanSpawnTab(const FSpawnTabArgs& SpawnArgs, TWeakPtr<FTabManager> WeakTabManager) const;
 
 public:
 	// Creates the widget that will be the body of the new tab
 	// Users will most commonly override just this.
-	virtual TSharedRef<SWidget> CreateTabBody(const FWorkflowTabSpawnInfo& Info) const;
+	UNREALED_API virtual TSharedRef<SWidget> CreateTabBody(const FWorkflowTabSpawnInfo& Info) const;
 
 	// Virtual to create an icon for the tab
-	virtual const FSlateBrush* GetTabIcon(const FWorkflowTabSpawnInfo& Info) const;
+	UNREALED_API virtual const FSlateBrush* GetTabIcon(const FWorkflowTabSpawnInfo& Info) const;
 	
 	// Virtual to create an icon for the menu entry
-	virtual const FSlateIcon& GetTabSpawnerIcon(const FWorkflowTabSpawnInfo& Info) const;
+	UNREALED_API virtual const FSlateIcon& GetTabSpawnerIcon(const FWorkflowTabSpawnInfo& Info) const;
 
 	// Create a tooltip widget for the tab
-	virtual TSharedPtr<SToolTip> CreateTabToolTipWidget(const FWorkflowTabSpawnInfo& Info) const;
+	UNREALED_API virtual TSharedPtr<SToolTip> CreateTabToolTipWidget(const FWorkflowTabSpawnInfo& Info) const;
 
 	// Virtual to get optional extra tooltip text for tab
 	virtual FText GetTabToolTipText(const FWorkflowTabSpawnInfo& Info) const
@@ -163,7 +163,7 @@ private:
 /////////////////////////////////////////////////////
 // FDocumentTabFactory
 
-class UNREALED_API FDocumentTabFactory : public FWorkflowTabFactory
+class FDocumentTabFactory : public FWorkflowTabFactory
 {
 public:
 	// Saves the state for this tab
@@ -189,8 +189,8 @@ public:
 	virtual bool IsPayloadSupported(TSharedRef<FTabPayload> Payload) const { return false; }
 
 	/** Creates a tab history node from the passed in payload */
-	virtual TSharedRef<struct FGenericTabHistory> CreateTabHistoryNode(TSharedPtr<FTabPayload> Payload);
+	UNREALED_API virtual TSharedRef<struct FGenericTabHistory> CreateTabHistoryNode(TSharedPtr<FTabPayload> Payload);
 
 protected:
-	FDocumentTabFactory(FName InIdentifier, TSharedPtr<class FAssetEditorToolkit> InHostingApp);
+	UNREALED_API FDocumentTabFactory(FName InIdentifier, TSharedPtr<class FAssetEditorToolkit> InHostingApp);
 };
diff --git a/Engine/Source/Editor/UnrealEd/Public/WorkflowOrientedApp/WorkflowTabManager.h b/Engine/Source/Editor/UnrealEd/Public/WorkflowOrientedApp/WorkflowTabManager.h
--- a/Engine/Source/Editor/UnrealEd/Public/WorkflowOrientedApp/WorkflowTabManager.h
+++ b/Engine/Source/Editor/UnrealEd/Public/WorkflowOrientedApp/WorkflowTabManager.h
@@ -26,26 +26,26 @@ namespace EDocumentOpenBehavior
 /////////////////////////////////////////////////////
 // FWorkflowAllowedTabSet
 
-class UNREALED_API FWorkflowAllowedTabSet
+class FWorkflowAllowedTabSet
 {
 public:
 	// Searches this set for a factory with the specified ID, or returns NULL.
-	TSharedPtr<class FWorkflowTabFactory> GetFactory(FName FactoryID) const;
+	UNREALED_API TSharedPtr<class FWorkflowTabFactory> GetFactory(FName FactoryID) const;
 
 	// Registers a factory with this set - must be unique.
-	void RegisterFactory(TSharedPtr<class FWorkflowTabFactory> Factory);
+	UNREALED_API void RegisterFactory(TSharedPtr<class FWorkflowTabFactory> Factory);
 
 	// Unregisters a factory with this specified ID.
-	void UnregisterFactory(FName FactoryID);
+	UNREALED_API void UnregisterFactory(FName FactoryID);
 
 	// Merges in a set of factories into this set - duplicates are OK.
-	void MergeInSet(const FWorkflowAllowedTabSet& OtherSet);
+	UNREALED_API void MergeInSet(const FWorkflowAllowedTabSet& OtherSet);
 
 	// Clears the set
-	void Clear();
+	UNREALED_API void Clear();
 
 	// Allows iteration over Factories
-	TMap< FName, TSharedPtr<class FWorkflowTabFactory> >::TIterator CreateIterator();
+	UNREALED_API TMap< FName, TSharedPtr<class FWorkflowTabFactory> >::TIterator CreateIterator();
 
 private:
 	// Map of available factories in this mode
@@ -57,15 +57,15 @@ private:
 /////////////////////////////////////////////////////
 // FTabInfo
 
-class UNREALED_API FTabInfo : public TSharedFromThis<FTabInfo>
+class FTabInfo : public TSharedFromThis<FTabInfo>
 {
 public:
-	FTabInfo(const TSharedRef<SDockTab>& InTab, const TSharedPtr<FDocumentTabFactory>& InSpawner, const TSharedPtr<class FDocumentTracker>& InTracker);
+	UNREALED_API FTabInfo(const TSharedRef<SDockTab>& InTab, const TSharedPtr<FDocumentTabFactory>& InSpawner, const TSharedPtr<class FDocumentTracker>& InTracker);
 
 	/** Returns TRUE if the payload used by this TabInfo is the same as passed in */
-	bool PayloadMatches(const TSharedPtr<FTabPayload> TestPayload) const;
+	UNREALED_API bool PayloadMatches(const TSharedPtr<FTabPayload> TestPayload) const;
 
-	static bool PayloadMatches(TSharedPtr<FTabPayload> A, TSharedPtr<FTabPayload> B);
+	static UNREALED_API bool PayloadMatches(TSharedPtr<FTabPayload> A, TSharedPtr<FTabPayload> B);
 
 	/**
 	 * Adds history data to immediately after the current location in the history list, wiping out any history after in the process
@@ -73,28 +73,28 @@ public:
 	 * @param InHistoryNode		The history node to add
 	 * @param bInSaveHistory	TRUE if history should be saved
 	 */
-	void AddTabHistory(TSharedPtr< struct FGenericTabHistory > InHistoryNode, bool bInSaveHistory = true);
+	UNREALED_API void AddTabHistory(TSharedPtr< struct FGenericTabHistory > InHistoryNode, bool bInSaveHistory = true);
 
 	/** Retrieves the history currently visible in this tab */
-	TSharedPtr<struct FGenericTabHistory> GetCurrentHistory() const;
+	UNREALED_API TSharedPtr<struct FGenericTabHistory> GetCurrentHistory() const;
 
 	/** Sets the current history and evoke it, this can either be an old or new history node */
-	void SetCurrentHistory(TSharedPtr<struct FGenericTabHistory> NewHistory, bool bInSaveHistory = true, bool bShouldRestore = true);
+	UNREALED_API void SetCurrentHistory(TSharedPtr<struct FGenericTabHistory> NewHistory, bool bInSaveHistory = true, bool bShouldRestore = true);
 
 	/** Single step forward in history */
-	FReply OnGoForwardInHistory();
+	UNREALED_API FReply OnGoForwardInHistory();
 
 	/** Single step back in history */
-	FReply OnGoBackInHistory();
+	UNREALED_API FReply OnGoBackInHistory();
 
 	/** Jumps backwards to the nearest valid history. If no history backwards is valid, jumps forward. In the case that there is no valid history, the tab is closed */
-	void JumpToNearestValidHistoryData();
+	UNREALED_API void JumpToNearestValidHistoryData();
 
 	/** Retrieves the factory from the current history data, which is being displayed by the tab */
-	TWeakPtr<FDocumentTabFactory> GetFactory() const;
+	UNREALED_API TWeakPtr<FDocumentTabFactory> GetFactory() const;
 
 	/** Retrieves the payload from the current history data, which is being displayed by the tab */
-	TSharedPtr<FTabPayload> GetPayload() const;
+	UNREALED_API TSharedPtr<FTabPayload> GetPayload() const;
 
 	/** Retrieves the tab used by this TabInfo */
 	TWeakPtr<class SDockTab> GetTab() const
@@ -103,16 +103,16 @@ public:
 	}
 
 	/** Jumps immediately to an index in the history if valid */
-	void GoToHistoryIndex(int32 InHistoryIdx);
+	UNREALED_API void GoToHistoryIndex(int32 InHistoryIdx);
 
 	/** Returns TRUE if stepping backward in history is allowed */
-	bool CanStepBackwardInHistory() const;
+	UNREALED_API bool CanStepBackwardInHistory() const;
 
 	/** Returns TRUE if stepping forward in history is allowed */
-	bool CanStepForwardInHistory() const;
+	UNREALED_API bool CanStepForwardInHistory() const;
 
 	/** Creates the history navigation widget used for navigating back and forward in history */
-	TSharedRef< SWidget > CreateHistoryNavigationWidget();
+	UNREALED_API TSharedRef< SWidget > CreateHistoryNavigationWidget();
 
 private:
 	/**
@@ -120,7 +120,7 @@ private:
 	 *
 	 * @param InMenuAnchor		This is the anchor the menu will use for positioning
 	 */
-	FReply OnMouseDownHistory( const FGeometry& MyGeometry, const FPointerEvent& MouseEvent, TWeakPtr< SMenuAnchor > InMenuAnchor );
+	UNREALED_API FReply OnMouseDownHistory( const FGeometry& MyGeometry, const FPointerEvent& MouseEvent, TWeakPtr< SMenuAnchor > InMenuAnchor );
 
 	/** 
 	 * Callback to create the history menu.
@@ -129,7 +129,7 @@ private:
 	 *
 	 * @return						The menu widget displaying all available history
 	 */
-	TSharedRef<SWidget> CreateHistoryMenu(bool bInBackHistory) const;
+	UNREALED_API TSharedRef<SWidget> CreateHistoryMenu(bool bInBackHistory) const;
 
 private:	
 	/** Cached history navigation widget */
@@ -243,19 +243,19 @@ protected:
 
 DECLARE_DELEGATE_OneParam(FOnTabActivated, TSharedPtr<SDockTab>);
 
-class UNREALED_API FDocumentTracker : public TSharedFromThis<FDocumentTracker>
+class FDocumentTracker : public TSharedFromThis<FDocumentTracker>
 {
 public:
-	FDocumentTracker();
-	~FDocumentTracker();
+	UNREALED_API FDocumentTracker();
+	UNREALED_API ~FDocumentTracker();
 
-	void ClearDocumentFactories();
-	void RegisterDocumentFactory(TSharedPtr<class FDocumentTabFactory> Factory);
+	UNREALED_API void ClearDocumentFactories();
+	UNREALED_API void RegisterDocumentFactory(TSharedPtr<class FDocumentTabFactory> Factory);
 
 	// Creates a new tab manager
-	void Initialize(TSharedPtr<FAssetEditorToolkit> InHostingApp);
+	UNREALED_API void Initialize(TSharedPtr<FAssetEditorToolkit> InHostingApp);
 
-	void SetTabManager(const TSharedRef<FTabManager>& InTabManager);
+	UNREALED_API void SetTabManager(const TSharedRef<FTabManager>& InTabManager);
 
 	/** Are we opening a new document, or restoring a previously opened document */
 	enum EOpenDocumentCause
@@ -291,37 +291,37 @@ public:
 	 *
 	 * @return							The tab modified
 	 */
-	TSharedPtr<SDockTab> OpenDocument(TSharedPtr<FTabPayload> InPayload, EOpenDocumentCause InOpenCause);
+	UNREALED_API TSharedPtr<SDockTab> OpenDocument(TSharedPtr<FTabPayload> InPayload, EOpenDocumentCause InOpenCause);
 
 	// Closes any tabs of the specified type that match the payload (or all of them if the payload is NULL)
-	void CloseTab(TSharedPtr<FTabPayload> Payload);
+	UNREALED_API void CloseTab(TSharedPtr<FTabPayload> Payload);
 
 	/** Cleans any invalid tabs, either moving them to a valid place in history or closing the tab altogether */
-	void CleanInvalidTabs();
+	UNREALED_API void CleanInvalidTabs();
 
 	// Calls OnTabRefreshed for each open tab (on the factories that created them)
-	void RefreshAllTabs() const;
+	UNREALED_API void RefreshAllTabs() const;
 
 	// Finds a factory that can deal with the supplied payload
-	TSharedPtr<FDocumentTabFactory> FindSupportingFactory(TSharedRef<FTabPayload> Payload) const;
+	UNREALED_API TSharedPtr<FDocumentTabFactory> FindSupportingFactory(TSharedRef<FTabPayload> Payload) const;
 
 	// Finds all tabs that match the specified identifier and optional payload, placing them in the specified array
-	void FindMatchingTabs(TSharedPtr<FTabPayload> Payload, /*inout*/ TArray< TSharedPtr<SDockTab> >& Results);
+	UNREALED_API void FindMatchingTabs(TSharedPtr<FTabPayload> Payload, /*inout*/ TArray< TSharedPtr<SDockTab> >& Results);
 
 	// Finds all tabs that match the specified identifier, placing them in the specified array
-	void FindAllTabsForFactory(const TWeakPtr<FDocumentTabFactory>& Factory, /*inout*/ TArray< TSharedPtr<SDockTab> >& Results);
+	UNREALED_API void FindAllTabsForFactory(const TWeakPtr<FDocumentTabFactory>& Factory, /*inout*/ TArray< TSharedPtr<SDockTab> >& Results);
 
 	// Saves the state of all tabs
-	void SaveAllState();
+	UNREALED_API void SaveAllState();
 
 	// Replaces the open payload in the specified tab with a new one; recreating the contents
-	void ReplacePayloadInTab(TSharedPtr<SDockTab> TargetTab, TSharedPtr<FTabPayload> NewPayload);
+	UNREALED_API void ReplacePayloadInTab(TSharedPtr<SDockTab> TargetTab, TSharedPtr<FTabPayload> NewPayload);
 
 	/** Get all the live documents tracked by this tracker */
-	TArray< TSharedPtr<SDockTab> > GetAllDocumentTabs() const;
+	UNREALED_API TArray< TSharedPtr<SDockTab> > GetAllDocumentTabs() const;
 
 	/** Returns the current active tab */
-	TSharedPtr<SDockTab> GetActiveTab() const;
+	UNREALED_API TSharedPtr<SDockTab> GetActiveTab() const;
 private:
 
 	// List of open tabs; used to bring one to front and focus it
@@ -356,28 +356,28 @@ private:
 
 private:
 	// Clean the spawned list 
-	FTabList& GetSpawnedList();
+	UNREALED_API FTabList& GetSpawnedList();
 
 	// Cleans the LastEditedTabInfo and returns pinned
-	TSharedPtr<FTabInfo> GetLastEditedTabInfo();
+	UNREALED_API TSharedPtr<FTabInfo> GetLastEditedTabInfo();
 
 	// Called by the global active tab changed callback; dispatches to individually registered callbacks
-	void OnActiveTabChanged(TSharedPtr<SDockTab> PreviouslyActive, TSharedPtr<SDockTab> NewlyActivated);
+	UNREALED_API void OnActiveTabChanged(TSharedPtr<SDockTab> PreviouslyActive, TSharedPtr<SDockTab> NewlyActivated);
 
 	// Called by the global tab manager foregrounds a tab; dispatches to individually registered callbacks
-	void OnTabForegrounded(TSharedPtr<SDockTab> PreviouslyActive, TSharedPtr<SDockTab> NewlyActivated);
+	UNREALED_API void OnTabForegrounded(TSharedPtr<SDockTab> PreviouslyActive, TSharedPtr<SDockTab> NewlyActivated);
 
 	/** Callback for navigating tabs with buttons */
-	FReply OnNavigateTab(FDocumentTracker::EOpenDocumentCause InCause);
+	UNREALED_API FReply OnNavigateTab(FDocumentTracker::EOpenDocumentCause InCause);
 
 	/** Helper function to find the tab in the foreground */
-	TWeakPtr< FTabInfo > FindTabInForeground();
+	UNREALED_API TWeakPtr< FTabInfo > FindTabInForeground();
 
 	/** Navigates to specific index in tab history, returns success if it was found */
-	bool NavigateToTabHistory(int32 InHistoryIdx);
+	UNREALED_API bool NavigateToTabHistory(int32 InHistoryIdx);
 
 	/** Returns the currently active tab history, may be null */
-	TSharedPtr<FGenericTabHistory> GetCurrentTabHistory();
+	UNREALED_API TSharedPtr<FGenericTabHistory> GetCurrentTabHistory();
 
 	/**
 	 * Navigates current tab
@@ -387,7 +387,7 @@ private:
 	 *
 	 * @return							The tab navigated and modified
 	 */
-	TSharedPtr<SDockTab> NavigateCurrentTab(TSharedPtr<FTabPayload> InPayload, EOpenDocumentCause InNavigateCause);
+	UNREALED_API TSharedPtr<SDockTab> NavigateCurrentTab(TSharedPtr<FTabPayload> InPayload, EOpenDocumentCause InNavigateCause);
 
 	/**
 	 * Opens a new tab
@@ -397,7 +397,7 @@ private:
 	 *
 	 * @return							The tab opened and modified
 	 */
-	TSharedPtr<SDockTab> OpenNewTab(TSharedPtr<FGenericTabHistory> InTabHistory, EOpenDocumentCause InOpenCause);
+	UNREALED_API TSharedPtr<SDockTab> OpenNewTab(TSharedPtr<FGenericTabHistory> InTabHistory, EOpenDocumentCause InOpenCause);
 
 	friend class FTabInfo;
 };
diff --git a/Engine/Source/Editor/UnrealEd/Public/WorldFolders.h b/Engine/Source/Editor/UnrealEd/Public/WorldFolders.h
--- a/Engine/Source/Editor/UnrealEd/Public/WorldFolders.h
+++ b/Engine/Source/Editor/UnrealEd/Public/WorldFolders.h
@@ -66,49 +66,49 @@ struct FActorPlacementFolder
 };
 
 /** Per-World Actor Folders UObject (used to support undo/redo reliably) */
-UCLASS()
-class UNREALED_API UWorldFolders : public UObject
+UCLASS(MinimalAPI)
+class UWorldFolders : public UObject
 {
 public:
 	GENERATED_BODY()
 
-	void Initialize(UWorld* InWorld);
-
-	void RebuildList();
-	bool AddFolder(const FFolder& InFolder);
-	bool RemoveFolder(const FFolder& InFolder, bool bShouldDeleteFolder = false);
-	bool RenameFolder(const FFolder& InOldFolder, const FFolder& InNewFolder);
-	bool IsFolderExpanded(const FFolder& InFolder) const;
-	bool SetIsFolderExpanded(const FFolder& InFolder, bool bIsExpanded);
-	FFolder GetActorEditorContextFolder(bool bMustMatchCurrentLevel) const;
-	bool SetActorEditorContextFolder(const FFolder& InFolder);
-	void PushActorEditorContext();
-	void PopActorEditorContext();
-	bool ContainsFolder(const FFolder& InFolder) const;
-	void ForEachFolder(TFunctionRef<bool(const FFolder&)> Operation);
-	void ForEachFolderWithRootObject(const FFolder::FRootObject& InFolderRootObject, TFunctionRef<bool(const FFolder&)> Operation);
-	void SaveState();
-	UWorld* GetWorld() const;
+	UNREALED_API void Initialize(UWorld* InWorld);
+
+	UNREALED_API void RebuildList();
+	UNREALED_API bool AddFolder(const FFolder& InFolder);
+	UNREALED_API bool RemoveFolder(const FFolder& InFolder, bool bShouldDeleteFolder = false);
+	UNREALED_API bool RenameFolder(const FFolder& InOldFolder, const FFolder& InNewFolder);
+	UNREALED_API bool IsFolderExpanded(const FFolder& InFolder) const;
+	UNREALED_API bool SetIsFolderExpanded(const FFolder& InFolder, bool bIsExpanded);
+	UNREALED_API FFolder GetActorEditorContextFolder(bool bMustMatchCurrentLevel) const;
+	UNREALED_API bool SetActorEditorContextFolder(const FFolder& InFolder);
+	UNREALED_API void PushActorEditorContext();
+	UNREALED_API void PopActorEditorContext();
+	UNREALED_API bool ContainsFolder(const FFolder& InFolder) const;
+	UNREALED_API void ForEachFolder(TFunctionRef<bool(const FFolder&)> Operation);
+	UNREALED_API void ForEachFolderWithRootObject(const FFolder::FRootObject& InFolderRootObject, TFunctionRef<bool(const FFolder&)> Operation);
+	UNREALED_API void SaveState();
+	UNREALED_API UWorld* GetWorld() const;
 
 	//~ Begin UObject
-	virtual void Serialize(FArchive& Ar) override;
+	UNREALED_API virtual void Serialize(FArchive& Ar) override;
 	//~ End UObject
 
 	//~ Begin Deprecated
-	FActorFolderProps* GetFolderProperties(const FFolder& InFolder);
+	UNREALED_API FActorFolderProps* GetFolderProperties(const FFolder& InFolder);
 	//~ End Deprecated
 
 private:
 
-	void BroadcastOnActorFolderCreated(const FFolder& InFolder);
-	void BroadcastOnActorFolderDeleted(const FFolder& InFolder);
-	void BroadcastOnActorFolderMoved(const FFolder& InSrcFolder, const FFolder& InDstFolder);
+	UNREALED_API void BroadcastOnActorFolderCreated(const FFolder& InFolder);
+	UNREALED_API void BroadcastOnActorFolderDeleted(const FFolder& InFolder);
+	UNREALED_API void BroadcastOnActorFolderMoved(const FFolder& InSrcFolder, const FFolder& InDstFolder);
 
-	FWorldFoldersImplementation& GetImpl(const FFolder& InFolder) const;
-	bool IsUsingPersistentFolders(const FFolder& InFolder) const;
+	UNREALED_API FWorldFoldersImplementation& GetImpl(const FFolder& InFolder) const;
+	UNREALED_API bool IsUsingPersistentFolders(const FFolder& InFolder) const;
 
-	FString GetWorldStateFilename() const;
-	void LoadState();
+	UNREALED_API FString GetWorldStateFilename() const;
+	UNREALED_API void LoadState();
 
 	TUniquePtr<FWorldPersistentFolders> PersistentFolders;
 	TUniquePtr<FWorldTransientFolders> TransientFolders;
@@ -125,4 +125,4 @@ private:
 	friend class FWorldFoldersImplementation;
 	friend class FWorldPersistentFolders;
 	friend class FWorldTransientFolders;
-};
\ No newline at end of file
+};
diff --git a/Engine/Source/Editor/UnrealEd/Public/WorldPartition/WorldPartitionBuilder.h b/Engine/Source/Editor/UnrealEd/Public/WorldPartition/WorldPartitionBuilder.h
--- a/Engine/Source/Editor/UnrealEd/Public/WorldPartition/WorldPartitionBuilder.h
+++ b/Engine/Source/Editor/UnrealEd/Public/WorldPartition/WorldPartitionBuilder.h
@@ -34,8 +34,8 @@ struct FCellInfo
 	static UNREALED_API FWorldBuilderCellCoord GetCellCount(const FBox& InBounds, const int32 InCellSize);
 };
 
-UCLASS(Abstract, Config=Engine)
-class UNREALED_API UWorldPartitionBuilder : public UObject
+UCLASS(Abstract, Config=Engine, MinimalAPI)
+class UWorldPartitionBuilder : public UObject
 {
 	GENERATED_UCLASS_BODY()
 
@@ -48,21 +48,21 @@ public:
 		IterativeCells2D,
 	};
 
-	bool RunBuilder(UWorld* World);
+	UNREALED_API bool RunBuilder(UWorld* World);
 
-	virtual bool RequiresCommandletRendering() const PURE_VIRTUAL(UWorldPartitionBuilder::RequiresCommandletRendering, return false;);
-	virtual ELoadingMode GetLoadingMode() const PURE_VIRTUAL(UWorldPartitionBuilder::GetLoadingMode, return ELoadingMode::Custom;);
+	UNREALED_API virtual bool RequiresCommandletRendering() const PURE_VIRTUAL(UWorldPartitionBuilder::RequiresCommandletRendering, return false;);
+	UNREALED_API virtual ELoadingMode GetLoadingMode() const PURE_VIRTUAL(UWorldPartitionBuilder::GetLoadingMode, return ELoadingMode::Custom;);
 	
-	bool Run(UWorld* World, FPackageSourceControlHelper& PackageHelper);
+	UNREALED_API bool Run(UWorld* World, FPackageSourceControlHelper& PackageHelper);
 
 	virtual bool PreWorldInitialization(UWorld* World, FPackageSourceControlHelper& PackageHelper) { return true; }
 
-	static bool SavePackages(const TArray<UPackage*>& Packages, FPackageSourceControlHelper& PackageHelper, bool bErrorsAsWarnings = false);
-	static bool DeletePackages(const TArray<UPackage*>& Packages, FPackageSourceControlHelper& PackageHelper, bool bErrorsAsWarnings = false);
-	static bool DeletePackages(const TArray<FString>& PackageNames, FPackageSourceControlHelper& PackageHelper, bool bErrorsAsWarnings = false);
+	static UNREALED_API bool SavePackages(const TArray<UPackage*>& Packages, FPackageSourceControlHelper& PackageHelper, bool bErrorsAsWarnings = false);
+	static UNREALED_API bool DeletePackages(const TArray<UPackage*>& Packages, FPackageSourceControlHelper& PackageHelper, bool bErrorsAsWarnings = false);
+	static UNREALED_API bool DeletePackages(const TArray<FString>& PackageNames, FPackageSourceControlHelper& PackageHelper, bool bErrorsAsWarnings = false);
 
 	DECLARE_DELEGATE_RetVal_TwoParams(bool, FModifiedFilesHandler, const TArray<FString>&, const FString&); /*bool ModifiedFileHander(ModifiedFiles, ChangeDescription) */
-	void SetModifiedFilesHandler(const FModifiedFilesHandler& ModifiedFilesHandler);
+	UNREALED_API void SetModifiedFilesHandler(const FModifiedFilesHandler& ModifiedFilesHandler);
 
 protected:
 	/**
@@ -71,7 +71,7 @@ protected:
 	 */
 	virtual bool PreRun(UWorld* World, FPackageSourceControlHelper& PackageHelper) { return true; }
 
-	virtual bool RunInternal(UWorld* World, const FCellInfo& InCellInfo, FPackageSourceControlHelper& PackageHelper) PURE_VIRTUAL(UWorldPartition::RunInternal, return false;);
+	UNREALED_API virtual bool RunInternal(UWorld* World, const FCellInfo& InCellInfo, FPackageSourceControlHelper& PackageHelper) PURE_VIRTUAL(UWorldPartition::RunInternal, return false;);
 
 	/**
 	 * Overridable method for derived classes to perform operations when world builder process completes.
@@ -94,8 +94,8 @@ protected:
 	 */
 	virtual bool CanProcessNonPartitionedWorlds() const { return false; }
 
-	bool OnFilesModified(const TArray<FString>& InModifiedFiles, const FString& InChangelistDescription) const;
-	bool OnPackagesModified(const TArray<UPackage*>& InModifiedPackages, const FString& InChangelistDescription) const;
+	UNREALED_API bool OnFilesModified(const TArray<FString>& InModifiedFiles, const FString& InChangelistDescription) const;
+	UNREALED_API bool OnPackagesModified(const TArray<UPackage*>& InModifiedPackages, const FString& InChangelistDescription) const;
 
 	/**
 	 * Test if the builder was provided the given parameter.
@@ -129,11 +129,11 @@ protected:
 	}
 
 	UE_DEPRECATED(5.3, "Please use OnFilesModified")
-	bool AutoSubmitFiles(const TArray<FString>& InModifiedFiles, const FString& InChangelistDescription) const;
+	UNREALED_API bool AutoSubmitFiles(const TArray<FString>& InModifiedFiles, const FString& InChangelistDescription) const;
 	UE_DEPRECATED(5.3, "Please use OnPackagesModified")
-	bool AutoSubmitPackages(const TArray<UPackage*>& InModifiedPackages, const FString& InChangelistDescription) const;
+	UNREALED_API bool AutoSubmitPackages(const TArray<UPackage*>& InModifiedPackages, const FString& InChangelistDescription) const;
 
-	virtual UWorld::InitializationValues GetWorldInitializationValues() const;
+	UNREALED_API virtual UWorld::InitializationValues GetWorldInitializationValues() const;
 
 	UE_DEPRECATED(5.3, "You must override the version that takes a World parameter")
 	virtual bool PreWorldInitialization(FPackageSourceControlHelper& PackageHelper) final { return true; }
@@ -150,10 +150,10 @@ protected:
 	FModifiedFilesHandler ModifiedFilesHandler;
 
 private:
-	void LoadDataLayers(UWorld* InWorld);
+	UNREALED_API void LoadDataLayers(UWorld* InWorld);
 
 	friend struct FWorldPartitionBuilderArgsScope;
-	static FString Args;
+	static UNREALED_API FString Args;
 };
 
 /**
diff --git a/Engine/Source/Editor/UnrealEd/Public/WorldPartition/WorldPartitionHLODsBuilder.h b/Engine/Source/Editor/UnrealEd/Public/WorldPartition/WorldPartitionHLODsBuilder.h
--- a/Engine/Source/Editor/UnrealEd/Public/WorldPartition/WorldPartitionHLODsBuilder.h
+++ b/Engine/Source/Editor/UnrealEd/Public/WorldPartition/WorldPartitionHLODsBuilder.h
@@ -75,43 +75,43 @@ enum class EHLODBuildStep : uint8
 };
 ENUM_CLASS_FLAGS(EHLODBuildStep);
 
-UCLASS()
-class UNREALED_API UWorldPartitionHLODsBuilder : public UWorldPartitionBuilder
+UCLASS(MinimalAPI)
+class UWorldPartitionHLODsBuilder : public UWorldPartitionBuilder
 {
 	GENERATED_UCLASS_BODY()
 public:
 	// UWorldPartitionBuilder interface begin
-	virtual bool RequiresCommandletRendering() const override;
+	UNREALED_API virtual bool RequiresCommandletRendering() const override;
 	virtual ELoadingMode GetLoadingMode() const override { return ELoadingMode::Custom; }
-	virtual bool PreWorldInitialization(UWorld* World, FPackageSourceControlHelper& PackageHelper) override;
+	UNREALED_API virtual bool PreWorldInitialization(UWorld* World, FPackageSourceControlHelper& PackageHelper) override;
 
 protected:
-	virtual bool RunInternal(UWorld* World, const FCellInfo& InCellInfo, FPackageSourceControlHelper& PackageHelper) override;
+	UNREALED_API virtual bool RunInternal(UWorld* World, const FCellInfo& InCellInfo, FPackageSourceControlHelper& PackageHelper) override;
 	virtual bool CanProcessNonPartitionedWorlds() const override { return true; }
 	// UWorldPartitionBuilder interface end
 
 	bool IsDistributedBuild() const { return bDistributedBuild; }
 	bool IsUsingBuildManifest() const { return !BuildManifest.IsEmpty(); }
-	bool ValidateParams() const;
+	UNREALED_API bool ValidateParams() const;
 
-	bool SetupHLODActors();
-	bool BuildHLODActors();
-	bool DeleteHLODActors();
-	bool SubmitHLODActors();
-	bool DumpStats();
+	UNREALED_API bool SetupHLODActors();
+	UNREALED_API bool BuildHLODActors();
+	UNREALED_API bool DeleteHLODActors();
+	UNREALED_API bool SubmitHLODActors();
+	UNREALED_API bool DumpStats();
 
-	bool GenerateBuildManifest(TMap<FString, int32>& FilesToBuilderMap) const;
-	bool GetHLODActorsToBuild(TArray<FGuid>& HLODActorsToBuild) const;
+	UNREALED_API bool GenerateBuildManifest(TMap<FString, int32>& FilesToBuilderMap) const;
+	UNREALED_API bool GetHLODActorsToBuild(TArray<FGuid>& HLODActorsToBuild) const;
 
-	TArray<TArray<FGuid>> GetHLODWorkloads(int32 NumWorkloads) const;
-	bool ValidateWorkload(const TArray<FGuid>& Workload) const;
+	UNREALED_API TArray<TArray<FGuid>> GetHLODWorkloads(int32 NumWorkloads) const;
+	UNREALED_API bool ValidateWorkload(const TArray<FGuid>& Workload) const;
 
-	bool CopyFilesToWorkingDir(const FString& TargetDir, const FHLODModifiedFiles& ModifiedFiles, TArray<FString>& BuildProducts);
-	bool CopyFilesFromWorkingDir(const FString& SourceDir);
+	UNREALED_API bool CopyFilesToWorkingDir(const FString& TargetDir, const FHLODModifiedFiles& ModifiedFiles, TArray<FString>& BuildProducts);
+	UNREALED_API bool CopyFilesFromWorkingDir(const FString& SourceDir);
 
-	bool ShouldRunStep(const EHLODBuildStep BuildStep) const;
+	UNREALED_API bool ShouldRunStep(const EHLODBuildStep BuildStep) const;
 
-	bool AddBuildProducts(const TArray<FString>& BuildProducts) const;
+	UNREALED_API bool AddBuildProducts(const TArray<FString>& BuildProducts) const;
 
 private:
 	UWorld* World;
@@ -135,4 +135,4 @@ private:
 	FString DistributedBuildManifest;
 	
 	FHLODModifiedFiles ModifiedFiles;
-};
\ No newline at end of file
+};
diff --git a/Engine/Source/Editor/UnrealEd/Public/WorldPartition/WorldPartitionLandscapeSplineMeshesBuilder.h b/Engine/Source/Editor/UnrealEd/Public/WorldPartition/WorldPartitionLandscapeSplineMeshesBuilder.h
--- a/Engine/Source/Editor/UnrealEd/Public/WorldPartition/WorldPartitionLandscapeSplineMeshesBuilder.h
+++ b/Engine/Source/Editor/UnrealEd/Public/WorldPartition/WorldPartitionLandscapeSplineMeshesBuilder.h
@@ -12,8 +12,8 @@ class UPackage;
 class ILandscapeSplineInterface;
 
 // Example Command Line: ProjectName MapName -run=WorldPartitionBuilderCommandlet -SCCProvider=Perforce -Builder=WorldPartitionLandscapeSplineMeshesBuilder (<Optional> -NewGridSize=Value)
-UCLASS()
-class UNREALED_API UWorldPartitionLandscapeSplineMeshesBuilder : public UWorldPartitionBuilder
+UCLASS(MinimalAPI)
+class UWorldPartitionLandscapeSplineMeshesBuilder : public UWorldPartitionBuilder
 {
 	GENERATED_UCLASS_BODY()
 
@@ -21,21 +21,21 @@ public:
 	// UWorldPartitionBuilder interface begin
 	virtual bool RequiresCommandletRendering() const override { return false; }
 	virtual ELoadingMode GetLoadingMode() const override { return ELoadingMode::Custom; }
-	virtual bool RunInternal(UWorld* World, const FCellInfo& InCellInfo, FPackageSourceControlHelper& PackageHelper) override;
+	UNREALED_API virtual bool RunInternal(UWorld* World, const FCellInfo& InCellInfo, FPackageSourceControlHelper& PackageHelper) override;
 	// UWorldPartitionBuilder interface end
 
 	// Helper function that creates a builder and runs it on an already loaded/initialized world
-	static bool RunOnInitializedWorld(UWorld* World);
+	static UNREALED_API bool RunOnInitializedWorld(UWorld* World);
 
 private:
-	void FilterStaticMeshComponents(TArray<class UStaticMeshComponent*>& InOutComponents);
-	void CloneStaticMeshComponentInActor(ALandscapeSplineMeshesActor* InActor, const UStaticMeshComponent* InMeshComponent);
-	void CloneStaticMeshComponent(const UStaticMeshComponent* InSrcMeshComponent, UStaticMeshComponent* DstMeshComponent);
-	int32 HashStaticMeshComponent(const UStaticMeshComponent* InComponent);
-	FName GetSplineCollisionProfileName(const UStaticMeshComponent* InMeshComponent);
-	ALandscapeSplineMeshesActor* GetOrCreatePartitionActorForComponent(class UWorld* InWorld, const UStaticMeshComponent* InMeshComponent, const FGuid& InLandscapeGuid);
+	UNREALED_API void FilterStaticMeshComponents(TArray<class UStaticMeshComponent*>& InOutComponents);
+	UNREALED_API void CloneStaticMeshComponentInActor(ALandscapeSplineMeshesActor* InActor, const UStaticMeshComponent* InMeshComponent);
+	UNREALED_API void CloneStaticMeshComponent(const UStaticMeshComponent* InSrcMeshComponent, UStaticMeshComponent* DstMeshComponent);
+	UNREALED_API int32 HashStaticMeshComponent(const UStaticMeshComponent* InComponent);
+	UNREALED_API FName GetSplineCollisionProfileName(const UStaticMeshComponent* InMeshComponent);
+	UNREALED_API ALandscapeSplineMeshesActor* GetOrCreatePartitionActorForComponent(class UWorld* InWorld, const UStaticMeshComponent* InMeshComponent, const FGuid& InLandscapeGuid);
 
 	uint32 NewGridSize;
 	
-	static class UStaticMesh* SplineEditorMesh;
-};
\ No newline at end of file
+	static UNREALED_API class UStaticMesh* SplineEditorMesh;
+};
diff --git a/Engine/Source/Editor/UnrealEd/Public/WorldPartition/WorldPartitionMiniMapBuilder.h b/Engine/Source/Editor/UnrealEd/Public/WorldPartition/WorldPartitionMiniMapBuilder.h
--- a/Engine/Source/Editor/UnrealEd/Public/WorldPartition/WorldPartitionMiniMapBuilder.h
+++ b/Engine/Source/Editor/UnrealEd/Public/WorldPartition/WorldPartitionMiniMapBuilder.h
@@ -10,8 +10,8 @@
 #include "WorldPartitionMiniMapBuilder.generated.h"
 
 
-UCLASS(config = Engine, defaultconfig)
-class UNREALED_API UWorldPartitionMiniMapBuilder : public UWorldPartitionBuilder
+UCLASS(config = Engine, defaultconfig, MinimalAPI)
+class UWorldPartitionMiniMapBuilder : public UWorldPartitionBuilder
 {
 	GENERATED_UCLASS_BODY()
 
@@ -21,9 +21,9 @@ public:
 	virtual ELoadingMode GetLoadingMode() const override { return ELoadingMode::IterativeCells2D; }
 
 protected:
-	virtual bool PreRun(UWorld* World, FPackageSourceControlHelper& PackageHelper) override;
-	virtual bool RunInternal(UWorld* World, const FCellInfo& InCellInfo, FPackageSourceControlHelper& PackageHelper) override;
-	virtual bool PostRun(UWorld* World, FPackageSourceControlHelper& PackageHelper, const bool bInRunSuccess) override;
+	UNREALED_API virtual bool PreRun(UWorld* World, FPackageSourceControlHelper& PackageHelper) override;
+	UNREALED_API virtual bool RunInternal(UWorld* World, const FCellInfo& InCellInfo, FPackageSourceControlHelper& PackageHelper) override;
+	UNREALED_API virtual bool PostRun(UWorld* World, FPackageSourceControlHelper& PackageHelper, const bool bInRunSuccess) override;
 	// UWorldPartitionBuilder interface end
 
 private:
@@ -36,4 +36,4 @@ private:
 	int32 MinimapImageSizeY;
 
 	bool bDebugCapture;
-};
\ No newline at end of file
+};
diff --git a/Engine/Source/Editor/UnrealEd/Public/WorldPartition/WorldPartitionRuntimeVirtualTextureBuilder.h b/Engine/Source/Editor/UnrealEd/Public/WorldPartition/WorldPartitionRuntimeVirtualTextureBuilder.h
--- a/Engine/Source/Editor/UnrealEd/Public/WorldPartition/WorldPartitionRuntimeVirtualTextureBuilder.h
+++ b/Engine/Source/Editor/UnrealEd/Public/WorldPartition/WorldPartitionRuntimeVirtualTextureBuilder.h
@@ -9,8 +9,8 @@
 class UWorldPartition;
 
 // Example Command Line: ProjectName MapName -run=WorldPartitionBuilderCommandlet -SCCProvider=Perforce -Builder=WorldPartitionRuntimeVirtualTextureBuilder
-UCLASS()
-class UNREALED_API UWorldPartitionRuntimeVirtualTextureBuilder : public UWorldPartitionBuilder
+UCLASS(MinimalAPI)
+class UWorldPartitionRuntimeVirtualTextureBuilder : public UWorldPartitionBuilder
 {
 	GENERATED_UCLASS_BODY()
 
@@ -18,9 +18,9 @@ public:
 	// UWorldPartitionBuilder interface begin
 	virtual bool RequiresCommandletRendering() const override { return true; }
 	virtual ELoadingMode GetLoadingMode() const override { return ELoadingMode::Custom; }
-	virtual bool RunInternal(UWorld* World, const FCellInfo& InCellInfo, FPackageSourceControlHelper& PackageHelper) override;
+	UNREALED_API virtual bool RunInternal(UWorld* World, const FCellInfo& InCellInfo, FPackageSourceControlHelper& PackageHelper) override;
 	// UWorldPartitionBuilder interface end
 
 	// Helper method to load actors that contribute to the Runtime Virtual Texture
-	static void LoadRuntimeVirtualTextureActors(UWorldPartition* WorldPartition, FWorldPartitionHelpers::FForEachActorWithLoadingResult& Result);
-};
\ No newline at end of file
+	static UNREALED_API void LoadRuntimeVirtualTextureActors(UWorldPartition* WorldPartition, FWorldPartitionHelpers::FForEachActorWithLoadingResult& Result);
+};
diff --git a/Engine/Source/Runtime/Engine/Public/Animation/MirrorSyncScope.h b/Engine/Source/Runtime/Engine/Public/Animation/MirrorSyncScope.h
--- a/Engine/Source/Runtime/Engine/Public/Animation/MirrorSyncScope.h
+++ b/Engine/Source/Runtime/Engine/Public/Animation/MirrorSyncScope.h
@@ -10,7 +10,7 @@ namespace UE { namespace Anim {
 
 class FAnimSyncGroupScope;
 
-class ENGINE_API FAnimNotifyMirrorContext : public IAnimNotifyEventContextDataInterface
+class FAnimNotifyMirrorContext : public IAnimNotifyEventContextDataInterface
 {
 	DECLARE_NOTIFY_CONTEXT_INTERFACE(FAnimNotifyMirrorContext)
 	public:
@@ -22,14 +22,14 @@ class ENGINE_API FAnimNotifyMirrorContext : public IAnimNotifyEventContextDataIn
 };
 	
 // Scoped graph message used to synchronize mirroring 
-class ENGINE_API FMirrorSyncScope : public IGraphMessage
+class FMirrorSyncScope : public IGraphMessage
 {
-	DECLARE_ANIMGRAPH_MESSAGE(FMirrorSyncScope);
+	DECLARE_ANIMGRAPH_MESSAGE_API(FMirrorSyncScope, ENGINE_API);
 public:
-	FMirrorSyncScope(const FAnimationBaseContext& InContext, const UMirrorDataTable* InMirrorDataTable);
-	virtual ~FMirrorSyncScope();
-	int32 GetMirrorScopeDepth() const;
-	virtual TUniquePtr<const IAnimNotifyEventContextDataInterface> MakeUniqueEventContextData() const override;
+	ENGINE_API FMirrorSyncScope(const FAnimationBaseContext& InContext, const UMirrorDataTable* InMirrorDataTable);
+	ENGINE_API virtual ~FMirrorSyncScope();
+	ENGINE_API int32 GetMirrorScopeDepth() const;
+	ENGINE_API virtual TUniquePtr<const IAnimNotifyEventContextDataInterface> MakeUniqueEventContextData() const override;
 private:
 	const UMirrorDataTable* MirrorDataTable = nullptr;
 	int32 MirrorScopeDepth = 1;
