Undo 25904929
#rnx

diff --git a/Engine/Source/Editor/UnrealEd/Public/Toolkits/GlobalEditorCommonCommands.h b/Engine/Source/Editor/UnrealEd/Public/Toolkits/GlobalEditorCommonCommands.h
--- a/Engine/Source/Editor/UnrealEd/Public/Toolkits/GlobalEditorCommonCommands.h
+++ b/Engine/Source/Editor/UnrealEd/Public/Toolkits/GlobalEditorCommonCommands.h
@@ -13,26 +13,26 @@ class SWindow;
 
 // Global editor common commands
 // Note: There is no real global command concept, so these must still be registered in each editor
-class FGlobalEditorCommonCommands : public TCommands< FGlobalEditorCommonCommands >
+class UNREALED_API FGlobalEditorCommonCommands : public TCommands< FGlobalEditorCommonCommands >
 {
 public:
-	UNREALED_API FGlobalEditorCommonCommands();
-	UNREALED_API ~FGlobalEditorCommonCommands();
+	FGlobalEditorCommonCommands();
+	~FGlobalEditorCommonCommands();
 
-	UNREALED_API virtual void RegisterCommands() override;
+	virtual void RegisterCommands() override;
 
-	static UNREALED_API void MapActions(TSharedRef<FUICommandList>& ToolkitCommands);
+	static void MapActions(TSharedRef<FUICommandList>& ToolkitCommands);
 
 protected:
-	static UNREALED_API void OnPressedCtrlTab(TSharedPtr<FUICommandInfo> TriggeringCommand);
-	static UNREALED_API void OnSummonedAssetPicker();
-	static UNREALED_API void OnSummonedConsoleCommandBox();
-	static UNREALED_API void OnOpenContentBrowserDrawer();
-	static UNREALED_API void OnOpenOutputLogDrawer();
+	static void OnPressedCtrlTab(TSharedPtr<FUICommandInfo> TriggeringCommand);
+	static void OnSummonedAssetPicker();
+	static void OnSummonedConsoleCommandBox();
+	static void OnOpenContentBrowserDrawer();
+	static void OnOpenOutputLogDrawer();
 
-	static UNREALED_API TSharedRef<SDockTab> SpawnAssetPicker(const FSpawnTabArgs& InArgs);
+	static TSharedRef<SDockTab> SpawnAssetPicker(const FSpawnTabArgs& InArgs);
 
-	static UNREALED_API TSharedPtr<IMenu> OpenPopupMenu(TSharedRef<SWidget> WindowContents, const FVector2D& PopupDesiredSize);
+	static TSharedPtr<IMenu> OpenPopupMenu(TSharedRef<SWidget> WindowContents, const FVector2D& PopupDesiredSize);
 public:
 	TSharedPtr<FUICommandInfo> FindInContentBrowser;
 	TSharedPtr<FUICommandInfo> SummonControlTabNavigation;
