[UnrealEd]
* Reverted the changes in 25904929 related to subclasses to TCommands<>.

This code relies on undefined behavior of how things are inlined. base class is a template with a static variable. One instance of that static variable exists in each dll/exe. To call TCommands<>::Get() would return the instance in the same dll.
But, because the subclass has dllexport on the type and not the members it will generate a Get function that is not inlined. The callsite when calling Get will actually then use the instantiated Get() function and not the inlined one.

So for

class API X : public TCommands<X> {}

it will (for now) work to call

X::Get() and get the same instantiated singletone.

Calling

TCommands<X>::Get() will definitely not work and return the singleton instantiated in the same dll/exe

diff --git a/Engine/Source/Editor/UnrealEd/Public/EditorViewportClient.h b/Engine/Source/Editor/UnrealEd/Public/EditorViewportClient.h
--- a/Engine/Source/Editor/UnrealEd/Public/EditorViewportClient.h
+++ b/Engine/Source/Editor/UnrealEd/Public/EditorViewportClient.h
@@ -55,16 +55,16 @@ namespace EDragTool
 /**
 * Unreal level editor actions
 */
-class FViewportNavigationCommands : public TCommands<FViewportNavigationCommands>
+class UNREALED_API FViewportNavigationCommands : public TCommands<FViewportNavigationCommands>
 {
 
 public:
-	UNREALED_API FViewportNavigationCommands();
+	FViewportNavigationCommands();
 
 	/**
 	* Initialize commands
 	*/
-	UNREALED_API virtual void RegisterCommands() override;
+	virtual void RegisterCommands() override;
 
 
 	TSharedPtr< FUICommandInfo > Forward;
diff --git a/Engine/Source/Editor/UnrealEd/Public/EditorViewportCommands.h b/Engine/Source/Editor/UnrealEd/Public/EditorViewportCommands.h
--- a/Engine/Source/Editor/UnrealEd/Public/EditorViewportCommands.h
+++ b/Engine/Source/Editor/UnrealEd/Public/EditorViewportCommands.h
@@ -15,10 +15,10 @@ class FUICommandList;
 /**
  * Class containing commands for editor viewport actions common to all viewports
  */
-class FEditorViewportCommands : public TCommands<FEditorViewportCommands>
+class UNREALED_API FEditorViewportCommands : public TCommands<FEditorViewportCommands>
 {
 public:
-	UNREALED_API FEditorViewportCommands();
+	FEditorViewportCommands();
 	
 	/** Viewport pane configurations */
 	TSharedPtr< FUICommandInfo > ViewportConfig_OnePane;
@@ -248,7 +248,7 @@ public:
 	TSharedPtr< FUICommandInfo > SurfaceSnapping;
 public:
 	/** Registers our commands with the binding system */
-	UNREALED_API virtual void RegisterCommands() override;
+	virtual void RegisterCommands() override;
 };
 
 /**
diff --git a/Engine/Source/Editor/UnrealEd/Public/GPUSkinCacheVisualizationMenuCommands.h b/Engine/Source/Editor/UnrealEd/Public/GPUSkinCacheVisualizationMenuCommands.h
--- a/Engine/Source/Editor/UnrealEd/Public/GPUSkinCacheVisualizationMenuCommands.h
+++ b/Engine/Source/Editor/UnrealEd/Public/GPUSkinCacheVisualizationMenuCommands.h
@@ -11,7 +11,7 @@
 
 class FEditorViewportClient;
 
-class FGPUSkinCacheVisualizationMenuCommands : public TCommands<FGPUSkinCacheVisualizationMenuCommands>
+class UNREALED_API FGPUSkinCacheVisualizationMenuCommands : public TCommands<FGPUSkinCacheVisualizationMenuCommands>
 {
 public:
 	enum class FGPUSkinCacheVisualizationType : uint8
@@ -38,15 +38,15 @@ public:
 	typedef TMultiMap<FName, FGPUSkinCacheVisualizationRecord> TGPUSkinCacheVisualizationModeCommandMap;
 	typedef TGPUSkinCacheVisualizationModeCommandMap::TConstIterator TCommandConstIterator;
 
-	UNREALED_API FGPUSkinCacheVisualizationMenuCommands();
+	FGPUSkinCacheVisualizationMenuCommands();
 
-	UNREALED_API TCommandConstIterator CreateCommandConstIterator() const;
+	TCommandConstIterator CreateCommandConstIterator() const;
 
-	static UNREALED_API void BuildVisualisationSubMenu(FMenuBuilder& Menu);
+	static void BuildVisualisationSubMenu(FMenuBuilder& Menu);
 
-	UNREALED_API virtual void RegisterCommands() override;
+	virtual void RegisterCommands() override;
 
-	UNREALED_API void BindCommands(FUICommandList& CommandList, const TSharedPtr<FEditorViewportClient>& Client) const;
+	void BindCommands(FUICommandList& CommandList, const TSharedPtr<FEditorViewportClient>& Client) const;
 
 	inline bool IsPopulated() const
 	{
@@ -54,11 +54,11 @@ public:
 	}
 
 private:
-	UNREALED_API void BuildCommandMap();
-	UNREALED_API bool AddCommandTypeToMenu(FMenuBuilder& Menu, const FGPUSkinCacheVisualizationType Type, bool bSeparatorBefore = false) const;
+	void BuildCommandMap();
+	bool AddCommandTypeToMenu(FMenuBuilder& Menu, const FGPUSkinCacheVisualizationType Type, bool bSeparatorBefore = false) const;
 
-	static UNREALED_API void ChangeGPUSkinCacheVisualizationMode(TWeakPtr<FEditorViewportClient> WeakClient, FName InName);
-	static UNREALED_API bool IsGPUSkinCacheVisualizationModeSelected(TWeakPtr<FEditorViewportClient> WeakClient, FName InName);
+	static void ChangeGPUSkinCacheVisualizationMode(TWeakPtr<FEditorViewportClient> WeakClient, FName InName);
+	static bool IsGPUSkinCacheVisualizationModeSelected(TWeakPtr<FEditorViewportClient> WeakClient, FName InName);
 
 	TGPUSkinCacheVisualizationModeCommandMap CommandMap;
 };
diff --git a/Engine/Source/Editor/UnrealEd/Public/Kismet2/DebuggerCommands.h b/Engine/Source/Editor/UnrealEd/Public/Kismet2/DebuggerCommands.h
--- a/Engine/Source/Editor/UnrealEd/Public/Kismet2/DebuggerCommands.h
+++ b/Engine/Source/Editor/UnrealEd/Public/Kismet2/DebuggerCommands.h
@@ -34,46 +34,46 @@ public:
 //////////////////////////////////////////////////////////////////////////
 // FPlayWorldCommands
 
-class FPlayWorldCommands : public TCommands<FPlayWorldCommands>
+class UNREALED_API FPlayWorldCommands : public TCommands<FPlayWorldCommands>
 {
 private:
 
 	friend class TCommands<FPlayWorldCommands>;
 
-	UNREALED_API FPlayWorldCommands();
+	FPlayWorldCommands();
 
 public:
 
 	// TCommands interface
-	UNREALED_API virtual void RegisterCommands() override;
+	virtual void RegisterCommands() override;
 	// End of TCommands interface
 
 	/**
 	 * Binds all global kismet commands to delegates
 	 */
-	static UNREALED_API void BindGlobalPlayWorldCommands();
+	static void BindGlobalPlayWorldCommands();
 
 	/** Populates a toolbar with the menu commands for play-world control (pause/resume/stop/possess/eject/step/show current loc) */
-	static UNREALED_API void BuildToolbar( FToolMenuSection& InSection, bool bIncludeLaunchButtonAndOptions = false );
+	static void BuildToolbar( FToolMenuSection& InSection, bool bIncludeLaunchButtonAndOptions = false );
 
 	/**
 	* Return the active widget that processes play world actions for PIE
 	*
 	*/
-	static UNREALED_API TWeakPtr<SGlobalPlayWorldActions> GetActiveGlobalPlayWorldActionsWidget();
+	static TWeakPtr<SGlobalPlayWorldActions> GetActiveGlobalPlayWorldActionsWidget();
 
 	/**
 	* Set the active widget that processes play world actions for PIE
 	*
 	*/
-	static UNREALED_API void SetActiveGlobalPlayWorldActionsWidget(TWeakPtr<SGlobalPlayWorldActions> ActiveWidget);
+	static void SetActiveGlobalPlayWorldActionsWidget(TWeakPtr<SGlobalPlayWorldActions> ActiveWidget);
 
 public:
 
 	/** 
 	 * A command list that can be passed around and isn't bound to an instance of any tool or editor. 
 	 */
-	static UNREALED_API TSharedPtr<FUICommandList> GlobalPlayWorldActions;
+	static TSharedPtr<FUICommandList> GlobalPlayWorldActions;
 
 public:
 
@@ -116,20 +116,20 @@ public:
 protected:
 
 	/** A weak pointer to the current active widget that processes PIE actions */
-	static UNREALED_API TWeakPtr<SGlobalPlayWorldActions> ActiveGlobalPlayWorldActionsWidget;
+	static TWeakPtr<SGlobalPlayWorldActions> ActiveGlobalPlayWorldActionsWidget;
 
 	/**
 	 * Generates menu content for the PIE combo button drop down menu
 	 *
 	 * @return	Menu content widget
 	 */
-	static UNREALED_API TSharedRef< SWidget > GeneratePlayMenuContent( TSharedRef<FUICommandList> InCommandList );
+	static TSharedRef< SWidget > GeneratePlayMenuContent( TSharedRef<FUICommandList> InCommandList );
 
 	// Add mobile PIE preview device commands
-	UNREALED_API void AddPIEPreviewDeviceCommands();
+	void AddPIEPreviewDeviceCommands();
 
 	// Add mobile PIE preview device actions
-	static UNREALED_API void AddPIEPreviewDeviceActions(const FPlayWorldCommands &Commands, FUICommandList &ActionList);
+	static void AddPIEPreviewDeviceActions(const FPlayWorldCommands &Commands, FUICommandList &ActionList);
 };
 
 
diff --git a/Engine/Source/Editor/UnrealEd/Public/RayTracingDebugVisualizationMenuCommands.h b/Engine/Source/Editor/UnrealEd/Public/RayTracingDebugVisualizationMenuCommands.h
--- a/Engine/Source/Editor/UnrealEd/Public/RayTracingDebugVisualizationMenuCommands.h
+++ b/Engine/Source/Editor/UnrealEd/Public/RayTracingDebugVisualizationMenuCommands.h
@@ -11,7 +11,7 @@
 
 class FEditorViewportClient;
 
-class FRayTracingDebugVisualizationMenuCommands : public TCommands<FRayTracingDebugVisualizationMenuCommands>
+class UNREALED_API FRayTracingDebugVisualizationMenuCommands : public TCommands<FRayTracingDebugVisualizationMenuCommands>
 {
 public:
 	struct FRayTracingDebugVisualizationRecord
@@ -28,15 +28,15 @@ public:
 		}
 	};
 
-	UNREALED_API FRayTracingDebugVisualizationMenuCommands();
+	FRayTracingDebugVisualizationMenuCommands();
 
-	static UNREALED_API void BuildVisualisationSubMenu(FMenuBuilder& Menu);
+	static void BuildVisualisationSubMenu(FMenuBuilder& Menu);
 
-	UNREALED_API virtual void RegisterCommands() override;
+	virtual void RegisterCommands() override;
 
-	UNREALED_API void BindCommands(FUICommandList& CommandList, const TSharedPtr<FEditorViewportClient>& Client) const;
+	void BindCommands(FUICommandList& CommandList, const TSharedPtr<FEditorViewportClient>& Client) const;
 
-	static UNREALED_API bool DebugModeShouldBeTonemapped(const FName& RayTracingDebugModeName);
+	static bool DebugModeShouldBeTonemapped(const FName& RayTracingDebugModeName);
 
 private:
 	void BuildCommandMap();
diff --git a/Engine/Source/Editor/UnrealEd/Public/Tools/StandardToolModeCommands.h b/Engine/Source/Editor/UnrealEd/Public/Tools/StandardToolModeCommands.h
--- a/Engine/Source/Editor/UnrealEd/Public/Tools/StandardToolModeCommands.h
+++ b/Engine/Source/Editor/UnrealEd/Public/Tools/StandardToolModeCommands.h
@@ -36,20 +36,20 @@ enum class EStandardToolModeCommands
  * call FStandardToolModeCommands::Get().FindStandardCommand() to get the registered UICommandInfo.
  * 
  */
-class FStandardToolModeCommands : public TCommands<FStandardToolModeCommands>
+class UNREALED_API FStandardToolModeCommands : public TCommands<FStandardToolModeCommands>
 {
 public:
-	UNREALED_API FStandardToolModeCommands();
+	FStandardToolModeCommands();
 
 	/**
 	 * Registers the set of standard commands. Call on module startup.
 	 */
-	UNREALED_API virtual void RegisterCommands() override;
+	virtual void RegisterCommands() override;
 
 	/**
 	 * Look up the UICommandInfo for a standard command
 	 */
-	UNREALED_API TSharedPtr<FUICommandInfo> FindStandardCommand(EStandardToolModeCommands Command) const;
+	TSharedPtr<FUICommandInfo> FindStandardCommand(EStandardToolModeCommands Command) const;
 
 protected:
 	TMap<EStandardToolModeCommands, TSharedPtr<FUICommandInfo>> Commands;
