Bring back bCheckReferencesOnDelete editor setting. This is used by stage opertors to prevent long running operations in fast paced environments.

# original CL 15854609

diff --git a/Engine/Source/Editor/UnrealEd/Classes/Settings/LevelEditorMiscSettings.h b/Engine/Source/Editor/UnrealEd/Classes/Settings/LevelEditorMiscSettings.h
--- a/Engine/Source/Editor/UnrealEd/Classes/Settings/LevelEditorMiscSettings.h
+++ b/Engine/Source/Editor/UnrealEd/Classes/Settings/LevelEditorMiscSettings.h
@@ -47,6 +47,10 @@ public:
 	UPROPERTY(EditAnywhere, config, Category = Editing, meta = (DisplayName = "Avoid Actor Relabel on Paste Selected"))
 	uint32 bAvoidRelabelOnPasteSelected:1;
 
+	/** If enabled, will force checking references during a delete. Note: disabling this is for advanced usage only and not recommended.*/
+	UPROPERTY(EditAnywhere, config, AdvancedDisplay, Category = Editing, meta = (DisplayName = "Check References on Delete"))
+	uint32 bCheckReferencesOnDelete:1;
+
 public:
 	/** If checked audio playing in the editor will continue to play even if the editor is in the background */
 	UPROPERTY(EditAnywhere, config, Category=Sound)
diff --git a/Engine/Source/Editor/UnrealEd/Private/Subsystems/EditorActorSubsystem.cpp b/Engine/Source/Editor/UnrealEd/Private/Subsystems/EditorActorSubsystem.cpp
--- a/Engine/Source/Editor/UnrealEd/Private/Subsystems/EditorActorSubsystem.cpp
+++ b/Engine/Source/Editor/UnrealEd/Private/Subsystems/EditorActorSubsystem.cpp
@@ -23,6 +23,7 @@
 #include "Engine/Selection.h"
 #include "LevelEditorViewport.h"
 #include "Settings/LevelEditorViewportSettings.h"
+#include "Settings/LevelEditorMiscSettings.h"
 #include "Subsystems/EditorElementSubsystem.h"
 #include "UObject/Stack.h"
 
@@ -285,7 +286,8 @@ void UEditorActorSubsystem::DeleteSelectedActors(UWorld* InWorld)
 	const FScopedTransaction Transaction(bComponentsSelected ? NSLOCTEXT("UnrealEd", "DeleteComponents", "Delete Components") : NSLOCTEXT("UnrealEd", "DeleteActors", "Delete Actors"));
 
 	FEditorDelegates::OnDeleteActorsBegin.Broadcast();
-	GEditor->edactDeleteSelected(InWorld);
+	const bool bCheckRef = GetDefault<ULevelEditorMiscSettings>()->bCheckReferencesOnDelete;
+	GEditor->edactDeleteSelected(InWorld, true, bCheckRef, bCheckRef);
 	FEditorDelegates::OnDeleteActorsEnd.Broadcast();
 }
 
diff --git a/Engine/Source/Editor/UnrealEd/Private/UnrealEdSrv.cpp b/Engine/Source/Editor/UnrealEd/Private/UnrealEdSrv.cpp
--- a/Engine/Source/Editor/UnrealEd/Private/UnrealEdSrv.cpp
+++ b/Engine/Source/Editor/UnrealEd/Private/UnrealEdSrv.cpp
@@ -1514,7 +1514,12 @@ bool UUnrealEdEngine::Exec_Edit( UWorld* InWorld, const TCHAR* Str, FOutputDevic
 						}
 
 						CommonActions->CopySelectedElements(SelectionSet);
-						CommonActions->DeleteSelectedElements(SelectionSet, InWorld, FTypedElementDeletionOptions());
+						const bool bCheckRef = GetDefault<ULevelEditorMiscSettings>()->bCheckReferencesOnDelete;
+						FTypedElementDeletionOptions Options;
+						Options
+							.SetWarnAboutReferences(bCheckRef)
+							.SetWarnAboutSoftReferences(bCheckRef);
+						CommonActions->DeleteSelectedElements(SelectionSet, InWorld, Options);
 
 						if (!bComponentsSelected)
 						{
@@ -1530,7 +1535,8 @@ bool UUnrealEdEngine::Exec_Edit( UWorld* InWorld, const TCHAR* Str, FOutputDevic
 			const FScopedTransaction Transaction(NSLOCTEXT("UnrealEd", "Cut", "Cut"));
 
 			edactCopySelected(InWorld);
-			edactDeleteSelected(InWorld);
+			const bool bCheckRef = GetDefault<ULevelEditorMiscSettings>()->bCheckReferencesOnDelete;
+			edactDeleteSelected(InWorld, true, bCheckRef, bCheckRef);
 		}
 		else
 		{
@@ -2587,7 +2593,12 @@ bool UUnrealEdEngine::Exec_Actor( UWorld* InWorld, const TCHAR* Str, FOutputDevi
 					else
 					{
 						FEditorDelegates::OnDeleteActorsBegin.Broadcast();
-						CommonActions->DeleteSelectedElements(SelectionSet, InWorld, FTypedElementDeletionOptions());
+						const bool bCheckRef = GetDefault<ULevelEditorMiscSettings>()->bCheckReferencesOnDelete;
+						FTypedElementDeletionOptions Options;
+						Options
+							.SetWarnAboutReferences(bCheckRef)
+							.SetWarnAboutSoftReferences(bCheckRef);
+						CommonActions->DeleteSelectedElements(SelectionSet, InWorld, Options);
 						FEditorDelegates::OnDeleteActorsEnd.Broadcast();
 					}
 				}
