Fixed an assert when gathering localization data from a WP world that had already been loaded as the editor world

This would cause an assert when re-initializing the WP world during the gather. We now ensure the editor always has an empty map prior to gathering any localization data from assets.

diff --git a/Engine/Source/Editor/UnrealEd/Private/Commandlets/GatherTextFromAssetsCommandlet.cpp b/Engine/Source/Editor/UnrealEd/Private/Commandlets/GatherTextFromAssetsCommandlet.cpp
--- a/Engine/Source/Editor/UnrealEd/Private/Commandlets/GatherTextFromAssetsCommandlet.cpp
+++ b/Engine/Source/Editor/UnrealEd/Private/Commandlets/GatherTextFromAssetsCommandlet.cpp
@@ -29,6 +29,7 @@
 #include "WorldPartition/WorldPartitionActorDescUtils.h"
 #include "Engine/Level.h"
 #include "Engine/World.h"
+#include "Editor.h"
 #include "EditorWorldUtils.h"
 #include "PackageHelperFunctions.h"
 #include "ShaderCompiler.h"
@@ -546,6 +547,17 @@ int32 UGatherTextFromAssetsCommandlet::Main(const FString& Params)
 		}
 	}
 
+	// If the editor has loaded a persistent world then create an empty world prior to starting the asset gather
+	// This avoids any issues when loading and initializing worlds during the gather, as WP needs to re-initialize the world
+	if (GEditor)
+	{
+		if (UWorld* EditorWorld = GEditor->GetEditorWorldContext().World();
+			EditorWorld && !FPackageName::IsTempPackage(FNameBuilder(EditorWorld->GetPackage()->GetFName()).ToView()))
+		{
+			GEditor->CreateNewMapForEditing(/*bPromptForSave*/false);
+		}
+	}
+
 	UE_LOG(LogGatherTextFromAssetsCommandlet, Display, TEXT("Discovering assets to gather..."));
 
 	FAssetRegistryModule& AssetRegistryModule = FModuleManager::LoadModuleChecked<FAssetRegistryModule>(TEXT("AssetRegistry"));
