Remove TypedElementViewportInteractionGizmoManipulation from TypedElement ViewportInteraction
* The same functionality is provided by the GizmoManipulationStarted and GizmoManipulationStopped.
* This simplifies the interface and removes the bespoke parameter added to branch between when manipulation starts and when it is ongoing to avoid expensive rendering invalidations
* There are some clients of TypedElementViewportInteractionGizmoManipulation::ApplyDeltaToActor, this CL does not address usage of this function.  Doing so will require further adoption of TypedElements among those systems

diff --git a/Engine/Source/Editor/EditorFramework/Private/Elements/Framework/TypedElementViewportInteraction.cpp b/Engine/Source/Editor/EditorFramework/Private/Elements/Framework/TypedElementViewportInteraction.cpp
--- a/Engine/Source/Editor/EditorFramework/Private/Elements/Framework/TypedElementViewportInteraction.cpp
+++ b/Engine/Source/Editor/EditorFramework/Private/Elements/Framework/TypedElementViewportInteraction.cpp
@@ -54,7 +54,7 @@ void FTypedElementViewportInteractionCustomization::GizmoManipulationStarted(con
 	InElementWorldHandle.NotifyMovementStarted();
 }
 
-void FTypedElementViewportInteractionCustomization::GizmoManipulationDeltaUpdate(const TTypedElement<ITypedElementWorldInterface>& InElementWorldHandle, const UE::Widget::EWidgetMode InWidgetMode, const EAxisList::Type InDragAxis, const FInputDeviceState& InInputState, const FTransform& InDeltaTransform, const FVector& InPivotLocation, ETypedElementViewportInteractionDragMovementType GizmoDeltaType)
+void FTypedElementViewportInteractionCustomization::GizmoManipulationDeltaUpdate(const TTypedElement<ITypedElementWorldInterface>& InElementWorldHandle, const UE::Widget::EWidgetMode InWidgetMode, const EAxisList::Type InDragAxis, const FInputDeviceState& InInputState, const FTransform& InDeltaTransform, const FVector& InPivotLocation)
 {
 	FTransform ElementWorldTransform;
 	if (InElementWorldHandle.GetWorldTransform(ElementWorldTransform))
@@ -191,16 +191,16 @@ void UTypedElementViewportInteraction::BeginGizmoManipulation(FTypedElementListC
 	});
 }
 
-void UTypedElementViewportInteraction::UpdateGizmoManipulation(FTypedElementListConstRef InElementsToMove, const UE::Widget::EWidgetMode InWidgetMode, const EAxisList::Type InDragAxis, const FInputDeviceState& InInputState, const FTransform& InDeltaTransform, ETypedElementViewportInteractionDragMovementType GizmoDeltaType)
+void UTypedElementViewportInteraction::UpdateGizmoManipulation(FTypedElementListConstRef InElementsToMove, const UE::Widget::EWidgetMode InWidgetMode, const EAxisList::Type InDragAxis, const FInputDeviceState& InInputState, const FTransform& InDeltaTransform)
 {
-	InElementsToMove->ForEachElementHandle([this, InWidgetMode, InDragAxis, &InInputState, &InDeltaTransform, GizmoDeltaType](const FTypedElementHandle& InElementToMove)
+	InElementsToMove->ForEachElementHandle([this, InWidgetMode, InDragAxis, &InInputState, &InDeltaTransform](const FTypedElementHandle& InElementToMove)
 	{
 		FTypedElementViewportInteractionElement ViewportInteractionElement = ResolveViewportInteractionElement(InElementToMove);
 		if (ViewportInteractionElement)
 		{
 			FVector PivotLocation = FVector::ZeroVector;
 			ViewportInteractionElement.GetGizmoPivotLocation(InWidgetMode, PivotLocation);
-			ViewportInteractionElement.GizmoManipulationDeltaUpdate(InWidgetMode, InDragAxis, InInputState, InDeltaTransform, PivotLocation, GizmoDeltaType);
+			ViewportInteractionElement.GizmoManipulationDeltaUpdate(InWidgetMode, InDragAxis, InInputState, InDeltaTransform, PivotLocation);
 		}
 		return true;
 	});
@@ -238,7 +238,7 @@ void UTypedElementViewportInteraction::ApplyDeltaToElement(const FTypedElementHa
 	{
 		FVector PivotLocation = FVector::ZeroVector;
 		ViewportInteractionElement.GetGizmoPivotLocation(InWidgetMode, PivotLocation);
-		ViewportInteractionElement.GizmoManipulationDeltaUpdate(InWidgetMode, InDragAxis, InInputState, InDeltaTransform, PivotLocation, ETypedElementViewportInteractionDragMovementType::None);
+		ViewportInteractionElement.GizmoManipulationDeltaUpdate(InWidgetMode, InDragAxis, InInputState, InDeltaTransform, PivotLocation);
 	}
 }
 
diff --git a/Engine/Source/Editor/EditorFramework/Public/Elements/Framework/TypedElementViewportInteraction.h b/Engine/Source/Editor/EditorFramework/Public/Elements/Framework/TypedElementViewportInteraction.h
--- a/Engine/Source/Editor/EditorFramework/Public/Elements/Framework/TypedElementViewportInteraction.h
+++ b/Engine/Source/Editor/EditorFramework/Public/Elements/Framework/TypedElementViewportInteraction.h
@@ -31,20 +31,6 @@ enum class ETypedElementViewportInteractionGizmoManipulationType : uint8
 	Click,
 };
 
-/**
- * Setting to identify if a movement is a drag operation, and if so, whether it is a first or subsequent move.  Some clients may apply
- * optimizations where certain logic (such as dirtying state) only runs on the first movement delta.
- */
-enum class ETypedElementViewportInteractionDragMovementType : uint8
-{
-	/** Not part of a drag operation */
-	None,
-	/** First delta movement of a drag operation */
-	FirstMove,
-	/** Subsequent delta movement of a drag operation */
-	SubsequentMove
-};
-
 /**
  * Customization used to allow asset editors (such as the level editor) to override the base behavior of viewport interaction.
  */
@@ -57,7 +43,7 @@ public:
 	virtual bool GetGizmoPivotLocation(const TTypedElement<ITypedElementWorldInterface>& InElementWorldHandle, const UE::Widget::EWidgetMode InWidgetMode, FVector& OutPivotLocation);
 	virtual void PreGizmoManipulationStarted(TArrayView<const FTypedElementHandle> InElementHandles, const UE::Widget::EWidgetMode InWidgetMode);
 	virtual void GizmoManipulationStarted(const TTypedElement<ITypedElementWorldInterface>& InElementWorldHandle, const UE::Widget::EWidgetMode InWidgetMode);
-	virtual void GizmoManipulationDeltaUpdate(const TTypedElement<ITypedElementWorldInterface>& InElementWorldHandle, const UE::Widget::EWidgetMode InWidgetMode, const EAxisList::Type InDragAxis, const FInputDeviceState& InInputState, const FTransform& InDeltaTransform, const FVector& InPivotLocation, ETypedElementViewportInteractionDragMovementType GizmoDeltaType);
+	virtual void GizmoManipulationDeltaUpdate(const TTypedElement<ITypedElementWorldInterface>& InElementWorldHandle, const UE::Widget::EWidgetMode InWidgetMode, const EAxisList::Type InDragAxis, const FInputDeviceState& InInputState, const FTransform& InDeltaTransform, const FVector& InPivotLocation);
 	virtual void GizmoManipulationStopped(const TTypedElement<ITypedElementWorldInterface>& InElementWorldHandle, const UE::Widget::EWidgetMode InWidgetMode, const ETypedElementViewportInteractionGizmoManipulationType InManipulationType);
 	virtual void PostGizmoManipulationStopped(TArrayView<const FTypedElementHandle> InElementHandles, const UE::Widget::EWidgetMode InWidgetMode);
 	virtual void MirrorElement(const TTypedElement<ITypedElementWorldInterface>& InElementWorldHandle, const FVector& InMirrorScale, const FVector& InPivotLocation);
@@ -98,7 +84,7 @@ public:
 	//~ See UTypedElementViewportInteraction for API docs
 	bool GetGizmoPivotLocation(const UE::Widget::EWidgetMode InWidgetMode, FVector& OutPivotLocation) const { return ViewportInteractionCustomization->GetGizmoPivotLocation(ElementWorldHandle, InWidgetMode, OutPivotLocation); }
 	void GizmoManipulationStarted(const UE::Widget::EWidgetMode InWidgetMode) const { ViewportInteractionCustomization->GizmoManipulationStarted(ElementWorldHandle, InWidgetMode); }
-	void GizmoManipulationDeltaUpdate(const UE::Widget::EWidgetMode InWidgetMode, const EAxisList::Type InDragAxis, const FInputDeviceState& InInputState, const FTransform& InDeltaTransform, const FVector& InPivotLocation, ETypedElementViewportInteractionDragMovementType GizmoDeltaType) const { ViewportInteractionCustomization->GizmoManipulationDeltaUpdate(ElementWorldHandle, InWidgetMode, InDragAxis, InInputState, InDeltaTransform, InPivotLocation, GizmoDeltaType); }
+	void GizmoManipulationDeltaUpdate(const UE::Widget::EWidgetMode InWidgetMode, const EAxisList::Type InDragAxis, const FInputDeviceState& InInputState, const FTransform& InDeltaTransform, const FVector& InPivotLocation) const { ViewportInteractionCustomization->GizmoManipulationDeltaUpdate(ElementWorldHandle, InWidgetMode, InDragAxis, InInputState, InDeltaTransform, InPivotLocation); }
 	void GizmoManipulationStopped(const UE::Widget::EWidgetMode InWidgetMode, const ETypedElementViewportInteractionGizmoManipulationType InManipulationType) const { ViewportInteractionCustomization->GizmoManipulationStopped(ElementWorldHandle, InWidgetMode, InManipulationType); }
 	void MirrorElement(const FVector& InMirrorScale, const FVector& InPivotLocation) const { ViewportInteractionCustomization->MirrorElement(ElementWorldHandle, InMirrorScale, InPivotLocation); }
 	bool GetFocusBounds(FBoxSphereBounds& OutBounds) const { return ViewportInteractionCustomization->GetFocusBounds(ElementWorldHandle, OutBounds); }
@@ -128,7 +114,7 @@ public:
 	 * Notify that the gizmo has manipulated the transform of the given set of elements.
 	 * @note This list should have been pre-normalized via UTypedElementSelectionSet::GetNormalizedSelection or UTypedElementSelectionSet::GetNormalizedElementList, or FLevelEditorViewportClient::GetElementsToManipulate.
 	 */
-	void UpdateGizmoManipulation(FTypedElementListConstRef InElementsToMove, const UE::Widget::EWidgetMode InWidgetMode, const EAxisList::Type InDragAxis, const FInputDeviceState& InInputState, const FTransform& InDeltaTransform, ETypedElementViewportInteractionDragMovementType GizmoDeltaType);
+	void UpdateGizmoManipulation(FTypedElementListConstRef InElementsToMove, const UE::Widget::EWidgetMode InWidgetMode, const EAxisList::Type InDragAxis, const FInputDeviceState& InInputState, const FTransform& InDeltaTransform);
 	
 	/**
 	 * Notify that the gizmo has finished manipulating the transform of the given set of elements.
diff --git a/Engine/Source/Editor/LevelEditor/Private/Elements/Actor/ActorElementLevelEditorViewportInteractionCustomization.cpp b/Engine/Source/Editor/LevelEditor/Private/Elements/Actor/ActorElementLevelEditorViewportInteractionCustomization.cpp
--- a/Engine/Source/Editor/LevelEditor/Private/Elements/Actor/ActorElementLevelEditorViewportInteractionCustomization.cpp
+++ b/Engine/Source/Editor/LevelEditor/Private/Elements/Actor/ActorElementLevelEditorViewportInteractionCustomization.cpp
@@ -30,7 +30,7 @@ void FActorElementLevelEditorViewportInteractionCustomization::GizmoManipulation
 	}
 }
 
-void FActorElementLevelEditorViewportInteractionCustomization::GizmoManipulationDeltaUpdate(const TTypedElement<ITypedElementWorldInterface>& InElementWorldHandle, const UE::Widget::EWidgetMode InWidgetMode, const EAxisList::Type InDragAxis, const FInputDeviceState& InInputState, const FTransform& InDeltaTransform, const FVector& InPivotLocation, ETypedElementViewportInteractionDragMovementType GizmoDeltaType)
+void FActorElementLevelEditorViewportInteractionCustomization::GizmoManipulationDeltaUpdate(const TTypedElement<ITypedElementWorldInterface>& InElementWorldHandle, const UE::Widget::EWidgetMode InWidgetMode, const EAxisList::Type InDragAxis, const FInputDeviceState& InInputState, const FTransform& InDeltaTransform, const FVector& InPivotLocation)
 {
 	AActor* Actor = ActorElementDataUtil::GetActorFromHandleChecked(InElementWorldHandle);
 
@@ -55,7 +55,7 @@ void FActorElementLevelEditorViewportInteractionCustomization::GizmoManipulation
 		ModifiedDeltaTransform.SetScale3D(AdjustedScale);
 	}
 
-	FActorElementEditorViewportInteractionCustomization::GizmoManipulationDeltaUpdate(InElementWorldHandle, InWidgetMode, InDragAxis, InInputState, ModifiedDeltaTransform, InPivotLocation, GizmoDeltaType);
+	FActorElementEditorViewportInteractionCustomization::GizmoManipulationDeltaUpdate(InElementWorldHandle, InWidgetMode, InDragAxis, InInputState, ModifiedDeltaTransform, InPivotLocation);
 
 	// Update the cameras from their locked actor (if any) only if the viewport is real-time enabled
 	GetMutableLevelEditorViewportClient()->UpdateLockedActorViewports(Actor, true);
diff --git a/Engine/Source/Editor/LevelEditor/Private/Elements/Component/ComponentElementLevelEditorViewportInteractionCustomization.cpp b/Engine/Source/Editor/LevelEditor/Private/Elements/Component/ComponentElementLevelEditorViewportInteractionCustomization.cpp
--- a/Engine/Source/Editor/LevelEditor/Private/Elements/Component/ComponentElementLevelEditorViewportInteractionCustomization.cpp
+++ b/Engine/Source/Editor/LevelEditor/Private/Elements/Component/ComponentElementLevelEditorViewportInteractionCustomization.cpp
@@ -32,7 +32,7 @@ void FComponentElementLevelEditorViewportInteractionCustomization::GizmoManipula
 	}
 }
 
-void FComponentElementLevelEditorViewportInteractionCustomization::GizmoManipulationDeltaUpdate(const TTypedElement<ITypedElementWorldInterface>& InElementWorldHandle, const UE::Widget::EWidgetMode InWidgetMode, const EAxisList::Type InDragAxis, const FInputDeviceState& InInputState, const FTransform& InDeltaTransform, const FVector& InPivotLocation, ETypedElementViewportInteractionDragMovementType GizmoDeltaType)
+void FComponentElementLevelEditorViewportInteractionCustomization::GizmoManipulationDeltaUpdate(const TTypedElement<ITypedElementWorldInterface>& InElementWorldHandle, const UE::Widget::EWidgetMode InWidgetMode, const EAxisList::Type InDragAxis, const FInputDeviceState& InInputState, const FTransform& InDeltaTransform, const FVector& InPivotLocation)
 {
 	UActorComponent* Component = ComponentElementDataUtil::GetComponentFromHandleChecked(InElementWorldHandle);
 
@@ -63,7 +63,7 @@ void FComponentElementLevelEditorViewportInteractionCustomization::GizmoManipula
 			ModifiedDeltaTransform.SetScale3D(AdjustedScale);
 		}
 
-		FComponentElementEditorViewportInteractionCustomization::GizmoManipulationDeltaUpdate(InElementWorldHandle, InWidgetMode, InDragAxis, InInputState, ModifiedDeltaTransform, InPivotLocation, GizmoDeltaType);
+		FComponentElementEditorViewportInteractionCustomization::GizmoManipulationDeltaUpdate(InElementWorldHandle, InWidgetMode, InDragAxis, InInputState, ModifiedDeltaTransform, InPivotLocation);
 	}
 }
 
diff --git a/Engine/Source/Editor/LevelEditor/Public/Elements/Actor/ActorElementLevelEditorViewportInteractionCustomization.h b/Engine/Source/Editor/LevelEditor/Public/Elements/Actor/ActorElementLevelEditorViewportInteractionCustomization.h
--- a/Engine/Source/Editor/LevelEditor/Public/Elements/Actor/ActorElementLevelEditorViewportInteractionCustomization.h
+++ b/Engine/Source/Editor/LevelEditor/Public/Elements/Actor/ActorElementLevelEditorViewportInteractionCustomization.h
@@ -9,7 +9,7 @@ class LEVELEDITOR_API FActorElementLevelEditorViewportInteractionCustomization :
 {
 public:
 	virtual void GizmoManipulationStarted(const TTypedElement<ITypedElementWorldInterface>& InElementWorldHandle, const UE::Widget::EWidgetMode InWidgetMode) override;
-	virtual void GizmoManipulationDeltaUpdate(const TTypedElement<ITypedElementWorldInterface>& InElementWorldHandle, const UE::Widget::EWidgetMode InWidgetMode, const EAxisList::Type InDragAxis, const FInputDeviceState& InInputState, const FTransform& InDeltaTransform, const FVector& InPivotLocation, ETypedElementViewportInteractionDragMovementType GizmoDeltaType) override;
+	virtual void GizmoManipulationDeltaUpdate(const TTypedElement<ITypedElementWorldInterface>& InElementWorldHandle, const UE::Widget::EWidgetMode InWidgetMode, const EAxisList::Type InDragAxis, const FInputDeviceState& InInputState, const FTransform& InDeltaTransform, const FVector& InPivotLocation) override;
 	virtual void GizmoManipulationStopped(const TTypedElement<ITypedElementWorldInterface>& InElementWorldHandle, const UE::Widget::EWidgetMode InWidgetMode, const ETypedElementViewportInteractionGizmoManipulationType InManipulationType) override;
 	virtual void PostGizmoManipulationStopped(TArrayView<const FTypedElementHandle> InElementHandles, const UE::Widget::EWidgetMode InWidgetMode) override;
 
diff --git a/Engine/Source/Editor/LevelEditor/Public/Elements/Component/ComponentElementLevelEditorViewportInteractionCustomization.h b/Engine/Source/Editor/LevelEditor/Public/Elements/Component/ComponentElementLevelEditorViewportInteractionCustomization.h
--- a/Engine/Source/Editor/LevelEditor/Public/Elements/Component/ComponentElementLevelEditorViewportInteractionCustomization.h
+++ b/Engine/Source/Editor/LevelEditor/Public/Elements/Component/ComponentElementLevelEditorViewportInteractionCustomization.h
@@ -11,7 +11,7 @@ class LEVELEDITOR_API FComponentElementLevelEditorViewportInteractionCustomizati
 {
 public:
 	virtual void GizmoManipulationStarted(const TTypedElement<ITypedElementWorldInterface>& InElementWorldHandle, const UE::Widget::EWidgetMode InWidgetMode) override;
-	virtual void GizmoManipulationDeltaUpdate(const TTypedElement<ITypedElementWorldInterface>& InElementWorldHandle, const UE::Widget::EWidgetMode InWidgetMode, const EAxisList::Type InDragAxis, const FInputDeviceState& InInputState, const FTransform& InDeltaTransform, const FVector& InPivotLocation, ETypedElementViewportInteractionDragMovementType GizmoDeltaType) override;
+	virtual void GizmoManipulationDeltaUpdate(const TTypedElement<ITypedElementWorldInterface>& InElementWorldHandle, const UE::Widget::EWidgetMode InWidgetMode, const EAxisList::Type InDragAxis, const FInputDeviceState& InInputState, const FTransform& InDeltaTransform, const FVector& InPivotLocation) override;
 	virtual void GizmoManipulationStopped(const TTypedElement<ITypedElementWorldInterface>& InElementWorldHandle, const UE::Widget::EWidgetMode InWidgetMode, const ETypedElementViewportInteractionGizmoManipulationType InManipulationType) override;
 
 	static void ValidateScale(const FVector& InOriginalPreDragScale, const EAxisList::Type InDragAxis, const FVector& InCurrentScale, const FVector& InBoxExtent, FVector& InOutScaleDelta, bool bInCheckSmallExtent);
diff --git a/Engine/Source/Editor/UnrealEd/Private/EditorEngine.cpp b/Engine/Source/Editor/UnrealEd/Private/EditorEngine.cpp
--- a/Engine/Source/Editor/UnrealEd/Private/EditorEngine.cpp
+++ b/Engine/Source/Editor/UnrealEd/Private/EditorEngine.cpp
@@ -2921,7 +2921,7 @@ void UEditorEngine::ApplyDeltaToActor(AActor* InActor,
 	FInputDeviceState InputState;
 	InputState.SetModifierKeyStates(bShiftDown, bAltDown, bControlDown, false);
 
-	FActorElementEditorViewportInteractionCustomization::ApplyDeltaToActor(InActor, bDelta, InTrans, InRot, InScale, GLevelEditorModeTools().PivotLocation, InputState, ETypedElementViewportInteractionDragMovementType::None);
+	FActorElementEditorViewportInteractionCustomization::ApplyDeltaToActor(InActor, bDelta, InTrans, InRot, InScale, GLevelEditorModeTools().PivotLocation, InputState);
 }
 
 void UEditorEngine::ApplyDeltaToComponent(USceneComponent* InComponent,
diff --git a/Engine/Source/Editor/UnrealEd/Private/Elements/Actor/ActorElementEditorViewportInteractionCustomization.cpp b/Engine/Source/Editor/UnrealEd/Private/Elements/Actor/ActorElementEditorViewportInteractionCustomization.cpp
--- a/Engine/Source/Editor/UnrealEd/Private/Elements/Actor/ActorElementEditorViewportInteractionCustomization.cpp
+++ b/Engine/Source/Editor/UnrealEd/Private/Elements/Actor/ActorElementEditorViewportInteractionCustomization.cpp
@@ -13,15 +13,190 @@
 #include "Elements/Interfaces/TypedElementHierarchyInterface.h"
 #include "Elements/Framework/TypedElementRegistry.h"
 
-void FActorElementEditorViewportInteractionCustomization::GizmoManipulationDeltaUpdate(const TTypedElement<ITypedElementWorldInterface>& InElementWorldHandle, const UE::Widget::EWidgetMode InWidgetMode, const EAxisList::Type InDragAxis, const FInputDeviceState& InInputState, const FTransform& InDeltaTransform, const FVector& InPivotLocation, ETypedElementViewportInteractionDragMovementType GizmoDeltaType)
+void FActorElementEditorViewportInteractionCustomization::GizmoManipulationStarted(
+	const TTypedElement<ITypedElementWorldInterface>& InElementWorldHandle,
+	const UE::Widget::EWidgetMode InWidgetMode)
 {
-	AActor* Actor = ActorElementDataUtil::GetActorFromHandleChecked(InElementWorldHandle);
+	AActor* RootActor = ActorElementDataUtil::GetActorFromHandleChecked(InElementWorldHandle);
+	
+	FTypedElementViewportInteractionCustomization::GizmoManipulationStarted(InElementWorldHandle, InWidgetMode);
+
+	TArray<AActor*> AffectedActors;
+	AffectedActors.Add(RootActor);
+	
+	// grab all attached actors to invalidate everybody 	
+	RootActor->GetAttachedActors(AffectedActors, /*bResetArray*/ false, /*bRecursivelyIncludeAttachedActors*/ true );
+
+	if (GEditor->IsDeltaModificationEnabled())
+	{
+		for (AActor* Actor : AffectedActors)
+		{
+			Actor->Modify();
+		}
+	}
+
+	if (!GIsDemoMode)
+	{
+		for (AActor* Actor : AffectedActors)
+		{
+			// We don't know at this point if it is translation only
+			constexpr bool bTranslationOnly = false;
+			Actor->InvalidateLightingCacheDetailed(bTranslationOnly);
+		}
+	}
+}
+
+void FActorElementEditorViewportInteractionCustomization::GizmoManipulationDeltaUpdate(const TTypedElement<ITypedElementWorldInterface>& InElementWorldHandle, const UE::Widget::EWidgetMode InWidgetMode, const EAxisList::Type InDragAxis, const FInputDeviceState& InInputState, const FTransform& InDeltaTransform, const FVector& InPivotLocation)
+{
+	AActor* RootActor = ActorElementDataUtil::GetActorFromHandleChecked(InElementWorldHandle);
+	
+	TArray<AActor*> AffectedActors;
+	AffectedActors.Add(RootActor);
+	
+	// grab all attached actors to invalidate everybody 	
+	RootActor->GetAttachedActors(AffectedActors, /*bResetArray*/ false, /*bRecursivelyIncludeAttachedActors*/ true );
 
 	const FVector DeltaTranslation = InDeltaTransform.GetTranslation();
 	const FRotator DeltaRotation = InDeltaTransform.Rotator();
 	const FVector DeltaScale3D = InDeltaTransform.GetScale3D();
+	
+	const bool bIsSimulatingInEditor = GEditor->IsSimulatingInEditor();
+
+	FNavigationLockContext LockNavigationUpdates(RootActor->GetWorld(), ENavigationLockReason::ContinuousEditorMove);
+
+	///////////////////
+	// Rotation
+	
+	const bool bRotatingActor = !DeltaRotation.IsZero();
+	if (bRotatingActor)
+	{			
+		if (RootActor->GetRootComponent())
+		{
+			const FRotator OriginalRotation = RootActor->GetRootComponent()->GetComponentRotation();
 
-	FActorElementEditorViewportInteractionCustomization::ApplyDeltaToActor(Actor, /*bDelta*/true, &DeltaTranslation, &DeltaRotation, &DeltaScale3D, InPivotLocation, InInputState, GizmoDeltaType);
+			RootActor->EditorApplyRotation(DeltaRotation, InInputState.bAltKeyDown, InInputState.bShiftKeyDown, InInputState.bCtrlKeyDown);
+
+			// Check to see if we should transform the rigid body
+			UPrimitiveComponent* RootPrimitiveComponent = Cast<UPrimitiveComponent>(RootActor->GetRootComponent());
+			if (bIsSimulatingInEditor && GIsPlayInEditorWorld && RootPrimitiveComponent)
+			{
+				FRotator ActorRotWind, ActorRotRem;
+				OriginalRotation.GetWindingAndRemainder(ActorRotWind, ActorRotRem);
+
+				const FQuat ActorQ = ActorRotRem.Quaternion();
+				const FQuat DeltaQ = DeltaRotation.Quaternion();
+				const FQuat ResultQ = DeltaQ * ActorQ;
+
+				const FRotator NewActorRotRem = FRotator(ResultQ);
+				FRotator DeltaRot = NewActorRotRem - ActorRotRem;
+				DeltaRot.Normalize();
+
+				// @todo SIE: Not taking into account possible offset between root component and actor
+				RootPrimitiveComponent->SetWorldRotation(OriginalRotation + DeltaRot);
+			}
+		}
+
+		FVector NewActorLocation = RootActor->GetActorLocation();
+		NewActorLocation -= InPivotLocation;
+		NewActorLocation = FRotationMatrix(DeltaRotation).TransformPosition(NewActorLocation);
+		NewActorLocation += InPivotLocation;
+		NewActorLocation -= RootActor->GetActorLocation();
+		RootActor->EditorApplyTranslation(NewActorLocation, InInputState.bAltKeyDown, InInputState.bShiftKeyDown, InInputState.bCtrlKeyDown);
+	}
+
+	///////////////////
+	// Translation
+	if (RootActor->GetRootComponent())
+	{
+		const FVector OriginalLocation = RootActor->GetRootComponent()->GetComponentLocation();
+
+		RootActor->EditorApplyTranslation(DeltaTranslation, InInputState.bAltKeyDown, InInputState.bShiftKeyDown, InInputState.bCtrlKeyDown);
+
+		// Check to see if we should transform the rigid body
+		UPrimitiveComponent* RootPrimitiveComponent = Cast<UPrimitiveComponent>(RootActor->GetRootComponent());
+		if (bIsSimulatingInEditor && GIsPlayInEditorWorld && RootPrimitiveComponent)
+		{
+			// @todo SIE: Not taking into account possible offset between root component and actor
+			RootPrimitiveComponent->SetWorldLocation(OriginalLocation + DeltaTranslation);
+		}
+	}
+
+	///////////////////
+	// Scaling
+	const bool bScalingActor = !DeltaScale3D.IsNearlyZero(0.000001f);
+	if (bScalingActor)
+	{
+		FVector ModifiedScale = DeltaScale3D;
+
+		// Note: With the new additive scaling method, this is handled in FLevelEditorViewportClient::ModifyScale
+		if (GEditor->UsePercentageBasedScaling())
+		{
+			// Get actor box extents
+			const FBox BoundingBox = RootActor->GetComponentsBoundingBox(true);
+			const FVector BoundsExtents = BoundingBox.GetExtent();
+
+			// Make sure scale on actors is clamped to a minimum and maximum size.
+			const float MinThreshold = 1.0f;
+
+			for (int32 Idx = 0; Idx < 3; Idx++)
+			{
+				if ((FMath::Pow(BoundsExtents[Idx], 2)) > BIG_NUMBER)
+				{
+					ModifiedScale[Idx] = 0.0f;
+				}
+				else if (SMALL_NUMBER < BoundsExtents[Idx])
+				{
+					const bool bBelowAllowableScaleThreshold = ((DeltaScale3D[Idx] + 1.0f) * BoundsExtents[Idx]) < MinThreshold;
+
+					if (bBelowAllowableScaleThreshold)
+					{
+						ModifiedScale[Idx] = (MinThreshold / BoundsExtents[Idx]) - 1.0f;
+					}
+				}
+			}
+		}
+		
+		// Flag actors to use old-style scaling or not
+		// @todo: Remove this hack once we have decided on the scaling method to use.
+		AActor::bUsePercentageBasedScaling = GEditor->UsePercentageBasedScaling();
+
+		RootActor->EditorApplyScale(
+			ModifiedScale,
+			&InPivotLocation,
+			InInputState.bAltKeyDown,
+			InInputState.bShiftKeyDown,
+			InInputState.bCtrlKeyDown
+			);
+	}
+
+	for (AActor* Actor : AffectedActors)
+	{
+		Actor->PostEditMove(false);
+	}
+}
+
+void FActorElementEditorViewportInteractionCustomization::GizmoManipulationStopped(
+	const TTypedElement<ITypedElementWorldInterface>& InElementWorldHandle,
+	const UE::Widget::EWidgetMode InWidgetMode,
+	const ETypedElementViewportInteractionGizmoManipulationType InManipulationType)
+{
+	FTypedElementViewportInteractionCustomization::GizmoManipulationStopped(
+		InElementWorldHandle, InWidgetMode, InManipulationType);
+
+	AActor* RootActor = ActorElementDataUtil::GetActorFromHandleChecked(InElementWorldHandle);
+	
+	TArray<AActor*> AffectedActors;
+	AffectedActors.Add(RootActor);
+	
+	// grab all attached actors to invalidate everybody 	
+	RootActor->GetAttachedActors(AffectedActors, /*bResetArray*/ false, /*bRecursivelyIncludeAttachedActors*/ true );
+
+	InElementWorldHandle.NotifyMovementEnded();
+	// Update the actor before leaving.
+	for (AActor* Actor : AffectedActors)
+	{
+		Actor->MarkPackageDirty();
+	}
 }
 
 void FActorElementEditorViewportInteractionCustomization::MirrorElement(const TTypedElement<ITypedElementWorldInterface>& InElementWorldHandle, const FVector& InMirrorScale, const FVector& InPivotLocation)
@@ -132,7 +307,7 @@ bool FActorElementEditorViewportInteractionCustomization::GetFocusBounds( const
 	return true;
 }
 
-void FActorElementEditorViewportInteractionCustomization::ApplyDeltaToActor(AActor* InActor, const bool InIsDelta, const FVector* InDeltaTranslationPtr, const FRotator* InDeltaRotationPtr, const FVector* InDeltaScalePtr, const FVector& InPivotLocation, const FInputDeviceState& InInputState, ETypedElementViewportInteractionDragMovementType GizmoDeltaType)
+void FActorElementEditorViewportInteractionCustomization::ApplyDeltaToActor(AActor* InActor, const bool InIsDelta, const FVector* InDeltaTranslationPtr, const FRotator* InDeltaRotationPtr, const FVector* InDeltaScalePtr, const FVector& InPivotLocation, const FInputDeviceState& InInputState)
 {
 	const bool bIsSimulatingInEditor = GEditor->IsSimulatingInEditor();
 
@@ -304,8 +479,7 @@ void FActorElementEditorViewportInteractionCustomization::ApplyDeltaToActor(AAct
 		Actor->MarkPackageDirty();
 	}
 	
-	// Don't need to invalidate lighting cache if this is a subsequent delta movement of the same actors
-	if (!GIsDemoMode && GizmoDeltaType != ETypedElementViewportInteractionDragMovementType::SubsequentMove)
+	if (!GIsDemoMode)
 	{
 		for (AActor* Actor : AffectedActors)
 		{
@@ -317,4 +491,4 @@ void FActorElementEditorViewportInteractionCustomization::ApplyDeltaToActor(AAct
 	{
 		Actor->PostEditMove(false);
 	}
-}
+}
\ No newline at end of file
diff --git a/Engine/Source/Editor/UnrealEd/Private/Elements/Component/ComponentElementEditorViewportInteractionCustomization.cpp b/Engine/Source/Editor/UnrealEd/Private/Elements/Component/ComponentElementEditorViewportInteractionCustomization.cpp
--- a/Engine/Source/Editor/UnrealEd/Private/Elements/Component/ComponentElementEditorViewportInteractionCustomization.cpp
+++ b/Engine/Source/Editor/UnrealEd/Private/Elements/Component/ComponentElementEditorViewportInteractionCustomization.cpp
@@ -28,7 +28,7 @@ bool FComponentElementEditorViewportInteractionCustomization::GetGizmoPivotLocat
 	return false;
 }
 
-void FComponentElementEditorViewportInteractionCustomization::GizmoManipulationDeltaUpdate(const TTypedElement<ITypedElementWorldInterface>& InElementWorldHandle, const UE::Widget::EWidgetMode InWidgetMode, const EAxisList::Type InDragAxis, const FInputDeviceState& InInputState, const FTransform& InDeltaTransform, const FVector& InPivotLocation, ETypedElementViewportInteractionDragMovementType GizmoDeltaType)
+void FComponentElementEditorViewportInteractionCustomization::GizmoManipulationDeltaUpdate(const TTypedElement<ITypedElementWorldInterface>& InElementWorldHandle, const UE::Widget::EWidgetMode InWidgetMode, const EAxisList::Type InDragAxis, const FInputDeviceState& InInputState, const FTransform& InDeltaTransform, const FVector& InPivotLocation)
 {
 	UActorComponent* Component = ComponentElementDataUtil::GetComponentFromHandleChecked(InElementWorldHandle);
 
diff --git a/Engine/Source/Editor/UnrealEd/Private/LevelEditorViewport.cpp b/Engine/Source/Editor/UnrealEd/Private/LevelEditorViewport.cpp
--- a/Engine/Source/Editor/UnrealEd/Private/LevelEditorViewport.cpp
+++ b/Engine/Source/Editor/UnrealEd/Private/LevelEditorViewport.cpp
@@ -1913,7 +1913,6 @@ FLevelEditorViewportClient::FLevelEditorViewportClient(const TSharedPtr<SLevelVi
 	, bLockedCameraView(true)
 	, bNeedToRestoreComponentBeingMovedFlag(false)
 	, bHasBegunGizmoManipulation(false)
-	, bGizmoDeltaApplied(false)
 	, bReceivedFocusRecently(false)
 	, bAlwaysShowModeWidgetAfterSelectionChanges(true)
 	, CachedElementsToManipulate(UTypedElementRegistry::GetInstance()->CreateElementList())
@@ -2850,7 +2849,6 @@ bool FLevelEditorViewportClient::InputWidgetDelta(FViewport* InViewport, EAxisLi
 									{
 										ElementsToManipulate = GetElementsToManipulate();
 										ViewportInteraction->BeginGizmoManipulation(ElementsToManipulate, GetWidgetMode());
-										bGizmoDeltaApplied = false;
 									}
 								}
 
@@ -3175,7 +3173,6 @@ void FLevelEditorViewportClient::TrackingStarted( const FInputEventState& InInpu
 		FTypedElementListConstRef ElementsToManipulate = GetElementsToManipulate();
 		ViewportInteraction->BeginGizmoManipulation(ElementsToManipulate, GetWidgetMode());
 		bHasBegunGizmoManipulation = true;
-		bGizmoDeltaApplied = false;
 
 		if (!bDuplicateActorsInProgress)
 		{
@@ -4031,11 +4028,8 @@ void FLevelEditorViewportClient::ApplyDeltaToSelectedElements(const FTransform&
 	FInputDeviceState InputState;
 	InputState.SetModifierKeyStates(IsShiftPressed(), IsAltPressed(), IsCtrlPressed(), IsCmdPressed());
 
-	ETypedElementViewportInteractionDragMovementType GizmoDeltaType = bGizmoDeltaApplied == false ? ETypedElementViewportInteractionDragMovementType::FirstMove : ETypedElementViewportInteractionDragMovementType::SubsequentMove;
-	bGizmoDeltaApplied = true;
-
 	FTypedElementListConstRef ElementsToManipulate = GetElementsToManipulate();
-	ViewportInteraction->UpdateGizmoManipulation(ElementsToManipulate, GetWidgetMode(), Widget ? Widget->GetCurrentAxis() : EAxisList::None, InputState, ModifiedDeltaTransform, GizmoDeltaType);
+	ViewportInteraction->UpdateGizmoManipulation(ElementsToManipulate, GetWidgetMode(), Widget ? Widget->GetCurrentAxis() : EAxisList::None, InputState, ModifiedDeltaTransform);
 }
 
 void FLevelEditorViewportClient::ApplyDeltaToElement(const FTypedElementHandle& InElementHandle, const FTransform& InDeltaTransform)
diff --git a/Engine/Source/Editor/UnrealEd/Public/Elements/Actor/ActorElementEditorViewportInteractionCustomization.h b/Engine/Source/Editor/UnrealEd/Public/Elements/Actor/ActorElementEditorViewportInteractionCustomization.h
--- a/Engine/Source/Editor/UnrealEd/Public/Elements/Actor/ActorElementEditorViewportInteractionCustomization.h
+++ b/Engine/Source/Editor/UnrealEd/Public/Elements/Actor/ActorElementEditorViewportInteractionCustomization.h
@@ -10,9 +10,11 @@ class AActor;
 class FActorElementEditorViewportInteractionCustomization : public FTypedElementViewportInteractionCustomization
 {
 public:
-	UNREALED_API virtual void GizmoManipulationDeltaUpdate(const TTypedElement<ITypedElementWorldInterface>& InElementWorldHandle, const UE::Widget::EWidgetMode InWidgetMode, const EAxisList::Type InDragAxis, const FInputDeviceState& InInputState, const FTransform& InDeltaTransform, const FVector& InPivotLocation, ETypedElementViewportInteractionDragMovementType GizmoDeltaType) override;
+	UNREALED_API virtual void GizmoManipulationStarted(const TTypedElement<ITypedElementWorldInterface>& InElementWorldHandle, const UE::Widget::EWidgetMode InWidgetMode) override;
+	UNREALED_API virtual void GizmoManipulationDeltaUpdate(const TTypedElement<ITypedElementWorldInterface>& InElementWorldHandle, const UE::Widget::EWidgetMode InWidgetMode, const EAxisList::Type InDragAxis, const FInputDeviceState& InInputState, const FTransform& InDeltaTransform, const FVector& InPivotLocation) override;
+	UNREALED_API virtual void GizmoManipulationStopped(const TTypedElement<ITypedElementWorldInterface>& InElementWorldHandle, const UE::Widget::EWidgetMode InWidgetMode, const ETypedElementViewportInteractionGizmoManipulationType InManipulationType) override;
 	UNREALED_API virtual void MirrorElement(const TTypedElement<ITypedElementWorldInterface>& InElementWorldHandle, const FVector& InMirrorScale, const FVector& InPivotLocation) override;
 	UNREALED_API virtual bool GetFocusBounds(const TTypedElement<ITypedElementWorldInterface>& InElementWorldHandle, FBoxSphereBounds& OutBounds) override;
 	
-	static UNREALED_API void ApplyDeltaToActor(AActor* InActor, const bool InIsDelta, const FVector* InDeltaTranslationPtr, const FRotator* InDeltaRotationPtr, const FVector* InDeltaScalePtr, const FVector& InPivotLocation, const FInputDeviceState& InInputState, ETypedElementViewportInteractionDragMovementType GizmoDeltaType);
+	UNREALED_API static void ApplyDeltaToActor(AActor* InActor, const bool InIsDelta, const FVector* InDeltaTranslationPtr, const FRotator* InDeltaRotationPtr, const FVector* InDeltaScalePtr, const FVector& InPivotLocation, const FInputDeviceState& InInputState);
 };
diff --git a/Engine/Source/Editor/UnrealEd/Public/Elements/Component/ComponentElementEditorViewportInteractionCustomization.h b/Engine/Source/Editor/UnrealEd/Public/Elements/Component/ComponentElementEditorViewportInteractionCustomization.h
--- a/Engine/Source/Editor/UnrealEd/Public/Elements/Component/ComponentElementEditorViewportInteractionCustomization.h
+++ b/Engine/Source/Editor/UnrealEd/Public/Elements/Component/ComponentElementEditorViewportInteractionCustomization.h
@@ -12,7 +12,7 @@ class FComponentElementEditorViewportInteractionCustomization : public FTypedEle
 {
 public:
 	UNREALED_API virtual bool GetGizmoPivotLocation(const TTypedElement<ITypedElementWorldInterface>& InElementWorldHandle, const UE::Widget::EWidgetMode InWidgetMode, FVector& OutPivotLocation) override;
-	UNREALED_API virtual void GizmoManipulationDeltaUpdate(const TTypedElement<ITypedElementWorldInterface>& InElementWorldHandle, const UE::Widget::EWidgetMode InWidgetMode, const EAxisList::Type InDragAxis, const FInputDeviceState& InInputState, const FTransform& InDeltaTransform, const FVector& InPivotLocation, ETypedElementViewportInteractionDragMovementType GizmoDeltaType) override;
+	UNREALED_API virtual void GizmoManipulationDeltaUpdate(const TTypedElement<ITypedElementWorldInterface>& InElementWorldHandle, const UE::Widget::EWidgetMode InWidgetMode, const EAxisList::Type InDragAxis, const FInputDeviceState& InInputState, const FTransform& InDeltaTransform, const FVector& InPivotLocation) override;
 	UNREALED_API virtual bool GetFocusBounds(const TTypedElement<ITypedElementWorldInterface>& InElementWorldHandle, FBoxSphereBounds& OutBounds) override;
 	
 	static UNREALED_API void ApplyDeltaToComponent(USceneComponent* InComponent, const bool InIsDelta, const FVector* InDeltaTranslationPtr, const FRotator* InDeltaRotationPtr, const FVector* InDeltaScalePtr, const FVector& InPivotLocation, const FInputDeviceState& InInputState);
diff --git a/Engine/Source/Editor/UnrealEd/Public/LevelEditorViewport.h b/Engine/Source/Editor/UnrealEd/Public/LevelEditorViewport.h
--- a/Engine/Source/Editor/UnrealEd/Public/LevelEditorViewport.h
+++ b/Engine/Source/Editor/UnrealEd/Public/LevelEditorViewport.h
@@ -879,9 +879,6 @@ public:
 	/** true if gizmo manipulation was started from a tracking event */
 	bool					bHasBegunGizmoManipulation;
 
-	/** set to true when a gizmo manipulation has produced a non-zero delta */
-	bool					bGizmoDeltaApplied;
-
 	/** Whether this viewport recently received focus. Used to determine whether component selection is permissible. */
 	bool bReceivedFocusRecently;
 
