Remove ICompileResultsLog as it is not needed to make the Engine module's dependency on UnrealEd private as done in 25942865. Something like this might be neccessary if we can remove that dependency entirely, but in that case it's likely that moving UBlueprint to a developer module is the most prudent refactor

diff --git a/Engine/Plugins/Animation/ControlRig/Source/ControlRigDeveloper/Private/ControlRigBlueprint.cpp b/Engine/Plugins/Animation/ControlRig/Source/ControlRigDeveloper/Private/ControlRigBlueprint.cpp
--- a/Engine/Plugins/Animation/ControlRig/Source/ControlRigDeveloper/Private/ControlRigBlueprint.cpp
+++ b/Engine/Plugins/Animation/ControlRig/Source/ControlRigDeveloper/Private/ControlRigBlueprint.cpp
@@ -1575,7 +1575,7 @@ void UControlRigBlueprint::HandleReportFromCompiler(EMessageSeverity::Type InSev
 		}
 	}
 
-	ICompilerResultsLog* Log = CurrentMessageLog ? CurrentMessageLog : &CompileLog;
+	FCompilerResultsLog* Log = CurrentMessageLog ? CurrentMessageLog : &CompileLog;
 	if (InSeverity == EMessageSeverity::Error)
 	{
 		Status = BS_Error;
diff --git a/Engine/Source/Editor/UnrealEd/Public/Kismet2/CompilerResultsLog.h b/Engine/Source/Editor/UnrealEd/Public/Kismet2/CompilerResultsLog.h
--- a/Engine/Source/Editor/UnrealEd/Public/Kismet2/CompilerResultsLog.h
+++ b/Engine/Source/Editor/UnrealEd/Public/Kismet2/CompilerResultsLog.h
@@ -29,7 +29,6 @@
 #include "EdGraphToken.h"
 #include "Logging/TokenizedMessage.h"
 #include "Misc/CompilationResult.h"
-#include "Interfaces/ICompilerResultsLog.h"
 #endif
 
 class FDelegateHandle;
@@ -67,7 +66,7 @@ public:
 };
 
 /** This class represents a log of compiler output lines (errors, warnings, and information notes), each of which can be a rich tokenized message */
-class FCompilerResultsLog : public ICompilerResultsLog
+class FCompilerResultsLog
 {
 	// Compiler event
 	struct FCompilerEvent
@@ -163,6 +162,62 @@ public:
 		SourcePath = InSourcePath;
 	}
 
+	/**
+	 * Write an error in to the compiler log.
+	 * Note: @@ will be replaced by node or pin links for nodes/pins passed via varargs
+	 */
+	template<typename... ArgTypes>
+	TSharedRef<FTokenizedMessage> Error(const TCHAR* Format, ArgTypes... Args)
+	{
+		++NumErrors;
+		TSharedRef<FTokenizedMessage> Line = FTokenizedMessage::Create(EMessageSeverity::Error);
+		InternalLogMessage(NAME_None, Format, Line, Args...);
+		return Line;
+	}
+
+	/**
+	 * Write a warning in to the compiler log.
+	 * Note: @@ will be replaced by node or pin links for nodes/pins passed via varargs
+	 */
+	template<typename... ArgTypes>
+	TSharedRef<FTokenizedMessage> Warning(const TCHAR* Format, ArgTypes... Args)
+	{
+		++NumWarnings;
+		TSharedRef<FTokenizedMessage> Line = FTokenizedMessage::Create(EMessageSeverity::Warning);
+		InternalLogMessage(NAME_None, Format, Line, Args...);
+		return Line;
+	}
+
+	/**
+	 * Write a warning in to the compiler log.
+	 * Note: @@ will be replaced by node or pin links for nodes/pins passed via varargs
+	 */
+	template<typename... ArgTypes>
+	void Warning(FName ID, const TCHAR* Format, ArgTypes... Args)
+	{
+		if (!IsMessageEnabled(ID))
+		{
+			return;
+		}
+
+		++NumWarnings;
+		TSharedRef<FTokenizedMessage> Line = FTokenizedMessage::Create(EMessageSeverity::Warning);
+		InternalLogMessage(ID, Format, Line, Args...);
+		return;
+	}
+
+	/**
+	 * Write a note in to the compiler log.
+	 * Note: @@ will be replaced by node or pin links for nodes/pins passed via varargs
+	 */
+	template<typename... ArgTypes>
+	TSharedRef<FTokenizedMessage> Note(const TCHAR* Format, ArgTypes... Args)
+	{
+		TSharedRef<FTokenizedMessage> Line = FTokenizedMessage::Create(EMessageSeverity::Info);
+		InternalLogMessage(NAME_None, Format, Line, Args...);
+		return Line;
+	}
+
 	/**
 	 * Store a potential error for a given node in the compiler log. All messages for the node can be committed to the log later by calling CommitPotentialMessages
 	 * Note: @@ will be replaced by node or pin links for nodes/pins passed via varargs 
@@ -285,11 +340,48 @@ public:
 	static UNREALED_API TSharedRef<IMessageLogListing> GetBlueprintMessageLog(UBlueprint* InBlueprint);
 
 	/** ICompilerResultsLog implementation */
-	UNREALED_API virtual void SetSilentMode(bool bValue) override { bSilentMode = bValue; };
+	UNREALED_API void SetSilentMode(bool bValue) { bSilentMode = bValue; };
 
 protected:
 	/** Helper method to add a child event to the given parent event scope */
 	UNREALED_API void AddChildEvent(TSharedPtr<FCompilerEvent>& ParentEventScope, TSharedRef<FCompilerEvent>& ChildEventScope);
+	void Tokenize(const TCHAR* Text, FTokenizedMessage& OutMessage, TArray<UEdGraphNode*>& OutSourceNode)
+	{
+		OutMessage.AddToken(FTextToken::Create(FText::FromString(Text)));
+	}
+
+	template<typename T, typename... ArgTypes>
+	void Tokenize(const TCHAR* Format, FTokenizedMessage& OutMessage, TArray<UEdGraphNode*>& OutSourceNode, T First, ArgTypes... Rest)
+	{
+		// read to next "@@":
+		if (const TCHAR* DelimiterStr = FCString::Strstr(Format, TEXT("@@")))
+		{
+			int32 TokenLength = UE_PTRDIFF_TO_INT32(DelimiterStr - Format);
+			OutMessage.AddToken(FTextToken::Create(FText::FromString(FString(TokenLength, Format))));
+			FEdGraphToken::Create(First, this, OutMessage, OutSourceNode);
+
+			const TCHAR* NextChunk = DelimiterStr + FCString::Strlen(TEXT("@@"));
+			if (*NextChunk)
+			{
+				Tokenize(NextChunk, OutMessage, OutSourceNode, Rest...);
+			}
+		}
+		else
+		{
+			Tokenize(Format, OutMessage, OutSourceNode);
+		}
+	}
+
+	template<typename... ArgTypes>
+	void InternalLogMessage(FName MessageID, const TCHAR* Format, const TSharedRef<FTokenizedMessage>& Message, ArgTypes... Args)
+	{
+		// Convention for SourceNode established by the original version of the compiler results log
+		// was to annotate the error on the first node we can find. I am preserving that behavior
+		// for this type safe, variadic version:
+		TArray<UEdGraphNode*> SourceNodes;
+		Tokenize(Format, *Message, SourceNodes, Args...);
+		InternalLogMessage(MessageID, Message, SourceNodes);
+	}
 
 	/** Links the UEdGraphNode with the LogLine: */
 	UNREALED_API void AnnotateNode(const TArray<UEdGraphNode*>& Nodes, TSharedRef<FTokenizedMessage> LogLine);
@@ -300,15 +392,14 @@ protected:
 	/** Internal helper method to recursively append event details into the MessageLog */
 	UNREALED_API void InternalLogEvent(const FCompilerEvent& InEvent, int32 InDepth = 0);
 
-	/** ICompilerResultsLog implementation */
-	UNREALED_API virtual void InternalLogMessage(FName MessageID, const TSharedRef<FTokenizedMessage>& Message, const TArray<UEdGraphNode*>& SourceNodes) override;
-	UNREALED_API virtual void FEdGraphToken_Create(const UObject* InObject, FTokenizedMessage& OutMessage, TArray<UEdGraphNode*>& OutSourceNodes) override;
-	UNREALED_API virtual void FEdGraphToken_Create(const UEdGraphPin* InPin, FTokenizedMessage& OutMessage, TArray<UEdGraphNode*>& OutSourceNodes) override;
-	UNREALED_API virtual void FEdGraphToken_Create(const TCHAR* String, FTokenizedMessage& OutMessage, TArray<UEdGraphNode*>& OutSourceNodes) override;
-	UNREALED_API virtual void FEdGraphToken_Create(const FField* InField, FTokenizedMessage& OutMessage, TArray<UEdGraphNode*>& OutSourceNodes) override;
-	virtual void IncrementErrorCount() override { ++NumErrors; }
-	virtual void IncrementWarningCount() override { ++NumWarnings; };
-	UNREALED_API virtual bool IsMessageEnabled(FName ID) override;
+	UNREALED_API void InternalLogMessage(FName MessageID, const TSharedRef<FTokenizedMessage>& Message, const TArray<UEdGraphNode*>& SourceNodes);
+	UNREALED_API void FEdGraphToken_Create(const UObject* InObject, FTokenizedMessage& OutMessage, TArray<UEdGraphNode*>& OutSourceNodes);
+	UNREALED_API void FEdGraphToken_Create(const UEdGraphPin* InPin, FTokenizedMessage& OutMessage, TArray<UEdGraphNode*>& OutSourceNodes);
+	UNREALED_API void FEdGraphToken_Create(const TCHAR* String, FTokenizedMessage& OutMessage, TArray<UEdGraphNode*>& OutSourceNodes);
+	UNREALED_API void FEdGraphToken_Create(const FField* InField, FTokenizedMessage& OutMessage, TArray<UEdGraphNode*>& OutSourceNodes);
+	void IncrementErrorCount() { ++NumErrors; }
+	void IncrementWarningCount() { ++NumWarnings; };
+	UNREALED_API bool IsMessageEnabled(FName ID);
 
 private:
 
diff --git a/Engine/Source/Runtime/Engine/Classes/Interfaces/ICompilerResultsLog.h b/Engine/Source/Runtime/Engine/Classes/Interfaces/ICompilerResultsLog.h
deleted file mode 100644
