Add back in call to OpenEditorForAsset in AssetEditorSubsystem.cpp a backout occurred in CL 22803476 which removed the call and the fix in CL 22823595 didn't bring back the call.

diff --git a/Engine/Source/Editor/UnrealEd/Private/Subsystems/AssetEditorSubsystem.cpp b/Engine/Source/Editor/UnrealEd/Private/Subsystems/AssetEditorSubsystem.cpp
--- a/Engine/Source/Editor/UnrealEd/Private/Subsystems/AssetEditorSubsystem.cpp
+++ b/Engine/Source/Editor/UnrealEd/Private/Subsystems/AssetEditorSubsystem.cpp
@@ -799,6 +799,7 @@ void UAssetEditorSubsystem::OpenEditorForAsset(const FString& AssetPathName, con
 		
 		if (Object != nullptr)
 		{
+			OpenEditorForAsset(Object, EToolkitMode::Standalone, TSharedPtr<IToolkitHost>(), true, OpenedMethod);
 		}
 	}
 }
