Adding new ObjectTools to attempt to find the UObject pointer on a given object.
This tool will take a given uobject then attempt to search for the target up to a CVar configurable depth
This tool can be useufl to determining how a particular uobject is a dependency of another.

diff --git a/Engine/Source/Editor/UnrealEd/Private/ObjectTools.cpp b/Engine/Source/Editor/UnrealEd/Private/ObjectTools.cpp
--- a/Engine/Source/Editor/UnrealEd/Private/ObjectTools.cpp
+++ b/Engine/Source/Editor/UnrealEd/Private/ObjectTools.cpp
@@ -162,6 +162,17 @@ void ReloadEditorWorldForReferenceReplacementIfNecessary(TArray< TWeakObjectPtr<
 
 namespace ObjectTools
 {
+	static int32 MaxTimesToCheckSameObject = 3;
+	static FAutoConsoleVariableRef CVarMaxTimesToCheckSameObject(TEXT("ObjectTools.MaxTimesToCheckSameObject"),
+		MaxTimesToCheckSameObject,
+		TEXT("Number of times to recurse on the same object when mapping property chains to objects."));
+
+	static int32 MaxRecursionDepth = 4;
+	static FAutoConsoleVariableRef CVarMaxRecursionDepth(TEXT("ObjectTools.MaxRecursionDepth"),
+		MaxRecursionDepth,
+		TEXT("How many times to recurse to find the object to search for"));
+
+
 	/** Returns true if the specified object can be displayed in a content browser */
 	bool IsObjectBrowsable( UObject* Obj )	// const
 	{
@@ -488,6 +499,67 @@ namespace ObjectTools
 
 	}
 
+	// recursive private function to determine the property paths that bring in a uobject
+	void RecursiveBuildPropertyMap_Helper(TMap<const UObject*, int32>& CheckedObjects, TMap<FString, const UObject*>& PropertyToObject, const FString& InPropertyStr, const UObject* InObject, int32 Depth)
+	{
+		for (TPropertyValueIterator<FObjectProperty> PIter(InObject->GetClass(), InObject); PIter; ++PIter)
+		{
+			FObjectProperty* Property = PIter.Key();
+			void* Value = const_cast<void*>(PIter->Value);
+			if (const UObject* ValueObject = Property->GetPropertyValue(Value))
+			{
+				if (int32* TimesEncountered = CheckedObjects.Find(ValueObject))
+				{
+					(*TimesEncountered)++;
+					if (*TimesEncountered > MaxTimesToCheckSameObject)
+					{
+						continue;
+					}
+				}
+				else
+				{
+					CheckedObjects.Add(ValueObject, 1);
+				}
+
+				const FString FullPropertyKey = InPropertyStr.IsEmpty() ? Property->GetName() : InPropertyStr + TEXT(" -> ") + Property->GetName();
+				FString TestKey = FullPropertyKey;
+				int32 ArrayIndex = 1;
+				while (PropertyToObject.Contains(TestKey))
+				{
+					TestKey = FString::Printf(TEXT("%s[%d]"), *FullPropertyKey, ArrayIndex);
+					ArrayIndex++;
+				}
+				PropertyToObject.Add(TestKey, ValueObject);
+				if (Depth < MaxRecursionDepth)
+				{
+					RecursiveBuildPropertyMap_Helper(CheckedObjects, PropertyToObject, TestKey, ValueObject, Depth + 1);
+				}
+			}
+		}
+	}
+
+	bool GatherPropertyChainsToObject(const UObject* SourceObject, const UObject* ObjectToSearchFor, TArray<FString>& OutFoundPropertyChains)
+	{
+		OutFoundPropertyChains.Empty();
+
+		if (SourceObject && ObjectToSearchFor)
+		{
+			TMap<const UObject*, int32> CheckedObjects;
+			TMap<FString, const UObject*> PropertyToObject;
+			RecursiveBuildPropertyMap_Helper(CheckedObjects, PropertyToObject, TEXT(""), SourceObject, 0);
+
+			for (const TPair<FString, const UObject*>& PropertyObjectPair : PropertyToObject)
+			{
+				const UObject* CurrentObject = PropertyObjectPair.Value;
+				if (CurrentObject && CurrentObject == ObjectToSearchFor)
+				{
+					OutFoundPropertyChains.Add(PropertyObjectPair.Key);
+				}
+			}
+		}
+		return OutFoundPropertyChains.Num() > 0;
+	}
+
 	/**
 	 * FArchiveTopLevelReferenceCollector constructor
 	 * @todo: comment
diff --git a/Engine/Source/Editor/UnrealEd/Public/ObjectTools.h b/Engine/Source/Editor/UnrealEd/Public/ObjectTools.h
--- a/Engine/Source/Editor/UnrealEd/Public/ObjectTools.h
+++ b/Engine/Source/Editor/UnrealEd/Public/ObjectTools.h
@@ -651,6 +651,20 @@ namespace ObjectTools
 	*/
 	UNREALED_API void GatherSubObjectsForReferenceReplacement(TSet<UObject*>& InObjects, TSet<UObject*>& ObjectsToExclude, TSet<UObject*>& OutObjectsAndSubObjects);
 
+	/**
+	 * Given a SourceObject and ObjectToSearchFor will attempt to find the ObjectToSearchFor on the given object.
+	 * If the ObjectToSearchFor is found this will report back properties needed to traverse to find this pointer.
+	 * This is useful to find out where a reference is being pulled in from.
+	 * 
+	 * Uses the CVar ObjectTools.MaxTimesToCheckSameObject to configure how many times the same object should be checked and traversed. This is to help avoid circular dependencies.
+	 * 
+	 * @param SourceObject - which object we will scan the properties on to attempt to find the given ObjectToSearchFor
+	 * @param ObjectToSearchFor - which object we will look for on each property of the SourceObject
+	 * @param OutFoundPropertyChains - an array of strings for which properties to traverse to find the ObjectToSearchFor
+	 * @return true when the ObjectToSearchFor is found, otherwise false.
+	 */
+	UNREALED_API bool GatherPropertyChainsToObject(const UObject* SourceObject, const UObject* ObjectToSearchFor, TArray<FString>& OutFoundPropertyChains);
+
 
 	/**
 	 * Batch version of UObject::GetArchetypeInstances.  Can be considerably faster in large worlds when querying multiple objects with the
