Iterative Cooking: Add optional invalidation of iterative results if the executable changes.
#rnx

diff --git a/Engine/Config/BaseEditor.ini b/Engine/Config/BaseEditor.ini
--- a/Engine/Config/BaseEditor.ini
+++ b/Engine/Config/BaseEditor.ini
@@ -352,6 +352,11 @@ bUseSoftGC=true
 SoftGCStartNumerator=5
 SoftGCDenominator=10
 
+; If true, iterative cooks will be allowed even if referenced ini values have changed.
+IterativeIgnoreIni=false
+; If true, iterative cooks will be allowed even if the UnrealEditor exes and dlls have changed.
+IterativeIgnoreExe=true
+
 ; CookProcessCount sets whether we cook singleprocess or multiprocess and how many processes there should be.
 ; CookProcessCount=(1 or less) is singleprocess. CookProcessCount=(N>1) is 1 director and N-1 cookworkers.
 CookProcessCount=1
diff --git a/Engine/Source/Editor/UnrealEd/Classes/CookOnTheSide/CookOnTheFlyServer.h b/Engine/Source/Editor/UnrealEd/Classes/CookOnTheSide/CookOnTheFlyServer.h
--- a/Engine/Source/Editor/UnrealEd/Classes/CookOnTheSide/CookOnTheFlyServer.h
+++ b/Engine/Source/Editor/UnrealEd/Classes/CookOnTheSide/CookOnTheFlyServer.h
@@ -1468,7 +1468,12 @@ private:
 	 * objects created by PreSave
 	 */
 	bool bCallIsCachedOnSaveCreatedObjects = false;
-
+	/** Command-line parameter; if true then iterative cooks will not be invalidated by ini changes. */
+	bool bIterativeIgnoreIni = false;
+	/** Command-line parameter; if true then iterative cooks will not be invalidated by exe changes. */
+	bool bIterativeIgnoreExe = false;
+	/** Whether we should calculate the exe's hash; might be true even if bIterativeIgnoreExe is true. */
+	bool bIterativeCalculateExe = true;
 	/** Timers for tracking how long we have been busy, to manage retries and warnings of deadlock */
 	double SaveBusyStartTimeSeconds = MAX_flt;
 	double SaveBusyRetryTimeSeconds = MAX_flt;
diff --git a/Engine/Source/Editor/UnrealEd/Private/Commandlets/CookCommandlet.cpp b/Engine/Source/Editor/UnrealEd/Private/Commandlets/CookCommandlet.cpp
--- a/Engine/Source/Editor/UnrealEd/Private/Commandlets/CookCommandlet.cpp
+++ b/Engine/Source/Editor/UnrealEd/Private/Commandlets/CookCommandlet.cpp
@@ -173,7 +173,7 @@ int32 UCookCommandlet::Main(const FString& CmdLineParams)
 	bUnversioned = Switches.Contains(TEXT("UNVERSIONED"));   // Save all cooked packages without versions. These are then assumed to be current version on load. This is dangerous but results in smaller patch sizes.
 	bCookEditorOptional = Switches.Contains(TEXT("EDITOROPTIONAL")); // Produce the optional editor package data alongside the cooked data.
 	bGenerateStreamingInstallManifests = Switches.Contains(TEXT("MANIFESTS"));   // Generate manifests for building streaming install packages
-	bIterativeCooking = Switches.Contains(TEXT("ITERATE"));
+	bIterativeCooking = Switches.Contains(TEXT("ITERATE")) || Switches.Contains(TEXT("ITERATIVE"));
 	bSkipEditorContent = Switches.Contains(TEXT("SKIPEDITORCONTENT")); // This won't save out any packages in Engine/Content/Editor*
 	bErrorOnEngineContentUse = Switches.Contains(TEXT("ERRORONENGINECONTENTUSE"));
 	bCookSinglePackage = Switches.Contains(TEXT("cooksinglepackagenorefs"));
diff --git a/Engine/Source/Editor/UnrealEd/Private/CookOnTheFlyServer.cpp b/Engine/Source/Editor/UnrealEd/Private/CookOnTheFlyServer.cpp
--- a/Engine/Source/Editor/UnrealEd/Private/CookOnTheFlyServer.cpp
+++ b/Engine/Source/Editor/UnrealEd/Private/CookOnTheFlyServer.cpp
@@ -66,6 +66,7 @@
 #include "HAL/PlatformProcess.h"
 #include "HAL/Runnable.h"
 #include "HAL/RunnableThread.h"
+#include "Hash/xxhash.h"
 #include "IMessageContext.h"
 #include "INetworkFileServer.h"
 #include "INetworkFileSystemModule.h"
@@ -6637,6 +6638,28 @@ void UCookOnTheFlyServer::SetInitializeConfigSettings(UE::Cook::FInitializeConfi
 	ParseCookFilters();
 
 	bCallIsCachedOnSaveCreatedObjects = FParse::Param(FCommandLine::Get(), TEXT("CallIsCachedOnSaveCreatedObjects"));
+
+	bIterativeIgnoreIni = false;
+	GConfig->GetBool(TEXT("CookSettings"), TEXT("IterativeIgnoreIni"), bIterativeIgnoreIni, GEditorIni);
+	bIterativeIgnoreIni =
+		!FParse::Param(FCommandLine::Get(), TEXT("iteraterequireini")) &&
+		!FParse::Param(FCommandLine::Get(), TEXT("iterativerequireini")) &&
+		(bIterativeIgnoreIni ||
+			IsCookFlagSet(ECookInitializationFlags::IgnoreIniSettingsOutOfDate) ||
+			FParse::Param(FCommandLine::Get(), TEXT("iterateignoreini")) ||
+			FParse::Param(FCommandLine::Get(), TEXT("iterativeignoreini")));
+	bIterativeCalculateExe = true;
+	bool bConfigSettingSetIterativeIgnoreExe = false;
+	GConfig->GetBool(TEXT("CookSettings"), TEXT("IterativeIgnoreExe"), bConfigSettingSetIterativeIgnoreExe, GEditorIni);
+	bIterativeIgnoreExe =
+		!FParse::Param(FCommandLine::Get(), TEXT("iteraterequireexe")) &&
+		!FParse::Param(FCommandLine::Get(), TEXT("iterativerequireexe")) &&
+		(bConfigSettingSetIterativeIgnoreExe ||
+			FParse::Param(FCommandLine::Get(), TEXT("iterateignoreexe")) ||
+			FParse::Param(FCommandLine::Get(), TEXT("iterativeignoreexe")));
+	// Calculate the exe hash if IterativeExeInvalidation is required by ini OR required by commandline
+	// It would be better to always calculate it, but we want to avoid the performance cost until it becomes more widely used
+	bIterativeCalculateExe = !bIterativeIgnoreExe || !bConfigSettingSetIterativeIgnoreExe;
 }
 
 void UCookOnTheFlyServer::ParseCookFilters()
@@ -7438,12 +7461,59 @@ void UCookOnTheFlyServer::ProcessAccessedIniSettings(const FConfigFile* Config,
 }
 
 static const TCHAR* TEXT_CookSettings(TEXT("CookSettings"));
+static FName ExecutableHashName(TEXT("ExecutableHash"));
+static FName ExecutableHashInvalidModuleName(TEXT("ExecutableHashInvalidModule"));
+
 
 TMap<FName, FString> UCookOnTheFlyServer::CalculateCookSettingStrings() const
 {
 	TMap<FName, FString> CookSettingStrings;
 	const FName NAME_CookMode(TEXT("CookMode"));
 
+	TArray<FModuleStatus> Modules;
+	FModuleManager::Get().QueryModules(Modules);
+	Modules.Sort([](const FModuleStatus& A, const FModuleStatus& B)
+		{
+			return A.FilePath.Compare(B.FilePath, ESearchCase::IgnoreCase) < 0;
+		});
+	IPlatformFile& PlatformFile = FPlatformFileManager::Get().GetPlatformFile();
+
+	// Calculate the executable hash by combining the module file hash of every loaded module
+	// TODO: Write the module file hash from UnrealBuildTool into the .modules file and read it
+	// here from the .modules file instead of calculating it on every cook.
+	if (bIterativeCalculateExe)
+	{
+		FString InvalidModule;
+		bool bValid = true;
+		FXxHash64Builder Hasher;
+		TArray<uint8> Buffer;
+		for (FModuleStatus& ModuleStatus : Modules)
+		{
+			TUniquePtr<IFileHandle> FileHandle(PlatformFile.OpenRead(*ModuleStatus.FilePath));
+			if (!FileHandle)
+			{
+				InvalidModule = ModuleStatus.FilePath;
+				break;
+			}
+			int64 FileSize = FileHandle->Size();
+			Buffer.SetNumUninitialized(FileSize, false /* bAllowShrinking */);
+			if (!FileHandle->Read(Buffer.GetData(), FileSize))
+			{
+				InvalidModule = ModuleStatus.FilePath;
+				break;
+			}
+			Hasher.Update(Buffer.GetData(), FileSize);
+		}
+		if (InvalidModule.IsEmpty())
+		{
+			CookSettingStrings.Add(ExecutableHashName, FString(*WriteToString<64>(Hasher.Finalize())));
+		}
+		else
+		{
+			CookSettingStrings.Add(ExecutableHashInvalidModuleName, InvalidModule);
+		}
+	}
+
 	CookSettingStrings.Add(FName(TEXT("Version")), TEXT("C7C76F79"));
 	if (IsDirectorCookByTheBook())
 	{
@@ -7472,14 +7542,60 @@ bool UCookOnTheFlyServer::ArePreviousCookSettingsCompatible(const TMap<FName, FS
 	const FConfigSection* CookSettings = ConfigFile.Find(TEXT_CookSettings);
 	if (CookSettings == nullptr)
 	{
+		UE_LOG(LogCook, Display, TEXT("Cook invalidated for CookSettings file %s is invalid. Clearing all cooked content."),
+			*GetCookSettingsFileName(TargetPlatform));
 		return false;
 	}
 
+	TSet<FName> IgnoreKeys;
+	IgnoreKeys.Add(ExecutableHashName);
+	IgnoreKeys.Add(ExecutableHashInvalidModuleName);
+
 	for (const TPair<FName, FString>& CurrentSetting : CurrentCookSettings)
 	{
+		if (IgnoreKeys.Contains(CurrentSetting.Key))
+		{
+			continue;
+		}
 		const FConfigValue* PreviousSetting = CookSettings->Find(CurrentSetting.Key);
 		if (!PreviousSetting || PreviousSetting->GetValue() != CurrentSetting.Value)
 		{
+			UE_LOG(LogCook, Display, TEXT("Cook invalidated for platform %s because %s has changed. Old: %s, New: %s. Clearing all cooked content."),
+				*TargetPlatform->PlatformName(), *CurrentSetting.Key.ToString(),
+				PreviousSetting ? *PreviousSetting->GetValue() : TEXT(""),
+				*CurrentSetting.Value);
+			return false;
+		}
+	}
+
+	if (!bIterativeIgnoreIni && IniSettingsOutOfDate(TargetPlatform))
+	{
+		UE_LOG(LogCook, Display, TEXT("Cook invalidated for platform %s because ini settings have changed. Clearing all cooked content."),
+			*TargetPlatform->PlatformName());
+		return false;
+	}
+
+	if (!bIterativeIgnoreExe)
+	{
+		const FString* CurrentHash = CurrentCookSettings.Find(ExecutableHashName);
+		if (!CurrentHash)
+		{
+			UE_LOG(LogCook, Display, TEXT("Cook invalidated for platform %s because current executable hash is invalid. Invalid module=%s. Clearing all cooked content."),
+				*TargetPlatform->PlatformName(), *CurrentCookSettings.FindRef(ExecutableHashInvalidModuleName));
+			return false;
+		}
+		const FConfigValue* PreviousHash = CookSettings->Find(ExecutableHashName);
+		if (!PreviousHash)
+		{
+			const FConfigValue* InvalidModuleName = CookSettings->Find(ExecutableHashInvalidModuleName);
+			UE_LOG(LogCook, Display, TEXT("Cook invalidated for platform %s because old executable hash is invalid. Invalid module=%s. Clearing all cooked content."),
+				*TargetPlatform->PlatformName(), InvalidModuleName ? *InvalidModuleName->GetValue() : TEXT(""));
+			return false;
+		}
+		if (!CurrentHash->Equals(*PreviousHash->GetValue(), ESearchCase::CaseSensitive))
+		{
+			UE_LOG(LogCook, Display, TEXT("Cook invalidated for platform %s because executable hash has changed. Old: %s, New: %s. Clearing all cooked content."),
+				*TargetPlatform->PlatformName(), *PreviousHash->GetValue(), **CurrentHash);
 			return false;
 		}
 	}
@@ -9935,31 +10051,8 @@ void UCookOnTheFlyServer::LoadBeginCookIterativeFlagsLocal(FBeginCookContext& Be
 			}
 			else if (!ArePreviousCookSettingsCompatible(PlatformContext.CurrentCookSettings, TargetPlatform))
 			{
-				UE_LOG(LogCook, Display, TEXT("Cook invalidated for platform %s because cook DLC settings have changed, clearing all cooked content."), *TargetPlatform->PlatformName());
 				bIterativeAllowed = false;
 			}
-			else
-			{
-				const bool bIsIniSettingsOutOfDate = IniSettingsOutOfDate(TargetPlatform);
-				if (bIsIniSettingsOutOfDate)
-				{
-					if (!IsCookFlagSet(ECookInitializationFlags::IgnoreIniSettingsOutOfDate))
-					{
-						UE_LOG(LogCook, Display, TEXT("Cook invalidated for platform %s because ini settings have changed, clearing all cooked content."), *TargetPlatform->PlatformName());
-						bIterativeAllowed = false;
-					}
-					else
-					{
-						UE_LOG(LogCook, Display, TEXT("Inisettings were out of date for platform %s but we are going with it anyway because IgnoreIniSettingsOutOfDate is set."), *TargetPlatform->PlatformName());
-						bIterativeAllowed = true;
-					}
-				}
-				else
-				{
-					UE_LOG(LogCook, Display, TEXT("Keeping cooked content for platform %s and cooking iteratively."), *TargetPlatform->PlatformName());
-					bIterativeAllowed = true;
-				}
-			}
 
 			if (bIterativeAllowed)
 			{
