Cooker: Fix missing function BeginCacheForCookedPlatformData on DiffPackageWriters.
#rnx

diff --git a/Engine/Source/Editor/UnrealEd/Private/Cooker/DiffPackageWriter.h b/Engine/Source/Editor/UnrealEd/Private/Cooker/DiffPackageWriter.h
--- a/Engine/Source/Editor/UnrealEd/Private/Cooker/DiffPackageWriter.h
+++ b/Engine/Source/Editor/UnrealEd/Private/Cooker/DiffPackageWriter.h
@@ -95,6 +95,10 @@ public:
 	{
 		Inner->MarkPackagesUpToDate(UpToDatePackages);
 	}
+	virtual EPackageWriterResult BeginCacheForCookedPlatformData(FBeginCacheForCookedPlatformDataInfo& Info) override
+	{
+		return Inner->BeginCacheForCookedPlatformData(Info);
+	}
 	virtual void UpdateSaveArguments(FSavePackageArgs& SaveArgs) override;
 	virtual bool IsAnotherSaveNeeded(FSavePackageResultStruct& PreviousResult, FSavePackageArgs& SaveArgs) override;
 	virtual TFuture<FCbObject> WriteMPCookMessageForPackage(FName PackageName) override
@@ -232,6 +236,10 @@ public:
 	{
 		Inner->MarkPackagesUpToDate(UpToDatePackages);
 	}
+	virtual EPackageWriterResult BeginCacheForCookedPlatformData(FBeginCacheForCookedPlatformDataInfo& Info) override
+	{
+		return Inner->BeginCacheForCookedPlatformData(Info);
+	}
 	virtual void UpdateSaveArguments(FSavePackageArgs& SaveArgs) override;
 	virtual bool IsAnotherSaveNeeded(FSavePackageResultStruct& PreviousResult, FSavePackageArgs& SaveArgs) override;
 	virtual TFuture<FCbObject> WriteMPCookMessageForPackage(FName PackageName) override
