Gameplay framework deprecation cleanup to remove many functions deprecated before 5.0
This is only the safe changes, some variables were not properly deprecated so will be removed in a future change

diff --git a/Engine/Source/Editor/Kismet/Private/BlueprintEditor.cpp b/Engine/Source/Editor/Kismet/Private/BlueprintEditor.cpp
--- a/Engine/Source/Editor/Kismet/Private/BlueprintEditor.cpp
+++ b/Engine/Source/Editor/Kismet/Private/BlueprintEditor.cpp
@@ -2608,11 +2608,6 @@ FText FBlueprintEditor::GetParentClassNameText() const
 	return (ParentClass != nullptr) ? ParentClass->GetDisplayNameText() : LOCTEXT("BlueprintEditor_NoParentClass", "None");
 }
 
-bool FBlueprintEditor::IsParentClassOfObjectABlueprint(const UBlueprint* Blueprint) const
-{
-	return FBlueprintEditorUtils::IsParentClassABlueprint(Blueprint);
-}
-
 bool FBlueprintEditor::IsParentClassABlueprint() const
 {
 	return FBlueprintEditorUtils::IsParentClassABlueprint(GetBlueprintObj());
diff --git a/Engine/Source/Editor/Kismet/Public/BlueprintEditor.h b/Engine/Source/Editor/Kismet/Public/BlueprintEditor.h
--- a/Engine/Source/Editor/Kismet/Public/BlueprintEditor.h
+++ b/Engine/Source/Editor/Kismet/Public/BlueprintEditor.h
@@ -434,10 +434,6 @@ public:
 	/** Returns true if able to compile */
 	virtual bool IsCompilingEnabled() const;
 
-	/** Returns true if the parent class is also a Blueprint */
-	UE_DEPRECATED(4.27, "Please use FBlueprintEditorUtils::IsParentClassABlueprint instead")
-	bool IsParentClassOfObjectABlueprint(const UBlueprint* Blueprint) const;
-
 	/** Returns true if the parent class of the Blueprint being edited is also a Blueprint */
 	bool IsParentClassABlueprint() const;
 
@@ -802,8 +798,6 @@ public:
 	void SelectAndDuplicateNode(UEdGraphNode* InNode);
 
 protected:
-	UE_DEPRECATED(4.26, "Please do any validation inside the UBlueprint class during compilation, extra errors during compiling only supplied by the designer can lead to design time only errors being reported and being missed during cooks/content validation.")
-	virtual void AppendExtraCompilerResults(TSharedPtr<class IMessageLogListing> ResultsListing) {}
 
 	/** Called during initialization of the blueprint editor to register commands and extenders. */
 	virtual void InitalizeExtenders();
diff --git a/Engine/Source/Editor/Kismet/Public/SSCSEditor.h b/Engine/Source/Editor/Kismet/Public/SSCSEditor.h
--- a/Engine/Source/Editor/Kismet/Public/SSCSEditor.h
+++ b/Engine/Source/Editor/Kismet/Public/SSCSEditor.h
@@ -496,28 +496,6 @@ private:
 		Unknown = 0xFC // ~FilteredInMask
 	};
 	uint8 FilterFlags;
-
-public:
-	UE_DEPRECATED(4.26, "Use IsNativeComponent() instead.")
-	bool IsNative() const { return IsNativeComponent(); }
-
-	UE_DEPRECATED(4.26, "Use IsInheritedSCSNode() instead.")
-	bool IsInheritedSCS() const { return IsInheritedSCSNode(); }
-
-	UE_DEPRECATED(4.26, "Use IsInheritedComponent() instead.")
-	bool IsInherited() const { return IsInheritedComponent(); }
-
-	UE_DEPRECATED(4.26, "Use IsUserInstancedComponent() instead.")
-	bool IsUserInstanced() const { return IsUserInstancedComponent(); }
-
-	UE_DEPRECATED(4.26, "Use CanEdit() instead.")
-	bool CanEditDefaults() const { return CanEdit(); }
-
-	UE_DEPRECATED(4.26, "Use SetObject() instead.")
-	void SetComponentTemplate(UActorComponent* Component) { SetObject(Component); }
-
-	UE_DEPRECATED(4.26, "Use RefreshFilteredState() instead. This API has been changed to an internal-only helper method.")
-	void UpdateCachedFilterState(bool bMatchesFilter, bool bUpdateParent) { SetCachedFilterState(bMatchesFilter, bUpdateParent); }
 };
 
 //////////////////////////////////////////////////////////////////////////
@@ -1114,16 +1092,6 @@ public:
 	 */
 	UActorComponent* AddNewComponent(UClass* NewComponentClass, UObject* Asset, const FAddNewComponentParams Params = FAddNewComponentParams());
 
-	/** Adds a component to the SCS tree */
-	UE_DEPRECATED(4.26, "Use version that takes parameter block")
-	UActorComponent* AddNewComponent(UClass* NewComponentClass, UObject* Asset, const bool bSkipMarkBlueprintModified, bool bSetFocusToNewItem = true)
-	{
-		FAddNewComponentParams Params;
-		Params.bSkipMarkBlueprintModified = bSkipMarkBlueprintModified;
-		Params.bSetFocusToNewItem = bSetFocusToNewItem;
-		return AddNewComponent(NewComponentClass, Asset, Params);
-	}
-
 	struct FAddedNodeDetails
 	{
 		FSCSEditorTreeNodePtrType NewNodePtr;
diff --git a/Engine/Source/Editor/UnrealEd/Classes/Editor/EditorEngine.h b/Engine/Source/Editor/UnrealEd/Classes/Editor/EditorEngine.h
--- a/Engine/Source/Editor/UnrealEd/Classes/Editor/EditorEngine.h
+++ b/Engine/Source/Editor/UnrealEd/Classes/Editor/EditorEngine.h
@@ -371,10 +371,6 @@ public:
 	UPROPERTY(EditAnywhere, config, Category=Advanced)
 	uint32 UseAxisIndicator:1;
 
-	UE_DEPRECATED(4.25, "This variable is no longer read.")
-	UPROPERTY(EditAnywhere, config, Category=Advanced)
-	uint32 GodMode:1;
-
 	UPROPERTY(EditAnywhere, config, Category=Advanced)
 	FString GameCommandLine;
 
@@ -1682,29 +1678,6 @@ public:
 	 */
 	virtual void CancelCookByTheBookInEditor() { }
 
-
-	/**
-	 * Makes a request to start a play from a Slate editor session
-	 * @param	bAtPlayerStart			Whether or not we would really like to use the game or level's PlayerStart vs the StartLocation
-	 * @param	DestinationViewport		Slate Viewport to play the game in, or NULL to spawn a standalone PIE window
-	 * @param	bInSimulateInEditor		True to start an in-editor simulation session, or false to kick off a play-in-editor session
-	 * @param	StartLocation			If specified, this is the location to play from (via a PlayerStartPIE - Play From Here)
-	 * @param	StartRotation			If specified, this is the rotation to start playing at
-	 * @param	DestinationConsole		Where to play the game - -1 means in editor, 0 or more is an index into the GConsoleSupportContainer
-	 * @param	bUseMobilePreview		True to enable mobile preview mode (PC platform only)
-	 * @param	bUseVRPreview			True to enable VR preview mode (PC platform only)
-	 */
-	UE_DEPRECATED(4.25, "Use the overload of RequestPlaySession which takes a FRequestPlaySessionParams instead.")
-	UNREALED_API void RequestPlaySession( bool bAtPlayerStart, TSharedPtr<class IAssetViewport> DestinationViewport, bool bInSimulateInEditor, const FVector* StartLocation = NULL, const FRotator* StartRotation = NULL, int32 DestinationConsole = -1, bool bUseMobilePreview = false, bool bUseVRPreview = false, bool bUseVulkanPreview = false);
-
-	// @todo gmp: temp hack for Rocket demo
-	UE_DEPRECATED(4.25, "Use the overload of RequestPlaySession which takes a FRequestPlaySessionParams instead.")
-	UNREALED_API void RequestPlaySession(const FVector* StartLocation, const FRotator* StartRotation, bool MobilePreview, bool VulkanPreview, const FString& MobilePreviewTargetDevice, FString AdditionalStandaloneLaunchParameters = TEXT(""));
-
-	/** Request to play a game on a remote device */
-	UE_DEPRECATED(4.25, "Use the overload of RequestPlaySession which takes a FRequestPlaySessionParams instead.")
-	UNREALED_API void RequestPlaySession( const FString& DeviceId, const FString& DeviceName );
-	
 	/** Request a play session (Play in Editor, Play in New Process, Launcher) with the configured parameters. See FRequestPlaySessionParams for more details. */
 	UNREALED_API void RequestPlaySession(const FRequestPlaySessionParams& InParams);
 
@@ -1756,12 +1729,6 @@ public:
 	/** Called when game client received input key */
 	UNREALED_API bool ProcessDebuggerCommands(const FKey InKey, const FModifierKeysState ModifierKeyState, EInputEvent EventType);
 
-	/**
-	 * Kicks off a "Play From Here" request that was most likely made during a transaction
-	 */
-	UE_DEPRECATED(4.25, "Call StartQueuedPlaySessionRequest, or override StartQueuedPlaySessionRequestImpl instead.")
-	UNREALED_API virtual void StartQueuedPlayMapRequest();
-
 	/** Kicks off a "Request Play Session" request if there is one set. This should not be called mid-transaction. */
 	UNREALED_API void StartQueuedPlaySessionRequest();
 
@@ -1792,18 +1759,6 @@ public:
 	 */
 	UNREALED_API virtual FString BuildPlayWorldURL(const TCHAR* MapName, bool bSpectatorMode = false, FString AdditionalURLOptions=FString());
 
-	/**
-	 * Starts a Play In Editor session
-	 *
-	 * @param	InWorld		World context
-	 * @param	bInSimulateInEditor	True to start an in-editor simulation session, or false to start a play-in-editor session
-	 */
-	UE_DEPRECATED(4.25, "Override StartPlayInEditorSession(FRequestPlaySessionParams& InRequestParams) instead. InWorld is GetEditorWorldContext().World(), and the other arguments come from InRequestParams now.") 
-	UNREALED_API virtual void PlayInEditor( UWorld* InWorld, bool bInSimulateInEditor, FPlayInEditorOverrides Overrides = FPlayInEditorOverrides());
-
-	UE_DEPRECATED(4.25, "Override CreateInnerProcessPIEGameInstance instead.")
-	UNREALED_API virtual UGameInstance* CreatePIEGameInstance(int32 InPIEInstance, bool bInSimulateInEditor, bool bAnyBlueprintErrors, bool bStartInSpectatorMode, bool bPlayNetDedicated, bool bPlayStereoscopic, float PIEStartTime);
-
 	/**
 	 * Kills the Play From Here session
 	 */
@@ -2622,19 +2577,6 @@ private:
 	 */
 	UNREALED_API bool LoadPreviewMesh( int32 Index );
 
-	/**
-	 * Login PIE instances with the online platform before actually creating any PIE worlds
-	 *
-	 * @param bAnyBlueprintErrors passthrough value of blueprint errors encountered during PIE creation
-	 * @param bStartInSpectatorMode passthrough value if it is expected that PIE will start in spectator mode
-	 * @param PIEStartTime passthrough value of the time that PIE was initiated by the user
-	 */
-	UE_DEPRECATED(4.25, "This is now done as part of StartPlayInEditorSession as online-login/non-logged in flow is now combined.")
-	UNREALED_API virtual void LoginPIEInstances(bool bAnyBlueprintErrors, bool bStartInSpectatorMode, double PIEStartTime);
-
-	/** Called when all PIE instances have been successfully logged in */
-	UE_DEPRECATED(4.25, "Override OnAllPIEInstancesStarted instead, which is called in both online and offline PIE sessions.")
-	UNREALED_API virtual void OnLoginPIEAllComplete();
 	/** Called when a recompilation of a module is beginning */
 	UNREALED_API void OnModuleCompileStarted(bool bIsAsyncCompile);
 
@@ -2645,17 +2587,6 @@ public:
 	/** Creates a pie world from the default entry map, used by clients that connect to a PIE server */
 	UNREALED_API UWorld* CreatePIEWorldFromEntry(FWorldContext &WorldContext, UWorld* InWorld, FString &PlayWorldMapName);
 
-	/**
-	 * Continue the creation of a single PIE world after a login was successful
-	 *
-	 * @param PieWorldContext world context for this PIE instance
-	 * @param PlayNetMode mode to create this PIE world in (as server, client, etc)
-	 * @param DataStruct data required to continue PIE creation, set at login time
-	 * @return	true if world created successfully
-	 */
-	UE_DEPRECATED(4.25, "This is now handled as part of OnLoginPIEComplete_Deferred.")
-	UNREALED_API bool CreatePIEWorldFromLogin(FWorldContext& PieWorldContext, EPlayNetMode PlayNetMode, FPieLoginStruct& DataStruct);
-
 	/** Called before creating PIE instance(s). */
 	UNREALED_API virtual FGameInstancePIEResult PreCreatePIEInstances(const bool bAnyBlueprintErrors, const bool bStartInSpectatorMode, const float PIEStartTime, const bool bSupportsOnlinePIE, int32& InNumOnlinePIEInstances);
 
@@ -2676,17 +2607,7 @@ public:
 	 * @return true if function succeeded, false if failed
 	 */
 	UNREALED_API bool GetSimulateInEditorViewTransform(FTransform& OutViewTransform) const;
-
-public:
-	/**
-	 * Non Online PIE creation flow, creates all instances of PIE at once when online isn't requested/required
-	 *
-	 * @param bAnyBlueprintErrors passthrough value of blueprint errors encountered during PIE creation
-	 * @param bStartInSpectatorMode passthrough value if it is expected that PIE will start in spectator mode
-	 */
-	UE_DEPRECATED(4.25, "This non-online flow is now deprecated and is shared with normal online flow. See OnLoginPIEComplete_Deferred instead.")
-	UNREALED_API void SpawnIntraProcessPIEWorlds(bool bAnyBlueprintErrors, bool bStartInSpectatorMode);
-
+	
 	/*
 	 * Spawns a PlayFromHere playerstart in the given world
 	 * @param	World		The World to spawn in (for PIE this may not be GWorld)
@@ -2949,13 +2870,6 @@ public:
 	*/
 	UNREALED_API FString GetPlayOnTargetPlatformName() const;
 
-	/**
-	* Launch a standalone instance using the Launcher to process required
-	* cooking, building, and deployment to a (possibly) remote device for testing.
-	*/
-	UE_DEPRECATED(4.25, "Use RequestPlaySession with parameters configured to provide a launcher target instead.")
-	UNREALED_API void PlayUsingLauncher();
-
 	/** 
 	* Cancel Play using Launcher on error 
 	* 
@@ -3176,10 +3090,7 @@ public:
 	// Launcher Worker
 	TSharedPtr<class ILauncherWorker> LauncherWorker;
 	
-	/** Function to run the Play On command for automation testing. */
-	UE_DEPRECATED(4.25, "Use RequestPlaySession and StartQueuedPlaySessionRequest instead.")
-	UNREALED_API void AutomationPlayUsingLauncher(const FString& InLauncherDeviceId);	
-
+	/** Load a map for automation testing */
 	UNREALED_API void AutomationLoadMap(const FString& MapName, bool bForceReload, FString* OutError);
 
 	/** This function should be called to notify the editor that new materials were added to our scene or some materials were modified */
diff --git a/Engine/Source/Editor/UnrealEd/Private/PlayLevel.cpp b/Engine/Source/Editor/UnrealEd/Private/PlayLevel.cpp
--- a/Engine/Source/Editor/UnrealEd/Private/PlayLevel.cpp
+++ b/Engine/Source/Editor/UnrealEd/Private/PlayLevel.cpp
@@ -977,94 +977,6 @@ void UEditorEngine::RequestPlaySession(const FRequestPlaySessionParams& InParams
 	}
 }
 
-// Deprecated, just format to match our new style.
-void UEditorEngine::RequestPlaySession(bool bAtPlayerStart, TSharedPtr<class IAssetViewport> DestinationViewport, bool bInSimulateInEditor, const FVector* StartLocation, const FRotator* StartRotation, int32 DestinationConsole, bool bUseMobilePreview, bool bUseVRPreview, bool bUseVulkanPreview)
-{
-	FRequestPlaySessionParams Params;
-
-	if (StartLocation)
-	{
-		Params.StartLocation = *StartLocation;
-		Params.StartRotation = StartRotation ? *StartRotation : FRotator::ZeroRotator;
-	}
-	if (DestinationViewport != nullptr)
-	{
-		Params.DestinationSlateViewport = DestinationViewport;
-	}
-
-	if (bInSimulateInEditor)
-	{
-		Params.WorldType = EPlaySessionWorldType::SimulateInEditor;
-	}
-
-	if (bUseVRPreview)
-	{
-		check(!bUseMobilePreview && !bUseVulkanPreview);
-		Params.SessionPreviewTypeOverride = EPlaySessionPreviewType::VRPreview;
-	}
-
-	if (bUseVulkanPreview)
-	{
-		check(!bUseMobilePreview && !bUseVRPreview);
-		Params.SessionPreviewTypeOverride = EPlaySessionPreviewType::VulkanPreview;
-		Params.SessionDestination = EPlaySessionDestinationType::NewProcess;
-	}
-
-	if (bUseMobilePreview)
-	{
-		check(!bUseVRPreview && !bUseVulkanPreview);
-		Params.SessionPreviewTypeOverride = EPlaySessionPreviewType::MobilePreview;
-		Params.SessionDestination = EPlaySessionDestinationType::NewProcess;
-	}
-	
-	RequestPlaySession(Params);
-}
-
-// Deprecated, forwards request onto the FRequestPlaySessionParams version.
-void UEditorEngine::RequestPlaySession(const FVector* StartLocation, const FRotator* StartRotation, bool MobilePreview, bool VulkanPreview, const FString& MobilePreviewTargetDevice, FString AdditionalLaunchParameters)
-{
-	FRequestPlaySessionParams Params;
-
-	if (MobilePreview)
-	{
-		check(!VulkanPreview);
-		Params.SessionDestination = EPlaySessionDestinationType::NewProcess;
-		Params.SessionPreviewTypeOverride = EPlaySessionPreviewType::MobilePreview;
-		Params.MobilePreviewTargetDevice = MobilePreviewTargetDevice;
-		Params.AdditionalStandaloneCommandLineParameters = AdditionalLaunchParameters;
-	}
-
-	if (VulkanPreview)
-	{
-		check(!MobilePreview);
-		Params.SessionDestination = EPlaySessionDestinationType::NewProcess;
-		Params.SessionPreviewTypeOverride = EPlaySessionPreviewType::VulkanPreview;
-		Params.AdditionalStandaloneCommandLineParameters = AdditionalLaunchParameters;
-	}
-	
-	if (StartLocation)
-	{
-		Params.StartLocation = *StartLocation;
-		Params.StartRotation = StartRotation ? *StartRotation : FRotator::ZeroRotator;
-	}
-
-	RequestPlaySession(Params);
-}
-
-// Deprecated, forwards request onto the FRequestPlaySessionParams version.
-void UEditorEngine::RequestPlaySession(const FString& DeviceId, const FString& DeviceName)
-{
-	FRequestPlaySessionParams::FLauncherDeviceInfo DeviceInfo;
-	DeviceInfo.DeviceId = DeviceId;
-	DeviceInfo.DeviceName = DeviceName;
-
-	FRequestPlaySessionParams Params;
-	Params.LauncherTargetDevice = DeviceInfo;
-	Params.SessionDestination = EPlaySessionDestinationType::Launcher;
-
-	RequestPlaySession(Params);
-}
-
 void UEditorEngine::CancelRequestPlaySession()
 {
 	FEditorDelegates::CancelPIE.Broadcast();
@@ -1139,12 +1051,6 @@ bool UEditorEngine::ProcessDebuggerCommands(const FKey InKey, const FModifierKey
 	return false;
 }
 
-// This function is deprecated, just call the non-deprecated version.
-void UEditorEngine::StartQueuedPlayMapRequest()
-{
-	StartQueuedPlaySessionRequest();
-}
-
 void UEditorEngine::StartQueuedPlaySessionRequest()
 {
 	if (!PlaySessionRequest.IsSet())
@@ -3747,20 +3653,4 @@ void UEditorEngine::StoreWindowSizeAndPositionForInstanceIndex(const int32 InVie
 	}
 }
 
-// Deprecated Stubs
-PRAGMA_DISABLE_DEPRECATION_WARNINGS
-UGameInstance* UEditorEngine::CreatePIEGameInstance(int32 InPIEInstance, bool bInSimulateInEditor, bool bAnyBlueprintErrors, bool bStartInSpectatorMode, bool bPlayNetDedicated, bool bPlayStereoscopic, float PIEStartTime)
-{
-	return nullptr;
-}
-
-void UEditorEngine::LoginPIEInstances(bool bAnyBlueprintErrors, bool bStartInSpectatorMode, double PIEStartTime)
-{}
-
-void UEditorEngine::OnLoginPIEAllComplete()
-{}
-
-void UEditorEngine::PlayInEditor(UWorld* InWorld, bool bInSimulateInEditor, FPlayInEditorOverrides Overrides /* = FPlayInEditorOverrides() */)
-{}
-PRAGMA_ENABLE_DEPRECATION_WARNINGS
 #undef LOCTEXT_NAMESPACE
diff --git a/Engine/Source/Editor/UnrealEd/Private/PlayLevelViaLauncher.cpp b/Engine/Source/Editor/UnrealEd/Private/PlayLevelViaLauncher.cpp
--- a/Engine/Source/Editor/UnrealEd/Private/PlayLevelViaLauncher.cpp
+++ b/Engine/Source/Editor/UnrealEd/Private/PlayLevelViaLauncher.cpp
@@ -481,22 +481,6 @@ void UEditorEngine::CancelPlayingViaLauncher()
 	}
 }
 
-// Deprecated, just formats a RequestPlaySession instead.
-void UEditorEngine::AutomationPlayUsingLauncher(const FString& InLauncherDeviceId)
-{
-	FRequestPlaySessionParams::FLauncherDeviceInfo LaunchedDeviceInfo;
-	LaunchedDeviceInfo.DeviceId = InLauncherDeviceId;
-	LaunchedDeviceInfo.DeviceName = InLauncherDeviceId.Right(InLauncherDeviceId.Find(TEXT("@")));
-
-	FRequestPlaySessionParams Params;
-	Params.LauncherTargetDevice = LaunchedDeviceInfo;
-
-	RequestPlaySession(Params);
-
-	// Immediately start our requested play session
-	StartQueuedPlaySessionRequest();
-}
-
 /** 
 * Cancel Play using Launcher on error 
 * 
@@ -745,16 +729,4 @@ FString UEditorEngine::GetPlayOnTargetPlatformName() const
 	return LastPlayUsingLauncherDeviceId.Left(LastPlayUsingLauncherDeviceId.Find(TEXT("@")));
 }
 
-void UEditorEngine::PlayUsingLauncher()
-{
-	// Deprecated, just a wrapper around RequestPlaySession now.
-	FRequestPlaySessionParams::FLauncherDeviceInfo DeviceInfo;
-	DeviceInfo.DeviceId = LastPlayUsingLauncherDeviceId;
-
-	FRequestPlaySessionParams Params;
-	Params.LauncherTargetDevice = DeviceInfo;
-	
-	RequestPlaySession(Params);
-}
-
 #undef LOCTEXT_NAMESPACE // "PlayLevel"
diff --git a/Engine/Source/Editor/UnrealEd/Public/Kismet2/KismetEditorUtilities.h b/Engine/Source/Editor/UnrealEd/Public/Kismet2/KismetEditorUtilities.h
--- a/Engine/Source/Editor/UnrealEd/Public/Kismet2/KismetEditorUtilities.h
+++ b/Engine/Source/Editor/UnrealEd/Public/Kismet2/KismetEditorUtilities.h
@@ -224,26 +224,6 @@ public:
 	/** Take a list of components that belong to a single Actor and add them to a blueprint as SCSNodes */
 	static UNREALED_API void AddComponentsToBlueprint(UBlueprint* Blueprint, const TArray<UActorComponent*>& Components, const FAddComponentsToBlueprintParams& Params = FAddComponentsToBlueprintParams());
 
-	UE_DEPRECATED(4.26, "Use version that uses parameter struct instead of individual parameters")
-	static void AddComponentsToBlueprint(UBlueprint* Blueprint, const TArray<UActorComponent*>& Components, EAddComponentToBPHarvestMode HarvestMode, USCS_Node* OptionalNewRootNode = nullptr, bool bKeepMobility = false)
-	{
-		FAddComponentsToBlueprintParams Params;
-		Params.HarvestMode = HarvestMode;
-		Params.OptionalNewRootNode = OptionalNewRootNode;
-		Params.bKeepMobility = bKeepMobility;
-		AddComponentsToBlueprint(Blueprint, Components, Params);
-	}
-
-	UE_DEPRECATED(4.25, "Use version that specifies harvest mode via enum instead of boolean")
-	static void AddComponentsToBlueprint(UBlueprint* Blueprint, const TArray<UActorComponent*>& Components, bool bHarvesting, USCS_Node* OptionalNewRootNode = nullptr, bool bKeepMobility = false)
-	{
-		FAddComponentsToBlueprintParams Params;
-		Params.HarvestMode = (bHarvesting ? EAddComponentToBPHarvestMode::Harvest_UseComponentName : EAddComponentToBPHarvestMode::None);
-		Params.OptionalNewRootNode = OptionalNewRootNode;
-		Params.bKeepMobility = bKeepMobility;
-		AddComponentsToBlueprint(Blueprint, Components, Params);
-	}
-
 	/** Parameter struct for customizing calls to AddActorsToBlueprint */
 	struct FAddActorsToBlueprintParams
 	{
@@ -324,48 +304,6 @@ public:
 	 */
 	static UNREALED_API UBlueprint* CreateBlueprintFromActor(const FName BlueprintName, UObject* Outer, AActor* Actor, const FCreateBlueprintFromActorParams& Params = FCreateBlueprintFromActorParams());
 
-	/** 
-	 * Take an Actor and generate a blueprint based on it. Uses the Actors type as the parent class. 
-	 * @param Path					The path to use when creating the package for the new blueprint
-	 * @param Actor					The actor to use as the template for the blueprint
-	 * @param bReplaceActor			If true, replace the actor in the scene with one based on the created blueprint
-	 * @param bKeepMobility			If true, The mobility of each actor components will be copy
-	 * @param bOpenInEditor			If true, open the created blueprint in the blueprint editor
-	 * @return The blueprint created from the actor
-	 */
-	UE_DEPRECATED(4.26, "Use version that passes parameters via parameter struct")
-	static UBlueprint* CreateBlueprintFromActor(const FString& Path, AActor* Actor, bool bReplaceActor, bool bKeepMobility = false, UClass* ParentClassOverride = nullptr, bool bOpenInEditor = true)
-	{
-		FCreateBlueprintFromActorParams Params;
-		Params.bReplaceActor = bReplaceActor;
-		Params.bKeepMobility = bKeepMobility;
-		Params.bOpenBlueprint = bOpenInEditor;
-		Params.ParentClassOverride = ParentClassOverride;
-		return CreateBlueprintFromActor(Path, Actor, Params);
-	}
-
-	/** 
-	 * Take an Actor and generate a blueprint based on it. Uses the Actors type as the parent class. 
-	 * @param BlueprintName			The name to use for the Blueprint
-	 * @param Outer					The outer object to create the blueprint within
-	 * @param Actor					The actor to use as the template for the blueprint
-	 * @param bReplaceActor			If true, replace the actor in the scene with one based on the created blueprint
-	 * @param bKeepMobility			If true, The mobility of each actor components will be copy
-	 * @param ParentClassOverride	The parent class to use when creating the blueprint. If null, the class of Actor will be used. If specified, must be a subclass of the Actor's class.
-	 * @param bOpenInEditor			If true, open the created blueprint in the blueprint editor
-	 * @return The blueprint created from the actor
-	 */
-	UE_DEPRECATED(4.26, "Use version that passes parameters via parameter struct")
-	static UBlueprint* CreateBlueprintFromActor(const FName BlueprintName, UObject* Outer, AActor* Actor, bool bReplaceInWorld, bool bKeepMobility = false, UClass* ParentClassOverride = nullptr, bool bOpenInEditor = true)
-	{
-		FCreateBlueprintFromActorParams Params;
-		Params.bReplaceActor = bReplaceInWorld;
-		Params.bKeepMobility = bKeepMobility;
-		Params.bOpenBlueprint = bOpenInEditor;
-		Params.ParentClassOverride = ParentClassOverride;
-		return CreateBlueprintFromActor(BlueprintName, Outer, Actor, Params);
-	}
-
 	/** Parameter struct for customizing calls to CreateBlueprintFromActors */
 	struct FCreateBlueprintFromActorsParams
 	{
@@ -425,46 +363,6 @@ public:
 	 */
 	static UNREALED_API UBlueprint* CreateBlueprintFromActors(const FName BlueprintName, UPackage* Package, const FCreateBlueprintFromActorsParams& Params);
 
-	/** 
-	 * Take a list of Actors and generate a blueprint based on it using the Actors as templates for child actor components.
-	 * @param Path					The path to use when creating the package for the new blueprint
-	 * @param Actors				The actors to use when creating child actor components
-	 * @param bReplaceActor			If true, replace the actor in the scene with one based on the created blueprint
-	 * @param ParentClass			The parent class to use when creating the blueprint.
-	 * @param bOpenInEditor			If true, open the created blueprint in the blueprint editor
-	 * @return The blueprint created from the actor
-	 */
-	UE_DEPRECATED(4.26, "Use version that passes parameters via parameter struct")
-	static UBlueprint* CreateBlueprintFromActors(const FString& Path, const TArray<AActor*>& Actors, bool bReplaceActor, UClass* ParentClass = AActor::StaticClass(), bool bOpenInEditor = true)
-	{
-		FCreateBlueprintFromActorsParams Params(Actors);
-		Params.bReplaceActors = bReplaceActor;
-		Params.ParentClass = ParentClass;
-		Params.bOpenBlueprint = bOpenInEditor;
-		return CreateBlueprintFromActors(Path, Params);
-	}
-
-	/** 
-	 * Take a list of Actors and generate a blueprint based on it using the Actors as templates for child actor components.
-	 * @param BlueprintName			The name to use for the Blueprint
-	 * @param Outer					The package to create the blueprint within
-	 * @param Actors				The actors to use when creating child actor components
-	 * @param bReplaceActor			If true, replace the actor in the scene with one based on the created blueprint
-	 * @param bKeepMobility			If true, The mobility of each actor components will be copy
-	 * @param ParentClass			The parent class to use when creating the blueprint.
-	 * @param bOpenInEditor			If true, open the created blueprint in the blueprint editor
-	 * @return The blueprint created from the actor
-	 */
-	UE_DEPRECATED(4.26, "Use version that passes parameters via parameter struct")
-	static UBlueprint* CreateBlueprintFromActors(const FName BlueprintName, UPackage* Package, const TArray<AActor*>& Actors, bool bReplaceInWorld, UClass* ParentClass = AActor::StaticClass(), bool bOpenInEditor = true)
-	{
-		FCreateBlueprintFromActorsParams Params(Actors);
-		Params.bReplaceActors = bReplaceInWorld;
-		Params.ParentClass = ParentClass;
-		Params.bOpenBlueprint = bOpenInEditor;
-		return CreateBlueprintFromActors(BlueprintName, Package, Params);
-	}
-
 	/** Parameter struct for customizing calls to CreateBlueprintFromActor */
 	struct FHarvestBlueprintFromActorsParams
 	{
@@ -504,23 +402,6 @@ public:
 	 */
 	static UNREALED_API UBlueprint* HarvestBlueprintFromActors(const FName BlueprintName, UPackage* Package, const TArray<AActor*>& Actors, const FHarvestBlueprintFromActorsParams& Params = FHarvestBlueprintFromActorsParams());
 
-	/** 
-	 * Take a list of Actors and generate a blueprint  by harvesting the components they have.
-	 * @param Path					The path to use when creating the package for the new blueprint
-	 * @param Actors				The actor list to use as the template for the new blueprint, typically this is the currently selected actors
-	 * @param bReplaceInWorld		If true, replace the selected actors in the scene with one based on the created blueprint
-	 * @param ParentClass			The parent class to use when creating the blueprint.
-	 * @return The blueprint created from the actors
-	 */
-	UE_DEPRECATED(4.26, "Use version that passes parameters via parameter struct")
-	static UBlueprint* HarvestBlueprintFromActors(const FString& Path, const TArray<AActor*>& Actors, bool bReplaceInWorld, UClass* ParentClass = AActor::StaticClass())
-	{
-		FHarvestBlueprintFromActorsParams Params;
-		Params.bReplaceActors = bReplaceInWorld;
-		Params.ParentClass = ParentClass;
-		return HarvestBlueprintFromActors(Path, Actors, Params);
-	}
-
 	/**
 	 * Take a list of Actors and update an existing Blueprint by harvesting the components they have. Essentially HarvestBlueprintFromActors, but 
 	 * updates an existing Blueprint rather than creating a new one.
@@ -530,24 +411,6 @@ public:
 	 */
 	static UNREALED_API UBlueprint* UpdateExistingBlueprintFromActors(const FString& Path, const TArray<AActor*>& Actors);
 
-	/**
-	 * Take a list of Actors and generate a blueprint  by harvesting the components they have. Uses AActor as parent class type as the parent class.
-	 * @param BlueprintName			The name to use for the Blueprint
-	 * @param Outer					The package to create the blueprint within
-	 * @param Actors				The actor list to use as the template for the new blueprint, typically this is the currently selected actors
-	 * @param bReplaceInWorld		If true, replace the selected actors in the scene with one based on the created blueprint
-	 * @param ParentClass			The parent class to use when creating the blueprint.
-	 * @return The blueprint created from the actors
-	 */
-	UE_DEPRECATED(4.26, "Use version that passes parameters via parameter struct")
-	static UBlueprint* HarvestBlueprintFromActors(const FName BlueprintName, UPackage* Package, const TArray<AActor*>& Actors, bool bReplaceInWorld, UClass* ParentClass = AActor::StaticClass())
-	{
-		FHarvestBlueprintFromActorsParams Params;
-		Params.bReplaceActors = bReplaceInWorld;
-		Params.ParentClass = ParentClass;
-		return HarvestBlueprintFromActors(BlueprintName, Package, Actors, Params);
-	}
-
 	/**
 	 * Updates this Actor's blueprint based on the actor itself. 
 	 * @return The number of properties that changes in the blueprint.
diff --git a/Engine/Source/Editor/UnrealEd/Public/PlayInEditorDataTypes.h b/Engine/Source/Editor/UnrealEd/Public/PlayInEditorDataTypes.h
--- a/Engine/Source/Editor/UnrealEd/Public/PlayInEditorDataTypes.h
+++ b/Engine/Source/Editor/UnrealEd/Public/PlayInEditorDataTypes.h
@@ -59,10 +59,6 @@ struct FPieLoginStruct
 	/** World context handle for this login */
 	FName WorldContextHandle;
 
-	/** Setting index for window positioning */
-	UE_DEPRECATED(4.25, "This is now done automatically via PlayInEditorSessionInfo")
-	int32 SettingsIndex;
-
 	/** X location for window positioning */
 	int32 NextX;
 	/** Y location for window positioning */
@@ -76,18 +72,15 @@ struct FPieLoginStruct
 	/** Which index is the instance of this Play in Editor session that will be created. */
 	int32 PIEInstanceIndex;
 
-	PRAGMA_DISABLE_DEPRECATION_WARNINGS
 	FPieLoginStruct(const FPieLoginStruct& InOther)
 	{
 		WorldContextHandle = InOther.WorldContextHandle;
-		SettingsIndex = InOther.SettingsIndex;
 		NextX = InOther.NextX;
 		NextY = InOther.NextY;
 		PIEStartTime = InOther.PIEStartTime;
 		GameInstancePIEParameters = InOther.GameInstancePIEParameters;
 		PIEInstanceIndex = InOther.PIEInstanceIndex;
 	}
-	PRAGMA_ENABLE_DEPRECATION_WARNINGS
 
 	FPieLoginStruct()
 		: WorldContextHandle(NAME_None)
@@ -95,25 +88,8 @@ struct FPieLoginStruct
 		, NextY(0)
 		, PIEStartTime(0)
 		, PIEInstanceIndex(-1)
-	{
-		PRAGMA_DISABLE_DEPRECATION_WARNINGS
-		SettingsIndex = 0;
-		PRAGMA_ENABLE_DEPRECATION_WARNINGS
-	}
-};
-
-
-/** Overrides you can pass when starting PIE to temporarily suppress the user's options for dedicated server, number of clients..etc. */
-struct UE_DEPRECATED(4.25, "Specify overrides by customizing the ULevelEditorPlaySettings in your PlaySessionRequest instead.") FPlayInEditorOverrides
-{
-	FPlayInEditorOverrides()
-		: bDedicatedServer()
-		, NumberOfClients()
 	{
 	}
-
-	TOptional<bool> bDedicatedServer;
-	TOptional<int32> NumberOfClients;
 };
 
 enum class EPlaySessionDestinationType : uint8
diff --git a/Engine/Source/Runtime/Engine/Classes/Camera/CameraModifier_CameraShake.h b/Engine/Source/Runtime/Engine/Classes/Camera/CameraModifier_CameraShake.h
--- a/Engine/Source/Runtime/Engine/Classes/Camera/CameraModifier_CameraShake.h
+++ b/Engine/Source/Runtime/Engine/Classes/Camera/CameraModifier_CameraShake.h
@@ -151,19 +151,6 @@ public:
 	 */
 	ENGINE_API virtual UCameraShakeBase* AddCameraShake(TSubclassOf<UCameraShakeBase> NewShake, const FAddCameraShakeParams& Params);
 
-	/** 
-	 * Adds a new active screen shake to be applied. 
-	 * @param NewShake - The class of camera shake to instantiate.
-	 * @param Scale - The scalar intensity to play the shake.
-	 * @param PlaySpace - Which coordinate system to play the shake in.
-	 * @param UserPlaySpaceRot - Coordinate system to play shake when PlaySpace == CAPS_UserDefined.
-	 */
-	UE_DEPRECATED(4.25, "Please use the new AddCameraShake method that takes a parameter struct.")
-	virtual UCameraShakeBase* AddCameraShake(TSubclassOf<UCameraShakeBase> NewShake, float Scale, ECameraShakePlaySpace PlaySpace= ECameraShakePlaySpace::CameraLocal, FRotator UserPlaySpaceRot = FRotator::ZeroRotator)
-	{
-		return AddCameraShake(NewShake, FAddCameraShakeParams(Scale, PlaySpace, UserPlaySpaceRot));
-	}
-
 	/**
 	 * Returns a list of currently active camera shakes.
 	 * @param ActiveCameraShakes - The array to fill up with shake information.
diff --git a/Engine/Source/Runtime/Engine/Classes/Camera/CameraTypes.h b/Engine/Source/Runtime/Engine/Classes/Camera/CameraTypes.h
--- a/Engine/Source/Runtime/Engine/Classes/Camera/CameraTypes.h
+++ b/Engine/Source/Runtime/Engine/Classes/Camera/CameraTypes.h
@@ -32,19 +32,6 @@ enum class ECameraShakePlaySpace : uint8
 	UserDefined,
 };
 
-/** Backwards compatible name for the camera shake play space enum, for C++ code. */
-namespace ECameraAnimPlaySpace
-{
-	UE_DEPRECATED(4.26, "Please use ECameraShakePlaySpace")
-	typedef ECameraShakePlaySpace Type;
-	UE_DEPRECATED(4.26, "Please use ECameraShakePlaySpace")
-	static const ECameraShakePlaySpace CameraLocal = ECameraShakePlaySpace::CameraLocal;
-	UE_DEPRECATED(4.26, "Please use ECameraShakePlaySpace")
-	static const ECameraShakePlaySpace World = ECameraShakePlaySpace::World;
-	UE_DEPRECATED(4.26, "Please use ECameraShakePlaySpace")
-	static const ECameraShakePlaySpace UserDefined = ECameraShakePlaySpace::UserDefined;
-}
-
 USTRUCT(BlueprintType)
 struct FMinimalViewInfo
 {
diff --git a/Engine/Source/Runtime/Engine/Classes/Camera/PlayerCameraManager.h b/Engine/Source/Runtime/Engine/Classes/Camera/PlayerCameraManager.h
--- a/Engine/Source/Runtime/Engine/Classes/Camera/PlayerCameraManager.h
+++ b/Engine/Source/Runtime/Engine/Classes/Camera/PlayerCameraManager.h
@@ -247,16 +247,6 @@ public:
 	/** Time at which interpolation started (when bEnableColorScaling and bEnableColorScaleInterp == true) */
 	float ColorScaleInterpStartTime;
 
-	/** Cached camera properties. */
-	UE_DEPRECATED(4.19, "This property is now deprecated, please use GetCameraCachePOV and SetCameraCachePOV functions instead.")
-	UPROPERTY(transient)
-	struct FCameraCacheEntry CameraCache;
-
-	/** Cached camera properties, one frame old. */
-	UE_DEPRECATED(4.19, "This property is now deprecated, please use GetLastFrameCameraCachePOV and SetLastFrameCameraCachePOV functions instead.")
-	UPROPERTY(transient)
-	struct FCameraCacheEntry LastFrameCameraCache;
-
 	/** Current ViewTarget */
 	UPROPERTY(transient)
 	struct FTViewTarget ViewTarget;
diff --git a/Engine/Source/Runtime/Engine/Classes/Components/SceneComponent.h b/Engine/Source/Runtime/Engine/Classes/Components/SceneComponent.h
--- a/Engine/Source/Runtime/Engine/Classes/Components/SceneComponent.h
+++ b/Engine/Source/Runtime/Engine/Classes/Components/SceneComponent.h
@@ -711,11 +711,6 @@ public:
 	/** Backwards compatibility: Used to convert old-style EAttachLocation to new-style EAttachmentRules */
 	static ENGINE_API void ConvertAttachLocation(EAttachLocation::Type InAttachLocation, EAttachmentRule& InOutLocationRule, EAttachmentRule& InOutRotationRule, EAttachmentRule& InOutScaleRule);
 
-	/** DEPRECATED - Use AttachToComponent() instead */
-	UE_DEPRECATED(4.17, "This function is deprecated, please use AttachToComponent() instead.")
-	UFUNCTION(BlueprintCallable, Category = "Transformation", meta = (DeprecationMessage = "OVERRIDE BAD MESSAGE", DisplayName = "AttachTo (Deprecated)", AttachType = "KeepRelativeOffset"))
-	ENGINE_API bool K2_AttachTo(USceneComponent* InParent, FName InSocketName = NAME_None, EAttachLocation::Type AttachType = EAttachLocation::KeepRelativeOffset, bool bWeldSimulatedBodies = true);
-
 	/**
 	* Attach this component to another scene component, optionally at a named socket. It is valid to call this on components whether or not they have been Registered, however from
 	* constructor or when not registered it is preferable to use SetupAttachment.
diff --git a/Engine/Source/Runtime/Engine/Classes/EdGraph/EdGraphNode.h b/Engine/Source/Runtime/Engine/Classes/EdGraph/EdGraphNode.h
--- a/Engine/Source/Runtime/Engine/Classes/EdGraph/EdGraphNode.h
+++ b/Engine/Source/Runtime/Engine/Classes/EdGraph/EdGraphNode.h
@@ -501,36 +501,6 @@ public:
 		FName NameParameter;
 	};
 
-	/** Create a new pin on this node using the supplied info, and return the new pin */
-	UE_DEPRECATED(4.19, "Use version that supplies Pin Category, SubCategory, and Name as an FName and uses PinContainerType instead of separate booleans for array, set, and map.")
-	ENGINE_API UEdGraphPin* CreatePin(
-		EEdGraphPinDirection Dir, 
-		const FNameParameterHelper PinCategory, 
-		const FNameParameterHelper PinSubCategory, 
-		UObject* PinSubCategoryObject, 
-		bool bIsArray, 
-		bool bIsReference, 
-		const FNameParameterHelper PinName, 
-		bool bIsConst = false, 
-		int32 Index = INDEX_NONE, 
-		bool bIsSet = false, 
-		bool bIsMap = false,
-		const FEdGraphTerminalType& ValueTerminalType = FEdGraphTerminalType());
-
-	/** Create a new pin on this node using the supplied info, and return the new pin */
-	UE_DEPRECATED(4.19, "Use version that supplies Pin Category, SubCategory, and Name as an FName and uses a parameter structure for optional paramaters.")
-	ENGINE_API UEdGraphPin* CreatePin(
-		EEdGraphPinDirection Dir,
-		const FNameParameterHelper PinCategory,
-		const FNameParameterHelper PinSubCategory,
-		UObject* PinSubCategoryObject,
-		const FNameParameterHelper PinName,
-		EPinContainerType PinContainerType = EPinContainerType::None,
-		bool bIsReference = false,
-		bool bIsConst = false,
-		int32 Index = INDEX_NONE,
-		const FEdGraphTerminalType& ValueTerminalType = FEdGraphTerminalType());
-
 	/** Parameter struct of less common options for CreatePin */
 	struct FCreatePinParams
 	{
@@ -572,12 +542,6 @@ public:
 	/** Create a new pin on this node using the supplied pin type, and return the new pin */
 	ENGINE_API UEdGraphPin* CreatePin(EEdGraphPinDirection Dir, const FEdGraphPinType& InPinType, const FName PinName, int32 Index = INDEX_NONE);
 
-	/** Create a new pin on this node using the supplied pin type, and return the new pin */
-	UE_DEPRECATED(4.19, "Use version that passes PinName as FName instead.")
-	UEdGraphPin* CreatePin(EEdGraphPinDirection Dir, const FEdGraphPinType& InPinType, const FString& PinName, int32 Index = INDEX_NONE)
-	{
-		return CreatePin(Dir, InPinType, FName(*PinName), Index);
-	}
 
 	/** Create a new pin on this node using the supplied pin type, and return the new pin */
 	//UE_DEPRECATED(4.19, "Remove when removing FString version. Exists just to resolve ambiguity")
diff --git a/Engine/Source/Runtime/Engine/Classes/Engine/GameViewportClient.h b/Engine/Source/Runtime/Engine/Classes/Engine/GameViewportClient.h
--- a/Engine/Source/Runtime/Engine/Classes/Engine/GameViewportClient.h
+++ b/Engine/Source/Runtime/Engine/Classes/Engine/GameViewportClient.h
@@ -91,12 +91,6 @@ public:
 	UPROPERTY(Config)
 	int32 MaxSplitscreenPlayers = 4;
 
-	/** if true then the title safe border is drawn
-	  * @deprecated - Use the cvar "r.DebugSafeZone.Mode=1".
-	  */
-	UE_DEPRECATED(4.26, "Use the cvar \"r.DebugSafeZone.Mode=1\".")
-	uint32 bShowTitleSafeZone:1;
-
 	/** If true, this viewport is a play in editor viewport */
 	uint32 bIsPlayInEditorViewport:1;
 
@@ -433,20 +427,6 @@ public:
 		return bDisableSplitScreenOverride;
 	}
 
-	/** Allows game code to disable splitscreen (useful when in menus) */
-	UE_DEPRECATED(4.24, "SetDisableSplitscreenOverride is deprecated. Please call UGameViewportClient::SetForceDisableSplitscreen(bDisabled) instead.")
-	void SetDisableSplitscreenOverride( const bool bDisabled )
-	{
-		SetForceDisableSplitscreen(bDisabled);
-	}
-
-	/** Determines whether splitscreen is forced to be turned off */
-	UE_DEPRECATED(4.24, "GetDisableSplitscreenOverride is deprecated. Please call UGameViewportClient::IsSplitscreenForceDisabled() instead.")
-	bool GetDisableSplitscreenOverride() const
-	{
-		return IsSplitscreenForceDisabled();
-	}
-
 	/** called before rending subtitles to allow the game viewport to determine the size of the subtitle area
 	 * @param Min top left bounds of subtitle region (0 to 1)
 	 * @param Max bottom right bounds of subtitle region (0 to 1)
diff --git a/Engine/Source/Runtime/Engine/Classes/Engine/LevelStreamingDynamic.h b/Engine/Source/Runtime/Engine/Classes/Engine/LevelStreamingDynamic.h
--- a/Engine/Source/Runtime/Engine/Classes/Engine/LevelStreamingDynamic.h
+++ b/Engine/Source/Runtime/Engine/Classes/Engine/LevelStreamingDynamic.h
@@ -106,6 +106,3 @@ private:
  	static ENGINE_API ULevelStreamingDynamic* LoadLevelInstance_Internal(const FLoadLevelInstanceParams& Params, bool& bOutSuccess);
 
 };
-
-UE_DEPRECATED(4.21, "ULevelStreamingKismet has been renamed to ULevelStreamingDynamic")
-typedef ULevelStreamingDynamic ULevelStreamingKismet;
diff --git a/Engine/Source/Runtime/Engine/Classes/Engine/SimpleConstructionScript.h b/Engine/Source/Runtime/Engine/Classes/Engine/SimpleConstructionScript.h
--- a/Engine/Source/Runtime/Engine/Classes/Engine/SimpleConstructionScript.h
+++ b/Engine/Source/Runtime/Engine/Classes/Engine/SimpleConstructionScript.h
@@ -62,12 +62,6 @@ class USimpleConstructionScript : public UObject
 	 */
 	ENGINE_API USceneComponent* GetSceneRootComponentTemplate(bool bShouldUseDefaultRoot = false, USCS_Node** OutSCSNode = nullptr) const;
 
-	UE_DEPRECATED(4.25, "Use version that specifies whether the default scene root should be used")
-	USceneComponent* GetSceneRootComponentTemplate(USCS_Node** OutSCSNode) const
-	{
-		return GetSceneRootComponentTemplate(false, OutSCSNode);
-	}
-
 	/** Saves the current state of SimpleConstructionScript and its nodes to the transaction buffer. */
 	ENGINE_API void SaveToTransactionBuffer();
 #endif
@@ -85,13 +79,6 @@ class USimpleConstructionScript : public UObject
 	const TArray<USCS_Node*>& GetAllNodes() const { return AllNodes; }
 #endif	
 
-	/** Return immutable references to nodes in tree as a flat list */
-	UE_DEPRECATED(4.27, "GetAllNodesConst is being removed as it provides no unique functionality that GetAllNodes cannot be used for.")
-	TArray<const USCS_Node*> GetAllNodesConst() const
-	{
-		return TArray<const USCS_Node*>(GetAllNodes());
-	}
-
 	/** Provides read-only access to the root node set */
 	const TArray<USCS_Node*>& GetRootNodes() const { return RootNodes; }
 
diff --git a/Engine/Source/Runtime/Engine/Classes/Engine/TimelineTemplate.h b/Engine/Source/Runtime/Engine/Classes/Engine/TimelineTemplate.h
--- a/Engine/Source/Runtime/Engine/Classes/Engine/TimelineTemplate.h
+++ b/Engine/Source/Runtime/Engine/Classes/Engine/TimelineTemplate.h
@@ -259,10 +259,6 @@ class UTimelineTemplate : public UObject
 	/** Get the name of the function we expect to find in the owning actor that we will bind the finished event to */
 	FName GetFinishedFunctionName() const { return FinishedFunctionName; }
 
-	/** Get the name of the funcerion we expect to find in the owning actor that we will bind event track with index EvenTrackIndex to */
-	UE_DEPRECATED(4.22, "Access the event track function name directly from the EventTrack instead.")
-	ENGINE_API FName GetEventTrackFunctionName(int32 EventTrackIndex) const;
-
 	/** Set a metadata value on the timeline */
 	ENGINE_API void SetMetaData(FName Key, FString Value);
 
diff --git a/Engine/Source/Runtime/Engine/Classes/Engine/World.h b/Engine/Source/Runtime/Engine/Classes/Engine/World.h
--- a/Engine/Source/Runtime/Engine/Classes/Engine/World.h
+++ b/Engine/Source/Runtime/Engine/Classes/Engine/World.h
@@ -264,12 +264,6 @@ public:
 	 */
 	bool StartTravel(UWorld* InCurrentWorld, const FURL& InURL);
 
-	UE_DEPRECATED(4.27, "UPackage::Guid has not been used by the engine for a long time. Please use StartTravel without a InGuid.")
-	bool StartTravel(UWorld* InCurrentWorld, const FURL& InURL, const FGuid& InGuid)
-	{
-		return StartTravel(InCurrentWorld, InURL);
-	}
-
 	/** @return whether a transition is already in progress */
 	FORCEINLINE bool IsInTransition() const
 	{
@@ -1672,11 +1666,6 @@ private:
 	/** Broadcasts whenever all the levels change */
 	FOnAllLevelsChangedEvent AllLevelsChangedEvent;
 
-	DECLARE_EVENT(UWorld, FOnBeginTearingDownEvent);
-
-	/** Broadcasted on UWorld::BeginTearingDown */
-	FOnBeginTearingDownEvent BeginTearingDownEvent;
-
 	/** Broadcasted when WorldPartition gets initialized */
 	DECLARE_EVENT_OneParam(UWorld, FWorldPartitionInitializedEvent, UWorldPartition*);
 	
@@ -2522,14 +2511,6 @@ public:
 	/** @return Returns the number of Controllers. */
 	int32 GetNumControllers() const;
 	
-	/** @return Returns an iterator for the pawn list. */
-	UE_DEPRECATED(4.24, "The PawnIterator is an inefficient mechanism for iterating pawns. Please use TActorIterator<PawnType> instead.")
-	FConstPawnIterator GetPawnIterator() const;
-	
-	/** @return Returns the number of Pawns. */
-	UE_DEPRECATED(4.23, "GetNumPawns is no longer a supported function on UWorld. The version that remains for backwards compatibility is significantly more expensive to call.")
-	int32 GetNumPawns() const;
-
 	/** @return Returns an iterator for the player controller list. */
 	FConstPlayerControllerIterator GetPlayerControllerIterator() const;
 
@@ -2797,18 +2778,7 @@ public:
 	 * @param	Controller	Controller to remove
 	 */
 	void RemoveController( AController* Controller );
-
-	UE_DEPRECATED(4.23, "There is no longer a reason to AddPawn to UWorld")
-	void AddPawn( APawn* Pawn ) { }
 	
-	/**
-	 * Removes the passed in pawn from the linked list of pawns.
-	 *
-	 * @param	Pawn	Pawn to remove
-	 */
-	UE_DEPRECATED(4.23, "RemovePawn has been deprecated and should no longer need to be called as PawnList is no longer maintained and Unpossess should be handled by EndPlay.")
-	void RemovePawn( APawn* Pawn ) const;
-
 	/**
 	 * Adds the passed in actor to the special network actor list
 	 * This list is used to specifically single out actors that are relevant for networking without having to scan the much large list
@@ -3775,10 +3745,6 @@ public:
 	/** Returns the AllLevelsChangedEvent member. */
 	FOnAllLevelsChangedEvent& OnAllLevelsChanged() { return AllLevelsChangedEvent; }
 
-	/** Returns the BeginTearingDownEvent member. */
-	UE_DEPRECATED(4.26, "OnBeginTearingDown has been replaced by FWorldDelegates::OnWorldBeginTearDown")
-	FOnBeginTearingDownEvent& OnBeginTearingDown() { return BeginTearingDownEvent; }
-
 	/** Returns the actor count. */
 	int32 GetProgressDenominator() const;
 	
@@ -3902,12 +3868,6 @@ public:
 	 */
 	void SeamlessTravel(const FString& InURL, bool bAbsolute = false);
 
-	UE_DEPRECATED(4.27, "UPackage::Guid has not been used by the engine for a long time. Please use SeamlessTravel without a NextMapGuid.")
-	void SeamlessTravel(const FString& InURL, bool bAbsolute, FGuid MapPackageGuid)
-	{
-		SeamlessTravel(InURL, bAbsolute);
-	}
-
 	/** @return whether we're currently in a seamless transition */
 	bool IsInSeamlessTravel() const;
 
diff --git a/Engine/Source/Runtime/Engine/Classes/GameFramework/Actor.h b/Engine/Source/Runtime/Engine/Classes/GameFramework/Actor.h
--- a/Engine/Source/Runtime/Engine/Classes/GameFramework/Actor.h
+++ b/Engine/Source/Runtime/Engine/Classes/GameFramework/Actor.h
@@ -1774,11 +1774,6 @@ public:
 	UFUNCTION(BlueprintCallable, meta=(BlueprintInternalUseOnly="true"))
 	ENGINE_API void FinishAddComponent(UActorComponent* Component, bool bManualAttachment, const FTransform& RelativeTransform);
 
-	/** DEPRECATED - Use AttachToComponent() instead */
-	UE_DEPRECATED(4.17, "Use AttachToComponent() instead.")
-	UFUNCTION(BlueprintCallable, meta = (DisplayName = "AttachRootComponentTo (Deprecated)", ScriptNoExport, AttachLocationType = "KeepRelativeOffset"), Category = "Transformation")
-	ENGINE_API void K2_AttachRootComponentTo(USceneComponent* InParent, FName InSocketName = NAME_None, EAttachLocation::Type AttachLocationType = EAttachLocation::KeepRelativeOffset, bool bWeldSimulatedBodies = true);
-
 	/**
 	 * Attaches the RootComponent of this Actor to the supplied component, optionally at a named socket. It is not valid to call this on components that are not Registered.
 	 * @param Parent					Parent to attach to.
@@ -1801,11 +1796,6 @@ public:
 	 */
 	ENGINE_API bool AttachToComponent(USceneComponent* Parent, const FAttachmentTransformRules& AttachmentRules, FName SocketName = NAME_None);
 
-	/** DEPRECATED - Use AttachToActor() instead */
-	UE_DEPRECATED(4.17, "Use AttachToActor() instead.")
-	UFUNCTION(BlueprintCallable, meta = (DisplayName = "AttachRootComponentToActor (Deprecated)", ScriptNoExport, AttachLocationType = "KeepRelativeOffset"), Category = "Transformation")
-	ENGINE_API void K2_AttachRootComponentToActor(AActor* InParentActor, FName InSocketName = NAME_None, EAttachLocation::Type AttachLocationType = EAttachLocation::KeepRelativeOffset, bool bWeldSimulatedBodies = true);
-
 	/**
 	 * Attaches the RootComponent of this Actor to the RootComponent of the supplied actor, optionally at a named socket.
 	 * @param ParentActor				Actor to attach this actor's RootComponent to
@@ -1828,11 +1818,6 @@ public:
 	UFUNCTION(BlueprintCallable, meta = (DisplayName = "Attach Actor To Actor", ScriptName = "AttachToActor", bWeldSimulatedBodies=true), Category = "Transformation")
 	ENGINE_API bool K2_AttachToActor(AActor* ParentActor, FName SocketName, EAttachmentRule LocationRule, EAttachmentRule RotationRule, EAttachmentRule ScaleRule, bool bWeldSimulatedBodies);
 
-	/** DEPRECATED - Use DetachFromActor() instead */
-	UE_DEPRECATED(4.17, "Use DetachFromActor() instead")
-	UFUNCTION(BlueprintCallable, meta=(DisplayName = "DetachActorFromActor (Deprecated)", ScriptNoExport), Category="Transformation")
-	ENGINE_API void DetachRootComponentFromParent(bool bMaintainWorldPosition = true);
-
 	/** 
 	 * Detaches the RootComponent of this Actor from any SceneComponent it is currently attached to. 
 	 * @param  LocationRule				How to handle translation when detaching.
diff --git a/Engine/Source/Runtime/Engine/Classes/GameFramework/CharacterMovementComponent.h b/Engine/Source/Runtime/Engine/Classes/GameFramework/CharacterMovementComponent.h
--- a/Engine/Source/Runtime/Engine/Classes/GameFramework/CharacterMovementComponent.h
+++ b/Engine/Source/Runtime/Engine/Classes/GameFramework/CharacterMovementComponent.h
@@ -2667,8 +2667,6 @@ public:
 	/** Apply a RootMotionSource to current root motion 
 	 *  @return LocalID for this Root Motion Source */
 	ENGINE_API uint16 ApplyRootMotionSource(TSharedPtr<FRootMotionSource> SourcePtr);
-	UE_DEPRECATED(4.26, "ApplyRootMotionSource no longer takes raw pointers. Pass in a TSharedPtr instead.")
-	ENGINE_API uint16 ApplyRootMotionSource(FRootMotionSource* SourcePtr);
 
 	/** Called during ApplyRootMotionSource call, useful for project-specific alerts for "something is about to be altering our movement" */
 	ENGINE_API virtual void OnRootMotionSourceBeingApplied(const FRootMotionSource* Source);
diff --git a/Engine/Source/Runtime/Engine/Classes/GameFramework/Controller.h b/Engine/Source/Runtime/Engine/Classes/GameFramework/Controller.h
--- a/Engine/Source/Runtime/Engine/Classes/GameFramework/Controller.h
+++ b/Engine/Source/Runtime/Engine/Classes/GameFramework/Controller.h
@@ -402,19 +402,6 @@ private:
 	// Hidden functions that don't make sense to use on this class.
 	HIDE_ACTOR_TRANSFORM_FUNCTIONS();
 
-public:
-
-	/** prepares path following component */
-	UE_DEPRECATED(4.20, "InitNavigationControl is depracated. The function lost its meaning due to better Engine-AIModule code separation")
-	virtual void InitNavigationControl(UObject*& PathFollowingComp) {}
-
-	/** If controller has any navigation-related components then this function
-	 *	makes them update their cached data 
-	 *	DEPRECATED: this functionality has been taken over by PathFollowingComponent 
-	 *		by observing newly possessed pawns (via OnNewPawn)
-	 */
-	UE_DEPRECATED(4.20, "UpdateNavigationComponents is depracated. The function lost its meaning due to better Engine-AIModule code separation.")
-	virtual void UpdateNavigationComponents() {}
 };
 
 
diff --git a/Engine/Source/Runtime/Engine/Classes/GameFramework/GameModeBase.h b/Engine/Source/Runtime/Engine/Classes/GameFramework/GameModeBase.h
--- a/Engine/Source/Runtime/Engine/Classes/GameFramework/GameModeBase.h
+++ b/Engine/Source/Runtime/Engine/Classes/GameFramework/GameModeBase.h
@@ -554,12 +554,6 @@ protected:
 	 */
 	ENGINE_API virtual APlayerController* ProcessClientTravel(FString& URL, bool bSeamless, bool bAbsolute);
 
-	UE_DEPRECATED(4.27, "UPackage::Guid has not been used by the engine for a long time. Please use ProcessClientTravel without a NextMapGuid.")
-	virtual APlayerController* ProcessClientTravel(FString& URL, FGuid NextMapGuid, bool bSeamless, bool bAbsolute)
-	{
-		return ProcessClientTravel(URL, bSeamless, bAbsolute);
-	}
-
 	/** Handles initializing a seamless travel player, handles logic similar to InitNewPlayer */
 	ENGINE_API virtual void InitSeamlessTravelPlayer(AController* NewController);
 
diff --git a/Engine/Source/Runtime/Engine/Classes/GameFramework/GameUserSettings.h b/Engine/Source/Runtime/Engine/Classes/GameFramework/GameUserSettings.h
--- a/Engine/Source/Runtime/Engine/Classes/GameFramework/GameUserSettings.h
+++ b/Engine/Source/Runtime/Engine/Classes/GameFramework/GameUserSettings.h
@@ -155,11 +155,6 @@ public:
 	UFUNCTION(BlueprintCallable, Category=Settings)
 	ENGINE_API virtual int32 GetOverallScalabilityLevel() const;
 
-	// Returns the current resolution scale and the range
-	UE_DEPRECATED(4.12, "Please call GetResolutionScaleInformationEx")
-	UFUNCTION(BlueprintCallable, Category = Settings, meta = (DeprecatedFunction, DisplayName = "GetResolutionScaleInformation_Deprecated"))
-	ENGINE_API void GetResolutionScaleInformation(float& CurrentScaleNormalized, int32& CurrentScaleValue, int32& MinScaleValue, int32& MaxScaleValue) const;
-
 	// Returns the current resolution scale and the range
 	UFUNCTION(BlueprintCallable, Category=Settings, meta=(DisplayName="GetResolutionScaleInformation"))
 	ENGINE_API void GetResolutionScaleInformationEx(float& CurrentScaleNormalized, float& CurrentScaleValue, float& MinScaleValue, float& MaxScaleValue) const;
@@ -168,11 +163,6 @@ public:
 	UFUNCTION(BlueprintCallable, Category = Settings)
 	ENGINE_API float GetResolutionScaleNormalized() const;
 
-	// Sets the current resolution scale
-	UE_DEPRECATED(4.12, "Please call SetResolutionScaleValueEx")
-	UFUNCTION(BlueprintCallable, Category=Settings, meta=(DeprecatedFunction, DisplayName="SetResolutionScaleValue_Deprecated"))
-	ENGINE_API void SetResolutionScaleValue(int32 NewScaleValue);
-
 	// Sets the current resolution scale
 	UFUNCTION(BlueprintCallable, Category=Settings, meta=(DisplayName="SetResolutionScaleValue"))
 	ENGINE_API void SetResolutionScaleValueEx(float NewScaleValue);
@@ -550,10 +540,6 @@ protected:
 	/** Sets the frame rate limit CVar to the passed in value, 0.0 indicates no limit */
 	static ENGINE_API void SetFrameRateLimitCVar(float InLimit);
 
-	/** Sets the sync interval limit CVar for consoles 0-3 */
-	UE_DEPRECATED(4.25, "Please use FPlatformRHIFramePacer::SetFramePace to set the desired frame rate")
-	static ENGINE_API void SetSyncIntervalCVar(int32 InInterval);
-
 	/** Sets the input latency mode 0 and 2 */
 	static ENGINE_API void SetSyncTypeCVar(int32 InInterval);
 
diff --git a/Engine/Source/Runtime/Engine/Classes/GameFramework/PlayerController.h b/Engine/Source/Runtime/Engine/Classes/GameFramework/PlayerController.h
--- a/Engine/Source/Runtime/Engine/Classes/GameFramework/PlayerController.h
+++ b/Engine/Source/Runtime/Engine/Classes/GameFramework/PlayerController.h
@@ -669,11 +669,6 @@ public:
 	UFUNCTION(exec)
 	ENGINE_API void ServerExec(const FString& Msg);
 
-	/** Return the client to the main menu gracefully */
-	UE_DEPRECATED(4.19, "As an FString, the ReturnReason parameter is not easily localized. Please use ClientReturnToMainMenuWithTextReason instead.")
-	UFUNCTION(Reliable, Client)
-	ENGINE_API virtual void ClientReturnToMainMenu(const FString& ReturnReason);
-
 	/** Return the client to the main menu gracefully */
 	UFUNCTION(Reliable, Client)
 	ENGINE_API virtual void ClientReturnToMainMenuWithTextReason(const FText& ReturnReason);
@@ -1117,13 +1112,6 @@ public:
 	UFUNCTION(unreliable, client, BlueprintCallable, Category="Game|Feedback")
 	ENGINE_API void ClientStartCameraShake(TSubclassOf<class UCameraShakeBase> Shake, float Scale = 1.f, ECameraShakePlaySpace PlaySpace = ECameraShakePlaySpace::CameraLocal, FRotator UserPlaySpaceRot = FRotator::ZeroRotator);
 
-	/** Backwards compatible method, for C++ code */
-	UE_DEPRECATED(4.26, "Please use ClientStartCameraShake")
-	void ClientPlayCameraShake(TSubclassOf<class UCameraShakeBase> Shake, float Scale = 1.f, ECameraShakePlaySpace PlaySpace = ECameraShakePlaySpace::CameraLocal, FRotator UserPlaySpaceRot = FRotator::ZeroRotator)
-	{
-		return ClientStartCameraShake(Shake, Scale, PlaySpace, UserPlaySpaceRot);
-	}
-
 	/** 
 	 * Play Camera Shake localized to a given source
 	 * @param Shake - Camera shake animation to play
@@ -1132,13 +1120,6 @@ public:
 	UFUNCTION(BlueprintCallable, Category="Game|Feedback")
 	ENGINE_API void ClientStartCameraShakeFromSource(TSubclassOf<class UCameraShakeBase> Shake, class UCameraShakeSourceComponent* SourceComponent);
 
-	/** Backwards compatible method, for C++ code */
-	UE_DEPRECATED(4.26, "Please use ClientStartCameraShakeFromSource")
-	void ClientPlayCameraShakeFromSource(TSubclassOf<class UCameraShakeBase> Shake, class UCameraShakeSourceComponent* SourceComponent)
-	{
-		return ClientStartCameraShakeFromSource(Shake, SourceComponent);
-	}
-
 	/**
 	 * Play sound client-side (so only the client will hear it)
 	 * @param Sound	- Sound to play
diff --git a/Engine/Source/Runtime/Engine/Classes/GameFramework/PlayerInput.h b/Engine/Source/Runtime/Engine/Classes/GameFramework/PlayerInput.h
--- a/Engine/Source/Runtime/Engine/Classes/GameFramework/PlayerInput.h
+++ b/Engine/Source/Runtime/Engine/Classes/GameFramework/PlayerInput.h
@@ -472,9 +472,6 @@ public:
 	/** Returns the mouse sensitivity along the Y-axis, or 1.0 if none are known. */
 	ENGINE_API float GetMouseSensitivityY();
 
-	UE_DEPRECATED(4.21, "Call axis specific GetMouseSensitivityX or GetMouseSensitivityY instead.")
-	float GetMouseSensitivity() { return GetMouseSensitivityX(); }
-
 	/** Returns whether an Axis Key is inverted */
 	ENGINE_API bool GetInvertAxisKey(const FKey AxisKey);
 
@@ -594,12 +591,6 @@ public:
 	/** Handles a touch input event.  Returns true. */
 	ENGINE_API bool InputTouch(uint32 Handle, ETouchType::Type Type, const FVector2D& TouchLocation, float Force, FDateTime DeviceTimestamp, uint32 TouchpadIndex);
 
-	UE_DEPRECATED(4.20, "InputTouch now takes a Force")
-	bool InputTouch(uint32 Handle, ETouchType::Type Type, const FVector2D& TouchLocation, FDateTime DeviceTimestamp, uint32 TouchpadIndex)
-	{
-		return InputTouch(Handle, Type, TouchLocation, 1.0f, DeviceTimestamp, TouchpadIndex);
-	}
-
 	/** Handles a motion input event.  Returns true. */
 	ENGINE_API bool InputMotion(const FVector& Tilt, const FVector& RotationRate, const FVector& Gravity, const FVector& Acceleration);
 
@@ -667,12 +658,6 @@ public:
 	/** @return raw value of the InKey */
 	ENGINE_API FVector GetRawVectorKeyValue(FKey InKey) const;
 
-	UE_DEPRECATED(4.21, "Use GetProcessedVectorKeyValue or GetRawVectorKeyValue instead. GetRawVectorKeyValue will have the same result as GetVectorKeyValue previously.")
-	FVector GetVectorKeyValue(FKey InKey) const
-	{
-		return GetRawVectorKeyValue(InKey);
-	}
-
 	/** @return true if alt key is pressed */
 	ENGINE_API bool IsAltPressed() const;
 
diff --git a/Engine/Source/Runtime/Engine/Classes/Kismet/KismetMathLibrary.h b/Engine/Source/Runtime/Engine/Classes/Kismet/KismetMathLibrary.h
--- a/Engine/Source/Runtime/Engine/Classes/Kismet/KismetMathLibrary.h
+++ b/Engine/Source/Runtime/Engine/Classes/Kismet/KismetMathLibrary.h
@@ -802,9 +802,6 @@ class UKismetMathLibrary : public UBlueprintFunctionLibrary
 	UFUNCTION(BlueprintPure, Category="Math|Float")
 	static ENGINE_API double Lerp(double A, double B, double Alpha);
 	
-	UE_DEPRECATED(4.19, "Use NormalizeToRange instead")
-	static ENGINE_API float InverseLerp(float A, float B, float Value);
-
 	/** Easeing  between A and B using a specified easing function */
 	UFUNCTION(BlueprintPure, meta = (DisplayName = "Ease", BlueprintInternalUseOnly = "true"), Category = "Math|Interpolation")
 	static ENGINE_API double Ease(double A, double B, double Alpha, TEnumAsByte<EEasingFunc::Type> EasingFunc, double BlendExp = 2, int32 Steps = 2);
@@ -1362,9 +1359,6 @@ class UKismetMathLibrary : public UBlueprintFunctionLibrary
 	UFUNCTION(BlueprintPure, meta=(ScriptMethod = "RotatorFromAxisAndAngle", Keywords="make construct build rotate rotation"), Category="Math|Vector")
 	static ENGINE_API FRotator RotatorFromAxisAndAngle(FVector Axis, float Angle);
 
-	UE_DEPRECATED(4.26, "Conv_VectorToQuaterion was renamed to Conv_VectorToQuaternion")
-	static ENGINE_API FQuat Conv_VectorToQuaterion(FVector InVec);
-
 	/**
 	 * Return the Quaternion orientation corresponding to the direction in which the vector points.
 	 * Similar to the FRotator version, returns a result without roll such that it preserves the up vector.
@@ -2094,9 +2088,6 @@ class UKismetMathLibrary : public UBlueprintFunctionLibrary
 	 */
 	UFUNCTION(BlueprintPure, meta=(DisplayName = "To Rotation (Vector4)", ScriptMethod = "Rotator", Keywords = "rotation rotate cast convert", BlueprintAutocast), Category = "Math|Conversions")
 	static ENGINE_API FRotator Conv_Vector4ToRotator(const FVector4& InVec);
-
-	UE_DEPRECATED(4.26, "Conv_Vector4ToQuaterion renamed to Conv_Vector4ToQuaternion")
-	static ENGINE_API FQuat Conv_Vector4ToQuaterion(const FVector4& InVec);
 	
 	/**
 	 * Return the Quaternion orientation corresponding to the direction in which the vector points.
@@ -3928,9 +3919,6 @@ class UKismetMathLibrary : public UBlueprintFunctionLibrary
 	UFUNCTION(BlueprintPure, meta=(ScriptMethod = "MakeRelative", Keywords="convert torelative"), Category="Math|Transform")
 	static ENGINE_API FTransform MakeRelativeTransform(const FTransform& A, const FTransform& RelativeTo);
 
-	UE_DEPRECATED(4.22, "Use MakeRelativeTransform instead, with reversed order of arguments.")
-	static ENGINE_API FTransform ConvertTransformToRelative(const FTransform& Transform, const FTransform& ParentTransform);
-
 	/** 
 	 * Returns the inverse of the given transform T.
 	 * 
diff --git a/Engine/Source/Runtime/Engine/Classes/Kismet/KismetMathLibrary.inl b/Engine/Source/Runtime/Engine/Classes/Kismet/KismetMathLibrary.inl
--- a/Engine/Source/Runtime/Engine/Classes/Kismet/KismetMathLibrary.inl
+++ b/Engine/Source/Runtime/Engine/Classes/Kismet/KismetMathLibrary.inl
@@ -984,12 +984,6 @@ FTransform UKismetMathLibrary::MakeRelativeTransform(const FTransform& A, const
 	return A.GetRelativeTransform(RelativeTo);
 }
 
-KISMET_MATH_INLINE
-FTransform UKismetMathLibrary::ConvertTransformToRelative(const FTransform& Transform, const FTransform& ParentTransform)
-{
-	return ParentTransform.GetRelativeTransform(Transform);
-}
-
 KISMET_MATH_INLINE
 FTransform UKismetMathLibrary::InvertTransform(const FTransform& T)
 {
@@ -1516,12 +1510,6 @@ FRotator UKismetMathLibrary::Conv_VectorToRotator(FVector InVec)
 	return InVec.ToOrientationRotator();
 }
 
-KISMET_MATH_FORCEINLINE
-FQuat UKismetMathLibrary::Conv_VectorToQuaterion(FVector InVec)
-{
-	return InVec.ToOrientationQuat();
-}
-
 KISMET_MATH_FORCEINLINE
 FQuat UKismetMathLibrary::Conv_VectorToQuaternion(FVector InVec)
 {
@@ -2128,12 +2116,6 @@ FRotator UKismetMathLibrary::Conv_Vector4ToRotator(const FVector4& InVec)
 	return InVec.ToOrientationRotator();
 }
 
-KISMET_MATH_FORCEINLINE
-FQuat UKismetMathLibrary::Conv_Vector4ToQuaterion(const FVector4& InVec)
-{
-	return InVec.ToOrientationQuat();
-}
-
 KISMET_MATH_FORCEINLINE
 FQuat UKismetMathLibrary::Conv_Vector4ToQuaternion(const FVector4& InVec)
 {
diff --git a/Engine/Source/Runtime/Engine/Private/Actor.cpp b/Engine/Source/Runtime/Engine/Private/Actor.cpp
--- a/Engine/Source/Runtime/Engine/Private/Actor.cpp
+++ b/Engine/Source/Runtime/Engine/Private/Actor.cpp
@@ -2198,16 +2198,6 @@ void AActor::SetAutoDestroyWhenFinished(bool bVal)
 	}
 }
 
-void AActor::K2_AttachRootComponentTo(USceneComponent* InParent, FName InSocketName, EAttachLocation::Type AttachLocationType /*= EAttachLocation::KeepRelativeOffset */, bool bWeldSimulatedBodies /*=true*/)
-{
-PRAGMA_DISABLE_DEPRECATION_WARNINGS
-	if (RootComponent && InParent)
-	{
-		RootComponent->K2_AttachTo(InParent, InSocketName, AttachLocationType, bWeldSimulatedBodies);
-	}
-PRAGMA_ENABLE_DEPRECATION_WARNINGS
-}
-
 bool AActor::K2_AttachToComponent(USceneComponent* Parent, FName SocketName, EAttachmentRule LocationRule, EAttachmentRule RotationRule, EAttachmentRule ScaleRule, bool bWeldSimulatedBodies)
 {
 	return AttachToComponent(Parent, FAttachmentTransformRules(LocationRule, RotationRule, ScaleRule, bWeldSimulatedBodies), SocketName);
@@ -2268,20 +2258,6 @@ void AActor::OnRep_AttachmentReplication()
 	}
 }
 
-void AActor::K2_AttachRootComponentToActor(AActor* InParentActor, FName InSocketName /*= NAME_None*/, EAttachLocation::Type AttachLocationType /*= EAttachLocation::KeepRelativeOffset */, bool bWeldSimulatedBodies /*=true*/)
-{
-	if (RootComponent && InParentActor)
-	{
-		USceneComponent* ParentDefaultAttachComponent = InParentActor->GetDefaultAttachComponent();
-		if (ParentDefaultAttachComponent)
-		{
-PRAGMA_DISABLE_DEPRECATION_WARNINGS
-			RootComponent->K2_AttachTo(ParentDefaultAttachComponent, InSocketName, AttachLocationType, bWeldSimulatedBodies );
-PRAGMA_ENABLE_DEPRECATION_WARNINGS
-		}
-	}
-}
-
 bool AActor::K2_AttachToActor(AActor* ParentActor, FName SocketName, EAttachmentRule LocationRule, EAttachmentRule RotationRule, EAttachmentRule ScaleRule, bool bWeldSimulatedBodies)
 {
 	return AttachToActor(ParentActor, FAttachmentTransformRules(LocationRule, RotationRule, ScaleRule, bWeldSimulatedBodies), SocketName);
@@ -2301,20 +2277,6 @@ bool AActor::AttachToActor(AActor* ParentActor, const FAttachmentTransformRules&
 	return false;
 }
 
-void AActor::DetachRootComponentFromParent(bool bMaintainWorldPosition)
-{
-	if(RootComponent)
-	{
-PRAGMA_DISABLE_DEPRECATION_WARNINGS
-		RootComponent->DetachFromParent(bMaintainWorldPosition);
-PRAGMA_ENABLE_DEPRECATION_WARNINGS		
-
-		// Clear AttachmentReplication struct
-		AttachmentReplication = FRepAttachment();
-		MARK_PROPERTY_DIRTY_FROM_NAME(AActor, AttachmentReplication, this);
-	}
-}
-
 void AActor::K2_DetachFromActor(EDetachmentRule LocationRule /*= EDetachmentRule::KeepRelative*/, EDetachmentRule RotationRule /*= EDetachmentRule::KeepRelative*/, EDetachmentRule ScaleRule /*= EDetachmentRule::KeepRelative*/)
 {
 	DetachFromActor(FDetachmentTransformRules(LocationRule, RotationRule, ScaleRule, true));
diff --git a/Engine/Source/Runtime/Engine/Private/Components/CharacterMovementComponent.cpp b/Engine/Source/Runtime/Engine/Private/Components/CharacterMovementComponent.cpp
--- a/Engine/Source/Runtime/Engine/Private/Components/CharacterMovementComponent.cpp
+++ b/Engine/Source/Runtime/Engine/Private/Components/CharacterMovementComponent.cpp
@@ -11367,11 +11367,6 @@ uint16 UCharacterMovementComponent::ApplyRootMotionSource(TSharedPtr<FRootMotion
 	return (uint16)ERootMotionSourceID::Invalid;
 }
 
-uint16 UCharacterMovementComponent::ApplyRootMotionSource(FRootMotionSource* SourcePtr)
-{
-	return ApplyRootMotionSource(TSharedPtr<FRootMotionSource>(SourcePtr));
-}
-
 void UCharacterMovementComponent::OnRootMotionSourceBeingApplied(const FRootMotionSource* Source)
 {
 }
diff --git a/Engine/Source/Runtime/Engine/Private/Components/SceneComponent.cpp b/Engine/Source/Runtime/Engine/Private/Components/SceneComponent.cpp
--- a/Engine/Source/Runtime/Engine/Private/Components/SceneComponent.cpp
+++ b/Engine/Source/Runtime/Engine/Private/Components/SceneComponent.cpp
@@ -1881,17 +1881,6 @@ void USceneComponent::SetupAttachment(class USceneComponent* InParent, FName InS
 	}
 }
 
-//This function is used for giving AttachTo different bWeldSimulatedBodies default, but only when called from BP
-bool USceneComponent::K2_AttachTo(class USceneComponent* InParent, FName InSocketName, EAttachLocation::Type AttachLocationType, bool bWeldSimulatedBodies /*= true*/)
-{
-PRAGMA_DISABLE_DEPRECATION_WARNINGS
-	FAttachmentTransformRules AttachmentRules(EAttachmentRule::KeepRelative, bWeldSimulatedBodies);
-	ConvertAttachLocation(AttachLocationType, AttachmentRules.LocationRule, AttachmentRules.RotationRule, AttachmentRules.ScaleRule);
-
-	return AttachToComponent(InParent, AttachmentRules, InSocketName);
-PRAGMA_ENABLE_DEPRECATION_WARNINGS
-}
-
 //This function is used for giving AttachToComponent different bWeldSimulatedBodies default, but only when called from BP
 bool USceneComponent::K2_AttachToComponent(USceneComponent* Parent, FName SocketName, EAttachmentRule LocationRule, EAttachmentRule RotationRule, EAttachmentRule ScaleRule, bool bWeldSimulatedBodies)
 {
diff --git a/Engine/Source/Runtime/Engine/Private/EdGraph/EdGraphNode.cpp b/Engine/Source/Runtime/Engine/Private/EdGraph/EdGraphNode.cpp
--- a/Engine/Source/Runtime/Engine/Private/EdGraph/EdGraphNode.cpp
+++ b/Engine/Source/Runtime/Engine/Private/EdGraph/EdGraphNode.cpp
@@ -350,26 +350,6 @@ UEdGraphPin* UEdGraphNode::CreatePin(EEdGraphPinDirection Dir, const FEdGraphPin
 	return NewPin;
 }
 
-UEdGraphPin* UEdGraphNode::CreatePin(EEdGraphPinDirection Dir, const FNameParameterHelper PinCategory, const FNameParameterHelper PinSubCategory, UObject* PinSubCategoryObject, bool bIsArray, bool bIsReference, const FNameParameterHelper PinName, bool bIsConst /*= false*/, int32 Index /*= INDEX_NONE*/, bool bIsSet /*= false*/, bool bIsMap /*= false*/, const FEdGraphTerminalType& ValueTerminalType /*= FEdGraphTerminalType()*/)
-{
-	PRAGMA_DISABLE_DEPRECATION_WARNINGS
-	return CreatePin(Dir, PinCategory, PinSubCategory, PinSubCategoryObject, PinName, FEdGraphPinType::ToPinContainerType(bIsArray, bIsSet, bIsMap), bIsReference, bIsConst, Index, ValueTerminalType);
-	PRAGMA_ENABLE_DEPRECATION_WARNINGS
-}
-
-UEdGraphPin* UEdGraphNode::CreatePin(EEdGraphPinDirection Dir, const FNameParameterHelper PinCategory, const FNameParameterHelper PinSubCategory, UObject* PinSubCategoryObject, const FNameParameterHelper PinName, EPinContainerType PinContainerType /* EPinContainerType::None */, bool bIsReference /* = false */, bool bIsConst /*= false*/, int32 Index /*= INDEX_NONE*/, const FEdGraphTerminalType& ValueTerminalType /*= FEdGraphTerminalType()*/)
-{
-	FCreatePinParams PinParams;
-	PinParams.ContainerType = PinContainerType;
-	PinParams.bIsConst = bIsConst;
-	PinParams.bIsReference = bIsReference;
-	PinParams.Index = Index;
-	PinParams.ValueTerminalType = ValueTerminalType;
-
-	return CreatePin(Dir, *PinCategory, *PinSubCategory, PinSubCategoryObject, *PinName, PinParams);
-}
-
-
 UEdGraphPin* UEdGraphNode::CreatePin(const EEdGraphPinDirection Dir, const FName PinCategory, const FName PinSubCategory, UObject* PinSubCategoryObject, const FName PinName, const FCreatePinParams& PinParams)
 {
 	FEdGraphPinType PinType(PinCategory, PinSubCategory, PinSubCategoryObject, PinParams.ContainerType, PinParams.bIsReference, PinParams.ValueTerminalType);
diff --git a/Engine/Source/Runtime/Engine/Private/GameModeBase.cpp b/Engine/Source/Runtime/Engine/Private/GameModeBase.cpp
--- a/Engine/Source/Runtime/Engine/Private/GameModeBase.cpp
+++ b/Engine/Source/Runtime/Engine/Private/GameModeBase.cpp
@@ -482,15 +482,9 @@ void AGameModeBase::ProcessServerTravel(const FString& URL, bool bAbsolute)
 		}
 	}
 
-	PRAGMA_DISABLE_DEPRECATION_WARNINGS
-	FGuid NextMapGuid = UEngine::GetPackageGuid(FName(*NextURL.Map), GetWorld()->IsPlayInEditor());
-	PRAGMA_ENABLE_DEPRECATION_WARNINGS
-
 	// Notify clients we're switching level and give them time to receive.
 	FString URLMod = NextURL.ToString();
-	PRAGMA_DISABLE_DEPRECATION_WARNINGS
-	APlayerController* LocalPlayer = ProcessClientTravel(URLMod, NextMapGuid, bSeamless, bAbsolute);
-	PRAGMA_ENABLE_DEPRECATION_WARNINGS
+	APlayerController* LocalPlayer = ProcessClientTravel(URLMod, bSeamless, bAbsolute);
 
 	World->NextURL = URLMod;
 	ENetMode NetMode = GetNetMode();
diff --git a/Engine/Source/Runtime/Engine/Private/GameUserSettings.cpp b/Engine/Source/Runtime/Engine/Private/GameUserSettings.cpp
--- a/Engine/Source/Runtime/Engine/Private/GameUserSettings.cpp
+++ b/Engine/Source/Runtime/Engine/Private/GameUserSettings.cpp
@@ -328,15 +328,6 @@ void UGameUserSettings::SetFrameRateLimitCVar(float InLimit)
 	GEngine->SetMaxFPS(FMath::Max(InLimit, 0.0f));
 }
 
-void UGameUserSettings::SetSyncIntervalCVar(int32 InInterval)
-{
-	static IConsoleVariable* SyncIntervalCVar = IConsoleManager::Get().FindConsoleVariable(TEXT("rhi.syncinterval"));
-	if (ensure(SyncIntervalCVar))
-	{
-		SyncIntervalCVar->Set(InInterval, ECVF_SetByCode);
-	}
-}
-
 void UGameUserSettings::SetSyncTypeCVar(int32 InType)
 {
 	static IConsoleVariable* SyncIntervalCVar = IConsoleManager::Get().FindConsoleVariable(TEXT("r.GTSyncType"));
@@ -766,14 +757,6 @@ int32 UGameUserSettings::GetOverallScalabilityLevel() const
 	return ScalabilityQuality.GetSingleQualityLevel();
 }
 
-void UGameUserSettings::GetResolutionScaleInformation(float& CurrentScaleNormalized, int32& CurrentScaleValue, int32& MinScaleValue, int32& MaxScaleValue) const
-{
-	CurrentScaleValue = ScalabilityQuality.ResolutionQuality;
-	MinScaleValue = MinResolutionScale;
-	MaxScaleValue = Scalability::MaxResolutionScale;
-	CurrentScaleNormalized = ((float)CurrentScaleValue - (float)MinScaleValue) / (float)(MaxScaleValue - MinScaleValue);
-}
-
 void UGameUserSettings::GetResolutionScaleInformationEx(float& CurrentScaleNormalized, float& CurrentScaleValue, float& MinScaleValue, float& MaxScaleValue) const
 {
 	CurrentScaleValue = ScalabilityQuality.ResolutionQuality;
@@ -790,11 +773,6 @@ float UGameUserSettings::GetResolutionScaleNormalized() const
 	return CurrentScaleNormalized;
 }
 
-void UGameUserSettings::SetResolutionScaleValue(int32 NewScaleValue)
-{
-	SetResolutionScaleValueEx((float)NewScaleValue);
-}
-
 void UGameUserSettings::SetResolutionScaleValueEx(float NewScaleValue)
 {
 	ScalabilityQuality.ResolutionQuality = FMath::Clamp(NewScaleValue, MinResolutionScale, Scalability::MaxResolutionScale);
diff --git a/Engine/Source/Runtime/Engine/Private/KismetMathLibrary.cpp b/Engine/Source/Runtime/Engine/Private/KismetMathLibrary.cpp
--- a/Engine/Source/Runtime/Engine/Private/KismetMathLibrary.cpp
+++ b/Engine/Source/Runtime/Engine/Private/KismetMathLibrary.cpp
@@ -417,11 +417,6 @@ void UKismetMathLibrary::MinOfByteArray(const TArray<uint8>& ByteArray, int32& I
 	MinValue = FMath::Min(ByteArray, &IndexOfMinValue);
 }
 
-float UKismetMathLibrary::InverseLerp(float A, float B, float Value)
-{
-	return NormalizeToRange(Value, A, B);
-}
-
 double UKismetMathLibrary::Ease(double A, double B, double Alpha, TEnumAsByte<EEasingFunc::Type> EasingFunc, double BlendExp, int32 Steps)
 {
 	return Lerp(A, B, EaseAlpha(Alpha, EasingFunc, BlendExp, Steps));
diff --git a/Engine/Source/Runtime/Engine/Private/PlayerController.cpp b/Engine/Source/Runtime/Engine/Private/PlayerController.cpp
--- a/Engine/Source/Runtime/Engine/Private/PlayerController.cpp
+++ b/Engine/Source/Runtime/Engine/Private/PlayerController.cpp
@@ -1896,11 +1896,6 @@ void APlayerController::ClientReturnToMainMenuWithTextReason_Implementation(cons
 	}
 }
 
-void APlayerController::ClientReturnToMainMenu_Implementation(const FString& ReturnReason)
-{
-	ClientReturnToMainMenuWithTextReason_Implementation(FText::FromString(ReturnReason));
-}
-
 bool APlayerController::SetPause( bool bPause, FCanUnpause CanUnpauseDelegate)
 {
 	bool bResult = false;
diff --git a/Engine/Source/Runtime/Engine/Private/TimelineTemplate.cpp b/Engine/Source/Runtime/Engine/Private/TimelineTemplate.cpp
--- a/Engine/Source/Runtime/Engine/Private/TimelineTemplate.cpp
+++ b/Engine/Source/Runtime/Engine/Private/TimelineTemplate.cpp
@@ -153,13 +153,6 @@ bool UTimelineTemplate::IsNewTrackNameValid(const FName NewTrackName) const
 			FindLinearColorTrackIndex(NewTrackName) == INDEX_NONE;
 }
 
-FName UTimelineTemplate::GetEventTrackFunctionName(int32 EventTrackIndex) const
-{
-	check(EventTrackIndex < EventTracks.Num());
-
-	return EventTracks[EventTrackIndex].GetFunctionName();
-}
-
 int32 UTimelineTemplate::FindMetaDataEntryIndexForKey(const FName Key) const
 {
 	for(int32 i=0; i<MetaDataArray.Num(); i++)
diff --git a/Engine/Source/Runtime/Engine/Private/World.cpp b/Engine/Source/Runtime/Engine/Private/World.cpp
--- a/Engine/Source/Runtime/Engine/Private/World.cpp
+++ b/Engine/Source/Runtime/Engine/Private/World.cpp
@@ -3453,7 +3453,6 @@ void UWorld::BeginTearingDown()
 {
 	bIsTearingDown = true;
 	UE_LOG(LogWorld, Log, TEXT("BeginTearingDown for %s"), *GetOutermost()->GetName());
-	BeginTearingDownEvent.Broadcast();
 
 	//Simultaneous similar edits that caused merge conflict. Taking both for now to unblock.
 	//Can likely be unified.
@@ -5597,35 +5596,6 @@ TUniquePtr<FPawnIteratorObject> FConstPawnIterator::operator->() const
 	return TUniquePtr<FPawnIteratorObject>(new FPawnIteratorObject(**Iterator));
 }
 
-FConstPawnIterator UWorld::GetPawnIterator() const
-{
-	return FConstPawnIterator(const_cast<UWorld*>(this)); // For backwards compat GetPawnIterator needs to remain const, but TActorIterator can't use a const UWorld.
-}
-
-int32 UWorld::GetNumPawns() const
-{
-	int32 NumPawns = 0;
-	for (APawn* Pawn : TActorRange<APawn>(const_cast<UWorld*>(this))) // For backwards compat GetNumPawns needs to remain const, but TActorRange can't use a const UWorld.
-	{
-		++NumPawns;
-	}
-
-	return NumPawns;
-}
-
-void UWorld::RemovePawn( APawn* Pawn ) const
-{
-	if (Pawn)
-	{
-		AController* Controller = Pawn->GetController();
-		if (Controller && (Controller->GetPawn() == Pawn))
-		{
-			Controller->UnPossess();
-		}
-	}
-}
-
-
 void UWorld::RegisterAutoActivateCamera(ACameraActor* CameraActor, int32 PlayerIndex)
 {
 	check(CameraActor);
diff --git a/Engine/Source/Runtime/Engine/Public/DataTableUtils.h b/Engine/Source/Runtime/Engine/Public/DataTableUtils.h
--- a/Engine/Source/Runtime/Engine/Public/DataTableUtils.h
+++ b/Engine/Source/Runtime/Engine/Public/DataTableUtils.h
@@ -20,10 +20,7 @@ enum class EDataTableExportFlags : uint8
 	/** Export text properties as their display string, rather than their complex lossless form. */
 	UseSimpleText = 1 << 1,
 
-	// DEPRECATED. Native properties/enums are always exported using their internal name, user struct/enums are always exported using the friendly names set in the editor
-
-	UsePrettyPropertyNames UE_DEPRECATED(4.23, "UsePrettyPropertyNames is deprecated, we now always use the unlocalized but readable authored names") = 1 << 6,
-	UsePrettyEnumNames UE_DEPRECATED(4.23, "UsePrettyEnumNames is deprecated, we now always use the unlocalized but readable authored names") = 1 << 7,
+	// Deprecated options removed, Native properties/enums are always exported using their internal name, user struct/enums are always exported using the friendly names set in the editor
 };
 ENUM_CLASS_FLAGS(EDataTableExportFlags);
 
diff --git a/Engine/Source/Runtime/Engine/Public/ViewportClient.h b/Engine/Source/Runtime/Engine/Public/ViewportClient.h
--- a/Engine/Source/Runtime/Engine/Public/ViewportClient.h
+++ b/Engine/Source/Runtime/Engine/Public/ViewportClient.h
@@ -101,9 +101,6 @@ public:
 	 */
 	virtual bool InputTouch(FViewport* Viewport, int32 ControllerId, uint32 Handle, ETouchType::Type Type, const FVector2D& TouchLocation, float Force, FDateTime DeviceTimestamp, uint32 TouchpadIndex) { return false; }
 
-	UE_DEPRECATED(4.20, "InputTouch now takes a Force")
-	bool InputTouch(FViewport* Viewport, int32 ControllerId, uint32 Handle, ETouchType::Type Type, const FVector2D& TouchLocation, FDateTime DeviceTimestamp, uint32 TouchpadIndex) { return InputTouch(Viewport, ControllerId, Handle, Type, TouchLocation, 1.0f, DeviceTimestamp, TouchpadIndex); }
-
 	/**
 	 * Check a gesture event received by the viewport.
 	 * If the viewport client uses the event, it should return true to consume it.
@@ -290,9 +287,6 @@ public:
 	 */
 	virtual EMouseCaptureMode GetMouseCaptureMode() const { return EMouseCaptureMode::CapturePermanently; }
 
-	UE_DEPRECATED(4.26, "Please call GetMouseCaptureMode() instead.")
-	void CaptureMouseOnClick() { GetMouseCaptureMode(); }
-
 	/**
 	 * Gets whether or not the viewport captures the Mouse on launch of the application
 	 * Technically this controls capture on the first window activate, so in situations
diff --git a/Engine/Source/Runtime/GameplayTags/Classes/GameplayTagContainer.h b/Engine/Source/Runtime/GameplayTags/Classes/GameplayTagContainer.h
--- a/Engine/Source/Runtime/GameplayTags/Classes/GameplayTagContainer.h
+++ b/Engine/Source/Runtime/GameplayTags/Classes/GameplayTagContainer.h
@@ -727,21 +727,6 @@ struct GAMEPLAYTAGS_API FGameplayTagNativeAdder
 	virtual void AddTags() = 0;
 };
 
-USTRUCT()
-struct 
-	UE_DEPRECATED(4.26, "FGameplayTagReferenceHelper has been deprecated, to find references right click a normal gameplay tag")
-	FGameplayTagReferenceHelper
-{
-	GENERATED_USTRUCT_BODY()
-
-	FGameplayTagReferenceHelper()
-	{
-	}
-
-	DECLARE_DELEGATE_RetVal_OneParam(FName, FOnGetGameplayTagName, void* /**RawOuterStructData*/);
-	FOnGetGameplayTagName OnGetGameplayTagName;
-};
-
 /** Helper struct: drop this in another struct to get an embedded create new tag widget. */
 USTRUCT()
 struct FGameplayTagCreationWidgetHelper
diff --git a/Engine/Source/Runtime/GameplayTags/Classes/GameplayTagsManager.h b/Engine/Source/Runtime/GameplayTags/Classes/GameplayTagsManager.h
--- a/Engine/Source/Runtime/GameplayTags/Classes/GameplayTagsManager.h
+++ b/Engine/Source/Runtime/GameplayTags/Classes/GameplayTagsManager.h
@@ -647,10 +647,6 @@ public:
 		return FString();
 	}
 
-	/** Returns "Categories" meta property from given struct, used for filtering by tag widget */
-	UE_DEPRECATED(4.22, "Please call GetCategoriesMetaFromField instead.")
-	FString GetCategoriesMetaFromStruct(UScriptStruct* Struct) const { return GetCategoriesMetaFromField(Struct); }
-
 	/** Returns "GameplayTagFilter" meta property from given function, used for filtering by tag widget for any parameters of the function that end up as BP pins */
 	FString GetCategoriesMetaFromFunction(const UFunction* Func, FName ParamName = NAME_None) const;
 
