Avoid applying actor editor context to preview actors being dragged from content browser.

diff --git a/Engine/Source/Editor/UnrealEd/Private/Subsystems/ActorEditorContextSubsystem.cpp b/Engine/Source/Editor/UnrealEd/Private/Subsystems/ActorEditorContextSubsystem.cpp
--- a/Engine/Source/Editor/UnrealEd/Private/Subsystems/ActorEditorContextSubsystem.cpp
+++ b/Engine/Source/Editor/UnrealEd/Private/Subsystems/ActorEditorContextSubsystem.cpp
@@ -51,7 +51,12 @@ void UActorEditorContextSubsystem::UnregisterClient(IActorEditorContextClient* C
 void UActorEditorContextSubsystem::ApplyContext(AActor* InActor)
 {
 	UWorld* World = GetWorld();
-	if (Clients.IsEmpty() || !World || !InActor || (InActor->GetWorld() != World) || InActor->IsChildActor())
+	if (Clients.IsEmpty() || !World)
+	{
+		return;
+	}
+
+	if (!InActor || (InActor->GetWorld() != World) || InActor->HasAnyFlags(RF_Transient) || InActor->IsChildActor())
 	{
 		return;
 	}
