Fix subchunking performed at end of cook when using ZenServer as cooked output store.

diff --git a/Engine/Source/Editor/UnrealEd/Private/Commandlets/AssetRegistryGenerator.cpp b/Engine/Source/Editor/UnrealEd/Private/Commandlets/AssetRegistryGenerator.cpp
--- a/Engine/Source/Editor/UnrealEd/Private/Commandlets/AssetRegistryGenerator.cpp
+++ b/Engine/Source/Editor/UnrealEd/Private/Commandlets/AssetRegistryGenerator.cpp
@@ -283,51 +283,6 @@ TArray<int32> FAssetRegistryGenerator::GetExplicitChunkIDs(const FName& PackageF
 	return PackageInputChunkIds;
 }
 
-class FPackageFileSizeVisitor : public IPlatformFile::FDirectoryStatVisitor
-{
-	TMap<FString, int64>& PackageFileSizes;
-public:
-	FPackageFileSizeVisitor(TMap<FString, int64>& InFileSizes)
-		: PackageFileSizes(InFileSizes)
-	{}
-	virtual bool Visit(const TCHAR* FilenameOrDirectory, const FFileStatData& StatData)
-	{
-		const TCHAR* Extensions[] = { TEXT(".uexp"), TEXT(".uasset"), TEXT(".ubulk"), TEXT(".ufont"), TEXT(".umap"), TEXT(".uptnl") };
-
-		if (StatData.bIsDirectory)
-			return true;
-
-		const TCHAR* Extension = FCString::Strrchr(FilenameOrDirectory, '.');
-		if (!Extension)
-			return true;
-
-		int32 ExtIndex = 0;
-		for (; ExtIndex < UE_ARRAY_COUNT(Extensions); ++ExtIndex)
-		{
-			if (0 == FCString::Stricmp(Extension, Extensions[ExtIndex]))
-				break;
-		}
-
-		if (ExtIndex >= UE_ARRAY_COUNT(Extensions))
-			return true;
-
-		int32 LengthWithoutExtension = UE_PTRDIFF_TO_INT32(Extension - FilenameOrDirectory);
-		FString FilenameWithoutExtension(LengthWithoutExtension, FilenameOrDirectory);
-
-		if (int64* CurrentPackageSize = PackageFileSizes.Find(FilenameWithoutExtension))
-		{
-			int64& TotalPackageSize = *CurrentPackageSize;
-			TotalPackageSize += StatData.FileSize;
-		}
-		else
-		{
-			PackageFileSizes.Add(FilenameWithoutExtension, StatData.FileSize);
-		}
-
-		return true;
-	}
-};
-
 static void ParseChunkLayerAssignment(TArray<FString> ChunkLayerAssignmentArray, TMap<int32, int32>& OutChunkLayerAssignment)
 {
 	OutChunkLayerAssignment.Empty();
@@ -466,16 +421,6 @@ bool FAssetRegistryGenerator::GenerateStreamingInstallManifest(int64 InOverrideC
 			}
 		}
 	}
-
-	TMap<FString, int64> PackageFileSizes;
-	if (MaxChunkSize > 0)
-	{
-		// ZENTODO: Pull file size information from the AssetRegistryState, since the packages do not exist on disk in zen
-		FString SandboxPath = InSandboxFile.GetSandboxDirectory();
-		SandboxPath.ReplaceInline(TEXT("[Platform]"), *Platform);
-		FPackageFileSizeVisitor PackageSearch(PackageFileSizes);
-		IFileManager::Get().IterateDirectoryStatRecursively(*SandboxPath, PackageSearch);
-	}
 	
 	bool bEnableGameOpenOrderSort = false;
 	bool bUseSecondaryOpenOrder = false;
@@ -607,6 +552,7 @@ bool FAssetRegistryGenerator::GenerateStreamingInstallManifest(int64 InOverrideC
 	// generate per-chunk pak list files
 	FDefaultPakFileRules DefaultPakFileRules;
 	bool bSucceeded = true;
+	TMap<FString, int64> PackageFileSizes;
 	for (int32 PakchunkIndex = 0; PakchunkIndex < FinalChunkManifests.Num() && bSucceeded; ++PakchunkIndex)
 	{
 		const FChunkPackageSet* Manifest = FinalChunkManifests[PakchunkIndex].Get();
@@ -631,6 +577,21 @@ bool FAssetRegistryGenerator::GenerateStreamingInstallManifest(int64 InOverrideC
 		TArray<FString> ChunkFilenames;
 		Manifest->GenerateValueArray(ChunkFilenames);
 
+		if (MaxChunkSize > 0)
+		{
+			PackageFileSizes.Reset();
+			PackageFileSizes.Reserve(Manifest->Num());
+			const TMap<FName, const FAssetPackageData*>& PackageDataMap = State.GetAssetPackageDataMap();
+			for (const TPair<FName, FString>& ManifestEntry : *Manifest)
+			{
+				const FAssetPackageData* const* ExistingPackageData = PackageDataMap.Find(ManifestEntry.Key);
+				if (ExistingPackageData)
+				{
+					PackageFileSizes.Emplace(ManifestEntry.Value.Replace(TEXT("[Platform]"), *Platform), (*ExistingPackageData)->DiskSize);
+				}
+			}
+		}
+
 		// Do not create any files if the chunk is empty and is not referenced by rules applied during staging
 		if (ChunkFilenames.IsEmpty())
 		{
