Fixed the class picker layout when mixing rows with and without documentation.
- Kept the same widths with or without the docs widget.

diff --git a/Engine/Source/Editor/UnrealEd/Private/SClassPickerDialog.cpp b/Engine/Source/Editor/UnrealEd/Private/SClassPickerDialog.cpp
--- a/Engine/Source/Editor/UnrealEd/Private/SClassPickerDialog.cpp
+++ b/Engine/Source/Editor/UnrealEd/Private/SClassPickerDialog.cpp
@@ -282,7 +282,7 @@ TSharedRef<ITableRow> SClassPickerDialog::GenerateListRow(TSharedPtr<FClassPicke
 			SNew(SHorizontalBox)
 			+SHorizontalBox::Slot()
 			.VAlign(VAlign_Center)
-			.FillWidth(0.5f)
+			.FillWidth(0.3f)
 			[
 				SNew(SButton)
 				.OnClicked(this, &SClassPickerDialog::OnDefaultClassPicked, ItemClass)
@@ -310,6 +310,7 @@ TSharedRef<ITableRow> SClassPickerDialog::GenerateListRow(TSharedPtr<FClassPicke
 			+SHorizontalBox::Slot()
 			.VAlign(VAlign_Center)
 			.Padding(10.0f, 0.0f, 4.0f, 0.0f)
+			.FillWidth(0.63f)
 			[
 				SNew(STextBlock)
 				.Text(Obj->GetDescription())
@@ -319,10 +320,10 @@ TSharedRef<ITableRow> SClassPickerDialog::GenerateListRow(TSharedPtr<FClassPicke
 			]
 			+SHorizontalBox::Slot()
 			.VAlign(VAlign_Center)
-			.AutoWidth()
+			.HAlign(HAlign_Center)
+			.FillWidth(0.07f)
 			[
 				SNew(SBox)
-				.WidthOverride(DocumentationWidget != SNullWidget::NullWidget ? 60.0f : 0.0f)
 				[
 					DocumentationWidget
 				]
