Copying and duplicating temporary actors while PIE is active is disabled.

In response to reports that copy/pasting actors during PIE causing crashes, copying and duplicating temporary actors has been disabled. This is a temporary measure until the copy/paste system during PIE be evaluated more closely.

diff --git a/Engine/Source/Editor/UnrealEd/Private/EditorServer.cpp b/Engine/Source/Editor/UnrealEd/Private/EditorServer.cpp
--- a/Engine/Source/Editor/UnrealEd/Private/EditorServer.cpp
+++ b/Engine/Source/Editor/UnrealEd/Private/EditorServer.cpp
@@ -3348,6 +3348,19 @@ bool UEditorEngine::CanCopySelectedActorsToClipboard( UWorld* InWorld, FCopySele
 				CopySelected.LevelAllActorsAreIn = Actor->GetLevel();
 			}
 
+			if (Actor->GetLevel())
+			{
+				if (UWorld* ActorWorld = Actor->GetLevel()->GetWorld())
+				{
+					// If the actor is in a PIE world but doesn't have an editor counterpart it means it's a temporary
+					// actor spawned to the world. These actors can cause issues when copied so have been disabled.
+					if (ActorWorld->WorldType == EWorldType::PIE && !GEditor->ObjectsThatExistInEditorWorld.Get(Actor))
+					{
+						return false;
+					}
+				}
+			}
+
 			if( Actor->GetLevel() != CopySelected.LevelAllActorsAreIn )
 			{
 				CopySelected.bAllActorsInSameLevel = false;
diff --git a/Engine/Source/Editor/UnrealEd/Private/Elements/Actor/ActorElementEditorWorldInterface.cpp b/Engine/Source/Editor/UnrealEd/Private/Elements/Actor/ActorElementEditorWorldInterface.cpp
--- a/Engine/Source/Editor/UnrealEd/Private/Elements/Actor/ActorElementEditorWorldInterface.cpp
+++ b/Engine/Source/Editor/UnrealEd/Private/Elements/Actor/ActorElementEditorWorldInterface.cpp
@@ -81,7 +81,19 @@ bool UActorElementEditorWorldInterface::DeleteElements(TArrayView<const FTypedEl
 bool UActorElementEditorWorldInterface::CanDuplicateElement(const FTypedElementHandle& InElementHandle)
 {
 	AActor* Actor = ActorElementDataUtil::GetActorFromHandle(InElementHandle);
-	return Actor != nullptr; // All actors can be duplicated
+	if (Actor && Actor->GetLevel())
+	{
+		if (UWorld* ActorWorld = Actor->GetLevel()->GetWorld())
+		{
+			// If the actor is in a PIE world but doesn't have an editor counterpart it means it's a temporary
+			// actor spawned to the world. These actors can cause issues when copied so have been disabled.
+			if (ActorWorld->WorldType == EWorldType::PIE && !GEditor->ObjectsThatExistInEditorWorld.Get(Actor))
+			{
+				return false;
+			}
+		}
+	}
+	return Actor != nullptr;
 }
 
 void UActorElementEditorWorldInterface::DuplicateElements(TArrayView<const FTypedElementHandle> InElementHandles, UWorld* InWorld, const FVector& InLocationOffset, TArray<FTypedElementHandle>& OutNewElements)
@@ -126,6 +138,19 @@ void UActorElementEditorWorldInterface::DuplicateElements(TArrayView<const FType
 
 bool UActorElementEditorWorldInterface::CanCopyElement(const FTypedElementHandle& InElementHandle)
 {
+	AActor* Actor = ActorElementDataUtil::GetActorFromHandle(InElementHandle);
+	if (Actor && Actor->GetLevel())
+	{
+		if (UWorld* ActorWorld = Actor->GetLevel()->GetWorld())
+		{
+			// If the actor is in a PIE world but doesn't have an editor counterpart it means it's a temporary
+			// actor spawned to the world. These actors can cause issues when copied so have been disabled.
+			if (ActorWorld->WorldType == EWorldType::PIE && !GEditor->ObjectsThatExistInEditorWorld.Get(Actor))
+			{
+				return false;
+			}
+		}
+	}
 	return true;
 }
 
