Adding a new Cook Metadata file that gets generated adjacent to the development asset registry. This will be the home for metadata about the cook that is unrelated to assets moving forward.

diff --git a/Engine/Source/Developer/CookMetadata/CookMetadata.Build.cs b/Engine/Source/Developer/CookMetadata/CookMetadata.Build.cs
new file mode 100644
--- /dev/null
+++ b/Engine/Source/Developer/CookMetadata/CookMetadata.Build.cs
@@ -0,0 +1,15 @@
+// Copyright Epic Games, Inc. All Rights Reserved.
+
+using UnrealBuildTool;
+
+public class CookMetadata : ModuleRules
+{
+	public CookMetadata(ReadOnlyTargetRules Target) : base(Target)
+	{
+		PrivateDependencyModuleNames.AddRange(
+			new string[] {
+				"Core"
+			}
+		);
+	}
+}
diff --git a/Engine/Source/Developer/CookMetadata/Private/CookMetadata.cpp b/Engine/Source/Developer/CookMetadata/Private/CookMetadata.cpp
new file mode 100644
--- /dev/null
+++ b/Engine/Source/Developer/CookMetadata/Private/CookMetadata.cpp
@@ -0,0 +1,49 @@
+// Copyright Epic Games, Inc. All Rights Reserved.
+
+#include "CookMetadata.h"
+
+#include "Hash/xxhash.h"
+#include "Memory/MemoryView.h"
+
+namespace UE::Cook
+{
+
+const FString& GetCookMetadataFilename()
+{
+	static const FString CookMetadataFilename(TEXT("CookMetadata.ucookmeta"));
+	return CookMetadataFilename;
+}
+
+constexpr uint32 COOK_METADATA_HEADER_MAGIC = 'UCMT';
+
+bool FCookMetadataState::Serialize(FArchive& Ar)
+{
+	uint32 MagicHeader = 0;
+	if (Ar.IsLoading())
+	{
+		Ar << MagicHeader;
+		if (MagicHeader != COOK_METADATA_HEADER_MAGIC)
+		{
+			return false; // not a metadata file.
+		}
+	}
+	else
+	{
+		MagicHeader = COOK_METADATA_HEADER_MAGIC;
+		Ar << MagicHeader;
+	}
+	Ar << Version;
+	Ar << PluginHierarchy;
+	Ar << AssociatedDevelopmentAssetRegistryHash;
+
+	return true;
+}
+
+/* static */ 
+uint64 FCookMetadataState::ComputeHashOfDevelopmentAssetRegistry(FMemoryView InSerializedDevelopmentAssetRegistry)
+{
+	return FXxHash64::HashBufferChunked(InSerializedDevelopmentAssetRegistry.GetData(), InSerializedDevelopmentAssetRegistry.GetSize(), 1ULL << 19).Hash;
+}
+
+
+} // end namespace
\ No newline at end of file
diff --git a/Engine/Source/Developer/CookMetadata/Public/CookMetadata.h b/Engine/Source/Developer/CookMetadata/Public/CookMetadata.h
new file mode 100644
--- /dev/null
+++ b/Engine/Source/Developer/CookMetadata/Public/CookMetadata.h
@@ -0,0 +1,118 @@
+// Copyright Epic Games, Inc. All Rights Reserved.
+
+#pragma once
+
+#include "Containers/Array.h"
+#include "Containers/UnrealString.h"
+#include "CoreTypes.h"
+#include "Memory/MemoryFwd.h"
+
+
+namespace UE::Cook
+{
+
+enum class ECookMetadataStateVersion : uint8
+{
+	PluginHierarchy = 1,
+
+	// Add new versions above this.
+	VersionCount,
+	LatestVersion = VersionCount - 1
+};
+
+
+/** The name and dependency information for a plugin that was enabled during cooking. */
+struct COOKMETADATA_API FCookMetadataPluginEntry
+{
+	FString Name;
+
+	// The dependencies are stored in the FCookMetadataPluginHierarchy::PluginDependencies array,
+	// and this is an index into it. From there you can get a further index into PluginsEnabledAtCook
+	// to get the plugin information.
+	// Example:
+	//	for (uint16 DependencyIndex = Plugin->DependencyIndexStart; DependencyIndex < Plugin->DependencyIndexEnd; DependencyIndex++)
+	//	{
+	//		const UE::Cook::FCookMetadataPluginEntry& DependentPlugin = PluginHierarchy.PluginsEnabledAtCook[PluginHierarchy.PluginDependencies[DependencyIndex]];
+	//	}
+	//
+	uint16 DependencyIndexStart;
+	uint16 DependencyIndexEnd;
+
+	uint16 DependencyCount() const { return DependencyIndexEnd - DependencyIndexStart; }
+
+	friend FArchive& operator<<(FArchive& Ar, FCookMetadataPluginEntry& Entry)
+	{
+		return Ar << Entry.Name << Entry.DependencyIndexStart << Entry.DependencyIndexEnd;
+	}
+};
+
+struct COOKMETADATA_API FCookMetadataPluginHierarchy
+{
+	// The list of plugins that were enabled during the cook that generated the FCookMetadataState,
+	// pruned to remove plugins that aren't enabled on the cook platform. If the cook was DLC, this
+	// is further pruned to only the DLC plugin and its dependencies.
+	TArray<FCookMetadataPluginEntry> PluginsEnabledAtCook;
+
+	// The list of plugin dependencies. FCookMetadataPluginEntry::DependencyIndexStart indexes into this
+	// array.
+	TArray<uint16> PluginDependencies;
+
+	// The list of root plugins for the project as defined by the Editor.ini file.
+	TArray<uint16> RootPlugins;
+
+	friend FArchive& operator<<(FArchive& Ar, FCookMetadataPluginHierarchy& Hierarchy)
+	{
+		return Ar << Hierarchy.PluginsEnabledAtCook << Hierarchy.PluginDependencies << Hierarchy.RootPlugins;
+	}
+};
+
+/**
+*	Structure serialized to disk to contain non-asset related metadata about a cook. This should always
+*	exist alongside a Development Asset Registry, and to ensure that the pair is not out of sync, users
+*	should validate the development asset registry they are using with GetAssociatedDevelopmentAssetRegistryHash().
+*/
+class COOKMETADATA_API FCookMetadataState
+{
+public:
+	FCookMetadataState() = default;
+	FCookMetadataState(const FCookMetadataState&) = delete;
+	FCookMetadataState(FCookMetadataState&& Rhs) = delete;
+	~FCookMetadataState() = default;
+
+	FCookMetadataState& operator=(const FCookMetadataState&) = delete;
+	FCookMetadataState& operator=(FCookMetadataState&& O) = delete;
+
+	
+	bool Serialize(FArchive& Ar);
+
+	// Plugin hierarchy information
+	void SetPluginHierarchyInfo(FCookMetadataPluginHierarchy&& InPluginHierarchy) { PluginHierarchy = MoveTemp(InPluginHierarchy); }
+	const FCookMetadataPluginHierarchy& GetPluginHierarchy() const { return PluginHierarchy; }
+
+	/**
+	*	Associated DevAR Hash.
+	* 
+	*	This is computed by reading the development asset registry file into a memory buffer and calling
+	*	ComputeHashOfDevelopmentAssetRegistry on the data.
+	* 
+	*	Use this to ensure that the files you are working with were produced by the same cook and didn't
+	*	get out of sync somehow.
+	* 
+	*	e.g.
+	*	uint64 CheckHash = GetAssociatedDevelopmentAssetRegistryHash();
+	*	bValidDevAr = ComputeHashOfDevelopmentAssetRegistry(MakeMemoryView(SerializedAssetRegistry)) == CheckHash;
+	*/
+	void SetAssociatedDevelopmentAssetRegistryHash(uint64 InHash) { AssociatedDevelopmentAssetRegistryHash = InHash; }
+	uint64 GetAssociatedDevelopmentAssetRegistryHash() const { return AssociatedDevelopmentAssetRegistryHash; }
+
+	static uint64 ComputeHashOfDevelopmentAssetRegistry(FMemoryView InSerializedDevelopmentAssetRegistry);
+private:
+	ECookMetadataStateVersion Version;
+	FCookMetadataPluginHierarchy PluginHierarchy;
+
+	uint64 AssociatedDevelopmentAssetRegistryHash;
+};
+
+COOKMETADATA_API const FString& GetCookMetadataFilename();
+
+}; // namespace
\ No newline at end of file
diff --git a/Engine/Source/Editor/UnrealEd/Classes/CookOnTheSide/CookOnTheFlyServer.h b/Engine/Source/Editor/UnrealEd/Classes/CookOnTheSide/CookOnTheFlyServer.h
--- a/Engine/Source/Editor/UnrealEd/Classes/CookOnTheSide/CookOnTheFlyServer.h
+++ b/Engine/Source/Editor/UnrealEd/Classes/CookOnTheSide/CookOnTheFlyServer.h
@@ -1217,6 +1217,10 @@ private:
 
 	UNREALED_API FString GetCookedAssetRegistryFilename(const FString& PlatformName);
 
+	// Return the filename to use for the cook metadata file, adjusted for DLC, sandbox, and platform.
+	FString GetCookedCookMetadataFilename(const FString& PlatformName);
+	void WriteCookMetadata(const ITargetPlatform* InTargetPlatform, uint64 InDevelopmentAssetRegistryHash);
+
 	/* @return Full path of the CachedEditorThumbnails.bin file in the sandbox */
 	UNREALED_API FString GetSandboxCachedEditorThumbnailsFilename();
 
diff --git a/Engine/Source/Editor/UnrealEd/Private/Commandlets/AssetRegistryGenerator.cpp b/Engine/Source/Editor/UnrealEd/Private/Commandlets/AssetRegistryGenerator.cpp
--- a/Engine/Source/Editor/UnrealEd/Private/Commandlets/AssetRegistryGenerator.cpp
+++ b/Engine/Source/Editor/UnrealEd/Private/Commandlets/AssetRegistryGenerator.cpp
@@ -13,6 +13,7 @@
 #include "Cooker/CookPackageData.h"
 #include "Cooker/CookPlatformManager.h"
 #include "Cooker/CookWorkerClient.h"
+#include "CookMetadata.h"
 #include "Commandlets/ChunkDependencyInfo.h"
 #include "Commandlets/IChunkDataGenerator.h"
 #include "Engine/AssetManager.h"
@@ -22,6 +23,7 @@
 #include "HAL/FileManager.h"
 #include "HAL/LowLevelMemTracker.h"
 #include "HAL/PlatformFileManager.h"
+#include "Hash/xxhash.h"
 #include "ICollectionManager.h"
 #include "Interfaces/ITargetPlatform.h"
 #include "IPlatformFileSandboxWrapper.h"
@@ -1670,7 +1672,8 @@ void FAssetRegistryGenerator::AddAssetToFileOrderRecursive(const FName& InPackag
 	}
 }
 
-bool FAssetRegistryGenerator::SaveAssetRegistry(const FString& SandboxPath, bool bSerializeDevelopmentAssetRegistry, bool bForceNoFilter)
+bool FAssetRegistryGenerator::SaveAssetRegistry(const FString& SandboxPath, bool bSerializeDevelopmentAssetRegistry, 
+	bool bForceNoFilter, uint64& OutDevelopmentAssetRegistryHash)
 {
 	LLM_SCOPE_BYTAG(Cooker_GeneratedAssetRegistry);
 	UE_LOG(LogAssetRegistryGenerator, Display, TEXT("Saving asset registry v%d."), FAssetRegistryVersion::Type::LatestVersion);
@@ -1714,10 +1717,28 @@ bool FAssetRegistryGenerator::SaveAssetRegistry(const FString& SandboxPath, bool
 
 			DevelopmentState.Save(SerializedAssetRegistry, DevelopmentSaveOptions);
 
+			UE::Tasks::TTask<uint64> HashTask = UE::Tasks::Launch(TEXT("HashDevelopmentAssetRegistry"),
+			 [&SerializedAssetRegistry]()
+			{
+				FMemoryView ToHash = MakeMemoryView(SerializedAssetRegistry);
+				return UE::Cook::FCookMetadataState::ComputeHashOfDevelopmentAssetRegistry(ToHash);
+			});
+
 			// Save the generated registry
 			FFileHelper::SaveArrayToFile(SerializedAssetRegistry, *PlatformSandboxPath);
 
-			UE_LOG(LogAssetRegistryGenerator, Display, TEXT("Generated development asset registry %s num assets %d, size is %5.2fkb"), *PlatformSandboxPath, State.GetNumAssets(), (float)SerializedAssetRegistry.Num() / 1024.f);
+			uint64 WaitStartTime = FPlatformTime::Cycles64();
+			uint64 DevArXxHash = HashTask.GetResult();
+			double WaitTime = FPlatformTime::ToSeconds64(FPlatformTime::Cycles64() - WaitStartTime);
+
+			UE_LOG(LogAssetRegistryGenerator, Display, 
+				TEXT("Generated development asset registry %s num assets %d, size is %5.2fkb, ")
+				TEXT("XxHash64 = " UINT64_FMT ", waited on hash %.2f seconds"), 
+				*PlatformSandboxPath, State.GetNumAssets(), (float)SerializedAssetRegistry.Num() / 1024.f, 
+				DevArXxHash, WaitTime
+				);
+
+			OutDevelopmentAssetRegistryHash = DevArXxHash;
 		}
 
 		if (bGenerateChunks)
diff --git a/Engine/Source/Editor/UnrealEd/Private/CookOnTheFlyServer.cpp b/Engine/Source/Editor/UnrealEd/Private/CookOnTheFlyServer.cpp
--- a/Engine/Source/Editor/UnrealEd/Private/CookOnTheFlyServer.cpp
+++ b/Engine/Source/Editor/UnrealEd/Private/CookOnTheFlyServer.cpp
@@ -39,6 +39,7 @@
 #include "Cooker/WorkerRequestsLocal.h"
 #include "Cooker/WorkerRequestsRemote.h"
 #include "CookerSettings.h"
+#include "CookMetadata.h"
 #include "CookOnTheFlyNetServer.h"
 #include "CookPackageSplitter.h"
 #include "DerivedDataCacheInterface.h"
@@ -729,7 +730,8 @@ void UCookOnTheFlyServer::AddCookOnTheFlyPlatformFromGameThread(ITargetPlatform*
 	// Functions in this section are not dependent upon each other and can be ordered arbitrarily or for async performance
 	InitializeShadersForCookOnTheFly(BeginContext.TargetPlatforms);
 	// SaveAssetRegistry is done in CookByTheBookFinished for CBTB, but we need at the start of CookOnTheFly to send as startup information to connecting clients
-	PlatformData->RegistryGenerator->SaveAssetRegistry(GetSandboxAssetRegistryFilename(), true);
+	uint64 DevelopmentAssetRegistryHash = 0;
+	PlatformData->RegistryGenerator->SaveAssetRegistry(GetSandboxAssetRegistryFilename(), true, false, DevelopmentAssetRegistryHash);
 
 	// Initialize systems that nothing in AddCookOnTheFlyPlatformFromGameThread references
 	// Functions in this section are not dependent upon each other and can be ordered arbitrarily or for async performance
@@ -8775,6 +8777,19 @@ FString UCookOnTheFlyServer::GetCookedAssetRegistryFilename(const FString& Platf
 	return CookedAssetRegistryFilename;
 }
 
+FString UCookOnTheFlyServer::GetCookedCookMetadataFilename(const FString& PlatformName )
+{
+	if (IsCookingDLC())
+	{
+		check(IsDirectorCookByTheBook());
+		const FString Filename = GetBaseDirectoryForDLC() / TEXT("Metadata") / UE::Cook::GetCookMetadataFilename();
+		return ConvertToFullSandboxPath(*Filename, true).Replace(TEXT("[Platform]"), *PlatformName);
+	}
+
+	const FString MetadataFilename = FPaths::ProjectDir() / TEXT("Metadata") / UE::Cook::GetCookMetadataFilename();
+	return ConvertToFullSandboxPath(*MetadataFilename, true).Replace(TEXT("[Platform]"), *PlatformName);
+}
+
 FString UCookOnTheFlyServer::GetSandboxCachedEditorThumbnailsFilename()
 {
 	if (IsCookingDLC())
@@ -9218,6 +9233,167 @@ void UCookOnTheFlyServer::CleanShaderCodeLibraries()
 	}
 }
 
+void UCookOnTheFlyServer::WriteCookMetadata(const ITargetPlatform* InTargetPlatform, uint64 InDevelopmentAssetRegistryHash)
+{
+	FString PlatformNameString = InTargetPlatform->PlatformName();
+
+	//
+	// Write the plugin hierarchy for the plugins we cooked. For DLC this is limited to the DLC plugin, otherwise
+	// we write out all plugins enabled for the cook.
+	//
+
+	TArray<TSharedRef<IPlugin>> EnabledPlugins = IPluginManager::Get().GetEnabledPlugins();
+
+	// Remove any plugins that don't support the target platform.
+	EnabledPlugins.RemoveAllSwap([InTargetPlatform](const TSharedRef<IPlugin>& Plugin)
+	{
+		return InTargetPlatform->IsEnabledForPlugin(Plugin.Get()) == false;
+	});
+
+	// Filter to the DLC plugin + dependencies if we are a DLC cook.
+	if (IsCookingDLC())
+	{
+		TMap<FStringView, TSharedRef<IPlugin>> EnabledPluginMap;
+		IPlugin* DlcPlugin = nullptr;
+		for (TSharedRef<IPlugin>& Plugin : EnabledPlugins)
+		{
+			EnabledPluginMap.Add(Plugin->GetName(), Plugin);
+			if (Plugin->GetName() == CookByTheBookOptions->DlcName)
+			{
+				check(DlcPlugin == nullptr);
+				DlcPlugin = &Plugin.Get();
+			}
+		}
+
+		if (DlcPlugin == nullptr)
+		{
+			UE_LOG(LogCook, Warning, TEXT("Dlc plugin %s not found, as a result including all plugins in cook metadata file."), *CookByTheBookOptions->DlcName);
+		}
+		else
+		{
+			TSet<FStringView> IncludedSet;
+
+			TArray<IPlugin*> DependencyStack;
+			IncludedSet.Add(DlcPlugin->GetName());
+			DependencyStack.Add(DlcPlugin);
+
+			while (DependencyStack.Num())
+			{
+				IPlugin* CurrentPlugin = DependencyStack.Pop();
+
+				for (const FPluginReferenceDescriptor& ChildPlugin : CurrentPlugin->GetDescriptor().Plugins)
+				{
+					TSharedRef<IPlugin>* FoundPlugin = EnabledPluginMap.Find(ChildPlugin.Name);
+					if (FoundPlugin)
+					{
+						bool bAlreadyExists = false;
+						IncludedSet.Add(ChildPlugin.Name, &bAlreadyExists);
+						if (bAlreadyExists == false)
+						{
+							DependencyStack.Add(&FoundPlugin->Get());
+						}
+					}
+				}
+			}
+
+			EnabledPlugins.RemoveAllSwap([&IncludedSet](TSharedRef<IPlugin>& EnabledPlugin)
+			{
+				return IncludedSet.Contains(EnabledPlugin->GetName()) == false;
+			}, false /* bAllowShrinking */);
+
+			UE_LOG(LogCook, Display, TEXT("Dlc cook narrowed cook plugin metadata list from %d to %d plugins"), EnabledPluginMap.Num(), EnabledPlugins.Num());
+		}
+	}
+
+	if (IntFitsIn<uint16>(EnabledPlugins.Num()) == false)
+	{
+		UE_LOG(LogCook, Warning, TEXT("Number of plugins exceeds 64k, unable to write cook metadata file (count = %d"), EnabledPlugins.Num());
+	}
+	else
+	{
+		TArray<UE::Cook::FCookMetadataPluginEntry> PluginsToAdd;
+		TMap<FString, uint16> IndexForPlugin;
+		TArray<uint16> PluginChildArray;
+
+		PluginsToAdd.AddDefaulted(EnabledPlugins.Num());
+		uint16 AddIndex = 0;
+		for (TSharedRef<IPlugin>& EnabledPlugin : EnabledPlugins)
+		{
+			UE::Cook::FCookMetadataPluginEntry& NewEntry = PluginsToAdd[AddIndex];
+			NewEntry.Name = EnabledPlugin->GetName();
+			IndexForPlugin.Add(NewEntry.Name, AddIndex);
+			AddIndex++;
+		}
+
+		// Construct the dependency list.
+		for (TSharedRef<IPlugin>& EnabledPlugin : EnabledPlugins)
+		{
+			uint16 SelfIndex = IndexForPlugin[EnabledPlugin->GetName()];
+			UE::Cook::FCookMetadataPluginEntry& Entry = PluginsToAdd[SelfIndex];
+
+			// We detect if this would overflow below and cancel the write - so while this could store
+			// bogus data, it won't get saved.
+			Entry.DependencyIndexStart = (uint16)PluginChildArray.Num();
+
+			const FPluginDescriptor& Descriptor = EnabledPlugin->GetDescriptor();
+			for (FPluginReferenceDescriptor ChildPlugin : Descriptor.Plugins)
+			{
+				if (uint16* ChildIndex = IndexForPlugin.Find(ChildPlugin.Name); ChildIndex != nullptr)
+				{
+					PluginChildArray.Add(*ChildIndex);
+				}
+				else
+				{
+					UE_LOG(LogCook, Display, TEXT("Dependent plugin %s wasn't found in enabled plugins list when creating cook metadata file... skipping"), *ChildPlugin.Name);
+				}
+			}
+
+			Entry.DependencyIndexEnd = PluginChildArray.Num();
+		}
+
+		// Read the root plugins from the config and make sure they exist. This is temporary until we 
+		// have the root qualification exists as part of the plugin definition.
+		TArray<uint16> RootPlugins;
+		{
+			TArray<FString> RootPluginListFromConfig;
+			GConfig->GetArray(TEXT("CookMetadata"), TEXT("RootPluginList"), RootPluginListFromConfig, GEditorIni);
+			for (const FString& Plugin : RootPluginListFromConfig)
+			{
+				uint16* RootPluginIndex = IndexForPlugin.Find(Plugin);
+				if (RootPluginIndex)
+				{
+					RootPlugins.Add(*RootPluginIndex);
+				}
+				else
+				{
+					UE_LOG(LogCook, Warning, TEXT("Unable to find root plugin %s from config list"), *Plugin);
+				}
+			}
+		}
+
+		if (IntFitsIn<uint16>(PluginChildArray.Num()) == false)
+		{
+			UE_LOG(LogCook, Warning, TEXT("Number of child plugins exceeds 64k, unable to write cook metadata file (count = %d)"), PluginChildArray.Num());
+		}
+		else
+		{
+			UE::Cook::FCookMetadataState MetadataState;
+			UE::Cook::FCookMetadataPluginHierarchy PluginHierarchy;
+
+			PluginHierarchy.PluginsEnabledAtCook = MoveTemp(PluginsToAdd);
+			PluginHierarchy.PluginDependencies = MoveTemp(PluginChildArray);
+			PluginHierarchy.RootPlugins = MoveTemp(RootPlugins);
+
+			MetadataState.SetPluginHierarchyInfo(MoveTemp(PluginHierarchy));
+			MetadataState.SetAssociatedDevelopmentAssetRegistryHash(InDevelopmentAssetRegistryHash);
+
+			FArrayWriter SerializedCookMetadata;
+			MetadataState.Serialize(SerializedCookMetadata);
+			FFileHelper::SaveArrayToFile(SerializedCookMetadata, *GetCookedCookMetadataFilename(PlatformNameString));
+		}
+	}
+}
+
 void UCookOnTheFlyServer::CookByTheBookFinished()
 {
 	using namespace UE::Cook;
@@ -9251,6 +9427,9 @@ void UCookOnTheFlyServer::CookByTheBookFinished()
 	check(!LibraryName.IsEmpty());
 	const bool bCacheShaderLibraries = IsUsingShaderCodeLibrary();
 
+	// The hashes of the entire files, per platform.
+	TArray<uint64> DevelopmentAssetRegistryHashes;
+
 	{
 		// Save modified asset registry with all streaming chunk info generated during cook
 		const FString SandboxRegistryFilename = GetSandboxAssetRegistryFilename();
@@ -9287,8 +9466,8 @@ void UCookOnTheFlyServer::CookByTheBookFinished()
 				TArray<FPackageData*> CookedPackageDatas;
 				TArray<FPackageData*> IgnorePackageDatas;
 
-				const FName& PlatformName = FName(*TargetPlatform->PlatformName());
-				FString PlatformNameString = PlatformName.ToString();
+				FString PlatformNameString = TargetPlatform->PlatformName();
+				FName PlatformName(*PlatformNameString);
 
 				PackageDatas->GetCookedPackagesForPlatform(TargetPlatform, CookedPackageDatas, IgnorePackageDatas);
 
@@ -9399,7 +9578,9 @@ void UCookOnTheFlyServer::CookByTheBookFinished()
 				}
 				{
 					UE_SCOPED_HIERARCHICAL_COOKTIMER(SaveRealAssetRegistry);
-					Generator.SaveAssetRegistry(SandboxRegistryFilename, bSaveDevelopmentAssetRegistry, bForceNoFilterAssetsFromAssetRegistry);
+					uint64 DevArHash = 0;
+					Generator.SaveAssetRegistry(SandboxRegistryFilename, bSaveDevelopmentAssetRegistry, bForceNoFilterAssetsFromAssetRegistry, DevArHash);
+					DevelopmentAssetRegistryHashes.Add(DevArHash);
 				}
 				{
 					Generator.PostSave();
@@ -9441,6 +9622,14 @@ void UCookOnTheFlyServer::CookByTheBookFinished()
 		}
 	}
 
+	// Write cook metadata file.for each platform
+	int32 PlatformIndex = 0;
+	for (const ITargetPlatform* TargetPlatform : PlatformManager->GetSessionPlatforms())
+	{
+		WriteCookMetadata(TargetPlatform, DevelopmentAssetRegistryHashes[PlatformIndex]);
+		PlatformIndex++;
+	}
+
 	FString ActualLibraryName = GenerateShaderCodeLibraryName(LibraryName, IsCookFlagSet(ECookInitializationFlags::IterateSharedBuild));
 	FShaderLibraryCooker::EndCookingLibrary(ActualLibraryName);
 	FShaderLibraryCooker::Shutdown();
diff --git a/Engine/Source/Editor/UnrealEd/Public/Commandlets/AssetRegistryGenerator.h b/Engine/Source/Editor/UnrealEd/Public/Commandlets/AssetRegistryGenerator.h
--- a/Engine/Source/Editor/UnrealEd/Public/Commandlets/AssetRegistryGenerator.h
+++ b/Engine/Source/Editor/UnrealEd/Public/Commandlets/AssetRegistryGenerator.h
@@ -207,7 +207,7 @@ public:
 	/**
 	* Saves generated asset registry data for each platform.
 	*/
-	bool SaveAssetRegistry(const FString& SandboxPath, bool bSerializeDevelopmentAssetRegistry = true, bool bForceNoFilterAssets = false);
+	bool SaveAssetRegistry(const FString& SandboxPath, bool bSerializeDevelopmentAssetRegistry, bool bForceNoFilterAssets, uint64& OutDevelopmentAssetRegistryHash);
 
 	/** 
 	 * Writes out CookerOpenOrder.log file 
diff --git a/Engine/Source/Editor/UnrealEd/UnrealEd.Build.cs b/Engine/Source/Editor/UnrealEd/UnrealEd.Build.cs
--- a/Engine/Source/Editor/UnrealEd/UnrealEd.Build.cs
+++ b/Engine/Source/Editor/UnrealEd/UnrealEd.Build.cs
@@ -106,6 +106,7 @@ public class UnrealEd : ModuleRules
 				"AnimGraph",
 				"AppFramework",
 				"CinematicCamera",
+				"CookMetadata",
 				"CurveEditor",
 				"DataLayerEditor",
 				"DerivedDataCache",
