Iterative Cooking - Invalidate packages for iterative cook if their classes are not allowed by class filters. This change uses the same system created for incremental cook and applies it to iterative cook.
Change class filtering to always use an allow list and a block list.
Change class filtering to enable allowlist all classes in a scriptpackage.
Move implementation of hybrid iterative's allow-all mode out into a separately parsed parameter.
Fix some Engine and UnrealEd classes that were not allowing EditorDomain to call Serialize on their class default object (which it needs to do to collect CustomVersions).
#rnx

diff --git a/Engine/Config/BaseEditor.ini b/Engine/Config/BaseEditor.ini
--- a/Engine/Config/BaseEditor.ini
+++ b/Engine/Config/BaseEditor.ini
@@ -356,6 +356,9 @@ SoftGCDenominator=10
 IterativeIgnoreIni=false
 ; If true, iterative cooks will be allowed even if the UnrealEditor exes and dlls have changed.
 IterativeIgnoreExe=true
+; If true, iterative cooks will not allow a package to cook iteratively unless all of its used classes
+; are in the [TargetDomain] filters for which classes can cook iteratively (see [TargetDomain] section below)
+IterativeUseClassFilters=false
 
 ; CookProcessCount sets whether we cook singleprocess or multiprocess and how many processes there should be.
 ; CookProcessCount=(1 or less) is singleprocess. CookProcessCount=(N>1) is 1 director and N-1 cookworkers.
@@ -432,60 +435,18 @@ SaveUnversioned=true
 
 
 [TargetDomain]
-; IterativeClassAllowListEnabled/IterativeClassAllowList/IterativeClassBlockList:
-;    Used in HybridIterative only.
-;    List of classes that may use iterative cooks. Allowed superclasses are also allowed.
+; IterativeClassAllowList/IterativeClassScriptPackageAllowList/IterativeClassBlockList:
+;    List of classes that may use incremental or iterative++ cooks. Allowed superclasses are also allowed.
 ;    Blocked subclasses are also blocked. BlockList overrides the AllowList.
-;    If IterativeClassAllowListEnabled is false, also classes are allowed unless blocked.
 ; +IterativeClassAllowList=/Long/PackageName.ClassName
+; +IterativeClassScriptPackageAllowList=/Script/PackageName
 ; +IterativeClassBlockList=/Long/PackageName.ClassName
-IterativeClassAllowListEnabled=true
-
-IterativeClassAllowList=/Script/Engine.Texture2D
-+IterativeClassAllowList=/Script/Engine.TextureCube
-+IterativeClassAllowList=/Script/Engine.SoundWave
-+IterativeClassAllowList=/Script/Engine.DialogueWave
-
-; Universally co-packaged classes
-+IterativeClassAllowList=/Script/CoreUObject.MetaData
-+IterativeClassAllowList=/Script/Engine.AssetImportData
-+IterativeClassAllowList=/Script/UnrealEd.SceneThumbnailInfo
-+IterativeClassAllowList=/Script/UnrealEd.SceneThumbnailInfoWithPrimitive
-
-; SoundCue and commonly co-packaged classes
-+IterativeClassAllowList=/Script/Engine.SoundCue
-+IterativeClassAllowList=/Script/AudioEditor.SoundCueGraph
-+IterativeClassAllowList=/Script/AudioEditor.SoundCueGraphNode_Root
-+IterativeClassAllowList=/Script/AudioEditor.SoundCueGraphNode
-+IterativeClassAllowList=/Script/Engine.SoundNodeDialoguePlayer
-+IterativeClassAllowList=/Script/Engine.SoundNodeWavePlayer
-+IterativeClassAllowList=/Script/Engine.SoundNodeRandom
-+IterativeClassAllowList=/Script/Engine.SoundNodeModulator
-+IterativeClassAllowList=/Script/Engine.SoundNodeMixer
-+IterativeClassAllowList=/Script/Engine.SoundNodeDistanceCrossFade
-+IterativeClassAllowList=/Script/Engine.SoundNodeSoundClass
-+IterativeClassAllowList=/Script/UnrealEd.EdGraphNode_Comment
-+IterativeClassAllowList=/Script/Engine.SoundNodeAttenuation
-
-; StaticMesh and commonly co-packaged classes
-; +IterativeClassAllowList=/Script/Engine.StaticMesh
-; +IterativeClassAllowList=/Script/Engine.BodySetup
-; +IterativeClassAllowList=/Script/NavigationSystem.NavCollision
-; +IterativeClassAllowList=/Script/Engine.StaticMeshDescriptionBulkData
-; +IterativeClassAllowList=/Script/UnrealEd.FbxStaticMeshImportData
-; +IterativeClassAllowList=/Script/Engine.StaticMeshSocket
-
-; MaterialInstanceConstant and commonly co-packaged classes
-; +IterativeClassAllowList=/Script/Engine.MaterialInstanceConstant
-; +IterativeClassAllowList=/Script/Engine.SceneThumbnailInfoWithPrimitive
-
-; AnimSequence and commonly co-packaged classes
-; +IterativeClassAllowList=/Script/Engine.AnimSequence
-; +IterativeClassAllowList=/Script/Engine.AnimDataModel
-; +IterativeClassAllowList=/Script/Engine.AnimCompress_PerTrackCompression
-; +IterativeClassAllowList=/Script/Engine.AnimCompress_BitwiseCompressOnly
-; +IterativeClassAllowList=/Script/AnimationModifiers.AnimationModifiersAssetUserData
-; +IterativeClassAllowList=/Script/UnrealEd.FbxAnimSequenceImportData
++IterativeScriptPackageAllowList=/Script/CoreUObject
++IterativeScriptPackageAllowList=/Script/Engine
++IterativeScriptPackageAllowList=/Script/UnrealEd
+
++IterativeScriptPackageAllowList=/Script/AudioEditor
+
 
 [CookPlatformDataCacheSettings]
 ; Add <ShortClassName>=<Integer> here to limit the number of instances that can have BeginCacheForCookedPlatformData called at once
diff --git a/Engine/Source/Editor/UnrealEd/Classes/CookOnTheSide/CookOnTheFlyServer.h b/Engine/Source/Editor/UnrealEd/Classes/CookOnTheSide/CookOnTheFlyServer.h
--- a/Engine/Source/Editor/UnrealEd/Classes/CookOnTheSide/CookOnTheFlyServer.h
+++ b/Engine/Source/Editor/UnrealEd/Classes/CookOnTheSide/CookOnTheFlyServer.h
@@ -1444,6 +1444,7 @@ private:
 	 * cooked in hybrid-iterative builds.
 	 */
 	bool bHybridIterativeEnabled = true;
+	bool bHybridIterativeAllowAllClasses = false;
 	bool bHiddenDependenciesDebug = false;
 	bool bOnlyEditorOnlyDebug = false;
 	bool bHiddenDependenciesClassPathFilterListIsAllowList = true;
diff --git a/Engine/Source/Editor/UnrealEd/Private/Commandlets/AssetRegistryGenerator.cpp b/Engine/Source/Editor/UnrealEd/Private/Commandlets/AssetRegistryGenerator.cpp
--- a/Engine/Source/Editor/UnrealEd/Private/Commandlets/AssetRegistryGenerator.cpp
+++ b/Engine/Source/Editor/UnrealEd/Private/Commandlets/AssetRegistryGenerator.cpp
@@ -49,6 +49,7 @@
 #include "Stats/StatsMisc.h"
 #include "String/Find.h"
 #include "String/ParseTokens.h"
+#include "TargetDomain/TargetDomainUtils.h"
 #include "Templates/UniquePtr.h"
 #include "UObject/SoftObjectPath.h"
 
@@ -1237,9 +1238,8 @@ void FAssetRegistryGenerator::ComputePackageDifferences(const FComputeDifference
 		{
 			// A package that was not explored in the previous cook. No need to record it
 		}
-		PRAGMA_DISABLE_DEPRECATION_WARNINGS
-		else if (CurrentPackageData->PackageGuid == PreviousPackageData->PackageGuid)
-		PRAGMA_ENABLE_DEPRECATION_WARNINGS
+		else if (ComputePackageDifferences_IsPackageFileUnchanged(Options, PackageName, *CurrentPackageData,
+			*PreviousPackageData))
 		{
 			if (PreviousPackageData->DiskSize < 0)
 			{
@@ -1447,6 +1447,28 @@ void FAssetRegistryGenerator::ComputePackageDifferences(const FComputeDifference
 	}
 }
 
+bool FAssetRegistryGenerator::ComputePackageDifferences_IsPackageFileUnchanged(
+	const FComputeDifferenceOptions& Options, FName PackageName, const FAssetPackageData& CurrentPackageData,
+	const FAssetPackageData& PreviousPackageData)
+{
+	PRAGMA_DISABLE_DEPRECATION_WARNINGS
+	if (CurrentPackageData.PackageGuid != PreviousPackageData.PackageGuid)
+	PRAGMA_ENABLE_DEPRECATION_WARNINGS
+	{
+		return false;
+	}
+
+	if (Options.bIterativeUseClassFilters)
+	{
+		if (!UE::TargetDomain::IsIterativeEnabled(PackageName, false /* bAllowAllClasses */))
+		{
+			return false;
+		}
+	}
+
+	return true;
+}
+
 FName FAssetRegistryGenerator::GetGeneratorPackage(FName PackageName, const FAssetRegistryState& InState)
 {
 	TConstArrayView<const FAssetData*> Assets = InState.GetAssetsByPackageName(PackageName);
diff --git a/Engine/Source/Editor/UnrealEd/Private/CookOnTheFlyServer.cpp b/Engine/Source/Editor/UnrealEd/Private/CookOnTheFlyServer.cpp
--- a/Engine/Source/Editor/UnrealEd/Private/CookOnTheFlyServer.cpp
+++ b/Engine/Source/Editor/UnrealEd/Private/CookOnTheFlyServer.cpp
@@ -6414,6 +6414,8 @@ void UCookOnTheFlyServer::Initialize( ECookMode::Type DesiredCookMode, ECookInit
 #if WITH_ADDITIONAL_CRASH_CONTEXTS
 	FGenericCrashContext::OnAdditionalCrashContextDelegate().AddUObject(this, &UCookOnTheFlyServer::DumpCrashContext);
 #endif
+
+	UE::EditorDomain::UtilsTargetDomainInit();
 }
 
 
@@ -7905,6 +7907,8 @@ void UCookOnTheFlyServer::PopulateCookedPackages(TArrayView<const ITargetPlatfor
 			FAssetRegistryGenerator::FComputeDifferenceOptions Options;
 			Options.bRecurseModifications = true;
 			Options.bRecurseScriptModifications = !IsCookFlagSet(ECookInitializationFlags::IgnoreScriptPackagesOutOfDate);
+			Options.bIterativeUseClassFilters = true;
+			GConfig->GetBool(TEXT("CookSettings"), TEXT("IterativeUseClassFilters"), Options.bIterativeUseClassFilters, GEditorIni);
 			FAssetRegistryGenerator::FAssetRegistryDifference Difference;
 			PlatformAssetRegistry.ComputePackageDifferences(Options, *PreviousAssetRegistry, Difference);
 			PreviousGeneratorPackages = MoveTemp(Difference.GeneratorPackages);
@@ -10089,6 +10093,7 @@ void FBeginCookConfigSettings::LoadLocal(FBeginCookContext& BeginContext)
 	bHybridIterativeEnabled &= !BeginContext.COTFS.IsCookingDLC();
 	// HybridIterative uses TargetDomain storage of dependencies which is only implemented in ZenStore
 	bHybridIterativeEnabled &= BeginContext.COTFS.IsUsingZenStore();
+	bHybridIterativeAllowAllClasses = BeginContext.COTFS.IsCookFlagSet(ECookInitializationFlags::Iterative);
 
 	FParse::Value(FCommandLine::Get(), TEXT("-CookShowInstigator="), CookShowInstigator);
 	LoadNeverCookLocal(BeginContext);
@@ -10112,6 +10117,7 @@ void FBeginCookConfigSettings::LoadLocal(FBeginCookContext& BeginContext)
 void UCookOnTheFlyServer::SetBeginCookConfigSettings(FBeginCookContext& BeginContext, UE::Cook::FBeginCookConfigSettings&& Settings)
 {
 	bHybridIterativeEnabled = Settings.bHybridIterativeEnabled;
+	bHybridIterativeAllowAllClasses = Settings.bHybridIterativeAllowAllClasses;
 	PackageDatas->SetBeginCookConfigSettings(Settings.CookShowInstigator);
 	SetNeverCookPackageConfigSettings(BeginContext, Settings);
 	if (!bFirstCookInThisProcessInitialized)
diff --git a/Engine/Source/Editor/UnrealEd/Private/Cooker/CookRequestCluster.cpp b/Engine/Source/Editor/UnrealEd/Private/Cooker/CookRequestCluster.cpp
--- a/Engine/Source/Editor/UnrealEd/Private/Cooker/CookRequestCluster.cpp
+++ b/Engine/Source/Editor/UnrealEd/Private/Cooker/CookRequestCluster.cpp
@@ -53,7 +53,6 @@ FRequestCluster::FRequestCluster(UCookOnTheFlyServer& InCOTFS)
 		bAllowSoftDependencies = false;
 	}
 
-	bHybridIterativeEnabled = COTFS.bHybridIterativeEnabled;
 	if (bErrorOnEngineContentUse)
 	{
 		DLCPath = FPaths::Combine(*COTFS.GetBaseDirectoryForDLC(), TEXT("Content"));
@@ -546,7 +545,7 @@ FRequestCluster::FGraphSearch::FGraphSearch(FRequestCluster& InCluster)
 	, AsyncResultsReadyEvent(EEventMode::ManualReset)
 {
 	AsyncResultsReadyEvent->Trigger();
-	bCookAttachmentsEnabled = !Cluster.bFullBuild && Cluster.bHybridIterativeEnabled;
+	bCookAttachmentsEnabled = !Cluster.bFullBuild && Cluster.COTFS.bHybridIterativeEnabled;
 	LastActivityTime = FPlatformTime::Seconds();
 	VertexAllocator.SetMaxBlockSize(1024);
 	VertexAllocator.SetMaxBlockSize(65536);
@@ -1082,9 +1081,9 @@ void FRequestCluster::FGraphSearch::ExploreVertexEdges(FVertexData& Vertex)
 			if (IsCookAttachmentsValid(PackageName, PlatformAttachments))
 			{
 				ICookedPackageWriter* PackageWriter = FetchPlatformData.Writer;
-				if (!Cluster.bFullBuild && Cluster.bHybridIterativeEnabled)
+				if (!Cluster.bFullBuild && Cluster.COTFS.bHybridIterativeEnabled)
 				{
-					if (IsIterativeEnabled(PackageName))
+					if (IsIterativeEnabled(PackageName, Cluster.COTFS.bHybridIterativeAllowAllClasses))
 					{
 						if (PlatformIndex == FirstSessionPlatformIndex)
 						{
diff --git a/Engine/Source/Editor/UnrealEd/Private/Cooker/CookRequestCluster.h b/Engine/Source/Editor/UnrealEd/Private/Cooker/CookRequestCluster.h
--- a/Engine/Source/Editor/UnrealEd/Private/Cooker/CookRequestCluster.h
+++ b/Engine/Source/Editor/UnrealEd/Private/Cooker/CookRequestCluster.h
@@ -335,7 +335,6 @@ private:
 	FBuildDefinitions& BuildDefinitions;
 	bool bAllowHardDependencies = true;
 	bool bAllowSoftDependencies = true;
-	bool bHybridIterativeEnabled = true;
 	bool bErrorOnEngineContentUse = false;
 	bool bPackageNamesComplete = false;
 	bool bDependenciesComplete = false;
diff --git a/Engine/Source/Editor/UnrealEd/Private/Cooker/CookTypes.cpp b/Engine/Source/Editor/UnrealEd/Private/Cooker/CookTypes.cpp
--- a/Engine/Source/Editor/UnrealEd/Private/Cooker/CookTypes.cpp
+++ b/Engine/Source/Editor/UnrealEd/Private/Cooker/CookTypes.cpp
@@ -698,6 +698,7 @@ void FInitializeConfigSettings::MoveToLocal(UCookOnTheFlyServer& COTFS)
 void FBeginCookConfigSettings::CopyFromLocal(const UCookOnTheFlyServer& COTFS)
 {
 	bHybridIterativeEnabled = COTFS.bHybridIterativeEnabled;
+	bHybridIterativeAllowAllClasses = COTFS.bHybridIterativeAllowAllClasses;
 	FParse::Value(FCommandLine::Get(), TEXT("-CookShowInstigator="), CookShowInstigator); // We don't store this on COTFS, so reparse it from commandLine
 	TSet<FName> COTFSNeverCookPackageList;
 	COTFS.PackageTracker->NeverCookPackageList.GetValues(COTFSNeverCookPackageList);
@@ -733,6 +734,7 @@ FCbWriter& operator<<(FCbWriter& Writer, const UE::Cook::FBeginCookConfigSetting
 {
 	Writer.BeginObject();
 	Writer << "HybridIterativeEnabled" << Value.bHybridIterativeEnabled;
+	Writer << "HybridIterativeAllowAllClasses" << Value.bHybridIterativeAllowAllClasses;
 	Writer << "CookShowInstigator" << Value.CookShowInstigator;
 	Writer << "NeverCookPackageList" << Value.NeverCookPackageList;
 	
@@ -754,6 +756,7 @@ bool LoadFromCompactBinary(FCbFieldView Field, UE::Cook::FBeginCookConfigSetting
 {
 	bool bOk = Field.IsObject();
 	bOk = LoadFromCompactBinary(Field["HybridIterativeEnabled"], OutValue.bHybridIterativeEnabled) & bOk;
+	bOk = LoadFromCompactBinary(Field["HybridIterativeAllowAllClasses"], OutValue.bHybridIterativeAllowAllClasses) & bOk;
 	bOk = LoadFromCompactBinary(Field["CookShowInstigator"], OutValue.CookShowInstigator) & bOk;
 	bOk = LoadFromCompactBinary(Field["NeverCookPackageList"], OutValue.NeverCookPackageList) & bOk;
 
diff --git a/Engine/Source/Editor/UnrealEd/Private/Cooker/CookTypes.h b/Engine/Source/Editor/UnrealEd/Private/Cooker/CookTypes.h
--- a/Engine/Source/Editor/UnrealEd/Private/Cooker/CookTypes.h
+++ b/Engine/Source/Editor/UnrealEd/Private/Cooker/CookTypes.h
@@ -344,6 +344,7 @@ namespace UE::Cook
 
 		FString CookShowInstigator;
 		bool bHybridIterativeEnabled = true;
+		bool bHybridIterativeAllowAllClasses = false;
 		TArray<FName> NeverCookPackageList;
 		TFastPointerMap<const ITargetPlatform*, TSet<FName>> PlatformSpecificNeverCookPackages;
 
diff --git a/Engine/Source/Editor/UnrealEd/Private/EditorDomain/EditorDomainUtils.cpp b/Engine/Source/Editor/UnrealEd/Private/EditorDomain/EditorDomainUtils.cpp
--- a/Engine/Source/Editor/UnrealEd/Private/EditorDomain/EditorDomainUtils.cpp
+++ b/Engine/Source/Editor/UnrealEd/Private/EditorDomain/EditorDomainUtils.cpp
@@ -16,6 +16,7 @@
 #include "DerivedDataRequestOwner.h"
 #include "DerivedDataValue.h"
 #include "Editor.h"
+#include "Engine/Blueprint.h"
 #include "Engine/StaticMesh.h"
 #include "HAL/CriticalSection.h"
 #include "HAL/FileManager.h"
@@ -196,10 +197,9 @@ TMap<FTopLevelAssetPath, EDomainUse> GClassBlockedUses;
 TMap<FName, EDomainUse> GPackageBlockedUses;
 TMultiMap<FTopLevelAssetPath, FTopLevelAssetPath> GConstructClasses;
 TSet<FTopLevelAssetPath> GTargetDomainClassBlockList;
-bool GTargetDomainClassUseAllowList = true;
-bool GTargetDomainClassEmptyAllowList = false;
 TArray<FTopLevelAssetPath> GGlobalConstructClasses;
 bool bGGlobalConstructClassesInitialized = false;
+bool bGUtilsTargetDomainInitialized = false;
 FBlake3Hash GGlobalConstructClassesHash;
 int64 GMaxBulkDataSize = -1;
 
@@ -507,10 +507,6 @@ FClassDigestData* FPrecacheClassDigest::GetRecursive(FTopLevelAssetPath ClassPat
 	// Fill in digest data config-driven flags
 	DigestData->EditorDomainUse = EDomainUse::LoadEnabled | EDomainUse::SaveEnabled;
 	DigestData->EditorDomainUse &= ~MapFindRef(GClassBlockedUses, ClassPath, EDomainUse::None);
-	if (!GTargetDomainClassUseAllowList)
-	{
-		DigestData->bTargetIterativeEnabled = !GTargetDomainClassBlockList.Contains(ClassPath);
-	}
 
 	// Fill in native-specific digest data, get the ParentName, and if non-native, get the native ancestor struct
 	UStruct* Struct = nullptr;
@@ -635,10 +631,6 @@ FClassDigestData* FPrecacheClassDigest::GetRecursive(FTopLevelAssetPath ClassPat
 			}
 			EnumSetFlagsAnd(DigestData->EditorDomainUse, EDomainUse::LoadEnabled | EDomainUse::SaveEnabled,
 				DigestData->EditorDomainUse, ParentDigest->EditorDomainUse);
-			if (!GTargetDomainClassUseAllowList)
-			{
-				DigestData->bTargetIterativeEnabled &= ParentDigest->bTargetIterativeEnabled;
-			}
 			ConstructClasses.Append(ParentDigest->ConstructClasses);
 			if (!StructAsClass)
 			{
@@ -994,7 +986,6 @@ TSet<FTopLevelAssetPath> ConstructTargetIterativeClassBlockList()
 
 void ConstructTargetIterativeClassAllowList()
 {
-	// We're using an allowlist with a blocklist override, so the blocklist is only needed when creating the allowlist
 	TSet<FTopLevelAssetPath> BlockListClassPaths = ConstructTargetIterativeClassBlockList();
 
 	// AllowList elements implicitly allow all parent classes, so instead of consulting a list and propagating
@@ -1008,8 +999,13 @@ void ConstructTargetIterativeClassAllowList()
 	TStringBuilder<256> NameStringBuffer;
 	TMap<FTopLevelAssetPath, TOptional<bool>> Visited;
 	auto EnableClassIfNotBlocked = [&Visited, &EnabledClassPaths, &BlockListClassPaths, &NameStringBuffer]
-		(const FTopLevelAssetPath& ClassPath, UStruct* Struct, bool& bOutIsBlocked, auto& EnableClassIfNotBlockedRef)
+		(UStruct* Struct, bool& bOutIsBlocked, auto& EnableClassIfNotBlockedRef)
 	{
+		FTopLevelAssetPath ClassPath(Struct);
+		if (!ClassPath.IsValid())
+		{
+			return;
+		}
 		int32 KeyHash = GetTypeHash(ClassPath);
 		TOptional<bool>& BlockedValue = Visited.FindOrAddByHash(KeyHash, ClassPath);
 		if (BlockedValue.IsSet())
@@ -1023,11 +1019,7 @@ void ConstructTargetIterativeClassAllowList()
 		UStruct* ParentStruct = Struct->GetSuperStruct();
 		if (ParentStruct)
 		{
-			FTopLevelAssetPath ParentStructPath(ParentStruct);
-			if (ParentStructPath.IsValid())
-			{
-				EnableClassIfNotBlockedRef(ParentStructPath, ParentStruct, bParentBlocked, EnableClassIfNotBlockedRef);
-			}
+			EnableClassIfNotBlockedRef(ParentStruct, bParentBlocked, EnableClassIfNotBlockedRef);
 		}
 
 		bOutIsBlocked = bParentBlocked || BlockListClassPaths.Contains(ClassPath);
@@ -1043,7 +1035,10 @@ void ConstructTargetIterativeClassAllowList()
 	};
 
 	TArray<FString> AllowListLeafNames;
+	TArray<FString> AllowListScriptPackages;
 	GConfig->GetArray(TEXT("TargetDomain"), TEXT("IterativeClassAllowList"), AllowListLeafNames, GEditorIni);
+	GConfig->GetArray(TEXT("TargetDomain"), TEXT("IterativeScriptPackageAllowList"), AllowListScriptPackages, GEditorIni);
+	TSet<UStruct*> AllowListClasses;
 	for (const FString& ClassPathString : AllowListLeafNames)
 	{
 		FTopLevelAssetPath ClassPath(ClassPathString);
@@ -1061,8 +1056,45 @@ void ConstructTargetIterativeClassAllowList()
 		{
 			continue;
 		}
+		AllowListClasses.Add(Struct);
+	}
+	UPackage* CoreUObjectPackage = UClass::StaticClass()->GetPackage();
+	UPackage* EnginePackage = UBlueprint::StaticClass()->GetPackage();
+	for (const FString& ScriptPackage : AllowListScriptPackages)
+	{
+		if (!FPackageName::IsScriptPackage(ScriptPackage))
+		{
+			continue;
+		}
+		UPackage* Package = FindObject<UPackage>(nullptr, *ScriptPackage);
+		if (!Package)
+		{
+			continue;
+		}
+		ForEachObjectWithPackage(Package, [Package, CoreUObjectPackage, EnginePackage, &AllowListClasses](UObject* Object)
+			{
+				UClass* Class = Cast<UClass>(Object);
+				if (!Class)
+				{
+					return true;
+				}
+				if (Package == CoreUObjectPackage || Package == EnginePackage)
+				{
+					// Skip some non-normal classes in CoreUObject and Engine; we crash on them because e.g. they do not have a CDO
+					if (FStringView(WriteToString<256>(Object->GetFName())).StartsWith(TEXT("Default_")))
+					{
+						return true;
+					}
+				}
+				AllowListClasses.Add(Class);
+				return true;
+			}, false /* bIncludeNestedObjects */);
+	}
+
+	for (UStruct* Struct : AllowListClasses)
+	{
 		bool bUnusedIsBlocked;
-		EnableClassIfNotBlocked(ClassPath, Struct, bUnusedIsBlocked, EnableClassIfNotBlocked);
+		EnableClassIfNotBlocked(Struct, bUnusedIsBlocked, EnableClassIfNotBlocked);
 	}
 
 	TArray<FTopLevelAssetPath> EnabledClassPathsArray = EnabledClassPaths.Array();
@@ -1114,9 +1146,6 @@ TMultiMap<FTopLevelAssetPath, FTopLevelAssetPath> ConstructConstructClasses()
 	return ConstructClasses;
 }
 
-FDelegateHandle GUtilsPostInitDelegate;
-void UtilsPostEngineInit();
-
 /** A default implementation of IPackageDigestCache that stores the Digests in a TMap. */
 class FDefaultPackageDigestCache : public IPackageDigestCache
 {
@@ -1176,53 +1205,21 @@ void UtilsInitialize()
 	GPackageBlockedUses = ConstructPackageNameBlockedUses();
 	GConstructClasses = ConstructConstructClasses();
 
-	bool bTargetDomainClassUseBlockList = true;
-	if (FParse::Param(FCommandLine::Get(), TEXT("fullcook")))
-	{
-		// Allow list is marked as used, but is initialized empty
-		bTargetDomainClassUseBlockList = false;
-		GTargetDomainClassUseAllowList = true;
-		GTargetDomainClassEmptyAllowList = true;
-	}
-	else if (FParse::Param(FCommandLine::Get(), TEXT("iterate")))
-	{
-		bTargetDomainClassUseBlockList = false;
-		GTargetDomainClassUseAllowList = false;
-	}
-	else
-	{
-		GConfig->GetBool(TEXT("TargetDomain"), TEXT("IterativeClassAllowListEnabled"), GTargetDomainClassUseAllowList, GEditorIni);
-		GTargetDomainClassEmptyAllowList = false;
-	}
-
-	if (!GTargetDomainClassUseAllowList && bTargetDomainClassUseBlockList)
-	{
-		GTargetDomainClassBlockList = ConstructTargetIterativeClassBlockList();
-	}
-
 	double MaxBulkDataSize = 64 * 1024;
 	GConfig->GetDouble(TEXT("EditorDomain"), TEXT("MaxBulkDataSize"), MaxBulkDataSize, GEditorIni);
 	GMaxBulkDataSize = static_cast<uint64>(MaxBulkDataSize);
 
-	// Constructing allowlists requires use of UStructs, and the early SetPackageResourceManager
-	// where UtilsInitialize is called is too early; trying to call UStruct->GetSchemaHash at that
-	// time will break the UClass. Defer the construction of allowlist-based data until OnPostEngineInit
-	GUtilsPostInitDelegate = FCoreDelegates::OnPostEngineInit.AddLambda([]() { UtilsPostEngineInit(); });
-
-
 	COOK_STAT(UE::EditorDomain::CookStats::Register());
 }
 
-void UtilsPostEngineInit()
+void UtilsTargetDomainInit()
 {
-	FCoreDelegates::OnPostEngineInit.Remove(GUtilsPostInitDelegate);
-	GUtilsPostInitDelegate.Reset();
-
-	// Note that constructing AllowLists depends on all BlockLists having been parsed already
-	if (GTargetDomainClassUseAllowList && !GTargetDomainClassEmptyAllowList)
+	if (bGUtilsTargetDomainInitialized)
 	{
-		ConstructTargetIterativeClassAllowList();
+		return;
 	}
+	bGUtilsTargetDomainInitialized = true;
+	ConstructTargetIterativeClassAllowList();
 }
 
 UE::DerivedData::FCacheKey GetEditorDomainPackageKey(const FIoHash& EditorDomainHash)
diff --git a/Engine/Source/Editor/UnrealEd/Private/EditorDomain/EditorDomainUtils.h b/Engine/Source/Editor/UnrealEd/Private/EditorDomain/EditorDomainUtils.h
--- a/Engine/Source/Editor/UnrealEd/Private/EditorDomain/EditorDomainUtils.h
+++ b/Engine/Source/Editor/UnrealEd/Private/EditorDomain/EditorDomainUtils.h
@@ -115,6 +115,9 @@ FClassDigestMap& GetClassDigests();
 /** Initializes some global config-driven values used by the EditorDomain and TargetDomain. */
 void UtilsInitialize();
 
+/** Initializes global config-driven values that are only needed by the TargetDomain. */
+void UtilsTargetDomainInit();
+
 #if ENABLE_COOK_STATS
 namespace CookStats
 {
diff --git a/Engine/Source/Editor/UnrealEd/Private/EditorTransaction.cpp b/Engine/Source/Editor/UnrealEd/Private/EditorTransaction.cpp
--- a/Engine/Source/Editor/UnrealEd/Private/EditorTransaction.cpp
+++ b/Engine/Source/Editor/UnrealEd/Private/EditorTransaction.cpp
@@ -1175,7 +1175,7 @@ void UTransBuffer::Initialize(SIZE_T InMaxMemory)
 // UObject interface.
 void UTransBuffer::Serialize( FArchive& Ar )
 {
-	check( !Ar.IsPersistent() );
+	check( !Ar.IsPersistent() || this->HasAnyFlags(RF_ClassDefaultObject) );
 
 	CheckState();
 
diff --git a/Engine/Source/Editor/UnrealEd/Private/TargetDomain/TargetDomainUtils.cpp b/Engine/Source/Editor/UnrealEd/Private/TargetDomain/TargetDomainUtils.cpp
--- a/Engine/Source/Editor/UnrealEd/Private/TargetDomain/TargetDomainUtils.cpp
+++ b/Engine/Source/Editor/UnrealEd/Private/TargetDomain/TargetDomainUtils.cpp
@@ -382,7 +382,7 @@ bool IsCookAttachmentsValid(FName PackageName, const FCookAttachments& CookAttac
 	return true;
 }
 
-bool IsIterativeEnabled(FName PackageName)
+bool IsIterativeEnabled(FName PackageName, bool bAllowAllClasses)
 {
 	IAssetRegistry* AssetRegistry = IAssetRegistry::Get();
 	if (!AssetRegistry)
@@ -396,25 +396,28 @@ bool IsIterativeEnabled(FName PackageName)
 	}
 	FAssetPackageData& PackageData = *PackageDataOpt;
 
-	UE::EditorDomain::FClassDigestMap& ClassDigests = UE::EditorDomain::GetClassDigests();
-	FReadScopeLock ClassDigestsScopeLock(ClassDigests.Lock);
-	for (FName ClassName : PackageData.ImportedClasses)
+	if (!bAllowAllClasses)
 	{
-		FTopLevelAssetPath ClassPath(WriteToString<256>(ClassName).ToView());
-		UE::EditorDomain::FClassDigestData* ExistingData = nullptr;
-		if (ClassPath.IsValid())
+		UE::EditorDomain::FClassDigestMap& ClassDigests = UE::EditorDomain::GetClassDigests();
+		FReadScopeLock ClassDigestsScopeLock(ClassDigests.Lock);
+		for (FName ClassName : PackageData.ImportedClasses)
 		{
-			ExistingData = ClassDigests.Map.Find(ClassPath);
-		}
-		if (!ExistingData)
-		{
-			// All allowlisted classes are added to ClassDigests at startup, so if the class is not in ClassDigests,
-			// it is not allowlisted
-			return false;
-		}
-		if (!ExistingData->bTargetIterativeEnabled)
-		{
-			return false;
+			FTopLevelAssetPath ClassPath(WriteToString<256>(ClassName).ToView());
+			UE::EditorDomain::FClassDigestData* ExistingData = nullptr;
+			if (ClassPath.IsValid())
+			{
+				ExistingData = ClassDigests.Map.Find(ClassPath);
+			}
+			if (!ExistingData)
+			{
+				// All allowlisted classes are added to ClassDigests at startup, so if the class is not in ClassDigests,
+				// it is not allowlisted
+				return false;
+			}
+			if (!ExistingData->bTargetIterativeEnabled)
+			{
+				return false;
+			}
 		}
 	}
 	return true;
diff --git a/Engine/Source/Editor/UnrealEd/Private/TargetDomain/TargetDomainUtils.h b/Engine/Source/Editor/UnrealEd/Private/TargetDomain/TargetDomainUtils.h
--- a/Engine/Source/Editor/UnrealEd/Private/TargetDomain/TargetDomainUtils.h
+++ b/Engine/Source/Editor/UnrealEd/Private/TargetDomain/TargetDomainUtils.h
@@ -62,7 +62,7 @@ void FetchCookAttachments(TArrayView<FName> PackageNames, const ITargetPlatform*
 bool IsCookAttachmentsValid(FName PackageName, const FCookAttachments& CookAttachments);
 
 /** Return whether iterative cook is enabled for the given packagename, based on used-class allowlist/blocklist. */
-bool IsIterativeEnabled(FName PackageName);
+bool IsIterativeEnabled(FName PackageName, bool bAllowAllClasses);
 
 
 /** Store extra information derived during save and used by the cooker for the given EditorDomain package. */
diff --git a/Engine/Source/Editor/UnrealEd/Public/Commandlets/AssetRegistryGenerator.h b/Engine/Source/Editor/UnrealEd/Public/Commandlets/AssetRegistryGenerator.h
--- a/Engine/Source/Editor/UnrealEd/Public/Commandlets/AssetRegistryGenerator.h
+++ b/Engine/Source/Editor/UnrealEd/Public/Commandlets/AssetRegistryGenerator.h
@@ -73,6 +73,8 @@ public:
 		bool bRecurseModifications;
 		/** if true, modified script / c++ packages are recursed, if false only asset references are recursed */
 		bool bRecurseScriptModifications;
+		/** If true, use the AllowList and DenyList for classes on a package's used classes to decide whether the package is iterable. */
+		bool bIterativeUseClassFilters;
 	};
 
 	/** Info about a GeneratorPackage (see ICookPackageSplitter) loaded from previous iterative cooks. */
@@ -131,7 +133,8 @@ public:
 	 * @param PreviousAssetPackageDataMap previously cooked asset package data
 	 * @param OutDifference the differences between the current and the previous state
 	 */
-	void ComputePackageDifferences(const FComputeDifferenceOptions& Options, const FAssetRegistryState& PreviousState, FAssetRegistryDifference& OutDifference);
+	void ComputePackageDifferences(const FComputeDifferenceOptions& Options, const FAssetRegistryState& PreviousState, 
+		FAssetRegistryDifference& OutDifference);
 
 	/** Computes just the list of packages in the PreviousState that no longer exist in the current state. */
 	void ComputePackageRemovals(const FAssetRegistryState& PreviousState, TArray<FName>& OutRemovedPackages,
@@ -286,6 +289,9 @@ private:
 
 	void SetOverridePackageDependencies(FName PackageName, TConstArrayView<FAssetDependency> OverridePackageDependencies);
 
+	bool ComputePackageDifferences_IsPackageFileUnchanged(const FComputeDifferenceOptions& Options, FName PackageName,
+		const FAssetPackageData& CurrentPackageData, const FAssetPackageData& PreviousPackageData);
+
 	/** State of the asset registry that is being built for this platform */
 	FAssetRegistryState State;
 	
diff --git a/Engine/Source/Runtime/Engine/Private/Elements/SMInstance/SMInstanceElementId.cpp b/Engine/Source/Runtime/Engine/Private/Elements/SMInstance/SMInstanceElementId.cpp
--- a/Engine/Source/Runtime/Engine/Private/Elements/SMInstance/SMInstanceElementId.cpp
+++ b/Engine/Source/Runtime/Engine/Private/Elements/SMInstance/SMInstanceElementId.cpp
@@ -385,7 +385,8 @@ USMInstanceElementIdMapTransactor::USMInstanceElementIdMapTransactor()
 #if WITH_EDITORONLY_DATA
 void USMInstanceElementIdMapTransactor::Serialize(FArchive& Ar)
 {
-	checkf(!Ar.IsPersistent(), TEXT("USMInstanceElementIdMapTransactor can only be serialized by transient archives!"));
+	checkf(!Ar.IsPersistent() || this->HasAnyFlags(RF_ClassDefaultObject),
+		TEXT("USMInstanceElementIdMapTransactor can only be serialized by transient archives!"));
 
 	if (OwnerEntry)
 	{
