Fix an issue where the gizmo would rotate on an axis from keybinds, even if the user had selected a valid axis on the gizmo

The axis hit proxy should have been cached, but was always being cleared if the control key was pressed.

diff --git a/Engine/Source/Editor/UnrealEd/Private/EditorViewportClient.cpp b/Engine/Source/Editor/UnrealEd/Private/EditorViewportClient.cpp
--- a/Engine/Source/Editor/UnrealEd/Private/EditorViewportClient.cpp
+++ b/Engine/Source/Editor/UnrealEd/Private/EditorViewportClient.cpp
@@ -3191,7 +3191,7 @@ void FEditorViewportClient::StartTrackingDueToInput( const struct FInputEventSta
 		}
 
 		// Start new tracking. Potentially reset the widget so that StartTracking can pick a new axis.
-		if ( Widget && ( !bDraggingByHandle || InputState.IsCtrlButtonPressed() ) )
+		if ( Widget && ( !bDraggingByHandle && InputState.IsCtrlButtonPressed() ) )
 		{
 			bWidgetAxisControlledByDrag = false;
 			Widget->SetCurrentAxis( EAxisList::None );
