Remove FBlueprintEditorUtils::RegenerateBlueprintClass - deprecated in 5.2, but long unused internally. Usage should be replaced with FKismetEditorUtilities::CompileBlueprint

diff --git a/Engine/Source/Editor/UnrealEd/Private/Kismet2/BlueprintEditorUtils.cpp b/Engine/Source/Editor/UnrealEd/Private/Kismet2/BlueprintEditorUtils.cpp
--- a/Engine/Source/Editor/UnrealEd/Private/Kismet2/BlueprintEditorUtils.cpp
+++ b/Engine/Source/Editor/UnrealEd/Private/Kismet2/BlueprintEditorUtils.cpp
@@ -946,118 +946,6 @@ void FBlueprintEditorUtils::RemoveAllLocalBookmarks(const UBlueprint* ForBluepri
 	}
 }
 
-/** 
- * Check FKismetCompilerContext::SetCanEverTickForActor
- */
-struct FSaveActorFlagsHelper
-{
-	bool bOverride;
-	bool bCanEverTick;
-	UClass * Class;
-
-	FSaveActorFlagsHelper(UClass * InClass) : Class(InClass)
-	{
-		bOverride = (AActor::StaticClass() == FBlueprintEditorUtils::FindFirstNativeClass(Class));
-		if(Class && bOverride)
-		{
-			AActor* CDActor = Cast<AActor>(Class->GetDefaultObject());
-			if(CDActor)
-			{
-				bCanEverTick = CDActor->PrimaryActorTick.bCanEverTick;
-			}
-		}
-	}
-
-	~FSaveActorFlagsHelper()
-	{
-		if(Class && bOverride)
-		{
-			AActor* CDActor = Cast<AActor>(Class->GetDefaultObject());
-			if(CDActor)
-			{
-				CDActor->PrimaryActorTick.bCanEverTick = bCanEverTick;
-			}
-		}
-	}
-};
-
-//////////////////////////////////////////////////////////////////////////
-
-/**
- * Archive built to go through and find any references to objects in the transient package, and then NULL those references
- */
-class FArchiveMoveSkeletalRefs : public FArchiveUObject
-{
-public:
-	FArchiveMoveSkeletalRefs(UBlueprint* TargetBP)
-		: TargetBlueprint(TargetBP)
-	{
-		ArIsObjectReferenceCollector = true;
-		this->SetIsPersistent(false);
-		ArIgnoreArchetypeRef = false;
-	}
-
-	void UpdateReferences()
-	{
-		if( TargetBlueprint != nullptr && (TargetBlueprint->BlueprintType != BPTYPE_MacroLibrary) )
-		{
-			if( ensureMsgf(TargetBlueprint->SkeletonGeneratedClass, TEXT("Blueprint %s is missing its skeleton generated class - known possible for assets on revision 1 older than 2088505"), *TargetBlueprint->GetName() ) )
-			{
-				TargetBlueprint->SkeletonGeneratedClass->GetDefaultObject()->Serialize(*this);
-			}
-			check(TargetBlueprint->GeneratedClass);
-			TargetBlueprint->GeneratedClass->GetDefaultObject()->Serialize(*this);
-
-			TArray<UObject*> SubObjs;
-			GetObjectsWithOuter(TargetBlueprint, SubObjs, true);
-
-			for (UObject* SubObj : SubObjs)
-			{
-				SubObj->Serialize(*this);
-			}
-
-			TargetBlueprint->bLegacyNeedToPurgeSkelRefs = false;
-		}
-	}
-
-protected:
-	UBlueprint* TargetBlueprint;
-
-	/** 
-	 * UObject serialize operator implementation
-	 *
-	 * @param Object	reference to Object reference
-	 * @return reference to instance of this class
-	 */
-	FArchive& operator<<( UObject*& Object )
-	{
-		// Check if this is a reference to an object existing in the transient package, and if so, NULL it.
-		if (Object != nullptr )
-		{
-			if( UClass* RefClass = Cast<UClass>(Object) )
-			{
-				const bool bIsValidBPGeneratedClass = RefClass->HasAnyClassFlags(CLASS_CompiledFromBlueprint) && RefClass->ClassGeneratedBy;
-				if (bIsValidBPGeneratedClass)
-				{
-					UClass* AuthClass = RefClass->GetAuthoritativeClass();
-					if (RefClass != AuthClass)
-					{
-						Object = AuthClass;
-					}
-				}
-			}
-		}
-
-		return *this;
-	}
-
-private:
-	// Want to make them HAVE to use a blueprint, so we can control what we replace refs on
-	FArchiveMoveSkeletalRefs()
-	{
-	}
-};
-
 //////////////////////////////////////////////////////////////////////////
 
 struct FRegenerationHelper
@@ -1343,249 +1231,6 @@ void FBlueprintEditorUtils::PreloadBlueprintSpecificData(UBlueprint* Blueprint)
 	}
 }
 
-UClass* FBlueprintEditorUtils::RegenerateBlueprintClass(UBlueprint* Blueprint, UClass* ClassToRegenerate, UObject* PreviousCDO)
-{
-	bool bRegenerated = false;
-
-	// Cache off the dirty flag for the package, so we can restore it later
-	UPackage* Package = Blueprint->GetOutermost();
-	bool bIsPackageDirty = Package ? Package->IsDirty() : false;
-
-	// Preload the blueprint and all its parts before refreshing nodes. 
-	// Otherwise, the nodes might not maintain their proper linkages... 
-	//
-	// This all should also happen here, first thing, before 
-	// bIsRegeneratingOnLoad is set, so that we can re-enter this function for 
-	// the same class further down the callstack (presumably from 
-	// PreloadInterfaces() or some other dependency load). This is here to 
-	// handle circular dependencies, where pre-loading a member here sets off a  
-	// subsequent load that in turn, relies on this class and requires this  
-	// class to be fully generated... A second call to this function with the 
-	// same class will continue to preload all it's members (from where it left
-	// off, since they're gated by a RF_NeedLoad check) and then fall through to
-	// finish compiling the class (while it's still technically pre-loading a
-	// member further up the stack).
-	if (!Blueprint->bHasBeenRegenerated)
-	{
-		UBlueprint::ForceLoadMetaData(Blueprint);
-		if (ensure(PreviousCDO))
-		{
-			UBlueprint::ForceLoadMembers(PreviousCDO, Blueprint);
-		}
-		UBlueprint::ForceLoadMembers(Blueprint);
-	}
-
-	if( ShouldRegenerateBlueprint(Blueprint) && !Blueprint->bHasBeenRegenerated )
-	{
-		Blueprint->bCachedDependenciesUpToDate = false;
-		Blueprint->bIsRegeneratingOnLoad = true;
-
-		// Cache off the linker index, if needed
-		FName GeneratedName, SkeletonName;
-		Blueprint->GetBlueprintCDONames(GeneratedName, SkeletonName);
-		int32 OldSkelLinkerIdx = INDEX_NONE;
-		int32 OldGenLinkerIdx = INDEX_NONE;
-		FLinkerLoad* OldLinker = Blueprint->GetLinker();
-		for( int32 i = 0; i < OldLinker->ExportMap.Num(); i++ )
-		{
-			FObjectExport& ThisExport = OldLinker->ExportMap[i];
-			if( ThisExport.ObjectName == SkeletonName )
-			{
-				OldSkelLinkerIdx = i;
-			}
-			else if( ThisExport.ObjectName == GeneratedName )
-			{
-				OldGenLinkerIdx = i;
-			}
-
-			if( OldSkelLinkerIdx != INDEX_NONE && OldGenLinkerIdx != INDEX_NONE )
-			{
-				break;
-			}
-		}
-
-		// Make sure the simple construction script is loaded, since the outer hierarchy isn't compatible with PreloadMembers past the root node
-		FBlueprintEditorUtils::PreloadConstructionScript(Blueprint);
-
-		// Preload Overridden Components
-		if (Blueprint->InheritableComponentHandler)
-		{
-			Blueprint->InheritableComponentHandler->PreloadAll();
-		}
-
-		// Purge any NULL graphs
-		FBlueprintEditorUtils::PurgeNullGraphs(Blueprint);
-
-		// Now that things have been preloaded, see what work needs to be done to refresh this blueprint
-		const bool bIsMacro = (Blueprint->BlueprintType == BPTYPE_MacroLibrary);
-		const bool bHasCode = !FBlueprintEditorUtils::IsDataOnlyBlueprint(Blueprint) && !bIsMacro;
-
-		// Make sure all used external classes/functions/structures/macros/etc are loaded and linked
-		FRegenerationHelper::LinkExternalDependencies(Blueprint);
-
-		bool bSkeletonUpToDate = FKismetEditorUtilities::GenerateBlueprintSkeleton(Blueprint);
-
-		const bool bDataOnlyClassThatMustBeRecompiled = !bHasCode && !bIsMacro
-			&& (!ClassToRegenerate || (Blueprint->ParentClass != ClassToRegenerate->GetSuperClass()));
-
-		UBlueprintGeneratedClass* BPGClassToRegenerate = Cast<UBlueprintGeneratedClass>(ClassToRegenerate);
-#if USE_UBER_GRAPH_PERSISTENT_FRAME
-		const bool bHasPendingUberGraphFrame = BPGClassToRegenerate
-			&& (BPGClassToRegenerate->UberGraphFramePointerProperty || BPGClassToRegenerate->UberGraphFunction);
-#else 
-		const bool bHasPendingUberGraphFrame = false;
-#endif //USE_UBER_GRAPH_PERSISTENT_FRAME
-
-		const bool bDefaultComponentMustBeAdded = !bHasCode 
-			&& BPGClassToRegenerate
-			&& SupportsConstructionScript(Blueprint) 
-			&& BPGClassToRegenerate->SimpleConstructionScript
-			&& (nullptr == BPGClassToRegenerate->SimpleConstructionScript->GetSceneRootComponentTemplate(true));
-		const bool bShouldBeRecompiled = bHasCode || bDataOnlyClassThatMustBeRecompiled || bHasPendingUberGraphFrame || bDefaultComponentMustBeAdded;
-
-		if (bShouldBeRecompiled)
-		{
-			// Make sure parent function calls are up to date
-			FBlueprintEditorUtils::ConformCallsToParentFunctions(Blueprint);
-
-			// Make sure events are up to date
-			FBlueprintEditorUtils::ConformImplementedEvents(Blueprint);
-			
-			// Make sure interfaces are up to date
-			FBlueprintEditorUtils::ConformImplementedInterfaces(Blueprint);
-
-			// Reconstruct all nodes, this will call AllocateDefaultPins, which ensures
-			// that nodes have a chance to create all the pins they'll expect when they compile.
-			// A good example of why this is necessary is UK2Node_BaseAsyncTask::AllocateDefaultPins
-			// and it's companion function UK2Node_BaseAsyncTask::ExpandNode.
-			FBlueprintEditorUtils::ReconstructAllNodes(Blueprint);
-
-			FBlueprintEditorUtils::ReplaceDeprecatedNodes(Blueprint);
-
-			// Compile the actual blueprint
-			EBlueprintCompileOptions Options = EBlueprintCompileOptions::IsRegeneratingOnLoad;
-			if(bSkeletonUpToDate)
-			{
-				Options |= EBlueprintCompileOptions::SkeletonUpToDate;
-			}
-			FKismetEditorUtilities::CompileBlueprint(Blueprint, Options, nullptr);
-		}
-		else if( bIsMacro )
-		{
-			// Just refresh all nodes in macro blueprints, but don't recompile
-			FBlueprintEditorUtils::RefreshAllNodes(Blueprint);
-
-			FBlueprintEditorUtils::ReplaceDeprecatedNodes(Blueprint);
-
-			if (ClassToRegenerate != nullptr)
-			{
-				UClass* OldSuperClass = ClassToRegenerate->GetSuperClass();
-				if ((OldSuperClass != nullptr) && OldSuperClass->HasAnyClassFlags(CLASS_NewerVersionExists))
-				{
-					UClass* NewSuperClass = OldSuperClass->GetAuthoritativeClass();
-					ensure(NewSuperClass == Blueprint->ParentClass);
-
-					// in case the macro's super class was re-instanced (it 
-					// would have re-parented this to a REINST_ class), for non-
-					// macro blueprints this would normally be reset in 
-					// CompileBlueprint (but since we don't compile macros, we 
-					// need to fix this up here)
-					ClassToRegenerate->SetSuperStruct(NewSuperClass);
-				}
-			}
-
-			// Flag macro blueprints as being up-to-date
-			Blueprint->Status = BS_UpToDate;
-		}
-		else
-		{
-			if (Blueprint->GeneratedClass != nullptr)
-			{
-				RemoveStaleFunctions(Cast<UBlueprintGeneratedClass>(Blueprint->GeneratedClass), Blueprint);
-			}
-
-			// No actual compilation work to be done, but try to conform the class and fix up anything that might need to be updated if the native base class has changed in any way
-			FKismetEditorUtilities::ConformBlueprintFlagsAndComponents(Blueprint);
-
-			if (Blueprint->GeneratedClass)
-			{
-				FBlueprintEditorUtils::RecreateClassMetaData(Blueprint, Blueprint->GeneratedClass, true);
-				Blueprint->GeneratedClass->ClassFlags &= ~CLASS_ReplicationDataIsSetUp;
-				Blueprint->GeneratedClass->SetUpRuntimeReplicationData();
-			}
-
-			// Flag data only blueprints as being up-to-date
-			Blueprint->Status = BS_UpToDate;
-		}
-		
-		// Patch the new CDOs to the old indices in the linker
-		if( Blueprint->SkeletonGeneratedClass )
-		{
-			PatchNewCDOIntoLinker(Blueprint->SkeletonGeneratedClass->GetDefaultObject(), OldLinker, OldSkelLinkerIdx, nullptr);
-		}
-		if( Blueprint->GeneratedClass )
-		{
-			PatchNewCDOIntoLinker(Blueprint->GeneratedClass->GetDefaultObject(), OldLinker, OldGenLinkerIdx, nullptr);
-		}
-
-		// Success or failure, there's no point in trying to recompile this class again when other objects reference it
-		// redo data only blueprints later, when we actually have a generated class
-		Blueprint->bHasBeenRegenerated = !FBlueprintEditorUtils::IsDataOnlyBlueprint(Blueprint) || Blueprint->GeneratedClass != nullptr; 
-
-		Blueprint->bIsRegeneratingOnLoad = false;
-
-		bRegenerated = bShouldBeRecompiled;
-
-		if (!FKismetEditorUtilities::IsClassABlueprintSkeleton(ClassToRegenerate))
-		{
-			if (!Blueprint->bRecompileOnLoad)
-			{
-				// If we didn't recompile, we still need to propagate flags, and instance components
-				FKismetEditorUtilities::ConformBlueprintFlagsAndComponents(Blueprint);
-			}
-
-			// Now that the CDO is valid, update the OwnedComponents, in case we've added or removed native components
-			if (AActor* MyActor = Cast<AActor>(Blueprint->GeneratedClass->GetDefaultObject()))
-			{
-				MyActor->ResetOwnedComponents();
-			}
-		}
-	}
-	else
-	{
-		if (Blueprint->GeneratedClass && !Blueprint->bHasBeenRegenerated && !Blueprint->bIsRegeneratingOnLoad)
-		{
-			FObjectDuplicationParameters Params(Blueprint->GeneratedClass, Blueprint->GeneratedClass->GetOuter());
-			Params.ApplyFlags = RF_Transient;
-			Params.DestName = *(FString("SKEL_COPY_") + Blueprint->GeneratedClass->GetName());
-			Blueprint->SkeletonGeneratedClass = (UClass*)StaticDuplicateObjectEx(Params);
-		}
-	}
-
-	if ( bRegenerated )
-	{		
-		// Fix any invalid metadata
-		UPackage* GeneratedClassPackage = Blueprint->GeneratedClass->GetOuterUPackage();
-		GeneratedClassPackage->GetMetaData()->RemoveMetaDataOutsidePackage();
-	}
-
-	bool const bNeedsSkelRefRemoval = !FKismetEditorUtilities::IsClassABlueprintSkeleton(ClassToRegenerate) && (Blueprint->SkeletonGeneratedClass != nullptr);
-	if (bNeedsSkelRefRemoval && Blueprint->bLegacyNeedToPurgeSkelRefs)
-	{
-		// Remove any references to the skeleton class, replacing them with refs to the generated class instead
-		FArchiveMoveSkeletalRefs SkelRefArchiver(Blueprint);
-		SkelRefArchiver.UpdateReferences();
-	}
-
-	// Restore the dirty flag
-	if( Package )
-	{
-		Package->SetDirtyFlag(bIsPackageDirty);
-	}
-
-	return bRegenerated ? Blueprint->GeneratedClass : nullptr;
-}
-
 void FBlueprintEditorUtils::LinkExternalDependencies(UBlueprint* Blueprint)
 {
 	FRegenerationHelper::LinkExternalDependencies(Blueprint);
diff --git a/Engine/Source/Editor/UnrealEd/Public/Kismet2/BlueprintEditorUtils.h b/Engine/Source/Editor/UnrealEd/Public/Kismet2/BlueprintEditorUtils.h
--- a/Engine/Source/Editor/UnrealEd/Public/Kismet2/BlueprintEditorUtils.h
+++ b/Engine/Source/Editor/UnrealEd/Public/Kismet2/BlueprintEditorUtils.h
@@ -208,12 +208,6 @@ public:
 
 	/** Helper function to punch through and honor UAnimGraphNode_Base::PreloadRequiredAssets, which formerly relied on loading assets during compile */
 	static UNREALED_API void PreloadBlueprintSpecificData(UBlueprint* Blueprint);
-
-	/**
-	 * Regenerates the class at class load time, and refreshes the blueprint
-	 */
-	UE_DEPRECATED(5.2, "Replace with FKismetEditorUtilities::CompileBlueprint")
-	static UNREALED_API UClass* RegenerateBlueprintClass(UBlueprint* Blueprint, UClass* ClassToRegenerate, UObject* PreviousCDO);
 	
 	/**
 	 * Links external dependencies
