Cooker SkipOnlyEditorOnly: Change commandline parsing to allow true/false values so we can override it back to off if config has turned it on.
Change commandline and config arguments to match the new name: OnlyEditorOnly -> SkipOnlyEditorOnly.
#rnx

diff --git a/Engine/Config/BaseEditor.ini b/Engine/Config/BaseEditor.ini
--- a/Engine/Config/BaseEditor.ini
+++ b/Engine/Config/BaseEditor.ini
@@ -318,7 +318,7 @@ MaxUsedMemory=16384
 [CookSettings]
 ; If true, the cook will only cook what is referenced by saved packages, rather than cooking every loaded package.
 ; Packages that are only referenced by editoronly properties will be skipped.
-OnlyEditorOnly=false
+SkipOnlyEditorOnly=false
 BulkDataRegistryEnabled=true
 ; HybridIterative uses the TargetDomain and BuildDependencies to decide when packages need to be recooked. If enabled,
 ; every cook is iterative for packages with allowed classes (IterativeClassAllowList) unless -noiterative is passed.
diff --git a/Engine/Source/Editor/UnrealEd/Private/CookOnTheFlyServer.cpp b/Engine/Source/Editor/UnrealEd/Private/CookOnTheFlyServer.cpp
--- a/Engine/Source/Editor/UnrealEd/Private/CookOnTheFlyServer.cpp
+++ b/Engine/Source/Editor/UnrealEd/Private/CookOnTheFlyServer.cpp
@@ -6574,13 +6574,21 @@ void UCookOnTheFlyServer::SetInitializeConfigSettings(UE::Cook::FInitializeConfi
 	// Debugging hidden dependencies
 	bOnlyEditorOnlyDebug = FParse::Param(FCommandLine::Get(), TEXT("OnlyEditorOnlyDebug"));
 	bSkipOnlyEditorOnly = false;
-	GConfig->GetBool(TEXT("CookSettings"), TEXT("OnlyEditorOnly"), bSkipOnlyEditorOnly, GEditorIni);
-	bSkipOnlyEditorOnly |= FParse::Param(FCommandLine::Get(), TEXT("OnlyEditorOnly"));
+	GConfig->GetBool(TEXT("CookSettings"), TEXT("SkipOnlyEditorOnly"), bSkipOnlyEditorOnly, GEditorIni);
+	FString ParamText;
+	if (FParse::Value(FCommandLine::Get(), TEXT("-SkipOnlyEditorOnly="), ParamText))
+	{
+		LexFromString(bSkipOnlyEditorOnly, *ParamText);
+	}
+	else if (FParse::Param(FCommandLine::Get(), TEXT("SkipOnlyEditorOnly")))
+	{
+		bSkipOnlyEditorOnly = true;
+	}
 	bSkipOnlyEditorOnly |= bOnlyEditorOnlyDebug;
 	if (bSkipOnlyEditorOnly)
 	{
 		UE_LOG(LogCook, Display,
-			TEXT("OnlyEditorOnly is enabled, unsolicited packages will not be cooked unless they are referenced from the cooked version of the instigator package."));
+			TEXT("SkipOnlyEditorOnly is enabled, unsolicited packages will not be cooked unless they are referenced from the cooked version of the instigator package."));
 	}
 
 	bHiddenDependenciesDebug = FParse::Param(FCommandLine::Get(), TEXT("HiddenDependenciesDebug"));
