Added parameter `-DlcRecook` to force a plugin to be cooked even if it has already been cooked before.

diff --git a/Engine/Source/Editor/UnrealEd/Classes/CookOnTheSide/CookOnTheFlyServer.h b/Engine/Source/Editor/UnrealEd/Classes/CookOnTheSide/CookOnTheFlyServer.h
--- a/Engine/Source/Editor/UnrealEd/Classes/CookOnTheSide/CookOnTheFlyServer.h
+++ b/Engine/Source/Editor/UnrealEd/Classes/CookOnTheSide/CookOnTheFlyServer.h
@@ -93,7 +93,7 @@ enum class ECookByTheBookOptions
 	DlcLoadMainAssetRegistry =			0x00020000, // If cooking DLC, populate the main game asset registry
 	ZenStore =							0x00040000, // Store cooked data in Zen Store
 	DlcReevaluateUncookedAssets =		0x00080000, // If cooking DLC, ignore assets in the base asset registry that were not cooked, so that this cook has an opportunity to cook the assets
-
+	DlcRecook =							0x00100000,	// If cooking DLC, force all the assets in the DLC to be cooked even if they are marked as already cooked.
 };
 ENUM_CLASS_FLAGS(ECookByTheBookOptions);
 
diff --git a/Engine/Source/Editor/UnrealEd/Private/Commandlets/CookCommandlet.cpp b/Engine/Source/Editor/UnrealEd/Private/Commandlets/CookCommandlet.cpp
--- a/Engine/Source/Editor/UnrealEd/Private/Commandlets/CookCommandlet.cpp
+++ b/Engine/Source/Editor/UnrealEd/Private/Commandlets/CookCommandlet.cpp
@@ -414,6 +414,7 @@ bool UCookCommandlet::CookByTheBook(const TArray<ITargetPlatform*>& Platforms)
 	CookOptions |= Switches.Contains(TEXT("CookAgainstFixedBase")) ? ECookByTheBookOptions::CookAgainstFixedBase : ECookByTheBookOptions::None;
 	CookOptions |= (Switches.Contains(TEXT("DlcLoadMainAssetRegistry")) || !bErrorOnEngineContentUse) ? ECookByTheBookOptions::DlcLoadMainAssetRegistry : ECookByTheBookOptions::None;
 	CookOptions |= Switches.Contains(TEXT("DlcReevaluateUncookedAssets")) ? ECookByTheBookOptions::DlcReevaluateUncookedAssets : ECookByTheBookOptions::None;
+	CookOptions |= Switches.Contains(TEXT("DlcRecook")) ? ECookByTheBookOptions::DlcRecook : ECookByTheBookOptions::None;
 	bool bCookList = Switches.Contains(TEXT("CookList"));
 
 	if (bCookSinglePackage)
diff --git a/Engine/Source/Editor/UnrealEd/Private/CookOnTheFlyServer.cpp b/Engine/Source/Editor/UnrealEd/Private/CookOnTheFlyServer.cpp
--- a/Engine/Source/Editor/UnrealEd/Private/CookOnTheFlyServer.cpp
+++ b/Engine/Source/Editor/UnrealEd/Private/CookOnTheFlyServer.cpp
@@ -11236,6 +11236,12 @@ void UCookOnTheFlyServer::RecordDLCPackagesFromBaseGame(FBeginCookContext& Begin
 			}
 		}
 	}
+
+	if (IsCookingDLC() && (!!(CookOptions & ECookByTheBookOptions::DlcRecook)))
+	{
+		// Mark all the packages in the DLC as not cooked to force them to be cooked again.
+		PackageDatas->ClearCookResultsForPlugin(CookByTheBookOptions->DlcName);
+	}
 }
 
 void UCookOnTheFlyServer::BeginCookPackageWriters(FBeginCookContext& BeginContext)
diff --git a/Engine/Source/Editor/UnrealEd/Private/Cooker/CookPackageData.cpp b/Engine/Source/Editor/UnrealEd/Private/Cooker/CookPackageData.cpp
--- a/Engine/Source/Editor/UnrealEd/Private/Cooker/CookPackageData.cpp
+++ b/Engine/Source/Editor/UnrealEd/Private/Cooker/CookPackageData.cpp
@@ -21,6 +21,7 @@
 #include "Engine/Console.h"
 #include "HAL/FileManager.h"
 #include "HAL/PlatformTime.h"
+#include "Interfaces/IPluginManager.h"
 #include "Interfaces/ITargetPlatform.h"
 #include "Misc/CommandLine.h"
 #include "Misc/CoreMiscDefines.h"
@@ -3075,6 +3076,29 @@ void FPackageDatas::ClearCookedPlatforms()
 	});
 }
 
+void FPackageDatas::ClearCookResultsForPlugin(const FString& InPluginName)
+{
+	TSharedPtr<IPlugin> Plugin = IPluginManager::Get().FindPlugin(InPluginName);
+	if (Plugin.IsValid())
+	{
+		FString PluginPath = Plugin->GetMountedAssetPath();
+		int AffectedPackagesCount = 0;
+
+		LockAndEnumeratePackageDatas([PluginPath, &AffectedPackagesCount](FPackageData* PackageData)
+			{
+				// Does this PackageData belong to the plugin?
+				const FName& PackageName = PackageData->GetPackageName();
+				if (PackageName.ToString().StartsWith(PluginPath))
+				{
+					PackageData->ClearCookResults();
+					AffectedPackagesCount++;
+				}
+			});
+
+		UE_LOG(LogCook, Display, TEXT("Cleared the cook results of %d packages because plugin '%s' was requested to be recooked."), AffectedPackagesCount, *InPluginName);
+	}
+}
+
 void FPackageDatas::OnRemoveSessionPlatform(const ITargetPlatform* TargetPlatform)
 {
 	LockAndEnumeratePackageDatas([TargetPlatform](FPackageData* PackageData)
diff --git a/Engine/Source/Editor/UnrealEd/Private/Cooker/CookPackageData.h b/Engine/Source/Editor/UnrealEd/Private/Cooker/CookPackageData.h
--- a/Engine/Source/Editor/UnrealEd/Private/Cooker/CookPackageData.h
+++ b/Engine/Source/Editor/UnrealEd/Private/Cooker/CookPackageData.h
@@ -1323,6 +1323,8 @@ public:
 	void Clear();
 	/** Set all platforms to not cooked in all PackageDatas. Used to e.g. invalidate previous cooks. */
 	void ClearCookedPlatforms();
+	/** Set all platforms to not cooked in all PackageDatas that belong to the given plugin. */
+	void ClearCookResultsForPlugin(const FString& InPluginName);
 	/** Remove all data about the given platform from all PackageDatas and other memory used by *this. */
 	void OnRemoveSessionPlatform(const ITargetPlatform* TargetPlatform);
 
