Ensure that when a package is saved multiple times (eg: when doing determinism testing) that we don't get into a situation where we attempt to save it a second time while it is still being asynchronously written the first time.

diff --git a/Engine/Source/Developer/IoStoreUtilities/Private/ZenStoreWriter.cpp b/Engine/Source/Developer/IoStoreUtilities/Private/ZenStoreWriter.cpp
--- a/Engine/Source/Developer/IoStoreUtilities/Private/ZenStoreWriter.cpp
+++ b/Engine/Source/Developer/IoStoreUtilities/Private/ZenStoreWriter.cpp
@@ -338,6 +338,11 @@ void FZenStoreWriter::WriteAdditionalFile(const FAdditionalFileInfo& Info, const
 
 	auto WriteToFile = [](const FString& Filename, const FIoBuffer& FileData)
 	{
+		ON_SCOPE_EXIT
+		{
+			UE::SavePackageUtilities::DecrementOutstandingAsyncWrites();
+		};
+
 		IFileManager& FileManager = IFileManager::Get();
 		int64 DataSize = IntCastChecked<int64>(FileData.DataSize());
 
@@ -364,6 +369,7 @@ void FZenStoreWriter::WriteAdditionalFile(const FAdditionalFileInfo& Info, const
 		UE_LOG(LogZenStoreWriter, Fatal, TEXT("Could not write to %s!"), *Filename);
 	};
 
+	UE::SavePackageUtilities::IncrementOutstandingAsyncWrites();
 	FileEntry.CompressedPayload = Async(EAsyncExecution::TaskGraph, [this, Info, FileData, WriteToFile]()
 	{
 		WriteToFile(Info.Filename, FileData);
diff --git a/Engine/Source/Editor/UnrealEd/Private/CookOnTheFlyServer.cpp b/Engine/Source/Editor/UnrealEd/Private/CookOnTheFlyServer.cpp
--- a/Engine/Source/Editor/UnrealEd/Private/CookOnTheFlyServer.cpp
+++ b/Engine/Source/Editor/UnrealEd/Private/CookOnTheFlyServer.cpp
@@ -5755,7 +5755,7 @@ void UCookOnTheFlyServer::SaveCookedPackage(UE::Cook::FSaveCookedPackageContext&
 			SaveArgs.SavePackageContext = Context.SavePackageContext;
 
 			Context.PackageWriter->UpdateSaveArguments(SaveArgs);
-			do
+			for(;;)
 			{
 				try
 				{
@@ -5769,7 +5769,19 @@ void UCookOnTheFlyServer::SaveCookedPackage(UE::Cook::FSaveCookedPackageContext&
 						*Package->GetName(), *TargetPlatform->PlatformName());
 					Context.SavePackageResult = ESavePackageResult::Error;
 				}
-			} while (Context.PackageWriter->IsAnotherSaveNeeded(Context.SavePackageResult, SaveArgs));
+
+				if (Context.PackageWriter->IsAnotherSaveNeeded(Context.SavePackageResult, SaveArgs))
+				{
+					// We must not try a second save of a package while the first save is still in flight.
+					// The optimal solution is to wait for ONLY the package that needs a second save, but we don't
+					// have the bookkeeping data to do that, so we have to wait for all async package writes to complete.
+					UPackage::WaitForAsyncFileWrites();
+				}
+				else
+				{
+					break;
+				}
+			}
 
 			// If package was actually saved check with asset manager to make sure it wasn't excluded for being a
 			// development or never cook package. But skip sending the warnings from this check if it was editor-only.
