Localization:
-Added some profiling to asset gathering to log the time parts of the asset gathering pipeline take.
- Downgraded warnings about no source files or no assets matching loc config criteria.
#rnx

diff --git a/Engine/Source/Editor/UnrealEd/Private/Commandlets/GatherTextFromAssetsCommandlet.cpp b/Engine/Source/Editor/UnrealEd/Private/Commandlets/GatherTextFromAssetsCommandlet.cpp
--- a/Engine/Source/Editor/UnrealEd/Private/Commandlets/GatherTextFromAssetsCommandlet.cpp
+++ b/Engine/Source/Editor/UnrealEd/Private/Commandlets/GatherTextFromAssetsCommandlet.cpp
@@ -559,16 +559,20 @@ int32 UGatherTextFromAssetsCommandlet::Main(const FString& Params)
 	}
 
 	UE_LOG(LogGatherTextFromAssetsCommandlet, Display, TEXT("Discovering assets to gather..."));
-
+	const double DiscoveringAssetsStartTime = FPlatformTime::Seconds();
+	const double SearchAssetRegistryStartTime = FPlatformTime::Seconds();
 	FAssetRegistryModule& AssetRegistryModule = FModuleManager::LoadModuleChecked<FAssetRegistryModule>(TEXT("AssetRegistry"));
 	IAssetRegistry& AssetRegistry = AssetRegistryModule.Get();
 	AssetRegistry.SearchAllAssets(true);
-	TArray<FAssetData> AssetDataArray;
+	UE_LOG(LogGatherTextFromAssetsCommandlet, Display, TEXT("Searching all assets took %.2f seconds."), FPlatformTime::Seconds() - SearchAssetRegistryStartTime)			;
 
+	TArray<FAssetData> AssetDataArray;
+	
 	{
 		FARFilter FirstPassFilter;
 
 		// Filter object paths to only those in any of the specified collections.
+		const double LoadCollectionObjectsStartTime = FPlatformTime::Seconds();
 		{
 			bool HasFailedToGetACollection = false;
 			FCollectionManagerModule& CollectionManagerModule = FCollectionManagerModule::GetModule();
@@ -586,7 +590,9 @@ int32 UGatherTextFromAssetsCommandlet::Main(const FString& Params)
 				return -1;
 			}
 		}
+		UE_LOG(LogGatherTextFromAssetsCommandlet, Display, TEXT("Loading all collection objects took %.2f seconds."), FPlatformTime::Seconds() - LoadCollectionObjectsStartTime);
 
+		const double FilterExcludedClassesStartTime = FPlatformTime::Seconds();
 		// Filter out any objects of the specified classes and their children at this point.
 		if (ShouldExcludeDerivedClasses)
 		{
@@ -606,22 +612,28 @@ int32 UGatherTextFromAssetsCommandlet::Main(const FString& Params)
 				}
 			}
 		}
+		UE_LOG(LogGatherTextFromAssetsCommandlet, Display, TEXT("Filtering excluded classes took %.2f seconds."), FPlatformTime::Seconds() - FilterExcludedClassesStartTime);
 
 		// Apply filter if valid to do so, get all assets otherwise.
 		if (FirstPassFilter.IsEmpty())
 		{
 			// @TODOLocalization: Logging that the first path filter is empty resulting in all assets being gathered can confuse users who generally rely on the second pass.
 			// Figure out a good way to still convey the information in a log or clog.
+			const double GetAllAssetsStartTime = FPlatformTime::Seconds();
 			AssetRegistry.GetAllAssets(AssetDataArray);
+			UE_LOG(LogGatherTextFromAssetsCommandlet, Display, TEXT("Loading all assets from asset registry took %.2f seconds."), FPlatformTime::Seconds() - GetAllAssetsStartTime);
 		}
 		else
 		{
+			const double GetAllAssetsWithFirstPassFilterStartTime = FPlatformTime::Seconds();
 			AssetRegistry.GetAssets(FirstPassFilter, AssetDataArray);
+			UE_LOG(LogGatherTextFromAssetsCommandlet, Display, TEXT("Getting all assets with first pass filter from asset registry took %.2f seconds."), FPlatformTime::Seconds() - GetAllAssetsWithFirstPassFilterStartTime);
 		}
 	}
 
 	if (!ShouldExcludeDerivedClasses)
 	{
+		const double ExcludeDerivedClassesStartTime = FPlatformTime::Seconds();
 		// Filter out any objects of the specified classes.
 		FARFilter ExcludeExactClassesFilter;
 		ExcludeExactClassesFilter.bRecursiveClasses = false;
@@ -650,12 +662,14 @@ int32 UGatherTextFromAssetsCommandlet::Main(const FString& Params)
 			{
 				return AssetsToExclude.Contains(AssetData);
 			});
+			UE_LOG(LogGatherTextFromAssetsCommandlet, Display, TEXT("Filtering otu derived classes took %.2f seconds."), FPlatformTime::Seconds() - ExcludeDerivedClassesStartTime);
 		}
 	}
 
 	// Note: AssetDataArray now contains all assets in the specified collections that are not instances of the specified excluded classes.
 
 	const FFuzzyPathMatcher FuzzyPathMatcher = FFuzzyPathMatcher(IncludePathFilters, ExcludePathFilters);
+	const double FilteringAssetsByIncludeExcludePathsStartTime = FPlatformTime::Seconds();
 	AssetDataArray.RemoveAll([&](const FAssetData& PartiallyFilteredAssetData) -> bool
 	{
 		FString PackageFilePathWithoutExtension;
@@ -700,17 +714,18 @@ int32 UGatherTextFromAssetsCommandlet::Main(const FString& Params)
 
 		return false;
 	});
+	UE_LOG(LogGatherTextFromAssetsCommandlet, Display, TEXT("Filtering assets by include exclude paths took %.2f seconds."), FPlatformTime::Seconds() - FilteringAssetsByIncludeExcludePathsStartTime);
 
 	if (AssetDataArray.Num() == 0)
 	{
-		UE_LOG(LogGatherTextFromAssetsCommandlet, Warning, TEXT("No assets matched the specified criteria."));
+		UE_LOG(LogGatherTextFromAssetsCommandlet, Display, TEXT("No assets matched the specified criteria."));
 		return 0;
 	}
 
 	// Discover the external actors for any worlds that are pending gather
 	{
 		UE_LOG(LogGatherTextFromAssetsCommandlet, Display, TEXT("Discovering external actors to gather..."));
-		
+		const double DiscoveringExternalActorsStartTime = FPlatformTime::Seconds();
 		TArray<FName> ExternalActorsSearchPaths;
 		AssetDataArray.RemoveAll([&ExternalActorsSearchPaths](const FAssetData& AssetData)
 		{
@@ -737,6 +752,7 @@ int32 UGatherTextFromAssetsCommandlet::Main(const FString& Params)
 		{
 			AssetRegistry.GetAssetsByPaths(ExternalActorsSearchPaths, AssetDataArray, /*bRecursive*/true);
 		}
+		UE_LOG(LogGatherTextFromAssetsCommandlet, Display, TEXT("Discovering external actors took %.2f seconds."), FPlatformTime::Seconds() - DiscoveringExternalActorsStartTime);
 	}
 
 	auto AppendPackagePendingGather = [this](const FName PackageNameToGather) -> FPackagePendingGather*
@@ -773,6 +789,7 @@ int32 UGatherTextFromAssetsCommandlet::Main(const FString& Params)
 			AppendPackagePendingGather(PackageNameToGather);
 		}
 	}
+	UE_LOG(LogGatherTextFromAssetsCommandlet, Display, TEXT("Discovering assests to gather took %.2f seconds."), FPlatformTime::Seconds() - DiscoveringAssetsStartTime);
 
 	FAssetGatherCacheMetrics AssetGatherCacheMetrics;
 	TMap<FString, FName> AssignedPackageLocalizationIds;
diff --git a/Engine/Source/Editor/UnrealEd/Private/Commandlets/GatherTextFromSourceCommandlet.cpp b/Engine/Source/Editor/UnrealEd/Private/Commandlets/GatherTextFromSourceCommandlet.cpp
--- a/Engine/Source/Editor/UnrealEd/Private/Commandlets/GatherTextFromSourceCommandlet.cpp
+++ b/Engine/Source/Editor/UnrealEd/Private/Commandlets/GatherTextFromSourceCommandlet.cpp
@@ -234,7 +234,7 @@ int32 UGatherTextFromSourceCommandlet::Main( const FString& Params )
 			SourceFileSearchFiltersString += FString(SourceFileSearchFiltersString.IsEmpty() ? TEXT("") : TEXT(", ")) + Filter;
 		}
 
-		UE_LOG(LogGatherTextFromSourceCommandlet, Warning, TEXT("The GatherTextFromSource commandlet couldn't find any source files matching (%s) in the specified directories:\n%s"), *SourceFileSearchFiltersString, *SpecifiedDirectoriesString);
+		UE_LOG(LogGatherTextFromSourceCommandlet, Display, TEXT("The GatherTextFromSource commandlet couldn't find any source files matching (%s) in the specified directories:\n%s"), *SourceFileSearchFiltersString, *SpecifiedDirectoriesString);
 		return 0;
 	}
 
