Don't apply editor context to reinstanced actors such as when compiling BP actors.

diff --git a/Engine/Source/Editor/UnrealEd/Private/Subsystems/ActorEditorContextSubsystem.cpp b/Engine/Source/Editor/UnrealEd/Private/Subsystems/ActorEditorContextSubsystem.cpp
--- a/Engine/Source/Editor/UnrealEd/Private/Subsystems/ActorEditorContextSubsystem.cpp
+++ b/Engine/Source/Editor/UnrealEd/Private/Subsystems/ActorEditorContextSubsystem.cpp
@@ -50,6 +50,11 @@ void UActorEditorContextSubsystem::UnregisterClient(IActorEditorContextClient* C
 
 void UActorEditorContextSubsystem::ApplyContext(AActor* InActor)
 {
+	if (GIsReinstancing)
+	{
+		return;
+	}
+
 	UWorld* World = GetWorld();
 	if (Clients.IsEmpty() || !World)
 	{
