Switch to using plugin properties for auto-detecting root plugins.

diff --git a/Engine/Source/Editor/UnrealEd/Private/CookOnTheFlyServer.cpp b/Engine/Source/Editor/UnrealEd/Private/CookOnTheFlyServer.cpp
--- a/Engine/Source/Editor/UnrealEd/Private/CookOnTheFlyServer.cpp
+++ b/Engine/Source/Editor/UnrealEd/Private/CookOnTheFlyServer.cpp
@@ -9357,6 +9357,7 @@ void UCookOnTheFlyServer::WriteCookMetadata(const ITargetPlatform* InTargetPlatf
 		}
 
 		// Construct the dependency list.
+		TArray<uint16> RootPlugins;
 		for (TSharedRef<IPlugin>& EnabledPlugin : EnabledPlugins)
 		{
 			uint16 SelfIndex = IndexForPlugin[EnabledPlugin->GetName()];
@@ -9367,6 +9368,13 @@ void UCookOnTheFlyServer::WriteCookMetadata(const ITargetPlatform* InTargetPlatf
 			Entry.DependencyIndexStart = (uint16)PluginChildArray.Num();
 
 			const FPluginDescriptor& Descriptor = EnabledPlugin->GetDescriptor();
+
+			// Root plugins are sealed && no code
+			if (Descriptor.bIsSealed && Descriptor.bNoCode)
+			{
+				RootPlugins.Add(SelfIndex);
+			}
+
 			for (FPluginReferenceDescriptor ChildPlugin : Descriptor.Plugins)
 			{
 				if (uint16* ChildIndex = IndexForPlugin.Find(ChildPlugin.Name); ChildIndex != nullptr)
@@ -9375,32 +9383,13 @@ void UCookOnTheFlyServer::WriteCookMetadata(const ITargetPlatform* InTargetPlatf
 				}
 				else
 				{
-					UE_LOG(LogCook, Display, TEXT("Dependent plugin %s wasn't found in enabled plugins list when creating cook metadata file... skipping"), *ChildPlugin.Name);
+					UE_LOG(LogCook, Display, TEXT("Dependent plugin \"%s\" referenced by \"%s\" wasn't found in enabled plugins list when creating cook metadata file... skipping"), *ChildPlugin.Name, *Descriptor.FriendlyName);
 				}
 			}
 
 			Entry.DependencyIndexEnd = PluginChildArray.Num();
 		}
 
-		// Read the root plugins from the config and make sure they exist. This is temporary until we 
-		// have the root qualification exists as part of the plugin definition.
-		TArray<uint16> RootPlugins;
-		{
-			TArray<FString> RootPluginListFromConfig;
-			GConfig->GetArray(TEXT("CookMetadata"), TEXT("RootPluginList"), RootPluginListFromConfig, GEditorIni);
-			for (const FString& Plugin : RootPluginListFromConfig)
-			{
-				uint16* RootPluginIndex = IndexForPlugin.Find(Plugin);
-				if (RootPluginIndex)
-				{
-					RootPlugins.Add(*RootPluginIndex);
-				}
-				else
-				{
-					UE_LOG(LogCook, Warning, TEXT("Unable to find root plugin %s from config list"), *Plugin);
-				}
-			}
-		}
 
 		// Add the /Engine and /Game pseudo plugins. These are placeholders for holding size information when unrealpak runs.
 		PluginsToAdd.AddDefaulted_GetRef().Name = TEXT("Engine");
