Remove filtering plugins based on platform as it's also too narrow. Without proper plugin transitivity we can't do this properly so we just include everything.

diff --git a/Engine/Source/Editor/UnrealEd/Private/CookOnTheFlyServer.cpp b/Engine/Source/Editor/UnrealEd/Private/CookOnTheFlyServer.cpp
--- a/Engine/Source/Editor/UnrealEd/Private/CookOnTheFlyServer.cpp
+++ b/Engine/Source/Editor/UnrealEd/Private/CookOnTheFlyServer.cpp
@@ -9268,17 +9268,11 @@ void UCookOnTheFlyServer::WriteCookMetadata(const ITargetPlatform* InTargetPlatf
 
 	TArray<TSharedRef<IPlugin>> EnabledPlugins = IPluginManager::Get().GetEnabledPlugins();
 
-	// Remove any plugins that don't support the target platform.
 	// NOTE: We can't use IsEnabledForPlugin because it has an issue with the AllowTargets list where preventing a plugin on
-	// a target at the uproject level doesn't get overridden by a dependent plugins' reference, so we remove plugins that
-	// are actually enabled.
-	// We switch to just directly checking the plugin if it's enabled for the platform. This won't handle transitive disabling,
-	// and also won't handle any direct disabling... however disabling should have already been taken care of by using GetEnabledPlugins
-	// as our base list.
-	EnabledPlugins.RemoveAllSwap([&PlatformNameString](const TSharedRef<IPlugin>& Plugin)
-	{
-		return Plugin->GetDescriptor().SupportsTargetPlatform(PlatformNameString) == false;
-	});
+	// a target at the uproject level doesn't get overridden by a dependent plugins' reference. This manifests as packages
+	// on disk during stage existing but the plugin isn't in the cook manifest. I wasn't able to find a way to fix this
+	// with the current plugin system, so we include all enabled plugins.
+
 
 	// Filter to the DLC plugin + dependencies if we are a DLC cook.
 	if (IsCookingDLC())
