Cooker OnlyEditorOnly: ExternalActors: Mark ExternalActors as not explorable explorable until they are loaded by the map package. The map package may skip the loading of some of them. Thsi will prevent us exploring and marking for cook all of the dependencies of a GameFeature used by a disabled Actor using that GameFeature. Modify the graph traversal so that we can mark the Actor as explorable later at the point where it loads.
#rnx

diff --git a/Engine/Source/Editor/UnrealEd/Private/CookOnTheFlyServer.cpp b/Engine/Source/Editor/UnrealEd/Private/CookOnTheFlyServer.cpp
--- a/Engine/Source/Editor/UnrealEd/Private/CookOnTheFlyServer.cpp
+++ b/Engine/Source/Editor/UnrealEd/Private/CookOnTheFlyServer.cpp
@@ -2897,7 +2897,8 @@ void UCookOnTheFlyServer::QueueDiscoveredPackage(UE::Cook::FPackageData& Package
 		DiscoveredPlatforms = ReachablePlatforms.GetPlatforms(*this, &Instigator,
 			TConstArrayView<const ITargetPlatform*>(), &BufferPlatforms);
 	}
-	if (PackageData.HasReachablePlatforms(DiscoveredPlatforms))
+	if (Instigator.Category != EInstigator::ForceExplorableSaveTimeSoftDependency &&
+		PackageData.HasReachablePlatforms(DiscoveredPlatforms))
 	{
 		// Not a new discovery; ignore
 		return;
@@ -4114,7 +4115,28 @@ void UCookOnTheFlyServer::ProcessUnsolicitedPackages(TArray<FName>* OutDiscovere
 			}
 		}
 
-		if (PackageData->FindOrAddPlatformData(CookerLoadingPlatformKey).IsReachable() || Instigator.Category == EInstigator::EditorOnlyLoad)
+		if (Instigator.Category == EInstigator::EditorOnlyLoad)
+		{
+			// This load was expected so we do not need to add a hidden dependency for it.
+			// If we are using legacy WhatGetsCookedRules, queue the package for cooking because it was loaded.
+			if (!bSkipOnlyEditorOnly)
+			{
+				QueueDiscoveredPackage(*PackageData, FInstigator(Instigator), EDiscoveredPlatformSet::CopyFromInstigator);
+			}
+		}
+		else if (bSkipOnlyEditorOnly &&
+					(Instigator.Category == EInstigator::ForceExplorableSaveTimeSoftDependency ||
+						(PackageTracker->NeverCookPackageList.Contains(PackageData->GetFileName()) &&
+							INDEX_NONE != UE::String::FindFirst(WriteToString<256>(PackageData->GetPackageName()),
+								ULevel::GetExternalActorsFolderName(), ESearchCase::IgnoreCase))))
+		{
+			// ONLYEDITORONLY_TODO: WorldPartition should mark up these loads with EInstigator::ForceExplorableSaveTimeSoftDependency rather than
+			// needing to use a naming convention. We should also mark them once, during Save, rather than marking them
+			// and reexploring them every time they are loaded.
+			Instigator.Category = EInstigator::ForceExplorableSaveTimeSoftDependency;
+			QueueDiscoveredPackage(*PackageData, FInstigator(Instigator), EDiscoveredPlatformSet::CopyFromInstigator);
+		}
+		else if (PackageData->FindOrAddPlatformData(CookerLoadingPlatformKey).IsReachable())
 		{
 			// This load was expected so we do not need to add a hidden dependency for it.
 			// If we are using legacy WhatGetsCookedRules, queue the package for cooking because it was loaded.
@@ -6883,6 +6905,7 @@ TArray<UE::Cook::FInstigator> UCookOnTheFlyServer::GetInstigatorChain(FName Pack
 			case EInstigator::EditorOnlyLoad: bGetNext = true; break;
 			case EInstigator::SaveTimeHardDependency: bGetNext = true; break;
 			case EInstigator::SaveTimeSoftDependency: bGetNext = true; break;
+			case EInstigator::ForceExplorableSaveTimeSoftDependency: bGetNext = true; break;
 			case EInstigator::GeneratedPackage: bGetNext = true; break;
 			default: break;
 		}
@@ -12184,7 +12207,8 @@ void UCookOnTheFlyServer::OnDiscoveredPackageDebug(FName PackageName, const UE::
 	switch (Instigator.Category)
 	{
 	case EInstigator::StartupPackage: [[fallthrough]];
-	case EInstigator::GeneratedPackage:
+	case EInstigator::GeneratedPackage: [[fallthrough]];
+	case EInstigator::ForceExplorableSaveTimeSoftDependency:
 		// Not a Hidden dependency
 		return;
 	default:
diff --git a/Engine/Source/Editor/UnrealEd/Private/Cooker/CookPackageData.cpp b/Engine/Source/Editor/UnrealEd/Private/Cooker/CookPackageData.cpp
--- a/Engine/Source/Editor/UnrealEd/Private/Cooker/CookPackageData.cpp
+++ b/Engine/Source/Editor/UnrealEd/Private/Cooker/CookPackageData.cpp
@@ -54,8 +54,8 @@ static FAutoConsoleVariableRef CVarPollAsyncPeriod(
 //////////////////////////////////////////////////////////////////////////
 // FPackageData
 FPackagePlatformData::FPackagePlatformData()
-	: bReachable(0), bVisitedByCluster(0), bSaveTimedOut(0), bCookable(1), bExplorable(1), bIterativelySkipped(0)
-	, bRegisteredForCachedObjectsInOuter(0), CookResults((uint8)ECookResult::NotAttempted)
+	: bReachable(0), bVisitedByCluster(0), bSaveTimedOut(0), bCookable(1), bExplorable(1), bExplorableOverride(0)
+	, bIterativelySkipped(0), bRegisteredForCachedObjectsInOuter(0), CookResults((uint8)ECookResult::NotAttempted)
 {
 }
 
@@ -67,6 +67,12 @@ void FPackagePlatformData::ResetReachable()
 	SetExplorable(true);
 }
 
+void FPackagePlatformData::MarkAsExplorable()
+{
+	ResetReachable();
+	SetExplorableOverride(true);
+}
+
 void FPackagePlatformData::MarkCookableForWorker(FCookWorkerClient& CookWorkerClient)
 {
 	SetReachable(true);
diff --git a/Engine/Source/Editor/UnrealEd/Private/Cooker/CookPackageData.h b/Engine/Source/Editor/UnrealEd/Private/Cooker/CookPackageData.h
--- a/Engine/Source/Editor/UnrealEd/Private/Cooker/CookPackageData.h
+++ b/Engine/Source/Editor/UnrealEd/Private/Cooker/CookPackageData.h
@@ -151,12 +151,26 @@ struct FPackagePlatformData
 
 	/**
 	 * Whether the package is searchable for transitive dependencies during Cluster evaluation.
-	 * This can be true even if the package is not cookable. Initially true, may be set false during Cluster search */
+	 * This can be true even if the package is not cookable. Initially true, may be set false during Cluster search
+	 */
 	bool IsExplorable() const { return bExplorable != 0; }
 	void SetExplorable(bool bValue) { bExplorable = (uint32)bValue; }
 
+	/**
+	 * Whether the package, which might be set as not explorable from IsRequestCookable, is set to explorable
+	 * based on conditions discovered during the cook.
+	 */
+	bool IsExplorableOverride() const { return bExplorableOverride != 0; }
+	void SetExplorableOverride(bool bValue) { bExplorableOverride = (uint32)bValue; }
+
 	/** All flags modified by reachability calculations are returned to default. */
 	void ResetReachable();
+	/**
+	 * Mark the platform as ExplorableOverride=true and reset all data necessary to reexplore it, including reachability.
+	 * Caller is responsbile for marking it again as reachable.
+	 */
+	void MarkAsExplorable();
+
 	/** Called on CookWorkers to indicate reachable,cookable,etc for packages sent from Director. */
 	void MarkCookableForWorker(FCookWorkerClient& CookWorkerClient);
 
@@ -181,6 +195,7 @@ private:
 	uint32 bSaveTimedOut : 1;
 	uint32 bCookable : 1;
 	uint32 bExplorable : 1;
+	uint32 bExplorableOverride : 1;
 	uint32 bIterativelySkipped : 1;
 	uint32 bRegisteredForCachedObjectsInOuter : 1;
 	uint32 CookResults : (int)ECookResult::NumBits;
diff --git a/Engine/Source/Editor/UnrealEd/Private/Cooker/CookRequestCluster.cpp b/Engine/Source/Editor/UnrealEd/Private/Cooker/CookRequestCluster.cpp
--- a/Engine/Source/Editor/UnrealEd/Private/Cooker/CookRequestCluster.cpp
+++ b/Engine/Source/Editor/UnrealEd/Private/Cooker/CookRequestCluster.cpp
@@ -25,6 +25,7 @@
 #include "Interfaces/ITargetPlatform.h"
 #include "Misc/RedirectCollector.h"
 #include "Misc/ReverseIterate.h"
+#include "Misc/ScopeExit.h"
 #include "Misc/StringBuilder.h"
 #include "String/Find.h"
 #include "UObject/CoreRedirects.h"
@@ -127,6 +128,21 @@ FRequestCluster::FRequestCluster(UCookOnTheFlyServer& InCOTFS, TRingBuffer<FDisc
 		{
 			NewReachablePlatforms = BufferPlatforms;
 		}
+		if (Discovery->Instigator.Category == EInstigator::ForceExplorableSaveTimeSoftDependency)
+		{
+			// This package was possibly previously marked as not explorable, but now it is marked as explorable.
+			// One example of this is externalactor packages - they are by default not cookable and not explorable
+			// (see comment in FRequestCluster::IsRequestCookable). But once WorldPartition loads them, we need to mark
+			// them as explored so that their imports are marked as expected and all of their soft dependencies
+			// are included.
+			for (const ITargetPlatform* TargetPlatform : NewReachablePlatforms)
+			{
+				if (TargetPlatform != CookerLoadingPlatformKey)
+				{
+					PackageData.FindOrAddPlatformData(TargetPlatform).MarkAsExplorable();
+				}
+			}
+		}
 
 		if (PackageData.HasReachablePlatforms(NewReachablePlatforms))
 		{
@@ -173,7 +189,8 @@ FRequestCluster::FRequestCluster(UCookOnTheFlyServer& InCOTFS, TRingBuffer<FDisc
 			// Adding packages to the cook should happen only for a few types of instigators, from external package
 			// requests, or during cluster exploration. If not expected, add a diagnostic message.
 			if (Discovery->Instigator.Category != EInstigator::SaveTimeHardDependency &&
-				Discovery->Instigator.Category != EInstigator::SaveTimeSoftDependency)
+				Discovery->Instigator.Category != EInstigator::SaveTimeSoftDependency &&
+				Discovery->Instigator.Category != EInstigator::ForceExplorableSaveTimeSoftDependency)
 			{
 				COTFS.OnDiscoveredPackageDebug(PackageData.GetPackageName(), Discovery->Instigator);
 			}
@@ -1487,6 +1504,13 @@ void FRequestCluster::IsRequestCookable(const ITargetPlatform* Platform, FName P
 		return;
 	}
 
+	FPackagePlatformData* PlatformData = PackageData.FindPlatformData(Platform);
+	bool bExplorableOverride = PlatformData ? PlatformData->IsExplorableOverride() : false;
+	ON_SCOPE_EXIT
+	{
+		bOutExplorable = bOutExplorable | bExplorableOverride;
+	};
+
 	FName FileName = PackageData.GetFileName();
 	if (InCOTFS.PackageTracker->NeverCookPackageList.Contains(FileName))
 	{
@@ -1496,19 +1520,13 @@ void FRequestCluster::IsRequestCookable(const ITargetPlatform* Platform, FName P
 			// NeverCookPackageList and checking naming convention here.
 			OutReason = ESuppressCookReason::NeverCook;
 			bOutCookable = false;
-			bOutExplorable = true;
-			if (InCOTFS.bSkipOnlyEditorOnly)
-			{
-				bOutExplorable = true;
-			}
-			else
-			{
-				// ONLYEDITORONLY_TODO
-				// Workaround to avoid changing behavior from legacy graph search. In legacy graph search ExternalActors
-				// were uncookable and unexplorable, and so packages referenced through externalactors referenced by e.g. GameFeatureData
-				// would not be cooked.
-				bOutExplorable = false;
-			}
+
+			// EXTERNALACTOR_TODO: We want to explore externalactors, because they add references to the cook that will
+			// otherwise not be found until the map package loads them and adds them as unsolicited packages
+			// But some externalactor packages will never be loaded by the generator, and we don't have a way to discover which
+			// ones will not be loaded until we load the Map and WorldPartition object.
+			// So set them to explorable = false until we implement an interface to determine which actors will be loaded up front.
+			bOutExplorable = false;
 		}
 		else
 		{
diff --git a/Engine/Source/Editor/UnrealEd/Private/Cooker/CookWorkerServer.cpp b/Engine/Source/Editor/UnrealEd/Private/Cooker/CookWorkerServer.cpp
--- a/Engine/Source/Editor/UnrealEd/Private/Cooker/CookWorkerServer.cpp
+++ b/Engine/Source/Editor/UnrealEd/Private/Cooker/CookWorkerServer.cpp
@@ -881,7 +881,8 @@ void FCookWorkerServer::QueueDiscoveredPackage(FDiscoveredPackageReplication&& D
 		DiscoveredPlatforms = Platforms.GetPlatforms(COTFS, &Instigator, OrderedSessionAndSpecialPlatforms, &BufferPlatforms);
 	}
 
-	if (PackageData.HasReachablePlatforms(DiscoveredPlatforms))
+	if (Instigator.Category != EInstigator::ForceExplorableSaveTimeSoftDependency &&
+		PackageData.HasReachablePlatforms(DiscoveredPlatforms))
 	{
 		// The CookWorker thought this was a new package, but the Director already knows about it; ignore the report
 		return;
diff --git a/Engine/Source/Runtime/CoreUObject/Public/UObject/ICookInfo.h b/Engine/Source/Runtime/CoreUObject/Public/UObject/ICookInfo.h
--- a/Engine/Source/Runtime/CoreUObject/Public/UObject/ICookInfo.h
+++ b/Engine/Source/Runtime/CoreUObject/Public/UObject/ICookInfo.h
@@ -52,6 +52,7 @@ namespace UE::Cook
 	CallbackMacro(EditorOnlyLoad, false) \
 	CallbackMacro(SaveTimeHardDependency, false) \
 	CallbackMacro(SaveTimeSoftDependency, false) \
+	CallbackMacro(ForceExplorableSaveTimeSoftDependency, false) \
 	CallbackMacro(GeneratedPackage, false) \
 
 
