Cooker: Fix support for AdditionalPluginDirectory plugins in projects that are outside of the UnrealEngine root directory. AssetRegistryGenerator needs to be passed a structure from the CookOnThe that knows how to map them rather than using the SandboxPlatformFile which is unaware of them.
#rnx

diff --git a/Engine/Source/Editor/UnrealEd/Classes/CookOnTheSide/CookOnTheFlyServer.h b/Engine/Source/Editor/UnrealEd/Classes/CookOnTheSide/CookOnTheFlyServer.h
--- a/Engine/Source/Editor/UnrealEd/Classes/CookOnTheSide/CookOnTheFlyServer.h
+++ b/Engine/Source/Editor/UnrealEd/Classes/CookOnTheSide/CookOnTheFlyServer.h
@@ -8,7 +8,6 @@
 #include "CookPackageSplitter.h"
 #include "HAL/PlatformMemory.h"
 #include "INetworkFileSystemModule.h"
-#include "IPlatformFileSandboxWrapper.h"
 #include "Misc/EnumClassFlags.h"
 #include "Misc/Optional.h"
 #include "Misc/PackageAccessTracking.h"
@@ -157,6 +156,7 @@ namespace UE::Cook
 	class FBuildDefinitions;
 	class FCachedDependencies;
 	class FCookDirector;
+	class FCookSandbox;
 	class FCookWorkerClient;
 	class FCookWorkerServer;
 	class FDiagnostics;
@@ -370,7 +370,7 @@ private:
 	int32 LoadBatchSize;
 
 	ECookInitializationFlags CookFlags = ECookInitializationFlags::None;
-	TUniquePtr<class FSandboxPlatformFile> SandboxFile;
+	TUniquePtr<UE::Cook::FCookSandbox> SandboxFile;
 	FString SandboxFileOutputDirectory;
 	TUniquePtr<FAsyncIODelete> AsyncIODelete; // Helper for deleting the old cook directory asynchronously
 	bool bIsSavingPackage = false; // used to stop recursive mark package dirty functions
diff --git a/Engine/Source/Editor/UnrealEd/Private/Commandlets/AssetRegistryGenerator.cpp b/Engine/Source/Editor/UnrealEd/Private/Commandlets/AssetRegistryGenerator.cpp
--- a/Engine/Source/Editor/UnrealEd/Private/Commandlets/AssetRegistryGenerator.cpp
+++ b/Engine/Source/Editor/UnrealEd/Private/Commandlets/AssetRegistryGenerator.cpp
@@ -12,6 +12,7 @@
 #include "Cooker/CookMPCollector.h"
 #include "Cooker/CookPackageData.h"
 #include "Cooker/CookPlatformManager.h"
+#include "Cooker/CookSandbox.h"
 #include "Cooker/CookWorkerClient.h"
 #include "CookMetadata.h"
 #include "Commandlets/ChunkDependencyInfo.h"
@@ -26,7 +27,6 @@
 #include "Hash/xxhash.h"
 #include "ICollectionManager.h"
 #include "Interfaces/ITargetPlatform.h"
-#include "IPlatformFileSandboxWrapper.h"
 #include "Misc/App.h"
 #include "Misc/ConfigCacheIni.h"
 #include "Misc/DataDrivenPlatformInfoRegistry.h"
@@ -336,7 +336,7 @@ static void AssignLayerChunkDelegate(const FAssignLayerChunkMap* ChunkManifest,
 }
 
 bool FAssetRegistryGenerator::GenerateStreamingInstallManifest(int64 InOverrideChunkSize, const TCHAR* InManifestSubDir,
-	FSandboxPlatformFile& InSandboxFile)
+	UE::Cook::FCookSandbox& InSandboxFile)
 {
 	const FString Platform = TargetPlatform->PlatformName();
 	FString TmpPackagingDir = GetTempPackagingDirectoryForPlatform(Platform);
@@ -665,7 +665,10 @@ bool FAssetRegistryGenerator::GenerateStreamingInstallManifest(int64 InOverrideC
 
 				for (const TSharedRef<IChunkDataGenerator>& ChunkDataGenerator : ChunkDataGenerators)
 				{
-					ChunkDataGenerator->GenerateChunkDataFiles(PakchunkIndex, PackagesInChunk, TargetPlatform, &InSandboxFile, ChunkFilenames);
+					// TOOD: Need to make a public interface for FCookSandbox and pass it into GenerateChunkDataFiles instead of the
+					// internal FSandboxPlatformFile, in case any of the generators need to correctly map files in remapped plugins
+					FSandboxPlatformFile& SandboxPlatformFile = InSandboxFile.GetSandboxPlatformFile();
+					ChunkDataGenerator->GenerateChunkDataFiles(PakchunkIndex, PackagesInChunk, TargetPlatform, &SandboxPlatformFile, ChunkFilenames);
 				}
 			}
 
@@ -804,8 +807,7 @@ bool FAssetRegistryGenerator::GenerateStreamingInstallManifest(int64 InOverrideC
 	{
 		FString ChunkManifestDirectory = FPaths::ProjectDir() / TEXT("Metadata") / TEXT("ChunkManifest");
 		ChunkManifestDirectory =
-			InSandboxFile.ConvertToAbsolutePathForExternalAppForWrite(*ChunkManifestDirectory)
-				.Replace(TEXT("[Platform]"), *Platform);
+			InSandboxFile.ConvertToAbsolutePathForExternalAppForWrite(*ChunkManifestDirectory, Platform);
 
 		IPlatformFile& PlatformFile = FPlatformFileManager::Get().GetPlatformFile();
 		if (!PlatformFile.CopyDirectoryTree(*ChunkManifestDirectory, *TmpPackagingDir, true))
@@ -819,7 +821,7 @@ bool FAssetRegistryGenerator::GenerateStreamingInstallManifest(int64 InOverrideC
 }
 
 void FAssetRegistryGenerator::CalculateChunkIdsAndAssignToManifest(const FName& PackageFName, const FString& PackagePathName,
-	const FString& SandboxFilename, const FString& LastLoadedMapName, FSandboxPlatformFile& InSandboxFile)
+	const FString& SandboxFilename, const FString& LastLoadedMapName, UE::Cook::FCookSandbox& InSandboxFile)
 {
 	TArray<int32> TargetChunks;
 	TArray<int32> ExistingChunkIDs;
@@ -964,7 +966,7 @@ void FAssetRegistryGenerator::InjectEncryptionData(FAssetRegistryState& TargetSt
 	}
 }
 
-bool FAssetRegistryGenerator::SaveManifests(FSandboxPlatformFile& InSandboxFile, int64 InOverrideChunkSize,
+bool FAssetRegistryGenerator::SaveManifests(UE::Cook::FCookSandbox& InSandboxFile, int64 InOverrideChunkSize,
 	const TCHAR* InManifestSubDir)
 {
 	LLM_SCOPE_BYTAG(Cooker_GeneratedAssetRegistry);
@@ -1553,7 +1555,7 @@ void FAssetRegistryGenerator::ComputePackageRemovals(const FAssetRegistryState&
 }
 
 void FAssetRegistryGenerator::FinalizeChunkIDs(const TSet<FName>& InCookedPackages,
-	const TSet<FName>& InDevelopmentOnlyPackages, FSandboxPlatformFile& InSandboxFile,
+	const TSet<FName>& InDevelopmentOnlyPackages, UE::Cook::FCookSandbox& InSandboxFile,
 	bool bGenerateStreamingInstallManifest)
 {
 	LLM_SCOPE_BYTAG(Cooker_GeneratedAssetRegistry);
@@ -1876,13 +1878,13 @@ bool FAssetRegistryGenerator::SaveAssetRegistry(const FString& SandboxPath, bool
 
 class FPackageCookerOpenOrderVisitor : public IPlatformFile::FDirectoryVisitor
 {
-	const FSandboxPlatformFile& SandboxFile;
+	const UE::Cook::FCookSandbox& SandboxFile;
 	const FString& PlatformSandboxPath;
 	const TSet<FStringView>& ValidExtensions;
 	TMultiMap<FString, FString>& PackageExtensions;
 public:
 	FPackageCookerOpenOrderVisitor(
-		const FSandboxPlatformFile& InSandboxFile,
+		const UE::Cook::FCookSandbox& InSandboxFile,
 		const FString& InPlatformSandboxPath,
 		const TSet<FStringView>& InValidExtensions,
 		TMultiMap<FString, FString>& OutPackageExtensions) :
@@ -1897,7 +1899,7 @@ public:
 		if (bIsDirectory)
 			return true;
 
-		FString Filename = FilenameOrDirectory;
+		const TCHAR* Filename = FilenameOrDirectory;
 		FStringView UnusedFilePath, FileBaseName, FileExtension;
 		FPathViews::Split(Filename, UnusedFilePath, FileBaseName, FileExtension);
 		if (ValidExtensions.Contains(FileExtension))
@@ -1909,8 +1911,7 @@ public:
 			}
 
 			FString PackageName;
-			Filename.ReplaceInline(*PlatformSandboxPath, *SandboxFile.GetSandboxDirectory());
-			FString AssetSourcePath = SandboxFile.ConvertFromSandboxPath(*Filename);
+			FString AssetSourcePath = SandboxFile.ConvertFromSandboxPathInPlatformRoot(Filename, PlatformSandboxPath);
 			FString StandardAssetSourcePath = FPaths::CreateStandardFilename(AssetSourcePath);
 			if (StandardAssetSourcePath.EndsWith(TEXT(".m.ubulk")))
 			{
@@ -1931,7 +1932,7 @@ public:
 	}
 };
 
-bool FAssetRegistryGenerator::WriteCookerOpenOrder(FSandboxPlatformFile& InSandboxFile)
+bool FAssetRegistryGenerator::WriteCookerOpenOrder(UE::Cook::FCookSandbox& InSandboxFile)
 {
 	LLM_SCOPE_BYTAG(Cooker_GeneratedAssetRegistry);
 	TSet<FName> PackageNameSet;
@@ -2344,7 +2345,7 @@ bool FAssetRegistryGenerator::CheckChunkAssetsAreNotInChild(const FChunkDependen
 }
 
 void FAssetRegistryGenerator::AddPackageToChunk(FChunkPackageSet& ThisPackageSet, FName InPkgName,
-	const FString& InSandboxFile, int32 PakchunkIndex, FSandboxPlatformFile& SandboxPlatformFile)
+	const FString& InSandboxFile, int32 PakchunkIndex, UE::Cook::FCookSandbox& SandboxPlatformFile)
 {
 	ThisPackageSet.Add(InPkgName, InSandboxFile);
 }
@@ -2354,7 +2355,7 @@ FString FAssetRegistryGenerator::GetTempPackagingDirectoryForPlatform(const FStr
 	return FPaths::ProjectSavedDir() / TEXT("TmpPackaging") / Platform;
 }
 
-void FAssetRegistryGenerator::FixupPackageDependenciesForChunks(FSandboxPlatformFile& InSandboxFile)
+void FAssetRegistryGenerator::FixupPackageDependenciesForChunks(UE::Cook::FCookSandbox& InSandboxFile)
 {
 	UE_LOG(LogAssetRegistryGenerator, Log, TEXT("Starting FixupPackageDependenciesForChunks..."));
 	SCOPE_LOG_TIME_IN_SECONDS(TEXT("... FixupPackageDependenciesForChunks complete."), nullptr);
diff --git a/Engine/Source/Editor/UnrealEd/Private/CookOnTheFlyServer.cpp b/Engine/Source/Editor/UnrealEd/Private/CookOnTheFlyServer.cpp
--- a/Engine/Source/Editor/UnrealEd/Private/CookOnTheFlyServer.cpp
+++ b/Engine/Source/Editor/UnrealEd/Private/CookOnTheFlyServer.cpp
@@ -29,6 +29,7 @@
 #include "Cooker/CookProfiling.h"
 #include "Cooker/CookRequestCluster.h"
 #include "Cooker/CookRequests.h"
+#include "Cooker/CookSandbox.h"
 #include "Cooker/CookTypes.h"
 #include "Cooker/CookWorkerClient.h"
 #include "Cooker/DiffPackageWriter.h"
@@ -8768,41 +8769,19 @@ void UCookOnTheFlyServer::SaveGlobalShaderMapFiles(const TArrayView<const ITarge
 
 FString UCookOnTheFlyServer::GetSandboxDirectory( const FString& PlatformName ) const
 {
-	FString Result;
-	Result = SandboxFile->GetSandboxDirectory();
-
-	Result.ReplaceInline(TEXT("[Platform]"), *PlatformName);
-
-	return Result;
+	return SandboxFile->GetSandboxDirectory(PlatformName);
 }
 
 FString UCookOnTheFlyServer::ConvertToFullSandboxPath( const FString &FileName, bool bForWrite ) const
 {
-	check( SandboxFile );
-
-	FString Result;
-	if (bForWrite)
-	{
-		if (FLooseCookedPackageWriter::TryConvertUncookedFilenameToCookedRemappedPluginFilename(
-			FileName, PluginsToRemap, SandboxFile->GetSandboxDirectory(), Result))
-		{
-			return Result;
-		}
-		Result = SandboxFile->ConvertToAbsolutePathForExternalAppForWrite(*FileName);
-	}
-	else
-	{
-		Result = SandboxFile->ConvertToAbsolutePathForExternalAppForRead(*FileName);
-	}
-
-	return Result;
+	check(SandboxFile);
+	return SandboxFile->ConvertToFullSandboxPath(FileName, bForWrite);
 }
 
 FString UCookOnTheFlyServer::ConvertToFullSandboxPath( const FString &FileName, bool bForWrite, const FString& PlatformName ) const
 {
-	FString Result = ConvertToFullSandboxPath( FileName, bForWrite );
-	Result.ReplaceInline(TEXT("[Platform]"), *PlatformName);
-	return Result;
+	check(SandboxFile);
+	return SandboxFile->ConvertToFullPlatformSandboxPath(FileName, bForWrite, PlatformName);
 }
 
 FString UCookOnTheFlyServer::GetSandboxAssetRegistryFilename()
@@ -10108,13 +10087,9 @@ void UCookOnTheFlyServer::CreateSandboxFile(FBeginCookContext& BeginContext)
 		SandboxFile.Reset();
 	}
 
-	// Local sandbox file wrapper. This will be used to handle path conversions, but will not be used to actually
-	// write/read files so we can safely use [Platform] token in the sandbox directory name and then replace it
-	// with the actual platform name.
 	// Filename lookups in the cooker must Use this SandboxFile to do path conversion to properly handle sandbox paths
 	// (outside of standard paths in particular).
-	SandboxFile = FSandboxPlatformFile::Create(false);
-	SandboxFile->Initialize(&FPlatformFileManager::Get().GetPlatformFile(), *FString::Printf(TEXT("-sandbox=\"%s\""), *OutputDirectory));
+	SandboxFile.Reset(new UE::Cook::FCookSandbox(OutputDirectory, PluginsToRemap));
 	SandboxFileOutputDirectory = OutputDirectory;
 
 }
@@ -10428,7 +10403,7 @@ UE::Cook::FCookSavePackageContext* UCookOnTheFlyServer::CreateSaveContext(const
 	else
 	{
 		PackageWriter = new FLooseCookedPackageWriter(ResolvedRootPath, ResolvedMetadataPath, TargetPlatform,
-			GetAsyncIODelete(), PluginsToRemap, MoveTemp(BeginCacheCallback));
+			GetAsyncIODelete(), *SandboxFile, MoveTemp(BeginCacheCallback));
 		WriterDebugName = TEXT("LooseCookedPackageWriter");
 	}
 
diff --git a/Engine/Source/Editor/UnrealEd/Private/Cooker/CookSandbox.cpp b/Engine/Source/Editor/UnrealEd/Private/Cooker/CookSandbox.cpp
new file mode 100644
--- /dev/null
+++ b/Engine/Source/Editor/UnrealEd/Private/Cooker/CookSandbox.cpp
@@ -0,0 +1,287 @@
+// Copyright Epic Games, Inc. All Rights Reserved.
+
+#include "CookSandbox.h"
+
+#include "CookTypes.h"
+#include "HAL/PlatformFileManager.h"
+#include "Interfaces/IPluginManager.h"
+#include "IPlatformFileSandboxWrapper.h"
+#include "Misc/App.h"
+#include "Misc/CString.h"
+#include "Misc/Paths.h"
+#include "String/Find.h"
+
+namespace UE::Cook
+{
+
+FCookSandbox::FCookSandbox(FStringView OutputDirectory, TArray<TSharedRef<IPlugin>>& InPluginsToRemap)
+	: PluginsToRemap(InPluginsToRemap)
+{
+	// Local sandbox file wrapper. This will be used to handle path conversions, but will not be used to actually
+	// write/read files so we can safely use [Platform] token in the sandbox directory name and then replace it
+	// with the actual platform name.
+	SandboxFile = FSandboxPlatformFile::Create(false);
+	SandboxFile->Initialize(&FPlatformFileManager::Get().GetPlatformFile(),
+		*FString::Printf(TEXT("-sandbox=\"%.*s\""), OutputDirectory.Len(), OutputDirectory.GetData()));
+}
+
+const FString& FCookSandbox::GetSandboxDirectory() const
+{
+	return SandboxFile->GetSandboxDirectory();
+}
+
+const FString& FCookSandbox::GetGameSandboxDirectoryName() const
+{
+	return SandboxFile->GetGameSandboxDirectoryName();
+}
+
+FString FCookSandbox::ConvertToAbsolutePathForExternalAppForWrite(const TCHAR* Filename) const
+{
+	return SandboxFile->ConvertToAbsolutePathForExternalAppForWrite(Filename);
+}
+
+FString FCookSandbox::ConvertFromSandboxPath(const TCHAR* Filename) const
+{
+	return ConvertFromSandboxPathInPlatformRoot(Filename, GetSandboxDirectory());
+}
+
+FString FCookSandbox::GetSandboxDirectory(const FString& PlatformName) const
+{
+	FString Result = SandboxFile->GetSandboxDirectory();
+	Result.ReplaceInline(TEXT("[Platform]"), *PlatformName);
+	return Result;
+}
+
+FString FCookSandbox::ConvertToAbsolutePathForExternalAppForWrite(const TCHAR* Filename, const FString& PlatformName) const
+{
+	FString Result = SandboxFile->ConvertToAbsolutePathForExternalAppForWrite(Filename);
+	Result.ReplaceInline(TEXT("[Platform]"), *PlatformName);
+	return Result;
+}
+
+FString FCookSandbox::ConvertFromSandboxPathInPlatformRoot(const TCHAR* Filename, FStringView PlatformSandboxRootDir) const
+{
+	FCookSandboxConvertCookedPathToPackageNameContext Context;
+	Context.SandboxRootDir = PlatformSandboxRootDir;
+	FillContext(Context);
+	return ConvertCookedPathToUncookedPath(Filename, Context);
+}
+
+bool FCookSandbox::TryConvertUncookedFilenameToCookedRemappedPluginFilename(FStringView FileName,
+	FString& OutCookedFileName, FStringView PlatformSandboxRootDir) const
+{
+	// Ideally this would be in the Sandbox File but it can't access the project or plugin
+	if (PluginsToRemap.IsEmpty())
+	{
+		return false;
+	}
+	if (PlatformSandboxRootDir.IsEmpty())
+	{
+		PlatformSandboxRootDir = GetSandboxDirectory();
+	}
+
+	for (const TSharedRef<IPlugin>& Plugin : PluginsToRemap)
+	{
+		// If these match, then this content is part of plugin that gets remapped when packaged/staged
+		if (FileName.StartsWith(Plugin->GetContentDir()))
+		{
+			FString SearchFor;
+			SearchFor /= Plugin->GetName() / TEXT("Content");
+			int32 FoundAt = UE::String::FindLast(FileName, SearchFor, ESearchCase::IgnoreCase);
+			check(FoundAt != INDEX_NONE);
+			// Strip off everything but <PluginName/Content/<remaing path to file>
+			FStringView SnippedOffPath = FileName.RightChop(FoundAt);
+			// Put this is in <sandbox path>/RemappedPlugins/<PluginName>/Content/<remaining path to file>
+			constexpr FStringView RemappedPluginsDirName(REMAPPED_PLUGINS);
+			OutCookedFileName.Reserve(PlatformSandboxRootDir.Len() + RemappedPluginsDirName.Len() + SnippedOffPath.Len() + 2);
+			OutCookedFileName = PlatformSandboxRootDir;
+			OutCookedFileName /= REMAPPED_PLUGINS;
+			OutCookedFileName /= SnippedOffPath;
+			return true;
+		}
+	}
+	return false;
+}
+
+FSandboxPlatformFile& FCookSandbox::GetSandboxPlatformFile()
+{
+	return *SandboxFile;
+}
+
+FString FCookSandbox::ConvertToFullSandboxPath(const FString& FileName, bool bForWrite) const
+{
+	return ConvertToFullSandboxPathInPlatformRoot(FileName, bForWrite, GetSandboxDirectory());
+}
+
+FString FCookSandbox::ConvertToFullPlatformSandboxPath(const FString& FileName, bool bForWrite,
+	const FString& PlatformName) const
+{
+	FString Result = ConvertToFullSandboxPath(FileName, bForWrite);
+	Result.ReplaceInline(TEXT("[Platform]"), *PlatformName);
+	return Result;
+}
+
+FString FCookSandbox::ConvertToFullSandboxPathInPlatformRoot(const FString& FileName, bool bForWrite,
+	FStringView PlatformSandboxRootDir) const
+{
+	FString Result;
+	if (bForWrite)
+	{
+		if (TryConvertUncookedFilenameToCookedRemappedPluginFilename(FileName, Result, PlatformSandboxRootDir))
+		{
+			return Result;
+		}
+		Result = SandboxFile->ConvertToAbsolutePathForExternalAppForWrite(*FileName);
+	}
+	else
+	{
+		Result = SandboxFile->ConvertToAbsolutePathForExternalAppForRead(*FileName);
+	}
+
+	return Result;
+}
+
+void FCookSandbox::FillContext(FCookSandboxConvertCookedPathToPackageNameContext& Context) const
+{
+	if (Context.SandboxRootDir.IsEmpty())
+	{
+		Context.SandboxRootDir = GetSandboxDirectory();
+	}
+	if (Context.UncookedRelativeRootDir.IsEmpty())
+	{
+		Context.UncookedRelativeRootDir = FPaths::GetRelativePathToRoot();
+	}
+	if (Context.SandboxProjectDir.IsEmpty())
+	{
+		Context.ScratchSandboxProjectDir = FPaths::Combine(Context.SandboxRootDir, FApp::GetProjectName()) + TEXT("/");
+		Context.SandboxProjectDir = Context.ScratchSandboxProjectDir;
+	}
+	if (Context.UncookedRelativeProjectDir.IsEmpty())
+	{
+		Context.ScratchUncookedRelativeProjectDir = FPaths::ProjectDir();
+		Context.UncookedRelativeProjectDir = Context.ScratchUncookedRelativeProjectDir;
+	}
+	Context.ScratchFileName.Reserve(1024);
+	Context.ScratchPackageName.Reserve(1024);
+}
+
+FName FCookSandbox::ConvertCookedPathToPackageName(FStringView CookedPath,
+	FCookSandboxConvertCookedPathToPackageNameContext& Context) const
+{
+	FString& UncookedFileName = ConvertCookedPathToUncookedPath(CookedPath, Context);
+	if (!FPackageName::TryConvertFilenameToLongPackageName(UncookedFileName, Context.ScratchPackageName))
+	{
+		return NAME_None;
+	}
+
+	return FName(*Context.ScratchPackageName);
+}
+
+FString& FCookSandbox::ConvertCookedPathToUncookedPath(FStringView CookedPath,
+	FCookSandboxConvertCookedPathToPackageNameContext& Context) const
+{
+	FString& UncookedFileName = Context.ScratchFileName;
+	UncookedFileName.Reset();
+	constexpr FStringView RemappedPluginsFolder(REMAPPED_PLUGINS);
+	constexpr FStringView ContentFolder(TEXTVIEW("Content/"));
+
+	// Check for remapped plugins' cooked content
+	if (PluginsToRemap.Num() > 0)
+	{
+		int32 RemappedIndex = UE::String::FindFirst(CookedPath, RemappedPluginsFolder, ESearchCase::IgnoreCase);
+		if (RemappedIndex >= 0)
+		{
+			// Snip everything up through the RemappedPlugins/ off so we can find the plugin it corresponds to
+			FStringView PluginPath = CookedPath.RightChop(RemappedIndex + RemappedPluginsFolder.Len() + 1);
+			// Find the plugin that owns this content
+			for (const TSharedRef<IPlugin>& Plugin : PluginsToRemap)
+			{
+				if (PluginPath.StartsWith(Plugin->GetName()))
+				{
+					UncookedFileName = Plugin->GetContentDir();
+					// Chop off the pluginName/Content since it's part of the full path
+					UncookedFileName /= PluginPath.RightChop(Plugin->GetName().Len() + ContentFolder.Len());
+					break;
+				}
+			}
+		}
+		// If we did not find a containing plugin, fall through to sandbox handling to construct UncookedFileName
+	}
+
+	if (UncookedFileName.IsEmpty())
+	{
+		auto BuildUncookedPath =
+			[&UncookedFileName](FStringView CookedPath, FStringView CookedRoot, FStringView UncookedRoot)
+		{
+			UncookedFileName.AppendChars(UncookedRoot.GetData(), UncookedRoot.Len());
+			UncookedFileName.AppendChars(CookedPath.GetData() + CookedRoot.Len(), CookedPath.Len() - CookedRoot.Len());
+		};
+
+		if (CookedPath.StartsWith(Context.SandboxRootDir))
+		{
+			// Optimized CookedPath.StartsWith(SandboxProjectDir) that does not compare all of SandboxRootDir again
+			if (CookedPath.Len() >= Context.SandboxProjectDir.Len() &&
+				0 == FCString::Strnicmp(
+					CookedPath.GetData() + Context.SandboxRootDir.Len(),
+					Context.SandboxProjectDir.GetData() + Context.SandboxRootDir.Len(),
+					Context.SandboxProjectDir.Len() - Context.SandboxRootDir.Len()))
+			{
+				BuildUncookedPath(CookedPath, Context.SandboxProjectDir, Context.UncookedRelativeProjectDir);
+			}
+			else
+			{
+				BuildUncookedPath(CookedPath, Context.SandboxRootDir, Context.UncookedRelativeRootDir);
+			}
+		}
+		else
+		{
+			TStringBuilder<1024> FullCookedFilename;
+			FPathViews::ToAbsolutePath(CookedPath, FullCookedFilename);
+			BuildUncookedPath(FullCookedFilename, Context.SandboxRootDir, Context.UncookedRelativeRootDir);
+		}
+	}
+	return Context.ScratchFileName;
+}
+
+FString FCookSandbox::ConvertPackageNameToCookedPath(FStringView PackageName,
+	FCookSandboxConvertCookedPathToPackageNameContext& Context) const
+{
+	FString UncookedFileName;
+	Context.ScratchPackageName = PackageName;
+	if (!FPackageName::TryConvertLongPackageNameToFilename(Context.ScratchPackageName, UncookedFileName))
+	{
+		return TEXT("");
+	}
+
+	FString CookedFileName;
+	if (TryConvertUncookedFilenameToCookedRemappedPluginFilename(UncookedFileName, CookedFileName,
+		Context.SandboxRootDir))
+	{
+		return CookedFileName;
+	}
+
+	auto BuildCookedPath =
+		[&CookedFileName](FStringView UncookedPath, FStringView CookedRoot, FStringView UncookedRoot)
+	{
+		CookedFileName.Reset();
+		CookedFileName.AppendChars(CookedRoot.GetData(), CookedRoot.Len());
+		CookedFileName.AppendChars(UncookedPath.GetData() + UncookedRoot.Len(), UncookedPath.Len() - UncookedRoot.Len());
+	};
+
+	if (UncookedFileName.StartsWith(Context.UncookedRelativeProjectDir))
+	{
+		BuildCookedPath(UncookedFileName, Context.SandboxProjectDir, Context.UncookedRelativeProjectDir);
+	}
+	else if (UncookedFileName.StartsWith(Context.UncookedRelativeRootDir))
+	{
+		BuildCookedPath(UncookedFileName, Context.SandboxRootDir, Context.UncookedRelativeRootDir);
+	}
+	else
+	{
+		CookedFileName.Empty();
+	}
+	return CookedFileName;
+}
+
+
+} // namespace UE::Cook
\ No newline at end of file
diff --git a/Engine/Source/Editor/UnrealEd/Private/Cooker/CookSandbox.h b/Engine/Source/Editor/UnrealEd/Private/Cooker/CookSandbox.h
new file mode 100644
--- /dev/null
+++ b/Engine/Source/Editor/UnrealEd/Private/Cooker/CookSandbox.h
@@ -0,0 +1,77 @@
+// Copyright Epic Games, Inc. All Rights Reserved.
+
+#pragma once
+
+#include "Containers/ArrayView.h"
+#include "Containers/StringView.h"
+#include "Containers/UnrealString.h"
+#include "IPlatformFileSandboxWrapper.h"
+#include "Templates/SharedPointer.h"
+#include "Templates/UniquePtr.h"
+
+class IPlugin;
+
+namespace UE::Cook
+{
+
+struct FCookSandboxConvertCookedPathToPackageNameContext;
+
+/**
+ * A wrapper around FSandboxPlatformFile that provides a similar interface but also handles cook-specific
+ * functionality like REMAPPED_PLUGINS
+ */
+class FCookSandbox
+{
+public:
+	FCookSandbox(FStringView OutputDirectory, TArray<TSharedRef<IPlugin>>& InPluginsToRemap);
+
+	// ISandboxPlatformFile
+	const FString& GetSandboxDirectory() const;
+	const FString& GetGameSandboxDirectoryName() const;
+	FString ConvertToAbsolutePathForExternalAppForWrite(const TCHAR* Filename) const;
+	FString ConvertFromSandboxPath(const TCHAR* Filename) const;
+
+	// Cooker API functions that handle the Platform mapping on top of ISandboxPlatformFile
+	FString GetSandboxDirectory(const FString& PlatformName) const;
+	FString ConvertToAbsolutePathForExternalAppForWrite(const TCHAR* Filename, const FString& PlatformName) const;
+	FString ConvertFromSandboxPathInPlatformRoot(const TCHAR* Filename, FStringView PlatformSandboxRootDir) const;
+
+	// Cooker API functions on top of ISandboxPlatformFile
+	FSandboxPlatformFile& GetSandboxPlatformFile();
+	FString ConvertToFullSandboxPath(const FString& FileName, bool bForWrite) const;
+	FString ConvertToFullPlatformSandboxPath(const FString& FileName, bool bForWrite,
+		const FString& PlatformName) const;
+	FString ConvertToFullSandboxPathInPlatformRoot(const FString& FileName, bool bForWrite,
+		FStringView PlatformSandboxRootdir) const;
+
+	bool TryConvertUncookedFilenameToCookedRemappedPluginFilename(FStringView FileName, FString& OutCookedFileName,
+		FStringView PlatformSandboxRootdir = FStringView()) const;
+
+	void FillContext(FCookSandboxConvertCookedPathToPackageNameContext& Context) const;
+	FString& ConvertCookedPathToUncookedPath(FStringView CookedPath,
+		FCookSandboxConvertCookedPathToPackageNameContext& Context) const;
+	FName ConvertCookedPathToPackageName(FStringView CookedPath,
+		FCookSandboxConvertCookedPathToPackageNameContext& Context) const;
+	FString ConvertPackageNameToCookedPath(FStringView PackageName,
+		FCookSandboxConvertCookedPathToPackageNameContext& Context) const;
+
+
+private:
+	TUniquePtr<FSandboxPlatformFile> SandboxFile;
+	TArray<TSharedRef<IPlugin>>& PluginsToRemap;
+};
+
+struct FCookSandboxConvertCookedPathToPackageNameContext
+{
+	FStringView SandboxRootDir;
+	FStringView UncookedRelativeRootDir;
+	FStringView SandboxProjectDir;
+	FStringView UncookedRelativeProjectDir;
+	FString ScratchSandboxProjectDir;
+	FString ScratchUncookedRelativeProjectDir;
+	FString ScratchFileName;
+	FString ScratchPackageName;
+};
+
+
+} // namespace UE::Cook
\ No newline at end of file
diff --git a/Engine/Source/Editor/UnrealEd/Private/Cooker/LooseCookedPackageWriter.cpp b/Engine/Source/Editor/UnrealEd/Private/Cooker/LooseCookedPackageWriter.cpp
--- a/Engine/Source/Editor/UnrealEd/Private/Cooker/LooseCookedPackageWriter.cpp
+++ b/Engine/Source/Editor/UnrealEd/Private/Cooker/LooseCookedPackageWriter.cpp
@@ -12,6 +12,7 @@
 #include "Cooker/AsyncIODelete.h"
 #include "Cooker/CompactBinaryTCP.h"
 #include "Cooker/CookPackageData.h"
+#include "Cooker/CookSandbox.h"
 #include "Cooker/CookTypes.h"
 #include "HAL/FileManager.h"
 #include "HAL/PlatformCrt.h"
@@ -21,7 +22,6 @@
 #include "HAL/PlatformString.h"
 #include "Hash/Blake3.h"
 #include "IO/IoHash.h"
-#include "Interfaces/IPluginManager.h"
 #include "Interfaces/ITargetPlatform.h"
 #include "Logging/LogCategory.h"
 #include "Logging/LogMacros.h"
@@ -59,11 +59,11 @@ LLM_DEFINE_TAG(Cooker_PackageStoreManifest);
 
 FLooseCookedPackageWriter::FLooseCookedPackageWriter(const FString& InOutputPath,
 	const FString& InMetadataDirectoryPath, const ITargetPlatform* InTargetPlatform, FAsyncIODelete& InAsyncIODelete,
-	const TArray<TSharedRef<IPlugin>>& InPluginsToRemap, FBeginCacheCallback&& InBeginCacheCallback)
+	UE::Cook::FCookSandbox& InSandboxFile, FBeginCacheCallback&& InBeginCacheCallback)
 	: OutputPath(InOutputPath)
 	, MetadataDirectoryPath(InMetadataDirectoryPath)
 	, TargetPlatform(*InTargetPlatform)
-	, PluginsToRemap(InPluginsToRemap)
+	, SandboxFile(InSandboxFile)
 	, AsyncIODelete(InAsyncIODelete)
 	, BeginCacheCallback(MoveTemp(InBeginCacheCallback))
 {
@@ -1000,20 +1000,15 @@ void FLooseCookedPackageWriter::GetAllCookedFiles()
 	}
 
 	const FString SandboxProjectDir = FPaths::Combine(OutputPath, FApp::GetProjectName()) + TEXT("/");
-	const FString RelativeRootDir = FPaths::GetRelativePathToRoot();
-	const FString RelativeProjectDir = FPaths::ProjectDir();
-	FString ScratchFileName;
-	ScratchFileName.Reserve(1024);
-	FString ScratchPackageName;
-	ScratchPackageName.Reserve(1024);
+
+	UE::Cook::FCookSandboxConvertCookedPathToPackageNameContext Context;
+	Context.SandboxRootDir = SandboxRootDir;
+	Context.SandboxProjectDir = SandboxProjectDir;
+	SandboxFile.FillContext(Context);
 
 	for (const FString& CookedFile : CookedFiles)
 	{
-		const FName PackageName = ConvertCookedPathToPackageName(
-			SandboxRootDir, RelativeRootDir,
-			SandboxProjectDir, RelativeProjectDir,
-			CookedFile, ScratchFileName, ScratchPackageName);
-
+		const FName PackageName = SandboxFile.ConvertCookedPathToPackageName(CookedFile, Context);
 		if (!PackageName.IsNone())
 		{
 			PackageNameToCookedFiles.FindOrAdd(PackageName).Add(CookedFile);
@@ -1025,15 +1020,15 @@ void FLooseCookedPackageWriter::FindAndDeleteCookedFilesForPackages(TConstArrayV
 {
 	const FString& SandboxRootDir = OutputPath;
 	const FString SandboxProjectDir = FPaths::Combine(OutputPath, FApp::GetProjectName()) + TEXT("/");
-	const FString RelativeRootDir = FPaths::GetRelativePathToRoot();
-	const FString RelativeProjectDir = FPaths::ProjectDir();
+
+	UE::Cook::FCookSandboxConvertCookedPathToPackageNameContext Context;
+	Context.SandboxRootDir = SandboxRootDir;
+	Context.SandboxProjectDir = SandboxProjectDir;
+	SandboxFile.FillContext(Context);
 
 	for (FName PackageName : PackageNames)
 	{
-		FString CookedFileName = ConvertPackageNameToCookedPath(
-			SandboxRootDir, RelativeRootDir,
-			SandboxProjectDir, RelativeProjectDir,
-			WriteToString<256>(PackageName));
+		FString CookedFileName = SandboxFile.ConvertPackageNameToCookedPath(WriteToString<256>(PackageName), Context);
 		if (CookedFileName.IsEmpty())
 		{
 			continue;
@@ -1072,151 +1067,6 @@ void FLooseCookedPackageWriter::FindAndDeleteCookedFilesForPackages(TConstArrayV
 	}
 }
 
-FName FLooseCookedPackageWriter::ConvertCookedPathToPackageName(
-	const FString& SandboxRootDir, const FString& RelativeRootDir,
-	const FString& SandboxProjectDir, const FString& RelativeProjectDir,
-	const FString& CookedPath, FString& ScratchFileName, FString& ScratchPackageName) const
-{
-	FString& UncookedFileName = ScratchFileName; 
-	UncookedFileName.Reset();
-	constexpr FStringView RemappedPluginsFolder(REMAPPED_PLUGINS);
-	constexpr FStringView ContentFolder(TEXTVIEW("Content/"));
-
-	// Check for remapped plugins' cooked content
-	if (PluginsToRemap.Num() > 0 && CookedPath.Contains(RemappedPluginsFolder))
-	{
-		int32 RemappedIndex = CookedPath.Find(RemappedPluginsFolder);
-		check(RemappedIndex >= 0);
-		// Snip everything up through the RemappedPlugins/ off so we can find the plugin it corresponds to
-		FString PluginPath = CookedPath.RightChop(RemappedIndex + RemappedPluginsFolder.Len() + 1);
-		// Find the plugin that owns this content
-		for (const TSharedRef<IPlugin>& Plugin : PluginsToRemap)
-		{
-			if (PluginPath.StartsWith(Plugin->GetName()))
-			{
-				UncookedFileName = Plugin->GetContentDir();
-				// Chop off the pluginName/Content since it's part of the full path
-				UncookedFileName /= PluginPath.RightChop(Plugin->GetName().Len() + ContentFolder.Len());
-				break;
-			}
-		}
-
-		// If UncookedFileName is empty fall through to sandbox handling
-	}
-
-	if (UncookedFileName.IsEmpty())
-	{
-		auto BuildUncookedPath =
-			[&UncookedFileName](const FString& CookedPath, const FString& CookedRoot, const FString& UncookedRoot)
-		{
-			UncookedFileName.AppendChars(*UncookedRoot, UncookedRoot.Len());
-			UncookedFileName.AppendChars(*CookedPath + CookedRoot.Len(), CookedPath.Len() - CookedRoot.Len());
-		};
-
-		if (CookedPath.StartsWith(SandboxRootDir))
-		{
-			// Optimized CookedPath.StartsWith(SandboxProjectDir) that does not compare all of SandboxRootDir again
-			if (CookedPath.Len() >= SandboxProjectDir.Len() &&
-				0 == FCString::Strnicmp(
-					*CookedPath + SandboxRootDir.Len(),
-					*SandboxProjectDir + SandboxRootDir.Len(),
-					SandboxProjectDir.Len() - SandboxRootDir.Len()))
-			{
-				BuildUncookedPath(CookedPath, SandboxProjectDir, RelativeProjectDir);
-			}
-			else
-			{
-				BuildUncookedPath(CookedPath, SandboxRootDir, RelativeRootDir);
-			}
-		}
-		else
-		{
-			FString FullCookedFilename = FPaths::ConvertRelativePathToFull(CookedPath);
-			BuildUncookedPath(FullCookedFilename, SandboxRootDir, RelativeRootDir);
-		}
-	}
-
-	if (!FPackageName::TryConvertFilenameToLongPackageName(UncookedFileName, ScratchPackageName))
-	{
-		return NAME_None;
-	}
-
-	return FName(*ScratchPackageName);
-}
-
-bool FLooseCookedPackageWriter::TryConvertUncookedFilenameToCookedRemappedPluginFilename(
-	FStringView FileName, TConstArrayView<TSharedRef<IPlugin>> InPluginsToRemap, FStringView SandboxDirectory,
-	FString& OutCookedFileName)
-{
-	// Ideally this would be in the Sandbox File but it can't access the project or plugin
-	if (InPluginsToRemap.IsEmpty())
-	{
-		return false;
-	}
-
-	for (const TSharedRef<IPlugin>& Plugin : InPluginsToRemap)
-	{
-		// If these match, then this content is part of plugin that gets remapped when packaged/staged
-		if (FileName.StartsWith(Plugin->GetContentDir()))
-		{
-			FString SearchFor;
-			SearchFor /= Plugin->GetName() / TEXT("Content");
-			int32 FoundAt = UE::String::FindLast(FileName, SearchFor, ESearchCase::IgnoreCase);
-			check(FoundAt != INDEX_NONE);
-			// Strip off everything but <PluginName/Content/<remaing path to file>
-			FStringView SnippedOffPath = FileName.RightChop(FoundAt);
-			// Put this is in <sandbox path>/RemappedPlugins/<PluginName>/Content/<remaing path to file>
-			constexpr FStringView RemappedPluginsDirName(REMAPPED_PLUGINS);
-			OutCookedFileName.Reserve(SandboxDirectory.Len() + RemappedPluginsDirName.Len() + SnippedOffPath.Len() + 2);
-			OutCookedFileName = SandboxDirectory;
-			OutCookedFileName /= REMAPPED_PLUGINS;
-			OutCookedFileName /= SnippedOffPath;
-			return true;
-		}
-	}
-	return false;
-}
-
-FString FLooseCookedPackageWriter::ConvertPackageNameToCookedPath(
-	const FString& SandboxRootDir, const FString& RelativeRootDir,
-	const FString& SandboxProjectDir, const FString& RelativeProjectDir,
-	FStringView PackageName) const
-{
-	FString UncookedFileName;
-	if (!FPackageName::TryConvertLongPackageNameToFilename(FString(PackageName), UncookedFileName))
-	{
-		return TEXT("");
-	}
-
-	FString CookedFileName;
-	if (TryConvertUncookedFilenameToCookedRemappedPluginFilename(UncookedFileName, PluginsToRemap,
-		SandboxRootDir, CookedFileName))
-	{
-		return CookedFileName;
-	}
-
-	auto BuildCookedPath =
-		[&CookedFileName](const FString& UncookedPath, const FString& CookedRoot, const FString& UncookedRoot)
-	{
-		CookedFileName.AppendChars(*CookedRoot, CookedRoot.Len());
-		CookedFileName.AppendChars(*UncookedPath + UncookedRoot.Len(), UncookedPath.Len() - UncookedRoot.Len());
-	};
-
-	if (UncookedFileName.StartsWith(RelativeProjectDir))
-	{
-		BuildCookedPath(UncookedFileName, SandboxProjectDir, RelativeProjectDir);
-	}
-	else if (UncookedFileName.StartsWith(RelativeRootDir))
-	{
-		BuildCookedPath(UncookedFileName, SandboxRootDir, RelativeRootDir);
-	}
-	else
-	{
-		CookedFileName.Empty();
-	}
-	return CookedFileName;
-}
-
 EPackageExtension FLooseCookedPackageWriter::BulkDataTypeToExtension(FBulkDataInfo::EType BulkDataType)
 {
 	switch (BulkDataType)
diff --git a/Engine/Source/Editor/UnrealEd/Private/Cooker/LooseCookedPackageWriter.h b/Engine/Source/Editor/UnrealEd/Private/Cooker/LooseCookedPackageWriter.h
--- a/Engine/Source/Editor/UnrealEd/Private/Cooker/LooseCookedPackageWriter.h
+++ b/Engine/Source/Editor/UnrealEd/Private/Cooker/LooseCookedPackageWriter.h
@@ -26,9 +26,10 @@ class FAsyncIODelete;
 class FAssetRegistryState;
 class FLargeMemoryWriter;
 class FMD5;
-class IPlugin;
 class ITargetPlatform;
 template <typename ReferencedType> class TRefCountPtr;
+namespace UE::Cook { class FCookSandbox; }
+namespace UE::Cook { struct FCookSandboxConvertCookedPathToPackageNameContext; }
 
 /** A CookedPackageWriter that saves cooked packages in separate .uasset,.uexp,.ubulk files in the Saved\Cooked\[Platform] directory. */
 class FLooseCookedPackageWriter : public TPackageWriterToSharedBuffer<ICookedPackageWriter>
@@ -38,7 +39,7 @@ public:
 
 	FLooseCookedPackageWriter(const FString& OutputPath, const FString& MetadataDirectoryPath,
 		const ITargetPlatform* TargetPlatform, FAsyncIODelete& InAsyncIODelete,
-		const TArray<TSharedRef<IPlugin>>& InPluginsToRemap, FBeginCacheCallback&& InBeginCacheCallback);
+		UE::Cook::FCookSandbox& InSandboxFile, FBeginCacheCallback&& InBeginCacheCallback);
 	~FLooseCookedPackageWriter();
 
 	virtual FCookCapabilities GetCookCapabilities() const override
@@ -71,9 +72,6 @@ public:
 	virtual FPackageWriterRecords::FPackage* ConstructRecord() override;
 
 	static EPackageExtension BulkDataTypeToExtension(FBulkDataInfo::EType BulkDataType);
-	static bool TryConvertUncookedFilenameToCookedRemappedPluginFilename(
-		FStringView FileName, TConstArrayView<TSharedRef<IPlugin>> InPluginsToRemap,
-		FStringView SandboxDirectory, FString& OutCookedFileName);
 
 private:
 
@@ -133,15 +131,6 @@ private:
 	void GetAllCookedFiles();
 	void FindAndDeleteCookedFilesForPackages(TConstArrayView<FName> PackageNames);
 
-	FName ConvertCookedPathToPackageName(
-		const FString& SandboxRootDir, const FString& RelativeRootDir,
-		const FString& SandboxProjectDir, const FString& RelativeProjectDir,
-		const FString& CookedPath, FString& ScratchFileName, FString& ScratchPackageName) const;
-	FString ConvertPackageNameToCookedPath(
-		const FString& SandboxRootDir, const FString& RelativeRootDir,
-		const FString& SandboxProjectDir, const FString& RelativeProjectDir,
-		FStringView PackageName) const;
-
 	void RemoveCookedPackagesByPackageName(TArrayView<const FName> PackageNamesToRemove, bool bRemoveRecords);
 	void AsyncSave(FRecord& Record, const FCommitPackageInfo& Info);
 
@@ -173,7 +162,7 @@ private:
 	FString MetadataDirectoryPath;
 	const ITargetPlatform& TargetPlatform;
 	TMap<FName, FOplogPackageInfo> Oplog;
-	const TArray<TSharedRef<IPlugin>>& PluginsToRemap;
+	UE::Cook::FCookSandbox& SandboxFile;
 	FAsyncIODelete& AsyncIODelete;
 	FBeginCacheCallback BeginCacheCallback;
 	bool bIterateSharedBuild = false;
diff --git a/Engine/Source/Editor/UnrealEd/Public/Commandlets/AssetRegistryGenerator.h b/Engine/Source/Editor/UnrealEd/Public/Commandlets/AssetRegistryGenerator.h
--- a/Engine/Source/Editor/UnrealEd/Public/Commandlets/AssetRegistryGenerator.h
+++ b/Engine/Source/Editor/UnrealEd/Public/Commandlets/AssetRegistryGenerator.h
@@ -18,7 +18,6 @@
 #include "UObject/SoftObjectPath.h"
 #include "UObject/UObjectHash.h"
 
-class FSandboxPlatformFile;
 class IAssetRegistry;
 class ITargetPlatform;
 class IChunkDataGenerator;
@@ -27,6 +26,7 @@ struct FChunkDependencyTreeNode;
 struct FCookTagList;
 struct FSoftObjectPath;
 namespace UE::Cook { class FAssetRegistryPackageMessage; }
+namespace UE::Cook { class FCookSandbox; }
 namespace UE::Cook { class FCookWorkerClient; }
 namespace UE::Cook { struct FPackageData; }
 
@@ -156,7 +156,7 @@ public:
 	 *        and a manifest of packagenames is written for each chunk
 	 */
 	void FinalizeChunkIDs(const TSet<FName>& CookedPackages, const TSet<FName>& DevelopmentOnlyPackages,
-		FSandboxPlatformFile& InSandboxFile, bool bGenerateStreamingInstallManifest);
+		UE::Cook::FCookSandbox& InSandboxFile, bool bGenerateStreamingInstallManifest);
 
 	/**
 	 * Register a chunk data generator with this generator.
@@ -204,7 +204,7 @@ public:
 	 * @param InManifestSubDir If non-null, the manifests are written into this subpath
 	 *        of the usual location.
 	 */
-	bool SaveManifests(FSandboxPlatformFile& InSandboxFile, int64 InOverrideChunkSize = 0,
+	bool SaveManifests(UE::Cook::FCookSandbox& InSandboxFile, int64 InOverrideChunkSize = 0,
 		const TCHAR* InManifestSubDir = nullptr);
 
 	/**
@@ -215,7 +215,7 @@ public:
 	/** 
 	 * Writes out CookerOpenOrder.log file 
 	 */
-	bool WriteCookerOpenOrder(FSandboxPlatformFile& InSandboxFile);
+	bool WriteCookerOpenOrder(UE::Cook::FCookSandbox& InSandboxFile);
 
 	/**
 	 * Follows an assets dependency chain to build up a list of package names in the same order as the runtime would attempt to load them
@@ -401,7 +401,7 @@ private:
 	 * 
 	 * @param the InSandboxFile used during cook
 	 */
-	void FixupPackageDependenciesForChunks(FSandboxPlatformFile& InSandboxFile);
+	void FixupPackageDependenciesForChunks(UE::Cook::FCookSandbox& InSandboxFile);
 
 	/**
 	 * Attaches encryption key guids into the registry data for encrypted primary assets
@@ -409,7 +409,7 @@ private:
 	void InjectEncryptionData(FAssetRegistryState& TargetState);
 
 	void AddPackageToChunk(FChunkPackageSet& ThisPackageSet, FName InPkgName,
-		const FString& InSandboxFile, int32 PakchunkIndex, FSandboxPlatformFile& SandboxPlatformFile);
+		const FString& InSandboxFile, int32 PakchunkIndex, UE::Cook::FCookSandbox& SandboxPlatformFile);
 
 	/**
 	 * Returns the path of the temporary packaging directory for the specified platform.
@@ -430,7 +430,7 @@ private:
 
 	/** Calculate the final ChunkIds used by the package and store the package in the manifest for each of those chunks. */
 	void CalculateChunkIdsAndAssignToManifest(const FName& PackageFName, const FString& PackagePathName,
-		const FString& SandboxFilename, const FString& LastLoadedMapName, FSandboxPlatformFile& InSandboxFile);
+		const FString& SandboxFilename, const FString& LastLoadedMapName, UE::Cook::FCookSandbox& InSandboxFile);
 
 	/** Deletes the temporary packaging directory for the specified platform */
 	bool CleanTempPackagingDirectory(const FString& Platform) const;
@@ -440,7 +440,7 @@ private:
 
 	/** Generates and saves streaming install chunk manifest */
 	bool GenerateStreamingInstallManifest(int64 InOverrideChunkSize, const TCHAR* InManifestSubDir,
-		FSandboxPlatformFile& InSandboxFile);
+		UE::Cook::FCookSandbox& InSandboxFile);
 
 	/** Gather a list of dependencies required by to completely load this package */
 	bool GatherAllPackageDependencies(FName PackageName, TArray<FName>& DependentPackageNames);
