Cooker DiffOnly: Fix crash due to DiffPackageWriter not respecting the contract for GetPreviousCookedBytes; ZenStoreWriter returns false and does not construct the out value, and the caller is required to handle that false code and not use the out value.
#rnx

diff --git a/Engine/Source/Editor/UnrealEd/Private/Cooker/DiffPackageWriter.cpp b/Engine/Source/Editor/UnrealEd/Private/Cooker/DiffPackageWriter.cpp
--- a/Engine/Source/Editor/UnrealEd/Private/Cooker/DiffPackageWriter.cpp
+++ b/Engine/Source/Editor/UnrealEd/Private/Cooker/DiffPackageWriter.cpp
@@ -149,7 +149,13 @@ void FDiffPackageWriter::WritePackageData(const FPackageInfo& Info, FLargeMemory
 
 	FArchiveStackTrace& Writer = static_cast<FArchiveStackTrace&>(ExportsArchive);
 	ICookedPackageWriter::FPreviousCookedBytesData PreviousInnerData;
-	Inner->GetPreviousCookedBytes(LocalInfo, PreviousInnerData);
+	if (!Inner->GetPreviousCookedBytes(LocalInfo, PreviousInnerData))
+	{
+		PreviousInnerData.Data.Reset();
+		PreviousInnerData.HeaderSize = 0;
+		PreviousInnerData.Size = 0;
+	}
+	check(PreviousInnerData.Data.Get() != nullptr || (PreviousInnerData.Size == 0 && PreviousInnerData.HeaderSize == 0));
 
 	FArchiveStackTrace::FPackageData PreviousPackageData;
 	PreviousPackageData.Data = PreviousInnerData.Data.Get();
