Fix an issue where the hovered/active gizmo hit proxy is incorrect when mouse tracking restarts

This was causing the indirect manipulation axis to persist as new buttons were pressed, which would result in the wrong axis getting manipulated.

diff --git a/Engine/Source/Editor/UnrealEd/Private/EditorViewportClient.cpp b/Engine/Source/Editor/UnrealEd/Private/EditorViewportClient.cpp
--- a/Engine/Source/Editor/UnrealEd/Private/EditorViewportClient.cpp
+++ b/Engine/Source/Editor/UnrealEd/Private/EditorViewportClient.cpp
@@ -3180,6 +3180,7 @@ void FEditorViewportClient::StartTrackingDueToInput( const struct FInputEventSta
 		{
 			MouseDeltaTracker->EndTracking( this );
 			bIsTracking = false;
+			CheckHoveredHitProxy(Viewport->GetHitProxy(CachedMouseX, CachedMouseY));
 		}
 
 		bDraggingByHandle = (Widget && Widget->GetCurrentAxis() != EAxisList::None);
