Fix for a goof not realizing the FMT macros didn't contain the percent sign.

diff --git a/Engine/Source/Editor/UnrealEd/Private/Commandlets/AssetRegistryGenerator.cpp b/Engine/Source/Editor/UnrealEd/Private/Commandlets/AssetRegistryGenerator.cpp
--- a/Engine/Source/Editor/UnrealEd/Private/Commandlets/AssetRegistryGenerator.cpp
+++ b/Engine/Source/Editor/UnrealEd/Private/Commandlets/AssetRegistryGenerator.cpp
@@ -1757,7 +1757,7 @@ bool FAssetRegistryGenerator::SaveAssetRegistry(const FString& SandboxPath, bool
 
 			UE_LOG(LogAssetRegistryGenerator, Display, 
 				TEXT("Generated development asset registry %s num assets %d, size is %5.2fkb, ")
-				TEXT("XxHash64 = " UINT64_FMT ", waited on hash %.2f seconds"), 
+				TEXT("XxHash64[LE] = 0x%" UINT64_X_FMT ", waited on hash %.2f seconds"),
 				*PlatformSandboxPath, State.GetNumAssets(), (float)SerializedAssetRegistry.Num() / 1024.f, 
 				DevArXxHash, WaitTime
 				);
