Add /Engine and /Game as dependencies to plugins in the ucookmeta file.

diff --git a/Engine/Source/Editor/UnrealEd/Private/CookOnTheFlyServer.cpp b/Engine/Source/Editor/UnrealEd/Private/CookOnTheFlyServer.cpp
--- a/Engine/Source/Editor/UnrealEd/Private/CookOnTheFlyServer.cpp
+++ b/Engine/Source/Editor/UnrealEd/Private/CookOnTheFlyServer.cpp
@@ -9352,6 +9352,12 @@ void UCookOnTheFlyServer::WriteCookMetadata(const ITargetPlatform* InTargetPlatf
 			AddIndex++;
 		}
 
+		// Add the /Engine and /Game pseudo plugins. These are placeholders for holding size information when unrealpak runs.
+		const uint16 EnginePluginIndex = (uint16)PluginsToAdd.Num();
+		PluginsToAdd.AddDefaulted_GetRef().Name = TEXT("Engine");
+		const uint16 GamePluginIndex = (uint16)PluginsToAdd.Num();
+		PluginsToAdd.AddDefaulted_GetRef().Name = TEXT("Game");
+
 		// Construct the dependency list.
 		TArray<uint16> RootPlugins;
 		for (TSharedRef<IPlugin>& EnabledPlugin : EnabledPlugins)
@@ -9383,14 +9389,19 @@ void UCookOnTheFlyServer::WriteCookMetadata(const ITargetPlatform* InTargetPlatf
 				}
 			}
 
+			//
+			// We've created two pseudo plugins Engine and Game, however no one depends on them explicitly.
+			// In order to facilitate the size computations, we inject artificial dependencies based on where
+			// the plugin was loaded from.
+			PluginChildArray.Add(EnginePluginIndex);
+			if (EnabledPlugin->GetLoadedFrom() == EPluginLoadedFrom::Project)
+			{
+				PluginChildArray.Add(GamePluginIndex);
+			}
+
 			Entry.DependencyIndexEnd = PluginChildArray.Num();
 		}
 
-
-		// Add the /Engine and /Game pseudo plugins. These are placeholders for holding size information when unrealpak runs.
-		PluginsToAdd.AddDefaulted_GetRef().Name = TEXT("Engine");
-		PluginsToAdd.AddDefaulted_GetRef().Name = TEXT("Game");
-
 		if (IntFitsIn<uint16>(PluginChildArray.Num()) == false)
 		{
 			UE_LOG(LogCook, Warning, TEXT("Number of child plugins exceeds 64k, unable to write cook metadata file (count = %d)"), PluginChildArray.Num());
