Actor Folders: Fix unloading of Level Instance(s) causing Folders to be created in the main world

#rnx

diff --git a/Engine/Source/Editor/UnrealEd/Private/EditorLevelUtils.cpp b/Engine/Source/Editor/UnrealEd/Private/EditorLevelUtils.cpp
--- a/Engine/Source/Editor/UnrealEd/Private/EditorLevelUtils.cpp
+++ b/Engine/Source/Editor/UnrealEd/Private/EditorLevelUtils.cpp
@@ -848,6 +848,14 @@ bool UEditorLevelUtils::RemoveLevelsFromWorld(TArray<ULevel*> InLevels, bool bCl
 
 	ULayersSubsystem* Layers = GEditor->GetEditorSubsystem<ULayersSubsystem>();
 
+	// Mark all Levels as being removed to prevent iteration on them while we are removing all of them.
+	// PrivateRemoveLevelFromWorld will end up calling UWorld::UpdateStreamingState.
+	// This can broadcast levels change events to listeners. By setting bIsBeingRemoved = true we make sure those levels do not get iterated and cause issues.
+	for (ULevel* Level : InLevels)
+	{
+		Level->bIsBeingRemoved = true;
+	}
+
 	for (ULevel* Level : InLevels)
 	{
 		Layers->RemoveLevelLayerInformation(Level);
@@ -871,6 +879,7 @@ bool UEditorLevelUtils::RemoveLevelsFromWorld(TArray<ULevel*> InLevels, bool bCl
 		// Keep Names in other list because unload of package will make the PackagesToUnload invalid.
 		PackageNames.Add(Level->GetOutermost()->GetFName());
 		ChangedWorlds.Add(OwningWorld);
+		Level->bIsBeingRemoved = false;
 	}
 
 	FText TransResetText(LOCTEXT("RemoveLevelTransReset", "Removing Levels from World"));
