MPCook: Disable CanSkipEditorReferencedPackagesWhenCooking on the MPCook director. Rewrite the calculation of SAVE_KeepEditorOnlyCookedPackages to make it clear that it is a suppression and what can cause the suppression.
#rnx

diff --git a/Engine/Source/Editor/UnrealEd/Private/CookOnTheFlyServer.cpp b/Engine/Source/Editor/UnrealEd/Private/CookOnTheFlyServer.cpp
--- a/Engine/Source/Editor/UnrealEd/Private/CookOnTheFlyServer.cpp
+++ b/Engine/Source/Editor/UnrealEd/Private/CookOnTheFlyServer.cpp
@@ -5873,10 +5873,13 @@ void FSaveCookedPackageContext::SetupPackage()
 	SaveFlags |= COTFS.IsCookFlagSet(ECookInitializationFlags::CookEditorOptional) ? SAVE_Optional : SAVE_None;
 
 	// removing editor only packages only works when cooking in commandlet and non iterative cooking
-	// also doesn't work in multiprocess cooking
-	bool bKeepEditorOnlyPackages = !(COTFS.IsCookByTheBookMode() && !COTFS.IsCookingInEditor());
-	bKeepEditorOnlyPackages |= COTFS.IsCookFlagSet(ECookInitializationFlags::Iterative);
-	SaveFlags |= bKeepEditorOnlyPackages ? SAVE_KeepEditorOnlyCookedPackages : SAVE_None;
+	bool bCanSkipEditorOnlyPackages = COTFS.IsCookByTheBookMode() && !COTFS.IsCookingInEditor();
+	bCanSkipEditorOnlyPackages &= !COTFS.IsCookFlagSet(ECookInitializationFlags::Iterative);
+	// MPCOOKTODO: it also doesn't work in multiprocess cooking, because GetPackage()->IsLoadedByEditorPropertiesOnly()
+	// might have been set to false on a CookWorker and is not replicated. Rather than fixing this, we should delete
+	// the feature and instead use SkipOnlyEditorOnly.
+	bCanSkipEditorOnlyPackages &= !COTFS.CookDirector.IsValid() && !COTFS.CookWorkerClient.IsValid();
+	SaveFlags |= bCanSkipEditorOnlyPackages ? SAVE_None : SAVE_KeepEditorOnlyCookedPackages;
 
 	// Use SandboxFile to do path conversion to properly handle sandbox paths (outside of standard paths in particular).
 	Filename = COTFS.ConvertToFullSandboxPath(*Filename, true);
