Cooker: Hidden Dependencies diagnostics: Add scopes to suppress false-positive hidden dependencies. Share the DebugOnlyEditorOnly code to remove other false-positives using analysis of the saved package.
#rnx

diff --git a/Engine/Source/Editor/Kismet/Private/BlueprintCompilationManager.cpp b/Engine/Source/Editor/Kismet/Private/BlueprintCompilationManager.cpp
--- a/Engine/Source/Editor/Kismet/Private/BlueprintCompilationManager.cpp
+++ b/Engine/Source/Editor/Kismet/Private/BlueprintCompilationManager.cpp
@@ -30,6 +30,7 @@
 #include "KismetCompiler.h"
 #include "Misc/ScopedSlowTask.h"
 #include "Misc/DataValidation.h"
+#include "Misc/PackageAccessTrackingOps.h"
 #include "ProfilingDebugging/ScopedTimers.h"
 #include "Serialization/ArchiveHasReferences.h"
 #include "Serialization/ArchiveReplaceObjectRef.h"
@@ -1217,6 +1218,7 @@ void FBlueprintCompilationManagerImpl::FlushCompilationQueueImpl(bool bSuppressB
 			DECLARE_SCOPE_HIERARCHICAL_COUNTER(ReconstructNodes)
 
 			UBlueprint* BP = CompilerData.BP;
+			UE_TRACK_REFERENCING_PACKAGE_SCOPED(BP, PackageAccessTrackingOps::NAME_CookerBuildObject);
 
 			ConformToParentAndInterfaces(BP);
 
diff --git a/Engine/Source/Editor/UnrealEd/Private/CookOnTheFlyServer.cpp b/Engine/Source/Editor/UnrealEd/Private/CookOnTheFlyServer.cpp
--- a/Engine/Source/Editor/UnrealEd/Private/CookOnTheFlyServer.cpp
+++ b/Engine/Source/Editor/UnrealEd/Private/CookOnTheFlyServer.cpp
@@ -5837,7 +5837,7 @@ void UCookOnTheFlyServer::SaveCookedPackage(UE::Cook::FSaveCookedPackageContext&
 
 bool UCookOnTheFlyServer::IsDebugRecordUnsolicited() const
 {
-	return bOnlyEditorOnlyDebug;
+	return bOnlyEditorOnlyDebug | bHiddenDependenciesDebug;
 }
 
 namespace UE::Cook
@@ -6264,11 +6264,11 @@ void FSaveCookedPackageContext::FinishPackage()
 			}
 		}
 	}
-	if (COTFS.bOnlyEditorOnlyDebug)
+	if (COTFS.IsDebugRecordUnsolicited())
 	{
 		COTFS.ProcessUnsolicitedPackages();
-		FDiagnostics::AnalyzeOnlyEditorOnlySave(COTFS, PackageData, PackageData.DetachUnsolicited(),
-			SaveReferences, ReachablePlatforms);
+		FDiagnostics::AnalyzeHiddenDependencies(COTFS, PackageData, PackageData.DetachUnsolicited(),
+			SaveReferences, ReachablePlatforms, COTFS.bOnlyEditorOnlyDebug, COTFS.bHiddenDependenciesDebug);
 	}
 
 	if (!bHasRetryErrorCode)
@@ -6632,6 +6632,11 @@ void UCookOnTheFlyServer::SetInitializeConfigSettings(UE::Cook::FInitializeConfi
 	bHiddenDependenciesDebug = FParse::Param(FCommandLine::Get(), TEXT("HiddenDependenciesDebug"));
 	if (bHiddenDependenciesDebug)
 	{
+		UE_LOG(LogCook, Display, TEXT("HiddenDependenciesDebug is enabled."));
+
+		// HiddenDependencies diagnostics rely on using SkipOnlyEditorOnly
+		bSkipOnlyEditorOnly = true;
+
 		FScopeLock HiddenDependenciesScopeLock(&HiddenDependenciesLock);
 
 		FString ClassPathListStr;
@@ -12241,6 +12246,8 @@ void UCookOnTheFlyServer::OnDiscoveredPackageDebug(FName PackageName, const UE::
 	ReportHiddenDependency(Instigator.Referencer, PackageName);
 }
 
+FName EngineTransientName(TEXT("/Engine/Transient"));
+
 void UCookOnTheFlyServer::OnObjectHandleReadDebug(TArrayView<const UObject*const> ReadObjects)
 {
 	using namespace UE::Cook;
@@ -12257,6 +12264,22 @@ void UCookOnTheFlyServer::OnObjectHandleReadDebug(TArrayView<const UObject*const
 	}
 
 	FName ReferencerPackageName = AccumulatedScopeData->PackageName;
+	if (ReferencerPackageName.IsNone() || ReferencerPackageName == EngineTransientName)
+	{
+		return;
+	}
+	TStringBuilder<256> ReferencerPackageNameStr;
+	ReferencerPackageName.ToString(ReferencerPackageNameStr);
+	if (FPackageName::IsTempPackage(ReferencerPackageNameStr) ||
+		FPackageName::IsVersePackage(ReferencerPackageNameStr))
+	{
+		return;
+	}
+
+	// Accelerate analysis and make hitting breakpoints more unique by ignoring dependencies that we have already
+	// logged for the most-recently-used referencer
+	static FName LastReferencer;
+	static TSet<FName> HandledDependencies;
 
 	TArray<FName, TInlineAllocator<16>> DependencyPackageNames;
 	for (const UObject* ReadObject : ReadObjects)
@@ -12275,8 +12298,22 @@ void UCookOnTheFlyServer::OnObjectHandleReadDebug(TArrayView<const UObject*const
 		{
 			continue;
 		}
+		if (ReferencerPackageName != LastReferencer)
+		{
+			LastReferencer = ReferencerPackageName;
+			HandledDependencies.Reset();
+		}
+		bool bAlreadyExists;
+		HandledDependencies.Add(DependencyPackageName, &bAlreadyExists);
+		if (bAlreadyExists)
+		{
+			continue;
+		}
 
-		if (FPackageName::IsScriptPackage(WriteToString<256>(DependencyPackageName)))
+		TStringBuilder<256> DependencyPackageNameStr;
+		DependencyPackageName.ToString(DependencyPackageNameStr);
+		if (FPackageName::IsScriptPackage(DependencyPackageNameStr) ||
+			FPackageName::IsTempPackage(DependencyPackageNameStr))
 		{
 			continue;
 		}
@@ -12394,8 +12431,14 @@ void UCookOnTheFlyServer::ReportHiddenDependency(FName Referencer, FName Depende
 				*PrimaryAssetData->AssetClassPath.ToString());
 		}
 	}
-	UE_LOG(LogCook, Warning, TEXT("Hidden dependency discovered: %s depends on %s%s."),
-		*Referencer.ToString(), *Dependency.ToString(), *PrimaryAssetClassPath);
+	FPackageData* ReferencerPackageData = PackageDatas->TryAddPackageDataByFileName(Referencer);
+	FPackageData* DependencyPackageData = PackageDatas->TryAddPackageDataByFileName(Dependency);
+	if (!ReferencerPackageData || !DependencyPackageData)
+	{
+		return;
+	}
+	TMap<FPackageData*, EInstigator>& Unsolicited = ReferencerPackageData->CreateOrGetUnsolicited();
+	Unsolicited.Add(DependencyPackageData, EInstigator::Unsolicited);
 }
 
 static
diff --git a/Engine/Source/Editor/UnrealEd/Private/Cooker/CookDiagnostics.cpp b/Engine/Source/Editor/UnrealEd/Private/Cooker/CookDiagnostics.cpp
--- a/Engine/Source/Editor/UnrealEd/Private/Cooker/CookDiagnostics.cpp
+++ b/Engine/Source/Editor/UnrealEd/Private/Cooker/CookDiagnostics.cpp
@@ -10,6 +10,7 @@
 #include "Cooker/CookTypes.h"
 #include "CookOnTheSide/CookLog.h"
 #include "CookOnTheSide/CookOnTheFlyServer.h"
+#include "Logging/LogMacros.h"
 #include "Misc/Optional.h"
 #include "Misc/StringBuilder.h"
 #include "Serialization/ArchiveSerializedPropertyChain.h"
@@ -21,6 +22,8 @@
 
 class ITargetPlatform;
 
+DEFINE_LOG_CATEGORY_STATIC(LogHiddenDependencies, Log, All);
+
 namespace UE::Cook
 {
 
@@ -96,9 +99,10 @@ public:
 	TArray<FPackageReferenceAndPropertyChain> Imports;
 };
 
-void FDiagnostics::AnalyzeOnlyEditorOnlySave(UCookOnTheFlyServer& COTFS, FPackageData& PackageData,
+void FDiagnostics::AnalyzeHiddenDependencies(UCookOnTheFlyServer& COTFS, FPackageData& PackageData,
 	TMap<FPackageData*, EInstigator>&& UnsolicitedForPackage, TSet<FPackageData*>& SaveReferences,
-	TConstArrayView<const ITargetPlatform*> ReachablePlatforms)
+	TConstArrayView<const ITargetPlatform*> ReachablePlatforms, bool bOnlyEditorOnlyDebug,
+	bool bHiddenDependenciesDebug)
 {
 	TOptional<TArray<FName>> ExpectedDependencies;
 	TOptional<TMap<FName, TMap<UObject*, TArray<FProperty*, TInlineAllocator<1>>>>> ExportsUsingPackageName;
@@ -118,11 +122,17 @@ void FDiagnostics::AnalyzeOnlyEditorOnlySave(UCookOnTheFlyServer& COTFS, FPackag
 	{
 		FPackageData* Unsolicited = UnsolicitedPair.Key;
 		EInstigator Instigator = UnsolicitedPair.Value;
-		if (SaveReferences.Contains(Unsolicited))
+		if (!bHiddenDependenciesDebug && SaveReferences.Contains(Unsolicited))
 		{
 			// This package is included by OnlyEditorOnly as well. Remove it from our list to investigate.
 			continue;
 		}
+		if (!bOnlyEditorOnlyDebug && Instigator != EInstigator::Unsolicited)
+		{
+			// Dependencies are only hidden if they're detected as Unsolicited; other instigator types
+			// are reported only for comparing SkipOnlyEditorOnly to legacy WhatGetsCooked rules.
+			continue;
+		}
 		FName UnsolicitedPackageName = Unsolicited->GetPackageName();
 
 		// If the package is not cookable on any of the platforms we are cooking, then even though it was loaded
@@ -262,155 +272,163 @@ void FDiagnostics::AnalyzeOnlyEditorOnlySave(UCookOnTheFlyServer& COTFS, FPackag
 		// Try to find out how it was loaded so we can give information about what code needs to
 		// be modified to support OnlyEditorOnly.
 
-		
-		// Serialize all referenced exports in the source package and record all the packages each one
-		// references. If any of those packages is the unsolicited package, print out the export and property
-		// that refers to it.
-		if (!ExportsUsingPackageName)
+		if (bOnlyEditorOnlyDebug)
 		{
-			ExportsUsingPackageName.Emplace();
-			TSet<UObject*> Exports;
-			TRingBuffer<UObject*> ExportsQueue;
-			ForEachObjectWithPackage(Package, [&Exports, &ExportsQueue](UObject* Object)
+			// Serialize all referenced exports in the source package and record all the packages each one
+			// references. If any of those packages is the unsolicited package, print out the export and property
+			// that refers to it.
+			if (!ExportsUsingPackageName)
+			{
+				ExportsUsingPackageName.Emplace();
+				TSet<UObject*> Exports;
+				TRingBuffer<UObject*> ExportsQueue;
+				ForEachObjectWithPackage(Package, [&Exports, &ExportsQueue](UObject* Object)
+					{
+						if (Object->HasAnyFlags(RF_Public))
+						{
+							bool bAlreadyExists;
+							Exports.Add(Object, &bAlreadyExists);
+							if (!bAlreadyExists)
+							{
+								ExportsQueue.Add(Object);
+							}
+						}
+						return true;
+					});
+				FWhyImportedCollector Collector(Package);
+				while (!ExportsQueue.IsEmpty())
 				{
-					if (Object->HasAnyFlags(RF_Public))
+					UObject* Export = ExportsQueue.PopFrontValue();
+					Export->Serialize(Collector);
+					for (UObject* Dependency : Collector.Exports)
 					{
 						bool bAlreadyExists;
-						Exports.Add(Object, &bAlreadyExists);
+						Exports.Add(Dependency, &bAlreadyExists);
 						if (!bAlreadyExists)
 						{
-							ExportsQueue.Add(Object);
+							ExportsQueue.Add(Dependency);
 						}
 					}
-					return true;
-				});
-			FWhyImportedCollector Collector(Package);
-			while (!ExportsQueue.IsEmpty())
-			{
-				UObject* Export = ExportsQueue.PopFrontValue();
-				Export->Serialize(Collector);
-				for (UObject* Dependency : Collector.Exports)
-				{
-					bool bAlreadyExists;
-					Exports.Add(Dependency, &bAlreadyExists);
-					if (!bAlreadyExists)
-					{
-						ExportsQueue.Add(Dependency);
-					}
-				}
-				for (FPackageReferenceAndPropertyChain& ImportChain : Collector.Imports)
-				{
-					TArray<FProperty*, TInlineAllocator<1>>& ExistingChain = ExportsUsingPackageName->FindOrAdd(ImportChain.PackageName).FindOrAdd(Export);
-					if (ExistingChain.IsEmpty())
+					for (FPackageReferenceAndPropertyChain& ImportChain : Collector.Imports)
 					{
-						ExistingChain = MoveTemp(ImportChain.Properties);
+						TArray<FProperty*, TInlineAllocator<1>>& ExistingChain = ExportsUsingPackageName->FindOrAdd(ImportChain.PackageName).FindOrAdd(Export);
+						if (ExistingChain.IsEmpty())
+						{
+							ExistingChain = MoveTemp(ImportChain.Properties);
+						}
 					}
+					Collector.Reset();
 				}
-				Collector.Reset();
 			}
-		}
-		TMap<UObject*, TArray<FProperty*, TInlineAllocator<1>>>* ExportsUsingUnsolicited = ExportsUsingPackageName->Find(UnsolicitedPackageName);
-		TStringBuilder<256> WhyReferenced;
+			TMap<UObject*, TArray<FProperty*, TInlineAllocator<1>>>* ExportsUsingUnsolicited = ExportsUsingPackageName->Find(UnsolicitedPackageName);
+			TStringBuilder<256> WhyReferenced;
 
-		int32 PackageNameLen = Package->GetName().Len();
-		if (ExportsUsingUnsolicited)
-		{
-			int32 Count = 0;
-			constexpr int32 MaxCount = 3;
-			WhyReferenced << TEXT("\n\tReferencers:");
-			for (TPair<UObject*, TArray<FProperty*, TInlineAllocator<1>>>&Pair : *ExportsUsingUnsolicited)
+			int32 PackageNameLen = Package->GetName().Len();
+			if (ExportsUsingUnsolicited)
 			{
-				UObject* Export = Pair.Key;
-				WhyReferenced << TEXT("\n\t\t");
-				if (Count++ == MaxCount)
-				{
-					WhyReferenced << TEXT("...");
-					break;
-				}
-				WhyReferenced << Export->GetClass()->GetPathName() << TEXT(" ");
-				WhyReferenced << Export->GetPathName().RightChop(PackageNameLen + 1);
-				WhyReferenced << TEXT("#");
-				if (Pair.Value.IsEmpty())
+				int32 Count = 0;
+				constexpr int32 MaxCount = 3;
+				WhyReferenced << TEXT("\n\tReferencers:");
+				for (TPair<UObject*, TArray<FProperty*, TInlineAllocator<1>>>&Pair : *ExportsUsingUnsolicited)
 				{
-					WhyReferenced << TEXT("<UnknownProperty>");
-				}
-				else
-				{
-					int32 NumProperties = Pair.Value.Num();
-					bool bAddedSeparator = false;
-					for (int32 PropertyIndex = 0; PropertyIndex < NumProperties; ++PropertyIndex)
+					UObject* Export = Pair.Key;
+					WhyReferenced << TEXT("\n\t\t");
+					if (Count++ == MaxCount)
 					{
-						FProperty* Property = Pair.Value[PropertyIndex];
-						if (PropertyIndex < NumProperties - 1 && Property->GetFName() == Pair.Value[PropertyIndex + 1]->GetFName() &&
-							Property->IsA(FArrayProperty::StaticClass()))
-						{
-							// Omit adding a duplicate entry for inner property of an array which has the same name as its array property
-							continue;
-						}
-						WhyReferenced << Property->GetName() << TEXT("#");
-						bAddedSeparator = true;
+						WhyReferenced << TEXT("...");
+						break;
 					}
-					if (bAddedSeparator)
+					WhyReferenced << Export->GetClass()->GetPathName() << TEXT(" ");
+					WhyReferenced << Export->GetPathName().RightChop(PackageNameLen + 1);
+					WhyReferenced << TEXT("#");
+					if (Pair.Value.IsEmpty())
 					{
-						WhyReferenced.RemoveSuffix(1);
+						WhyReferenced << TEXT("<UnknownProperty>");
+					}
+					else
+					{
+						int32 NumProperties = Pair.Value.Num();
+						bool bAddedSeparator = false;
+						for (int32 PropertyIndex = 0; PropertyIndex < NumProperties; ++PropertyIndex)
+						{
+							FProperty* Property = Pair.Value[PropertyIndex];
+							if (PropertyIndex < NumProperties - 1 && Property->GetFName() == Pair.Value[PropertyIndex + 1]->GetFName() &&
+								Property->IsA(FArrayProperty::StaticClass()))
+							{
+								// Omit adding a duplicate entry for inner property of an array which has the same name as its array property
+								continue;
+							}
+							WhyReferenced << Property->GetName() << TEXT("#");
+							bAddedSeparator = true;
+						}
+						if (bAddedSeparator)
+						{
+							WhyReferenced.RemoveSuffix(1);
+						}
 					}
 				}
 			}
-		}
-		else
-		{
-			// If we didn't find any exports referring to the unsolicited package, then just print out the list of class types
-			// in the package to show which code needs to be investigated.
-			enum class EClassPriority
+			else
 			{
-				PrimaryUAsset,
-				Asset,
-				Public,
-				Private,
-				Count,
-			};
-			TMap<UClass*, EClassPriority> ObjectClasses;
-			FName PackageLeafName = FName(*FPaths::GetBaseFilename(Package->GetName(), true /* bRemovePath */));
-			ForEachObjectWithPackage(Package, [&ObjectClasses, PackageLeafName](UObject* Object)
+				// If we didn't find any exports referring to the unsolicited package, then just print out the list of class types
+				// in the package to show which code needs to be investigated.
+				enum class EClassPriority
 				{
-					UClass* Class = Object->GetClass();
-					EClassPriority NewPriority = EClassPriority::Private;
-					if (Object->IsAsset())
+					PrimaryUAsset,
+					Asset,
+					Public,
+					Private,
+					Count,
+				};
+				TMap<UClass*, EClassPriority> ObjectClasses;
+				FName PackageLeafName = FName(*FPaths::GetBaseFilename(Package->GetName(), true /* bRemovePath */));
+				ForEachObjectWithPackage(Package, [&ObjectClasses, PackageLeafName](UObject* Object)
 					{
-						NewPriority = Object->GetFName() == PackageLeafName ? EClassPriority::PrimaryUAsset : EClassPriority::Asset;
-					}
-					else
+						UClass* Class = Object->GetClass();
+						EClassPriority NewPriority = EClassPriority::Private;
+						if (Object->IsAsset())
+						{
+							NewPriority = Object->GetFName() == PackageLeafName ? EClassPriority::PrimaryUAsset : EClassPriority::Asset;
+						}
+						else
+						{
+							NewPriority = Object->HasAnyFlags(RF_Public) ? EClassPriority::Public : EClassPriority::Private;
+						}
+						EClassPriority& OldPriority = ObjectClasses.FindOrAdd(Class, NewPriority);
+						OldPriority = static_cast<EClassPriority>(FMath::Min(static_cast<int32>(OldPriority), static_cast<int32>(NewPriority)));
+						return true;
+					});
+				WhyReferenced << TEXT("\n\tNo exports found referencing the dependency. Classes in the referencer package:");
+				int32 ClassCount = 0;
+				constexpr int32 MaxClassCount = 10;
+				ObjectClasses.ValueSort([](EClassPriority A, EClassPriority B)
 					{
-						NewPriority = Object->HasAnyFlags(RF_Public) ? EClassPriority::Public : EClassPriority::Private;
-					}
-					EClassPriority& OldPriority = ObjectClasses.FindOrAdd(Class, NewPriority);
-					OldPriority = static_cast<EClassPriority>(FMath::Min(static_cast<int32>(OldPriority), static_cast<int32>(NewPriority)));
-					return true;
-				});
-			WhyReferenced << TEXT("\n\tNo exports found referencing the dependency. Classes in the referencer package:");
-			int32 ClassCount = 0;
-			constexpr int32 MaxClassCount = 10;
-			ObjectClasses.ValueSort([](EClassPriority A, EClassPriority B)
+						if (A != B)
+						{
+							return (int32)A < (int32)B;
+						}
+						return false;
+					});
+				for (TPair<UClass*, EClassPriority>& ClassPair : ObjectClasses)
 				{
-					if (A != B)
+					if (ClassCount++ >= MaxClassCount)
 					{
-						return (int32)A < (int32)B;
+						break;
 					}
-					return false;
-				});
-			for (TPair<UClass*, EClassPriority>& ClassPair : ObjectClasses)
-			{
-				if (ClassCount++ >= MaxClassCount)
-				{
-					break;
+					WhyReferenced << TEXT("\n\t\t") << ClassPair.Key->GetPathName();
 				}
-				WhyReferenced << TEXT("\n\t\t") << ClassPair.Key->GetPathName();
 			}
-		}
 
-		UE_LOG(LogCook, Display, TEXT("OnlyEditorOnly: Skipped adding %s %s -> %s%s"),
-			LexToString(Instigator), *WriteToString<256>(PackageData.GetPackageName()),
-			*WriteToString<256>(Unsolicited->GetPackageName()), *WhyReferenced);
+			UE_LOG(LogHiddenDependencies, Display, TEXT("Skipped adding %s -> (%s) %s%s"),
+				*WriteToString<256>(PackageData.GetPackageName()), LexToString(Instigator),
+				*WriteToString<256>(Unsolicited->GetPackageName()), *WhyReferenced);
+		}
+		else
+		{
+			UE_LOG(LogHiddenDependencies, Display, TEXT("HiddenDependency %s -> (%s) %s"),
+				*WriteToString<256>(PackageData.GetPackageName()), LexToString(Instigator),
+				*WriteToString<256>(Unsolicited->GetPackageName()));
+		}
 	}
 }
 
diff --git a/Engine/Source/Editor/UnrealEd/Private/Cooker/CookDiagnostics.h b/Engine/Source/Editor/UnrealEd/Private/Cooker/CookDiagnostics.h
--- a/Engine/Source/Editor/UnrealEd/Private/Cooker/CookDiagnostics.h
+++ b/Engine/Source/Editor/UnrealEd/Private/Cooker/CookDiagnostics.h
@@ -22,9 +22,10 @@ public:
 	 * are added to the cook using OnlyEditorOnly rules. Silently ignores expected differences, but logs a message
 	 * with diagnostics for unexpected differences.
 	 */
-	static void AnalyzeOnlyEditorOnlySave(UCookOnTheFlyServer& COTFS, FPackageData& PackageData,
+	static void AnalyzeHiddenDependencies(UCookOnTheFlyServer& COTFS, FPackageData& PackageData,
 		TMap<FPackageData*, EInstigator>&& UnsolicitedForPackage, TSet<FPackageData*>& SaveReferences,
-		TConstArrayView<const ITargetPlatform*> ReachablePlatforms);
+		TConstArrayView<const ITargetPlatform*> ReachablePlatforms, bool bOnlyEditorOnlyDebug,
+		bool bHiddenDependenciesDebug);
 
 };
 
diff --git a/Engine/Source/Editor/UnrealEd/Private/EditorEngine.cpp b/Engine/Source/Editor/UnrealEd/Private/EditorEngine.cpp
--- a/Engine/Source/Editor/UnrealEd/Private/EditorEngine.cpp
+++ b/Engine/Source/Editor/UnrealEd/Private/EditorEngine.cpp
@@ -6269,6 +6269,7 @@ void UEditorEngine::DoConvertActors( const TArray<AActor*>& ActorsToConvert, UCl
 void UEditorEngine::NotifyToolsOfObjectReplacement(const TMap<UObject*, UObject*>& OldToNewInstanceMap)
 {
 	// Allow any other observers to act upon the object replacement
+	UE_TRACK_REFERENCING_OPNAME_SCOPED(PackageAccessTrackingOps::NAME_ResetContext);
 	FCoreUObjectDelegates::OnObjectsReplaced.Broadcast(OldToNewInstanceMap);
 }
 
diff --git a/Engine/Source/Runtime/AssetRegistry/Private/AssetRegistry.cpp b/Engine/Source/Runtime/AssetRegistry/Private/AssetRegistry.cpp
--- a/Engine/Source/Runtime/AssetRegistry/Private/AssetRegistry.cpp
+++ b/Engine/Source/Runtime/AssetRegistry/Private/AssetRegistry.cpp
@@ -6344,6 +6344,7 @@ void FAssetRegistryImpl::UpdateInheritanceBuffer(Impl::FClassInheritanceBuffer&
 {
 	TRACE_CPUPROFILER_EVENT_SCOPE(UAssetRegistryImpl::UpdateTemporaryCaches)
 	LLM_SCOPE(ELLMTag::AssetRegistry);
+	UE_TRACK_REFERENCING_OPNAME_SCOPED(PackageAccessTrackingOps::NAME_ResetContext);
 
 	TMap<UClass*, TSet<UClass*>> NativeSubclasses = GetAllDerivedClasses();
 
