FArchiveStackTrace: Deprecate the public interface and move it into a private helper class of DiffPackageWriter.
#rnx

diff --git a/Engine/Source/Editor/UnrealEd/Private/CookOnTheFlyServer.cpp b/Engine/Source/Editor/UnrealEd/Private/CookOnTheFlyServer.cpp
--- a/Engine/Source/Editor/UnrealEd/Private/CookOnTheFlyServer.cpp
+++ b/Engine/Source/Editor/UnrealEd/Private/CookOnTheFlyServer.cpp
@@ -114,7 +114,6 @@
 #include "ProjectDescriptor.h"
 #include "SceneUtils.h"
 #include "Serialization/ArchiveCountMem.h"
-#include "Serialization/ArchiveStackTrace.h"
 #include "Serialization/ArchiveUObject.h"
 #include "Serialization/ArrayReader.h"
 #include "Serialization/ArrayWriter.h"
diff --git a/Engine/Source/Editor/UnrealEd/Private/Cooker/DiffPackageWriter.cpp b/Engine/Source/Editor/UnrealEd/Private/Cooker/DiffPackageWriter.cpp
--- a/Engine/Source/Editor/UnrealEd/Private/Cooker/DiffPackageWriter.cpp
+++ b/Engine/Source/Editor/UnrealEd/Private/Cooker/DiffPackageWriter.cpp
@@ -14,7 +14,6 @@
 #include "Misc/Parse.h"
 #include "Misc/OutputDevice.h"
 #include "ProfilingDebugging/CookStats.h"
-#include "Serialization/ArchiveStackTrace.h"
 #include "Serialization/LargeMemoryWriter.h"
 #include "Templates/UniquePtr.h"
 #include "Templates/UnrealTemplate.h"
@@ -147,7 +146,7 @@ void FDiffPackageWriter::WritePackageData(const FPackageInfo& Info, FLargeMemory
 	FPackageInfo LocalInfo(Info);
 	Inner->CompleteExportsArchiveForDiff(LocalInfo, ExportsArchive);
 
-	FArchiveStackTrace& Writer = static_cast<FArchiveStackTrace&>(ExportsArchive);
+	FDiffWriterArchive& Writer = static_cast<FDiffWriterArchive&>(ExportsArchive);
 	ICookedPackageWriter::FPreviousCookedBytesData PreviousInnerData;
 	if (!Inner->GetPreviousCookedBytes(LocalInfo, PreviousInnerData))
 	{
@@ -157,7 +156,7 @@ void FDiffPackageWriter::WritePackageData(const FPackageInfo& Info, FLargeMemory
 	}
 	check(PreviousInnerData.Data.Get() != nullptr || (PreviousInnerData.Size == 0 && PreviousInnerData.HeaderSize == 0));
 
-	FArchiveStackTrace::FPackageData PreviousPackageData;
+	FDiffWriterArchive::FPackageData PreviousPackageData;
 	PreviousPackageData.Data = PreviousInnerData.Data.Get();
 	PreviousPackageData.Size = PreviousInnerData.Size;
 	PreviousPackageData.HeaderSize = PreviousInnerData.HeaderSize;
@@ -196,12 +195,12 @@ TUniquePtr<FLargeMemoryWriter> FDiffPackageWriter::CreateLinkerArchive(FName Pac
 	{
 		check(MultiOutputIndex < 2);
 		// Each difference will be logged with its Serialize call stack trace
-		return TUniquePtr<FLargeMemoryWriter>(new FArchiveStackTrace(Asset, *PackageName.ToString(),
+		return TUniquePtr<FLargeMemoryWriter>(new FDiffWriterArchive(Asset, *PackageName.ToString(),
 			true /* bInCollectCallstacks */, &DiffMap[MultiOutputIndex]));
 	}
 	else
 	{
-		return TUniquePtr<FLargeMemoryWriter>(new FArchiveStackTrace(Asset, *PackageName.ToString(),
+		return TUniquePtr<FLargeMemoryWriter>(new FDiffWriterArchive(Asset, *PackageName.ToString(),
 			false /* bInCollectCallstacks */));
 	}
 }
@@ -211,20 +210,20 @@ TUniquePtr<FLargeMemoryWriter> FDiffPackageWriter::CreateLinkerExportsArchive(FN
 	// When cooking, exports are serialized into a separate archive. The serialization callstack offsets
 	// and stack traces are collected into a separate callstack collection and appended to the overall
 	// callstacks for the entire package. DiffOnly cooks saves a package twice. The first pass collects
-	// the serialization offsets without the stack traces and then creates a FArchiveDiffMap. In the second
+	// the serialization offsets without the stack traces and then creates a FDiffWriterDiffMap. In the second
 	// pass, the diff map is used to collect offsets AND the entire stack trace for mismatching package data.
 	// In the first pass, callstack offsets will be relative to the export archive and adjusted when appending
 	// the callstacks to the overall package callstacks. In the second pass, the offset (ExportsDiffMapOffset) is
 	// known and the callstack will be relative to the beginning of the package.
 	
-	ExportsCallstacks = MakeUnique<FArchiveCallstacks>(Asset);
+	ExportsCallstacks = MakeUnique<FDiffWriterCallstacks>(Asset);
 	const int64 PreAllocateBytes = 0;
 	const bool bIsPersistent = true;
 
 	check(MultiOutputIndex < 2);
 	if (bDiffCallstack)
 	{
-		return MakeUnique<FArchiveStackTraceMemoryWriter>(
+		return MakeUnique<FDiffWriterArchiveMemoryWriter>(
 			*ExportsCallstacks,
 			&DiffMap[MultiOutputIndex],
 			ExportsDiffMapOffset[MultiOutputIndex],
@@ -234,7 +233,7 @@ TUniquePtr<FLargeMemoryWriter> FDiffPackageWriter::CreateLinkerExportsArchive(FN
 	}
 	else
 	{
-		return MakeUnique<FArchiveStackTraceMemoryWriter>(
+		return MakeUnique<FDiffWriterArchiveMemoryWriter>(
 			*ExportsCallstacks,
 			nullptr,
 			0,
diff --git a/Engine/Source/Editor/UnrealEd/Private/Cooker/DiffPackageWriter.h b/Engine/Source/Editor/UnrealEd/Private/Cooker/DiffPackageWriter.h
--- a/Engine/Source/Editor/UnrealEd/Private/Cooker/DiffPackageWriter.h
+++ b/Engine/Source/Editor/UnrealEd/Private/Cooker/DiffPackageWriter.h
@@ -3,7 +3,7 @@
 #pragma once
 
 #include "AssetRegistry/AssetRegistryState.h"
-#include "Serialization/ArchiveStackTrace.h"
+#include "Cooker/DiffWriterArchive.h"
 #include "Serialization/PackageWriter.h"
 
 /** A CookedPackageWriter that diffs output from the current cook with the file that was saved in the previous cook. */
@@ -122,8 +122,8 @@ private:
 	void ConditionallyDumpObjList();
 	void ConditionallyDumpObjects();
 
-	FArchiveDiffMap DiffMap[2];
-	TUniquePtr<FArchiveCallstacks> ExportsCallstacks;
+	FDiffWriterDiffMap DiffMap[2];
+	TUniquePtr<FDiffWriterCallstacks> ExportsCallstacks;
 	FBeginPackageInfo BeginInfo;
 	TUniquePtr<ICookedPackageWriter> Inner;
 	FString DumpObjListParams;
diff --git a/Engine/Source/Editor/UnrealEd/Private/Cooker/DiffWriterArchive.cpp b/Engine/Source/Editor/UnrealEd/Private/Cooker/DiffWriterArchive.cpp
new file mode 100644
--- /dev/null
+++ b/Engine/Source/Editor/UnrealEd/Private/Cooker/DiffWriterArchive.cpp
@@ -0,0 +1,1548 @@
+// Copyright Epic Games, Inc. All Rights Reserved.
+
+#include "Cooker/DiffWriterArchive.h"
+
+#include "Compression/CompressionUtil.h"
+#include "HAL/FileManager.h"
+#include "HAL/PlatformFileManager.h"
+#include "HAL/PlatformStackWalk.h"
+#include "Misc/CommandLine.h"
+#include "Misc/FileHelper.h"
+#include "Misc/OutputDeviceHelper.h"
+#include "Misc/PackageName.h"
+#include "Misc/Paths.h"
+#include "Misc/ScopeExit.h"
+#include "Serialization/StaticMemoryReader.h"
+#include "UObject/LinkerLoad.h"
+#include "UObject/ObjectMacros.h"
+#include "UObject/Package.h"
+#include "UObject/PropertyTempVal.h"
+#include "UObject/UnrealType.h"
+#include "UObject/UObjectGlobals.h"
+#include "UObject/UObjectThreadContext.h"
+
+DEFINE_LOG_CATEGORY_STATIC(LogDiffWriter, Log, All);
+
+/** Helper class that holds runtime generated constants for log output formatting */
+struct FDiffFormatHelper
+{
+public:
+	const TCHAR* const Indent;
+	const TCHAR* const LineTerminator;
+
+	FDiffFormatHelper()
+		: Indent(FCString::Spc(FOutputDeviceHelper::FormatLogLine(ELogVerbosity::Warning, LogDiffWriter.GetCategoryName(), TEXT(""), GPrintLogTimes).Len()))
+		, LineTerminator(TEXT("\n")) // because LINE_TERMINATOR doesn't work well wit EC
+	{}
+	static FDiffFormatHelper& Get()
+	{
+		static FDiffFormatHelper Instance;
+		return Instance;
+	}
+};
+
+class FIgnoreDiffManager
+{
+	int32 IgnoreCount = 0;
+	int32 DisabledCount = 0;
+
+public:
+	FIgnoreDiffManager() = default;
+	void Push()
+	{
+		IgnoreCount++;
+	}
+	void Pop()
+	{
+		IgnoreCount--;
+		check(IgnoreCount >= 0);
+	}
+	bool ShouldIgnoreDiff() const
+	{
+		return !!IgnoreCount;
+	}
+	void PushDisabled()
+	{
+		DisabledCount++;
+	}
+	void PopDisabled()
+	{
+		DisabledCount--;
+		check(DisabledCount >= 0);
+	}
+	bool ShouldBypassDiff() const
+	{
+		return !!DisabledCount;
+	}
+};
+
+static FIgnoreDiffManager GIgnoreDiffManager;
+
+static const ANSICHAR* DebugDataStackMarker = "\r\nDebugDataStack:\r\n";
+
+FDiffWriterCallstacks::FCallstackData::FCallstackData(TUniquePtr<ANSICHAR[]>&& InCallstack, UObject* InSerializedObject, FProperty* InSerializedProperty)
+	: Callstack(MoveTemp(InCallstack))
+	, SerializedProp(InSerializedProperty)
+{
+	if (InSerializedObject)
+	{
+		SerializedObjectName = InSerializedObject->GetFullName();
+	}
+	if (InSerializedProperty)
+	{
+		SerializedPropertyName = InSerializedProperty->GetFullName();
+	}
+}
+
+FString FDiffWriterCallstacks::FCallstackData::ToString(const TCHAR* CallstackCutoffText) const
+{
+	FString HumanReadableString;
+
+	const TCHAR* const LineTerminator = FDiffFormatHelper::Get().LineTerminator;
+	const TCHAR* const Indent = FDiffFormatHelper::Get().Indent;
+	
+	FString StackTraceText = Callstack.Get();
+	if (CallstackCutoffText != nullptr)
+	{
+		// If the cutoff string is provided, remove all functions starting with the one specifiec in the cutoff string
+		int32 CutoffIndex = StackTraceText.Find(CallstackCutoffText, ESearchCase::CaseSensitive);
+		if (CutoffIndex > 0)
+		{
+			CutoffIndex = StackTraceText.Find(TEXT("\n"), ESearchCase::CaseSensitive, ESearchDir::FromEnd, CutoffIndex - 1);
+			if (CutoffIndex > 0)
+			{
+				StackTraceText = StackTraceText.Left(CutoffIndex + 1);
+			}
+		}
+	}
+
+	TArray<FString> StackLines;
+	StackTraceText.ParseIntoArrayLines(StackLines);
+	for (FString& StackLine : StackLines)
+	{
+		if (StackLine.StartsWith(TEXT("0x")))
+		{
+			int32 CutoffIndex = StackLine.Find(TEXT(" "), ESearchCase::CaseSensitive);
+			if (CutoffIndex >= -1 && CutoffIndex < StackLine.Len() - 2)
+			{
+				StackLine.MidInline(CutoffIndex + 1, MAX_int32, false);
+			}
+		}
+		HumanReadableString += Indent;
+		HumanReadableString += StackLine;
+		HumanReadableString += LineTerminator;
+	}
+
+	if (!SerializedObjectName.IsEmpty())
+	{
+		HumanReadableString += LineTerminator;
+		HumanReadableString += Indent;
+		HumanReadableString += TEXT("Serialized Object: ");
+		HumanReadableString += SerializedObjectName;
+		HumanReadableString += LineTerminator;
+	}
+	if (!SerializedPropertyName.IsEmpty())
+	{
+		if (SerializedObjectName.IsEmpty())
+		{
+			HumanReadableString += LineTerminator;
+		}
+		HumanReadableString += Indent;
+		HumanReadableString += TEXT("Serialized Property: ");
+		HumanReadableString += SerializedPropertyName;
+		HumanReadableString += LineTerminator;
+	}
+	return HumanReadableString;
+}
+
+FDiffWriterCallstacks::FCallstackData FDiffWriterCallstacks::FCallstackData::Clone() const
+{
+	TUniquePtr<ANSICHAR[]> CallstackCopy;
+	if (const int32 Len = FCStringAnsi::Strlen(Callstack.Get()); Len > 0)
+	{
+		CallstackCopy = MakeUnique<ANSICHAR[]>(Len + 1);
+		FMemory::Memcpy(CallstackCopy.Get(), Callstack.Get(), Len + 1);
+	}
+
+	FCallstackData Clone(MoveTemp(CallstackCopy), nullptr, SerializedProp);
+	Clone.SerializedObjectName = SerializedObjectName;
+
+	return Clone;
+}
+
+FDiffWriterCallstacks::FDiffWriterCallstacks(UObject* InAsset)
+	: Asset(InAsset)
+	, bCallstacksDirty(true)
+	, StackTraceSize(65535)
+	, LastSerializeCallstack(nullptr)
+	, TotalSize(0)
+{
+	StackTrace = MakeUnique<ANSICHAR[]>(StackTraceSize);
+	StackTrace[0] = 0;
+}
+
+FName FDiffWriterCallstacks::GetAssetClass() const
+{
+	return Asset != nullptr ? Asset->GetClass()->GetFName() : NAME_None;
+}
+
+ANSICHAR* FDiffWriterCallstacks::AddUniqueCallstack(bool bIsCollectingCallstacks, UObject* SerializedObject, FProperty* SerializedProperty, uint32& OutCallstackCRC)
+{
+	ANSICHAR* Callstack = nullptr;
+	if (bIsCollectingCallstacks)
+	{
+		OutCallstackCRC = FCrc::StrCrc32(StackTrace.Get());
+
+		if (FCallstackData* ExistingCallstack = UniqueCallstacks.Find(OutCallstackCRC))
+		{
+			Callstack = ExistingCallstack->Callstack.Get();
+		}
+		else
+		{
+			const int32 Len = FCStringAnsi::Strlen(StackTrace.Get()) + 1;
+			TUniquePtr<ANSICHAR[]> NewCallstack = MakeUnique<ANSICHAR[]>(Len);
+			FCStringAnsi::Strcpy(NewCallstack.Get(), Len, StackTrace.Get());
+			FCallstackData& NewEntry = UniqueCallstacks.Add(OutCallstackCRC, FCallstackData(MoveTemp(NewCallstack), SerializedObject, SerializedProperty));
+			Callstack = NewEntry.Callstack.Get();
+		}
+	}
+	else
+	{
+		OutCallstackCRC = 0;
+	}
+	return Callstack;
+}
+
+void FDiffWriterCallstacks::Add(
+	int64 Offset,
+	int64 Length,
+	UObject* SerializedObject,
+	FProperty* SerializedProperty,
+	TArrayView<const FName> DebugDataStack,
+	bool bIsCollectingCallstacks,
+	bool bCollectCurrentCallstack,
+	int32 StackIgnoreCount)
+{
+	if (GIgnoreDiffManager.ShouldBypassDiff())
+	{
+		return;
+	}
+
+	const int64 CurrentOffset = Offset;
+	TotalSize = FMath::Max(TotalSize, CurrentOffset + Length); 
+
+	const bool bShouldCollectCallstack = bIsCollectingCallstacks && bCollectCurrentCallstack && !GIgnoreDiffManager.ShouldIgnoreDiff();
+	if (bShouldCollectCallstack)
+	{
+		StackTrace[0] = '\0';
+		FPlatformStackWalk::StackWalkAndDump(StackTrace.Get(), StackTraceSize, StackIgnoreCount);
+		//if we have a debug name stack, plaster it onto the end of the current stack buffer so that it's a part of the unique stack entry.
+		if (DebugDataStack.Num() > 0)
+		{
+			FCStringAnsi::Strcat(StackTrace.Get(), StackTraceSize, DebugDataStackMarker);
+
+			const FString SubIndent = FString(FDiffFormatHelper::Get().Indent) + FString(TEXT("    "));
+
+			bool bIsIndenting = true;
+			for (const auto& DebugData : DebugDataStack)
+			{
+				if (bIsIndenting)
+				{
+					FCStringAnsi::Strcat(StackTrace.Get(), StackTraceSize, TCHAR_TO_ANSI(*SubIndent));
+				}
+
+				ANSICHAR DebugName[NAME_SIZE];
+				DebugData.GetPlainANSIString(DebugName);
+				FCStringAnsi::Strcat(StackTrace.Get(), StackTraceSize, DebugName);
+
+				//these are special-cased, as we assume they'll be followed by object/property names and want the names on the same line for readability's sake.
+				const bool bIsPropertyLabel = (DebugData == TEXT("SerializeScriptProperties") || DebugData == TEXT("PropertySerialize") || DebugData == TEXT("SerializeTaggedProperty"));
+				const ANSICHAR* const LineEnd = bIsPropertyLabel ? ": " : "\r\n";
+				FCStringAnsi::Strcat(StackTrace.Get(), StackTraceSize, LineEnd);
+				bIsIndenting = !bIsPropertyLabel;
+			}
+		}
+		// Make sure we compare the new stack trace with the last one in the next if statement
+		bCallstacksDirty = true;
+	}
+
+	if (LastSerializeCallstack == nullptr || (bCallstacksDirty && FCStringAnsi::Strcmp(LastSerializeCallstack, StackTrace.Get()) != 0))
+	{
+		uint32 CallstackCRC = 0;
+		if (CallstackAtOffsetMap.Num() == 0 || CurrentOffset > CallstackAtOffsetMap.Last().Offset)
+		{
+			// New data serialized at the end of archive buffer
+			LastSerializeCallstack = AddUniqueCallstack(bIsCollectingCallstacks, SerializedObject, SerializedProperty, CallstackCRC);
+			CallstackAtOffsetMap.Add(FCallstackAtOffset {CurrentOffset, CallstackCRC, GIgnoreDiffManager.ShouldIgnoreDiff()});
+		}
+		else
+		{
+			// This happens usually after Seek() so we need to find the exiting offset or insert a new one
+			const int32 CallstackToUpdateIndex = GetCallstackIndexAtOffset(CurrentOffset);
+			check(CallstackToUpdateIndex != -1);
+			FCallstackAtOffset& CallstackToUpdate = CallstackAtOffsetMap[CallstackToUpdateIndex];
+			LastSerializeCallstack = AddUniqueCallstack(bIsCollectingCallstacks, SerializedObject, SerializedProperty, CallstackCRC);
+			if (CallstackToUpdate.Offset == CurrentOffset)
+			{
+				CallstackToUpdate.Callstack = CallstackCRC;
+			}
+			else
+			{
+				// Insert a new callstack
+				check(CallstackToUpdate.Offset < CurrentOffset);
+				CallstackAtOffsetMap.Insert(FCallstackAtOffset {CurrentOffset, CallstackCRC, GIgnoreDiffManager.ShouldIgnoreDiff()}, CallstackToUpdateIndex + 1);
+			}
+		}
+		check(CallstackCRC != 0 || !bShouldCollectCallstack);
+	}
+	else if (LastSerializeCallstack)
+	{
+		// Skip callstack comparison on next serialize call unless we grab a stack trace
+		bCallstacksDirty = false;
+	}
+}
+
+int32 FDiffWriterCallstacks::GetCallstackIndexAtOffset(int64 Offset, int32 MinOffsetIndex) const
+{
+	if (Offset < 0 || Offset > TotalSize || MinOffsetIndex < 0 || MinOffsetIndex >= CallstackAtOffsetMap.Num())
+	{
+		return -1;
+	}
+
+	// Find the index of the offset the InOffset maps to
+	int32 OffsetForCallstackIndex = -1;
+	int32 MaxOffsetIndex = CallstackAtOffsetMap.Num() - 1;
+
+	// Binary search
+	for (; MinOffsetIndex <= MaxOffsetIndex; )
+	{
+		int32 SearchIndex = (MinOffsetIndex + MaxOffsetIndex) / 2;
+		if (CallstackAtOffsetMap[SearchIndex].Offset < Offset)
+		{
+			MinOffsetIndex = SearchIndex + 1;
+		}
+		else if (CallstackAtOffsetMap[SearchIndex].Offset > Offset)
+		{
+			MaxOffsetIndex = SearchIndex - 1;
+		}
+		else
+		{
+			OffsetForCallstackIndex = SearchIndex;
+			break;
+		}
+	}
+	
+	if (OffsetForCallstackIndex == -1)
+	{
+		// We didn't find the exact offset value so let's try to find the first one that is lower than the requested one
+		MinOffsetIndex = FMath::Min(MinOffsetIndex, CallstackAtOffsetMap.Num() - 1);
+		for (int32 FirstLowerOffsetIndex = MinOffsetIndex; FirstLowerOffsetIndex >= 0; --FirstLowerOffsetIndex)
+		{
+			if (CallstackAtOffsetMap[FirstLowerOffsetIndex].Offset < Offset)
+			{
+				OffsetForCallstackIndex = FirstLowerOffsetIndex;
+				break;
+			}
+		}
+		check(OffsetForCallstackIndex != -1);
+		check(CallstackAtOffsetMap[OffsetForCallstackIndex].Offset < Offset);
+		check(OffsetForCallstackIndex == (CallstackAtOffsetMap.Num() - 1) || CallstackAtOffsetMap[OffsetForCallstackIndex + 1].Offset > Offset);
+	}
+
+	return OffsetForCallstackIndex;
+}
+
+void FDiffWriterCallstacks::Append(const FDiffWriterCallstacks& Other, int64 Offset)
+{
+	for (const FCallstackAtOffset& OtherOffset : Other.CallstackAtOffsetMap)
+	{
+		FCallstackAtOffset& New = CallstackAtOffsetMap.Add_GetRef(OtherOffset);
+		New.Offset += Offset;
+	}
+
+	CallstackAtOffsetMap.Sort([](const FCallstackAtOffset& LHS,const FCallstackAtOffset& RHS)
+	{
+		return LHS.Offset < RHS.Offset;
+	});
+
+	for (const TPair<uint32, FCallstackData>& Kv : Other.UniqueCallstacks)
+	{
+		if (FCallstackData* Existing = UniqueCallstacks.Find(Kv.Key))
+		{
+			if (LastSerializeCallstack == Existing->Callstack.Get())
+			{
+				LastSerializeCallstack = nullptr;
+			}
+			UniqueCallstacks.Remove(Kv.Key);
+		}
+		UniqueCallstacks.Emplace(Kv.Key, Kv.Value.Clone());
+	}
+
+	TotalSize = FMath::Max(TotalSize, Other.TotalSize);
+}
+
+FDiffWriterArchiveWriter::FDiffWriterArchiveWriter(
+	FArchive& InInner,
+	FDiffWriterCallstacks& InCallstacks,
+	const FDiffWriterDiffMap* InDiffMap,
+	int64 InDiffMapOffset)
+		: FArchiveProxy(InInner)
+		, Callstacks(InCallstacks)
+		, DiffMap(InDiffMap)
+		, DiffMapOffset(InDiffMapOffset)
+		, bInnerArchiveDisabled(false)
+{
+}
+
+FDiffWriterArchiveWriter::~FDiffWriterArchiveWriter() = default;
+
+void FDiffWriterArchiveWriter::Serialize(void* Data, int64 Length)
+{
+	static struct FBreakAtOffsetSettings
+	{
+		FString PackageToBreakOn;
+		int64 OffsetToBreakOn;
+
+		FBreakAtOffsetSettings()
+			: OffsetToBreakOn(-1)
+		{
+			if (!FParse::Param(FCommandLine::Get(), TEXT("cooksinglepackage")))
+			{
+				return;
+			}
+
+			FString Package;
+			if (!FParse::Value(FCommandLine::Get(), TEXT("map="), Package))
+			{
+				return;
+			}
+
+			int64 Offset;
+			if (!FParse::Value(FCommandLine::Get(), TEXT("diffonlybreakoffset="), Offset) || Offset <= 0)
+			{
+				return;
+			}
+
+			OffsetToBreakOn = Offset;
+			PackageToBreakOn = TEXT("/") + FPackageName::GetShortName(Package);
+		}
+	} BreakAtOffsetSettings;
+
+	const int64 CurrentOffset = DiffMapOffset + Tell();
+
+	if (BreakAtOffsetSettings.OffsetToBreakOn >= 0 && BreakAtOffsetSettings.OffsetToBreakOn >= CurrentOffset && BreakAtOffsetSettings.OffsetToBreakOn < CurrentOffset + Length)
+	{
+		FString ArcName = GetArchiveName();
+		int32 SubnameIndex = ArcName.Find(BreakAtOffsetSettings.PackageToBreakOn, ESearchCase::IgnoreCase, ESearchDir::FromEnd);
+		if (SubnameIndex >= 0)
+		{
+			int32 SubnameEndIndex = SubnameIndex + BreakAtOffsetSettings.PackageToBreakOn.Len();
+			if (SubnameEndIndex == ArcName.Len() || ArcName[SubnameEndIndex] == TEXT('.'))
+			{
+				UE_DEBUG_BREAK();
+			}
+		}
+	}
+
+	if (Length > 0)
+	{
+		UObject* SerializedObject = SerializeContext ? SerializeContext->SerializedObject : nullptr;
+		TArrayView<const FName> DebugStack;
+		DebugStack = DebugDataStack;
+
+		const bool bIsCollectingCallstacks = DiffMap != nullptr;
+		const bool bCollectCurrentCallstack = bIsCollectingCallstacks && DiffMap->ContainsOffset(CurrentOffset);
+
+		Callstacks.Add(
+			CurrentOffset,
+			Length,
+			SerializedObject,
+			GetSerializedProperty(),
+			DebugStack,
+			bIsCollectingCallstacks,
+			bCollectCurrentCallstack,
+			StackIgnoreCount);
+	}
+
+	if (bInnerArchiveDisabled == false)
+	{
+		InnerArchive.Serialize(Data, Length);
+	}
+}
+
+void FDiffWriterArchiveWriter::SetSerializeContext(FUObjectSerializeContext* Context)
+{
+	SerializeContext = Context;
+
+	if (bInnerArchiveDisabled == false)
+	{
+		InnerArchive.SetSerializeContext(Context);
+	}
+}
+
+FUObjectSerializeContext* FDiffWriterArchiveWriter::GetSerializeContext()
+{
+	return SerializeContext;
+}
+
+FDiffWriterArchiveMemoryWriter::FDiffWriterArchiveMemoryWriter(
+	FDiffWriterCallstacks& Callstacks,
+	const FDiffWriterDiffMap* DiffMap,
+	const int64 DiffMapOffset,
+	const int64 PreAllocateBytes,
+	bool bIsPersistent,
+	const TCHAR* Filename)
+		: FLargeMemoryWriter(PreAllocateBytes, bIsPersistent, Filename)
+		, StackTraceWriter(*this, Callstacks, DiffMap, DiffMapOffset)
+{
+	// Hack to prevent recursive calls to serialize when passing in this.
+	StackTraceWriter.SetDisableInnerArchive(true);
+	StackTraceWriter.SetStackIgnoreCount(StackTraceWriter.GetStackIgnoreCount() + 1);
+}
+
+void FDiffWriterArchiveMemoryWriter::Serialize(void* Memory, int64 Length)
+{
+	StackTraceWriter.Serialize(Memory, Length);
+	FLargeMemoryWriter::Serialize(Memory, Length);
+}
+
+void FDiffWriterArchiveMemoryWriter::SetSerializeContext(FUObjectSerializeContext* Context)
+{
+	StackTraceWriter.SetSerializeContext(Context);
+}
+
+FUObjectSerializeContext* FDiffWriterArchiveMemoryWriter::GetSerializeContext()
+{
+	return StackTraceWriter.GetSerializeContext();
+}
+
+FDiffWriterArchive::FDiffWriterArchive(UObject* InAsset, const TCHAR* InFilename, bool bInCollectCallstacks, const FDiffWriterDiffMap* InDiffMap)
+	: FLargeMemoryWriter(0, false, InFilename)
+	, Callstacks(InAsset)
+	, StackTraceWriter(*this, Callstacks, InDiffMap, 0)
+{
+	// Hack to prevent recursive calls to serialize when passing in this.
+	StackTraceWriter.SetDisableInnerArchive(true);
+	StackTraceWriter.SetStackIgnoreCount(StackTraceWriter.GetStackIgnoreCount() + 1);
+}
+
+FDiffWriterArchive::~FDiffWriterArchive()
+{
+}
+
+void FDiffWriterArchive::Serialize(void* Memory, int64 Length)
+{
+	StackTraceWriter.Serialize(Memory, Length);
+	FLargeMemoryWriter::Serialize(Memory, Length);
+}
+
+void FDiffWriterArchive::SetSerializeContext(FUObjectSerializeContext* Context)
+{
+	StackTraceWriter.SetSerializeContext(Context);
+}
+
+FUObjectSerializeContext* FDiffWriterArchive::GetSerializeContext()
+{
+	return StackTraceWriter.GetSerializeContext();
+}
+
+namespace
+{
+	bool ShouldDumpPropertyValueState(FProperty* Prop)
+	{
+		if (Prop->IsA<FNumericProperty>()
+			|| Prop->IsA<FStrProperty>()
+			|| Prop->IsA<FBoolProperty>()
+			|| Prop->IsA<FNameProperty>())
+		{
+			return true;
+		}
+
+		if (FArrayProperty* ArrayProp = CastField<FArrayProperty>(Prop))
+		{
+			return ShouldDumpPropertyValueState(ArrayProp->Inner);
+		}
+
+		if (FMapProperty* MapProp = CastField<FMapProperty>(Prop))
+		{
+			return ShouldDumpPropertyValueState(MapProp->KeyProp) && ShouldDumpPropertyValueState(MapProp->ValueProp);
+		}
+
+		if (FSetProperty* SetProp = CastField<FSetProperty>(Prop))
+		{
+			return ShouldDumpPropertyValueState(SetProp->ElementProp);
+		}
+
+		if (FStructProperty* StructProp = CastField<FStructProperty>(Prop))
+		{
+			if (StructProp->Struct == TBaseStructure<FVector>::Get()
+				|| StructProp->Struct == TBaseStructure<FGuid>::Get())
+			{
+				return true;
+			}
+		}
+
+		return false;
+	}
+}
+
+void FDiffWriterArchiveWriter::Compare(
+	const FPackageData& SourcePackage,
+	const FPackageData& DestPackage,
+	const FDiffWriterCallstacks& Callstacks,
+	const FDiffWriterDiffMap& DiffMap,
+	const TCHAR* AssetFilename,
+	const TCHAR* CallstackCutoffText,
+	const int64 MaxDiffsToLog,
+	int32& InOutDiffsLogged,
+	TMap<FName, FArchiveDiffStats>& OutStats,
+	bool bSuppressLogging)
+{
+	const TCHAR* const Indent = FDiffFormatHelper::Get().Indent;
+	const TCHAR* const LineTerminator = FDiffFormatHelper::Get().LineTerminator;
+	const int64 SourceSize = SourcePackage.Size - SourcePackage.StartOffset;
+	const int64 DestSize = DestPackage.Size - DestPackage.StartOffset;
+	const int64 SizeToCompare = FMath::Min(SourceSize, DestSize);
+	const FName AssetClass = Callstacks.GetAssetClass();
+	
+	if (SourceSize != DestSize)
+	{
+		UE_CLOG(!bSuppressLogging, LogDiffWriter, Warning, TEXT("%s: Size mismatch: on disk: %lld vs memory: %lld"), AssetFilename, SourceSize, DestSize);
+		int64 SizeDiff = DestPackage.Size - SourcePackage.Size;
+		OutStats.FindOrAdd(AssetClass).DiffSize += SizeDiff;
+	}
+
+	FString LastDifferenceCallstackDataText;
+	int32 LastDifferenceCallstackOffsetIndex = -1;
+	int64 NumDiffsLocal = 0;
+	int64 NumDiffsLoggedLocal = 0;
+	int64 FirstUnreportedDiffIndex = -1;
+
+	for (int64 LocalOffset = 0; LocalOffset < SizeToCompare; ++LocalOffset)
+	{
+		const int64 SourceAbsoluteOffset = LocalOffset + SourcePackage.StartOffset;
+		const int64 DestAbsoluteOffset = LocalOffset + DestPackage.StartOffset;
+
+		const uint8 SourceByte = SourcePackage.Data[SourceAbsoluteOffset];
+		const uint8 DestByte   = DestPackage  .Data[DestAbsoluteOffset];
+		if (SourceByte == DestByte)
+		{
+			continue;
+		}
+
+		bool bDifferenceLogged = false;
+		ON_SCOPE_EXIT
+		{
+			if (bDifferenceLogged)
+			{
+				InOutDiffsLogged++;
+				NumDiffsLoggedLocal++;
+			}
+		};
+
+		if (DiffMap.ContainsOffset(DestAbsoluteOffset))
+		{
+			int32 DifferenceCallstackoffsetIndex = Callstacks.GetCallstackIndexAtOffset(DestAbsoluteOffset, FMath::Max(LastDifferenceCallstackOffsetIndex, 0));
+			ON_SCOPE_EXIT
+			{
+				LastDifferenceCallstackOffsetIndex = DifferenceCallstackoffsetIndex;
+			};
+
+			if (DifferenceCallstackoffsetIndex < 0)
+			{
+				UE_CLOG(!bSuppressLogging, LogDiffWriter, Warning, TEXT("%s: Difference at offset %lld (absolute offset: %lld), unknown callstack"), AssetFilename, LocalOffset, DestAbsoluteOffset);
+				continue;
+			}
+
+			if (DifferenceCallstackoffsetIndex == LastDifferenceCallstackOffsetIndex)
+			{
+				continue;
+			}
+
+			const FDiffWriterCallstacks::FCallstackAtOffset& CallstackAtOffset = Callstacks.GetCallstack(DifferenceCallstackoffsetIndex);
+			const FDiffWriterCallstacks::FCallstackData& DifferenceCallstackData = Callstacks.GetCallstackData(CallstackAtOffset);
+			FString DifferenceCallstackDataText = DifferenceCallstackData.ToString(CallstackCutoffText);
+			if (LastDifferenceCallstackDataText.Compare(DifferenceCallstackDataText, ESearchCase::CaseSensitive) == 0)
+			{
+				continue;
+			}
+			ON_SCOPE_EXIT
+			{
+				LastDifferenceCallstackDataText = MoveTemp(DifferenceCallstackDataText);
+			};
+
+			if (!CallstackAtOffset.bIgnore && (MaxDiffsToLog < 0 || InOutDiffsLogged < MaxDiffsToLog))
+			{
+				FString BeforePropertyVal;
+				FString AfterPropertyVal;
+				if (FProperty* SerProp = DifferenceCallstackData.SerializedProp)
+				{
+					if (SourceSize == DestSize && ShouldDumpPropertyValueState(SerProp))
+					{
+						// Walk backwards until we find a callstack which wasn't from the given property
+						int64 OffsetX = DestAbsoluteOffset;
+						for (;;)
+						{
+							if (OffsetX == 0)
+							{
+								break;
+							}
+
+							const int32 CallstackIndex = Callstacks.GetCallstackIndexAtOffset(OffsetX - 1, 0);
+							const FDiffWriterCallstacks::FCallstackAtOffset& PreviousCallstack = Callstacks.GetCallstack(CallstackIndex);
+							const FDiffWriterCallstacks::FCallstackData& PreviousCallstackData = Callstacks.GetCallstackData(PreviousCallstack);
+							if (PreviousCallstackData.SerializedProp != SerProp)
+							{
+								break;
+							}
+
+							--OffsetX;
+						}
+
+						FPropertyTempVal SourceVal(SerProp);
+						FPropertyTempVal DestVal  (SerProp);
+
+						FStaticMemoryReader SourceReader(&SourcePackage.Data[SourceAbsoluteOffset - (DestAbsoluteOffset - OffsetX)], SourcePackage.Size - SourceAbsoluteOffset);
+						FStaticMemoryReader DestReader(&DestPackage.Data[OffsetX], DestPackage.Size - DestAbsoluteOffset);
+
+						SourceVal.Serialize(SourceReader);
+						DestVal  .Serialize(DestReader);
+
+									if (!SourceReader.IsError() && !DestReader.IsError())
+									{
+										SourceVal.ExportText(BeforePropertyVal);
+										DestVal  .ExportText(AfterPropertyVal);
+									}
+								}
+							}
+
+				FString DiffValues;
+				if (BeforePropertyVal != AfterPropertyVal)
+				{
+					DiffValues = FString::Printf(TEXT("\r\n%sBefore: %s\r\n%sAfter:  %s"), Indent, *BeforePropertyVal, Indent, *AfterPropertyVal);
+				}
+
+				FString DebugDataStackText;
+				//check for a debug data stack as part of the unique stack entry, and log it out if we find it.
+				FString FullStackText = DifferenceCallstackData.Callstack.Get();
+				int32 DebugDataIndex = FullStackText.Find(ANSI_TO_TCHAR(DebugDataStackMarker), ESearchCase::CaseSensitive);
+				if (DebugDataIndex > 0)
+				{
+					DebugDataStackText = FString::Printf(TEXT("\r\n%s"), FDiffFormatHelper::Get().Indent) + FullStackText.RightChop(DebugDataIndex + 2);
+				}
+
+				UE_CLOG(
+					!bSuppressLogging,
+					LogDiffWriter,
+					Warning,
+					TEXT("%s: Difference at offset %lld%s (absolute offset: %lld): byte %d on disk, byte %d in memory, callstack:%s%s%s%s%s"),
+					AssetFilename,
+					CallstackAtOffset.Offset - DestPackage.StartOffset,
+					DestAbsoluteOffset > CallstackAtOffset.Offset ? *FString::Printf(TEXT("(+%lld)"), DestAbsoluteOffset - CallstackAtOffset.Offset) : TEXT(""),
+					DestAbsoluteOffset,
+					SourceByte, DestByte,
+					LineTerminator,
+					LineTerminator,
+					*DifferenceCallstackDataText,
+					*DiffValues,
+					*DebugDataStackText
+				);
+
+				const int BytesToLog = 128;
+				UE_CLOG(
+					!bSuppressLogging,
+					LogDiffWriter,
+					Display,
+					TEXT("%s: Logging %d bytes around absolute offset: %lld (%016X) in the on disk (existing) package, (which corresponds to offset %lld (%016X) in the in-memory package)"),
+					AssetFilename,
+					BytesToLog,
+					SourceAbsoluteOffset,
+					SourceAbsoluteOffset,
+					DestAbsoluteOffset,
+					DestAbsoluteOffset
+				);
+				FCompressionUtil::LogHexDump(SourcePackage.Data, SourcePackage.Size, SourceAbsoluteOffset - BytesToLog / 2, SourceAbsoluteOffset + BytesToLog / 2);
+
+				UE_CLOG(
+					!bSuppressLogging,
+					LogDiffWriter,
+					Display,
+					TEXT("%s: Logging %d bytes around absolute offset: %lld (%016X) in the in memory (new) package"),
+					AssetFilename,
+					BytesToLog,
+					DestAbsoluteOffset,
+					DestAbsoluteOffset
+				);
+				FCompressionUtil::LogHexDump(DestPackage.Data, DestPackage.Size, DestAbsoluteOffset - BytesToLog / 2, DestAbsoluteOffset + BytesToLog / 2);
+
+				bDifferenceLogged = true;
+			}
+			else if (FirstUnreportedDiffIndex == -1)
+			{
+				FirstUnreportedDiffIndex = DestAbsoluteOffset;
+			}
+			OutStats.FindOrAdd(AssetClass).NumDiffs++;
+			NumDiffsLocal++;
+		}
+		else
+		{
+			// Each byte will count as a difference but without callstack data there's no way around it
+			OutStats.FindOrAdd(AssetClass).NumDiffs++;
+			NumDiffsLocal++;
+			if (FirstUnreportedDiffIndex == -1)
+			{
+				FirstUnreportedDiffIndex = DestAbsoluteOffset;
+			}
+		}
+		OutStats.FindOrAdd(AssetClass).DiffSize++;
+	}
+
+	if (MaxDiffsToLog >= 0 && NumDiffsLocal > NumDiffsLoggedLocal)
+	{
+		if (FirstUnreportedDiffIndex != -1)
+		{
+			UE_CLOG(!bSuppressLogging, LogDiffWriter, Warning, TEXT("%s: %lld difference(s) not logged (first at absolute offset: %lld)."), AssetFilename, NumDiffsLocal - NumDiffsLoggedLocal, FirstUnreportedDiffIndex);
+		}
+		else
+		{
+			UE_CLOG(!bSuppressLogging, LogDiffWriter, Warning, TEXT("%s: %lld difference(s) not logged."), AssetFilename, NumDiffsLocal - NumDiffsLoggedLocal);
+		}
+	}
+}
+
+void FDiffWriterArchive::CompareWith(const TCHAR* InFilename, const int64 TotalHeaderSize, const TCHAR* CallstackCutoffText, const int32 MaxDiffsToLog, TMap<FName, FArchiveDiffStats>& OutStats)
+{
+	TUniquePtr<uint8> SourcePackageBytes;
+	FPackageData SourcePackage;
+	UE::ArchiveStackTrace::LoadPackageIntoMemory(InFilename, SourcePackage, SourcePackageBytes);
+	CompareWith(SourcePackage, InFilename, TotalHeaderSize, CallstackCutoffText, MaxDiffsToLog, OutStats);
+}
+
+void FDiffWriterArchive::CompareWith(const FPackageData& SourcePackage, const TCHAR* FileDisplayName, const int64 TotalHeaderSize,
+	const TCHAR* CallstackCutoffText, const int32 MaxDiffsToLog, TMap<FName, FArchiveDiffStats>&OutStats,
+	const FDiffWriterArchiveWriter::EPackageHeaderFormat PackageHeaderFormat /* = FDiffWriterArchiveWriter::EPackageHeaderFormat::PackageFileSummary */)
+{
+	const FName AssetClass = Callstacks.GetAssetClass();
+	OutStats.FindOrAdd(AssetClass).NewFileTotalSize = TotalSize();
+	if (SourcePackage.Size == 0)
+	{
+		UE_LOG(LogDiffWriter, Warning, TEXT("New package: %s"), *GetArchiveName());
+		OutStats.FindOrAdd(AssetClass).DiffSize = OutStats.FindOrAdd(AssetClass).NewFileTotalSize;
+		return;
+	}
+
+	FPackageData DestPackage;
+	DestPackage.Data = GetData();
+	DestPackage.Size = TotalSize();
+	DestPackage.HeaderSize = TotalHeaderSize;
+	DestPackage.StartOffset = 0;
+
+	UE_LOG(LogDiffWriter, Display, TEXT("Comparing: %s"), *GetArchiveName());
+	UE_LOG(LogDiffWriter, Display, TEXT("Asset class: %s"), *AssetClass.ToString());
+
+	int32 NumLoggedDiffs = 0;
+
+	FPackageData SourcePackageHeader = SourcePackage;
+	SourcePackageHeader.Size = SourcePackageHeader.HeaderSize;
+	SourcePackageHeader.HeaderSize = 0;
+	SourcePackageHeader.StartOffset = 0;
+
+	FPackageData DestPackageHeader = DestPackage;
+	DestPackageHeader.Size = TotalHeaderSize;
+	DestPackageHeader.HeaderSize = 0;
+	DestPackageHeader.StartOffset = 0;
+
+	FDiffWriterArchiveWriter::Compare(SourcePackageHeader, DestPackageHeader, Callstacks, StackTraceWriter.GetDiffMap(), FileDisplayName, CallstackCutoffText, MaxDiffsToLog, NumLoggedDiffs, OutStats);
+
+	if (TotalHeaderSize > 0 && OutStats.FindOrAdd(AssetClass).NumDiffs > 0)
+	{
+		FDiffWriterArchiveWriter::DumpPackageHeaderDiffs(SourcePackage, DestPackage, FileDisplayName, MaxDiffsToLog, PackageHeaderFormat);
+	}
+
+	FPackageData SourcePackageExports = SourcePackage;
+	SourcePackageExports.HeaderSize = 0;
+	SourcePackageExports.StartOffset = SourcePackage.HeaderSize;
+
+	FPackageData DestPackageExports = DestPackage;
+	DestPackageExports.HeaderSize = 0;
+	DestPackageExports.StartOffset = TotalHeaderSize;
+
+	FString AssetName;
+	if (DestPackage.HeaderSize > 0)
+	{
+		AssetName = FPaths::ChangeExtension(FileDisplayName, TEXT("uexp"));
+	}
+	else
+	{
+		AssetName = FileDisplayName;
+	}
+
+	FDiffWriterArchiveWriter::Compare(SourcePackageExports, DestPackageExports, Callstacks, StackTraceWriter.GetDiffMap(), *AssetName, CallstackCutoffText, MaxDiffsToLog, NumLoggedDiffs, OutStats);
+
+	// Optionally save out any differences we detected.
+	const FArchiveDiffStats& Stats = OutStats.FindOrAdd(AssetClass);
+	if (Stats.NumDiffs > 0)
+	{
+		static struct FDiffOutputSettings
+		{
+			FString DiffOutputDir;
+
+			FDiffOutputSettings()
+			{
+				FString Dir;
+				if (!FParse::Value(FCommandLine::Get(), TEXT("diffoutputdir="), Dir))
+				{
+					return;
+				}
+
+				FPaths::NormalizeDirectoryName(Dir);
+				DiffOutputDir = MoveTemp(Dir) + TEXT("/");
+			}
+		} DiffOutputSettings;
+
+		// Only save out the differences if we have a -diffoutputdir set.
+		if (!DiffOutputSettings.DiffOutputDir.IsEmpty())
+		{
+			FString OutputFilename = FPaths::ConvertRelativePathToFull(FileDisplayName);
+			FString SavedDir       = FPaths::ConvertRelativePathToFull(FPaths::ProjectSavedDir());
+			if (OutputFilename.StartsWith(SavedDir))
+			{
+				OutputFilename.ReplaceInline(*SavedDir, *DiffOutputSettings.DiffOutputDir);
+
+				IFileManager& FileManager = IFileManager::Get();
+
+				// Copy the original asset as '.before.uasset'.
+				{
+					TUniquePtr<FArchive> DiffUAssetArchive(FileManager.CreateFileWriter(*FPaths::SetExtension(OutputFilename, TEXT(".before.") + FPaths::GetExtension(FileDisplayName))));
+					DiffUAssetArchive->Serialize(SourcePackageHeader.Data + SourcePackageHeader.StartOffset, SourcePackageHeader.Size - SourcePackageHeader.StartOffset);
+				}
+				{
+					TUniquePtr<FArchive> DiffUExpArchive(FileManager.CreateFileWriter(*FPaths::SetExtension(OutputFilename, TEXT(".before.uexp"))));
+					DiffUExpArchive->Serialize(SourcePackageExports.Data + SourcePackageExports.StartOffset, SourcePackageExports.Size - SourcePackageExports.StartOffset);
+				}
+
+				// Save out the in-memory data as '.after.uasset'.
+				{
+					TUniquePtr<FArchive> DiffUAssetArchive(FileManager.CreateFileWriter(*FPaths::SetExtension(OutputFilename, TEXT(".after.") + FPaths::GetExtension(FileDisplayName))));
+					DiffUAssetArchive->Serialize(DestPackageHeader.Data + DestPackageHeader.StartOffset, DestPackageHeader.Size - DestPackageHeader.StartOffset);
+				}
+				{
+					TUniquePtr<FArchive> DiffUExpArchive(FileManager.CreateFileWriter(*FPaths::SetExtension(OutputFilename, TEXT(".after.uexp"))));
+					DiffUExpArchive->Serialize(DestPackageExports.Data + DestPackageExports.StartOffset, DestPackageExports.Size - DestPackageExports.StartOffset);
+				}
+			}
+			else
+			{
+				UE_LOG(LogDiffWriter, Warning, TEXT("Package '%s' doesn't seem to be writing to the Saved directory - skipping writing diff"), *OutputFilename);
+			}
+		}
+	}
+}
+
+bool FDiffWriterArchiveWriter::GenerateDiffMap(const FPackageData& SourcePackage, const FPackageData& DestPackage, const FDiffWriterCallstacks& Callstacks, int32 MaxDiffsToFind, FDiffWriterDiffMap& OutDiffMap)
+{
+	bool bIdentical = true;
+	int32 LastDifferenceCallstackOffsetIndex = -1;
+	FDiffWriterCallstacks::FCallstackData* DifferenceCallstackData = nullptr;
+
+	const int64 SourceSize = SourcePackage.Size - SourcePackage.StartOffset;
+	const int64 DestSize = DestPackage.Size - DestPackage.StartOffset;
+	const int64 SizeToCompare = FMath::Min(SourceSize, DestSize);
+	
+	for (int64 LocalOffset = 0; LocalOffset < SizeToCompare; ++LocalOffset)
+	{
+		const int64 SourceAbsoluteOffset = LocalOffset + SourcePackage.StartOffset;
+		const int64 DestAbsoluteOffset = LocalOffset + DestPackage.StartOffset;
+		if (SourcePackage.Data[SourceAbsoluteOffset] != DestPackage.Data[DestAbsoluteOffset])
+		{
+			bIdentical = false;
+			if (OutDiffMap.Num() < MaxDiffsToFind)
+			{
+				const int32 DifferenceCallstackOffsetIndex = Callstacks.GetCallstackIndexAtOffset(DestAbsoluteOffset, FMath::Max<int32>(LastDifferenceCallstackOffsetIndex, 0));
+				if (DifferenceCallstackOffsetIndex >= 0 && DifferenceCallstackOffsetIndex != LastDifferenceCallstackOffsetIndex)
+				{
+					const FDiffWriterCallstacks::FCallstackAtOffset& CallstackAtOffset = Callstacks.GetCallstack(DifferenceCallstackOffsetIndex);
+					if (!CallstackAtOffset.bIgnore)
+					{
+						FDiffWriterDiffInfo OffsetAndSize;
+						OffsetAndSize.Offset = CallstackAtOffset.Offset;
+						OffsetAndSize.Size = Callstacks.GetSerializedDataSizeForOffsetIndex(DifferenceCallstackOffsetIndex);
+						OutDiffMap.Add(OffsetAndSize);
+					}
+				}
+				LastDifferenceCallstackOffsetIndex = DifferenceCallstackOffsetIndex;
+			}
+		}
+	}
+
+	if (SourceSize < DestSize)
+	{
+		bIdentical = false;
+
+		// Add all the remaining callstacks to the diff map
+		for (int32 OffsetIndex = LastDifferenceCallstackOffsetIndex + 1; OffsetIndex < Callstacks.Num() && OutDiffMap.Num() < MaxDiffsToFind; ++OffsetIndex)
+		{
+			const FDiffWriterCallstacks::FCallstackAtOffset& CallstackAtOffset = Callstacks.GetCallstack(OffsetIndex);
+			// Compare against the size without start offset as all callstack offsets are absolute (from the merged header + exports file)
+			if (CallstackAtOffset.Offset < DestPackage.Size)
+			{
+				if (!CallstackAtOffset.bIgnore)
+				{
+					FDiffWriterDiffInfo OffsetAndSize;
+					OffsetAndSize.Offset = CallstackAtOffset.Offset;
+					OffsetAndSize.Size = Callstacks.GetSerializedDataSizeForOffsetIndex(OffsetIndex);
+					OutDiffMap.Add(OffsetAndSize);
+				}
+			}
+			else
+			{
+				break;
+			}
+		}
+	}
+	else if (SourceSize > DestSize)
+	{
+		bIdentical = false;
+	}
+	return bIdentical;
+}
+
+bool FDiffWriterArchive::GenerateDiffMap(const TCHAR* InFilename, int64 TotalHeaderSize, int32 MaxDiffsToFind, FDiffWriterDiffMap& OutDiffMap)
+{
+	TUniquePtr<uint8> SourcePackageBytes;
+	FPackageData SourcePackage;
+	if (!UE::ArchiveStackTrace::LoadPackageIntoMemory(InFilename, SourcePackage, SourcePackageBytes))
+	{
+		return false;
+	}
+	return GenerateDiffMap(SourcePackage, TotalHeaderSize, MaxDiffsToFind, OutDiffMap);
+}
+
+bool FDiffWriterArchive::GenerateDiffMap(const FPackageData& SourcePackage, int64 TotalHeaderSize, int32 MaxDiffsToFind, FDiffWriterDiffMap& OutDiffMap)
+{
+	check(MaxDiffsToFind > 0);
+
+	bool bIdentical = true;
+	bool bHeaderIdentical = true;
+	bool bExportsIdentical = true;
+
+	FPackageData DestPackage;
+	DestPackage.Data = GetData();
+	DestPackage.Size = TotalSize();
+	DestPackage.HeaderSize = TotalHeaderSize;
+	DestPackage.StartOffset = 0;
+
+	{
+		FPackageData SourcePackageHeader = SourcePackage;
+		SourcePackageHeader.Size = SourcePackageHeader.HeaderSize;
+		SourcePackageHeader.HeaderSize = 0;
+		SourcePackageHeader.StartOffset = 0;
+
+		FPackageData DestPackageHeader = DestPackage;
+		DestPackageHeader.Size = TotalHeaderSize;
+		DestPackageHeader.HeaderSize = 0;
+		DestPackageHeader.StartOffset = 0;
+
+		bHeaderIdentical = FDiffWriterArchiveWriter::GenerateDiffMap(SourcePackageHeader, DestPackageHeader, Callstacks, MaxDiffsToFind, OutDiffMap);
+	}
+
+	{
+		FPackageData SourcePackageExports = SourcePackage;
+		SourcePackageExports.HeaderSize = 0;
+		SourcePackageExports.StartOffset = SourcePackage.HeaderSize;
+
+		FPackageData DestPackageExports = DestPackage;
+		DestPackageExports.HeaderSize = 0;
+		DestPackageExports.StartOffset = TotalHeaderSize;
+
+		bExportsIdentical = FDiffWriterArchiveWriter::GenerateDiffMap(SourcePackageExports, DestPackageExports, Callstacks, MaxDiffsToFind, OutDiffMap);
+	}
+
+	bIdentical = bHeaderIdentical && bExportsIdentical;
+
+	return bIdentical;
+}
+
+bool FDiffWriterArchive::IsIdentical(const TCHAR* InFilename, int64 BufferSize, const uint8* BufferData)
+{
+	TUniquePtr<uint8> SourcePackageBytes;
+	FPackageData SourcePackage;
+	if (!UE::ArchiveStackTrace::LoadPackageIntoMemory(InFilename, SourcePackage, SourcePackageBytes))
+	{
+		return false;
+	}
+
+	return IsIdentical(SourcePackage, BufferSize, BufferData);
+}
+
+bool FDiffWriterArchive::IsIdentical(const FPackageData& SourcePackage, int64 BufferSize, const uint8* BufferData)
+{
+	bool bIdentical = false;
+	if (BufferSize == SourcePackage.Size)
+	{
+		bIdentical = (FMemory::Memcmp(SourcePackage.Data, BufferData, BufferSize) == 0);
+	}
+	else
+	{
+		bIdentical = false;
+	}
+	return bIdentical;
+}
+
+FLinkerLoad* FDiffWriterArchiveWriter::CreateLinkerForPackage(FUObjectSerializeContext* LoadContext, const FString& InPackageName, const FString& InFilename, const FPackageData& PackageData)
+{
+	// First create a temp package to associate the linker with
+	UPackage* Package = FindObjectFast<UPackage>(nullptr, *InPackageName);
+	if (!Package)
+	{
+		Package = CreatePackage(*InPackageName);
+	}
+	// Create an archive for the linker. The linker will take ownership of it.
+	FLargeMemoryReader* PackageReader = new FLargeMemoryReader(PackageData.Data, PackageData.Size, ELargeMemoryReaderFlags::None, *InPackageName);	
+	FLinkerLoad* Linker = FLinkerLoad::CreateLinker(LoadContext, Package, FPackagePath::FromLocalPath(InFilename), LOAD_NoVerify, PackageReader);
+
+	if (Linker && Package)
+	{
+		Package->SetPackageFlags(PKG_ForDiffing);
+	}
+
+	return Linker;
+}
+
+static FString GetTableKey(const FLinkerLoad* Linker, const FObjectExport& Export)
+{
+	FName ClassName = Export.ClassIndex.IsNull() ? FName(NAME_Class) : Linker->ImpExp(Export.ClassIndex).ObjectName;
+	return FString::Printf(TEXT("%s %s.%s"),
+		*ClassName.ToString(),
+		!Export.OuterIndex.IsNull() ? *Linker->ImpExp(Export.OuterIndex).ObjectName.ToString() : *FPackageName::GetShortName(Linker->LinkerRoot),
+		*Export.ObjectName.ToString());
+}
+
+static FString GetTableKey(const FLinkerLoad* Linker, const FObjectImport& Import)
+{
+	return FString::Printf(TEXT("%s %s.%s"),
+		*Import.ClassName.ToString(),
+		!Import.OuterIndex.IsNull() ? *Linker->ImpExp(Import.OuterIndex).ObjectName.ToString() : TEXT("NULL"),
+		*Import.ObjectName.ToString());
+}
+
+static inline FString GetTableKey(const FLinkerLoad* Linker, const FName& Name)
+{
+	return *Name.ToString();
+}
+
+static inline FString GetTableKey(const FLinkerLoad* Linker, FNameEntryId Id)
+{
+	return FName::GetEntry(Id)->GetPlainNameString();
+}
+
+static inline FString GetTableKeyForIndex(const FLinkerLoad* Linker, FPackageIndex Index)
+{
+	if (Index.IsNull())
+	{
+		return TEXT("NULL");
+	}
+	else if (Index.IsExport())
+	{
+		return GetTableKey(Linker, Linker->Exp(Index));
+	}
+	else
+	{
+		return GetTableKey(Linker, Linker->Imp(Index));
+	}
+}
+
+static bool ComparePackageIndices(FLinkerLoad* SourceLinker, FLinkerLoad* DestLinker, const FPackageIndex& SourceIndex, const FPackageIndex& DestIndex);
+
+static bool CompareTableItem(FLinkerLoad* SourceLinker, FLinkerLoad* DestLinker, const FName& SourceName, const FName& DestName)
+{
+	return SourceName == DestName;
+}
+
+static bool CompareTableItem(FLinkerLoad* SourceLinker, FLinkerLoad* DestLinker, FNameEntryId SourceName, FNameEntryId DestName)
+{
+	return SourceName == DestName;
+}
+
+static FString ConvertItemToText(const FName& Name, FLinkerLoad* Linker)
+{
+	return Name.ToString();
+}
+
+static FString ConvertItemToText(FNameEntryId Id, FLinkerLoad* Linker)
+{
+	return FName::GetEntry(Id)->GetPlainNameString();
+}
+
+static bool CompareTableItem(FLinkerLoad* SourceLinker, FLinkerLoad* DestLinker, const FObjectImport& SourceImport, const FObjectImport& DestImport)
+{
+	if (SourceImport.ObjectName != DestImport.ObjectName ||
+		SourceImport.ClassName != DestImport.ClassName ||
+		SourceImport.ClassPackage != DestImport.ClassPackage ||
+		!ComparePackageIndices(SourceLinker, DestLinker, SourceImport.OuterIndex, DestImport.OuterIndex))
+	{
+		return false;
+	}
+	else
+	{
+		return true;
+	}
+}
+
+static FString ConvertItemToText(const FObjectImport& Import, FLinkerLoad* Linker)
+{
+	return FString::Printf(
+		TEXT("%s ClassPackage: %s"),
+		*GetTableKey(Linker, Import),
+		*Import.ClassPackage.ToString()
+	);
+}
+
+static bool CompareTableItem(FLinkerLoad* SourceLinker, FLinkerLoad* DestLinker, const FObjectExport& SourceExport, const FObjectExport& DestExport)
+{
+	if (SourceExport.ObjectName != DestExport.ObjectName ||
+		SourceExport.PackageFlags != DestExport.PackageFlags ||
+		SourceExport.ObjectFlags != DestExport.ObjectFlags ||
+		SourceExport.SerialSize != DestExport.SerialSize ||
+		SourceExport.bForcedExport != DestExport.bForcedExport ||
+		SourceExport.bNotForClient != DestExport.bNotForClient ||
+		SourceExport.bNotForServer != DestExport.bNotForServer ||
+		SourceExport.bNotAlwaysLoadedForEditorGame != DestExport.bNotAlwaysLoadedForEditorGame ||
+		SourceExport.bIsAsset != DestExport.bIsAsset ||
+		SourceExport.bIsInheritedInstance != DestExport.bIsInheritedInstance ||
+		SourceExport.bGeneratePublicHash != DestExport.bGeneratePublicHash ||
+		!ComparePackageIndices(SourceLinker, DestLinker, SourceExport.TemplateIndex, DestExport.TemplateIndex) ||
+		!ComparePackageIndices(SourceLinker, DestLinker, SourceExport.OuterIndex, DestExport.OuterIndex) ||
+		!ComparePackageIndices(SourceLinker, DestLinker, SourceExport.ClassIndex, DestExport.ClassIndex) ||
+		!ComparePackageIndices(SourceLinker, DestLinker, SourceExport.SuperIndex, DestExport.SuperIndex))
+	{
+		return false;
+	}
+	else
+	{
+		return true;
+	}
+}
+
+static bool IsImportMapIdentical(FLinkerLoad* SourceLinker, FLinkerLoad* DestLinker)
+{
+	bool bIdentical = (SourceLinker->ImportMap.Num() == DestLinker->ImportMap.Num());
+	if (bIdentical)
+	{
+		for (int32 ImportIndex = 0; ImportIndex < SourceLinker->ImportMap.Num(); ++ImportIndex)
+		{
+			if (!CompareTableItem(SourceLinker, DestLinker, SourceLinker->ImportMap[ImportIndex], DestLinker->ImportMap[ImportIndex]))
+			{
+				bIdentical = false;
+				break;
+			}
+		}
+	}
+	return bIdentical;
+}
+
+static bool ComparePackageIndices(FLinkerLoad* SourceLinker, FLinkerLoad* DestLinker, const FPackageIndex& SourceIndex, const FPackageIndex& DestIndex)
+{
+	if (SourceIndex.IsNull() && DestIndex.IsNull())
+	{
+		return true;
+	}
+
+	if (SourceIndex.IsExport() && DestIndex.IsExport())
+	{
+		int32 SourceArrayIndex = SourceIndex.ToExport();
+		int32 DestArrayIndex   = DestIndex  .ToExport();
+
+		if (!SourceLinker->ExportMap.IsValidIndex(SourceArrayIndex) || !DestLinker->ExportMap.IsValidIndex(DestArrayIndex))
+		{
+			UE_LOG(LogDiffWriter, Warning, TEXT("Invalid export indices found, source: %d (of %d), dest: %d (of %d)"), SourceArrayIndex, SourceLinker->ExportMap.Num(), DestArrayIndex, DestLinker->ExportMap.Num());
+			return false;
+		}
+
+		const FObjectExport& SourceOuterExport = SourceLinker->Exp(SourceIndex);
+		const FObjectExport& DestOuterExport   = DestLinker  ->Exp(DestIndex);
+
+		FString SourceOuterExportKey = GetTableKey(SourceLinker, SourceOuterExport);
+		FString DestOuterExportKey   = GetTableKey(DestLinker,   DestOuterExport);
+
+		return SourceOuterExportKey == DestOuterExportKey;
+	}
+
+	if (SourceIndex.IsImport() && DestIndex.IsImport())
+	{
+		int32 SourceArrayIndex = SourceIndex.ToImport();
+		int32 DestArrayIndex   = DestIndex  .ToImport();
+
+		if (!SourceLinker->ImportMap.IsValidIndex(SourceArrayIndex) || !DestLinker->ImportMap.IsValidIndex(DestArrayIndex))
+		{
+			UE_LOG(LogDiffWriter, Warning, TEXT("Invalid import indices found, source: %d (of %d), dest: %d (of %d)"), SourceArrayIndex, SourceLinker->ImportMap.Num(), DestArrayIndex, DestLinker->ImportMap.Num());
+			return false;
+		}
+
+		const FObjectImport& SourceOuterImport = SourceLinker->Imp(SourceIndex);
+		const FObjectImport& DestOuterImport   = DestLinker  ->Imp(DestIndex);
+
+		FString SourceOuterImportKey = GetTableKey(SourceLinker, SourceOuterImport);
+		FString DestOuterImportKey   = GetTableKey(DestLinker,   DestOuterImport);
+
+		return SourceOuterImportKey == DestOuterImportKey;
+	}
+
+	return false;
+}
+
+static FString ConvertItemToText(const FObjectExport& Export, FLinkerLoad* Linker)
+{
+	FName ClassName = Export.ClassIndex.IsNull() ? FName(NAME_Class) : Linker->ImpExp(Export.ClassIndex).ObjectName;
+	return FString::Printf(TEXT("%s Super: %s, Template: %s, Flags: %d, Size: %lld, PackageFlags: %d, ForcedExport: %d, NotForClient: %d, NotForServer: %d, NotAlwaysLoadedForEditorGame: %d, IsAsset: %d, IsInheritedInstance: %d, GeneratePublicHash: %d"),
+		*GetTableKey(Linker, Export),
+		*GetTableKeyForIndex(Linker, Export.SuperIndex),
+		*GetTableKeyForIndex(Linker, Export.TemplateIndex),
+		(int32)Export.ObjectFlags,
+		Export.SerialSize,
+		Export.PackageFlags,
+		Export.bForcedExport,
+		Export.bNotForClient,
+		Export.bNotForServer,
+		Export.bNotAlwaysLoadedForEditorGame,
+		Export.bIsAsset,
+		Export.bIsInheritedInstance,
+		Export.bGeneratePublicHash);
+}
+
+static bool IsExportMapIdentical(FLinkerLoad* SourceLinker, FLinkerLoad* DestLinker)
+{
+	bool bIdentical = (SourceLinker->ExportMap.Num() == DestLinker->ExportMap.Num());
+	if (bIdentical)
+	{
+		for (int32 ExportIndex = 0; ExportIndex < SourceLinker->ExportMap.Num(); ++ExportIndex)
+		{
+			if (!CompareTableItem(SourceLinker, DestLinker, SourceLinker->ExportMap[ExportIndex], DestLinker->ExportMap[ExportIndex]))
+			{
+				bIdentical = false;
+				break;
+			}
+		}
+	}
+	return bIdentical;
+}
+
+/** Structure that holds an item from the NameMap/ImportMap/ExportMap in a TSet for diffing */
+template <typename T>
+struct TTableItem
+{
+	/** The key generated for this item */
+	FString Key;
+	/** Pointer to the original item */
+	const T* Item;
+	/** Index in the original *Map (table). Only for information purposes. */
+	int32 Index;
+
+	TTableItem(FString&& InKey, const T* InItem, int32 InIndex)
+		: Key(MoveTemp(InKey))
+		, Item(InItem)
+		, Index(InIndex)
+	{
+	}
+
+	FORCENOINLINE friend uint32 GetTypeHash(const TTableItem& TableItem)
+	{
+		return GetTypeHash(TableItem.Key);
+	}
+
+	FORCENOINLINE friend bool operator==(const TTableItem& Lhs, const TTableItem& Rhs)
+	{
+		return Lhs.Key == Rhs.Key;
+	}
+};
+
+/** Dumps differences between Linker tables */
+template <typename T>
+static void DumpTableDifferences(
+	FLinkerLoad* SourceLinker, 
+	FLinkerLoad* DestLinker, 
+	TArray<T>& SourceTable, 
+	TArray<T>& DestTable,
+	const TCHAR* AssetFilename,
+	const TCHAR* ItemName,
+	const int32 MaxDiffsToLog
+)
+{
+	const TCHAR* const LineTerminator = FDiffFormatHelper::Get().LineTerminator;
+	const TCHAR* const Indent = FDiffFormatHelper::Get().Indent;
+
+	FString HumanReadableString;
+	int32 LoggedDiffs = 0;
+	int32 NumDiffs = 0;
+
+	TSet<TTableItem<T>> SourceSet;
+	TSet<TTableItem<T>> DestSet;
+
+	SourceSet.Reserve(SourceTable.Num());
+	DestSet.Reserve(DestTable.Num());
+
+	for (int32 Index = 0; Index < SourceTable.Num(); ++Index)
+	{
+		const T& Item = SourceTable[Index];
+		SourceSet.Add(TTableItem<T>(GetTableKey(SourceLinker, Item), &Item, Index));
+	}
+	for (int32 Index = 0; Index < DestTable.Num(); ++Index)
+	{
+		const T& Item = DestTable[Index];
+		DestSet.Add(TTableItem<T>(GetTableKey(DestLinker, Item), &Item, Index));
+	}
+
+	// Determine the list of items removed from the source package and added to the dest package
+	TSet<TTableItem<T>> RemovedItems = SourceSet.Difference(DestSet);
+	TSet<TTableItem<T>> AddedItems   = DestSet.Difference(SourceSet);
+
+	// Add changed items as added-and-removed
+	for (const TTableItem<T>& ChangedSourceItem : SourceSet)
+	{
+		if (const TTableItem<T>* ChangedDestItem = DestSet.Find(ChangedSourceItem))
+		{
+			if (!CompareTableItem(SourceLinker, DestLinker, *ChangedSourceItem.Item, *ChangedDestItem->Item))
+			{
+				RemovedItems.Add(ChangedSourceItem);
+				AddedItems  .Add(*ChangedDestItem);
+			}
+		}
+	}
+
+	// Sort all additions and removals by index
+	RemovedItems.Sort([](const TTableItem<T>& Lhs, const TTableItem<T>& Rhs){ return Lhs.Index < Rhs.Index; });
+	AddedItems.  Sort([](const TTableItem<T>& Lhs, const TTableItem<T>& Rhs){ return Lhs.Index < Rhs.Index; });
+
+	// Dump all changes
+	for (const TTableItem<T>& RemovedItem : RemovedItems)
+	{
+		HumanReadableString += Indent;
+		HumanReadableString += FString::Printf(TEXT("-[%d] %s"), RemovedItem.Index, *ConvertItemToText(*RemovedItem.Item, SourceLinker));
+		HumanReadableString += LineTerminator;
+	}
+	for (const TTableItem<T>& AddedItem : AddedItems)
+	{
+		HumanReadableString += Indent;
+		HumanReadableString += FString::Printf(TEXT("+[%d] %s"), AddedItem.Index, *ConvertItemToText(*AddedItem.Item, DestLinker));
+		HumanReadableString += LineTerminator;
+	}
+
+	// For now just log everything out. When this becomes too spammy, respect the MaxDiffsToLog parameter
+	NumDiffs = RemovedItems.Num() + AddedItems.Num();
+	LoggedDiffs = NumDiffs;
+
+	if (NumDiffs > LoggedDiffs)
+	{
+		HumanReadableString += Indent;
+		HumanReadableString += FString::Printf(TEXT("+ %d differences not logged."), (NumDiffs - LoggedDiffs));
+		HumanReadableString += LineTerminator;
+	}
+
+	UE_LOG(
+		LogDiffWriter,
+		Warning,
+		TEXT("%s: %sMap is different (%d %ss in source package vs %d %ss in dest package):%s%s"),		
+		AssetFilename,
+		ItemName,
+		SourceTable.Num(),
+		ItemName,
+		DestTable.Num(),
+		ItemName,
+		LineTerminator,
+		*HumanReadableString);
+}
+
+static void DumpPackageHeaderDiffs_LinkerLoad(
+	const FDiffWriterArchiveWriter::FPackageData& SourcePackage,
+	const FDiffWriterArchiveWriter::FPackageData& DestPackage,
+	const FString& AssetFilename,
+	const int32 MaxDiffsToLog)
+{
+	FString AssetPathName = FPaths::Combine(*FPaths::GetPath(AssetFilename.Mid(AssetFilename.Find(TEXT(":"), ESearchCase::CaseSensitive) + 1)), *FPaths::GetBaseFilename(AssetFilename));
+	// The root directory could have a period in it (d:/Release5.0/EngineTest/Saved/Cooked),
+	// which is not a valid character for a LongPackageName. Remove it.
+	for (TCHAR c : FStringView(INVALID_LONGPACKAGE_CHARACTERS))
+	{
+		AssetPathName.ReplaceCharInline(c, TEXT('_'), ESearchCase::CaseSensitive);
+	}
+	FString SourceAssetPackageName = FPaths::Combine(TEXT("/Memory"), TEXT("/SourceForDiff"), *AssetPathName);
+	FString DestAssetPackageName = FPaths::Combine(TEXT("/Memory"), TEXT("/DestForDiff"), *AssetPathName);
+	check(FPackageName::IsValidLongPackageName(SourceAssetPackageName, true /* bIncludeReadOnlyRoots */));
+
+	TGuardValue<bool> GuardIsSavingPackage(GIsSavingPackage, false);
+	TGuardValue<int32> GuardAllowUnversionedContentInEditor(GAllowUnversionedContentInEditor, 1);
+	TGuardValue<int32> GuardAllowCookedDataInEditorBuilds(GAllowCookedDataInEditorBuilds, 1);
+
+	FLinkerLoad* SourceLinker = nullptr;
+	FLinkerLoad* DestLinker = nullptr;
+	// Create linkers. Note there's no need to clean them up here since they will be removed by the package associated with them
+	{
+		TRefCountPtr<FUObjectSerializeContext> LinkerLoadContext(FUObjectThreadContext::Get().GetSerializeContext());
+		BeginLoad(LinkerLoadContext);
+		SourceLinker = FDiffWriterArchiveWriter::CreateLinkerForPackage(LinkerLoadContext, SourceAssetPackageName, AssetFilename, SourcePackage);
+		EndLoad(SourceLinker ? SourceLinker->GetSerializeContext() : LinkerLoadContext.GetReference());
+	}
+	
+	{
+		TRefCountPtr<FUObjectSerializeContext> LinkerLoadContext(FUObjectThreadContext::Get().GetSerializeContext());
+		BeginLoad(LinkerLoadContext);
+		DestLinker = FDiffWriterArchiveWriter::CreateLinkerForPackage(LinkerLoadContext, DestAssetPackageName, AssetFilename, DestPackage);
+		EndLoad(DestLinker ? DestLinker->GetSerializeContext() : LinkerLoadContext.GetReference());
+	}
+
+	if (SourceLinker && DestLinker)
+	{
+		if (SourceLinker->NameMap != DestLinker->NameMap)
+		{
+			DumpTableDifferences<FNameEntryId>(SourceLinker, DestLinker, SourceLinker->NameMap, DestLinker->NameMap, *AssetFilename, TEXT("Name"), MaxDiffsToLog);
+		}
+
+		if (!IsImportMapIdentical(SourceLinker, DestLinker))
+		{
+			DumpTableDifferences<FObjectImport>(SourceLinker, DestLinker, SourceLinker->ImportMap, DestLinker->ImportMap, *AssetFilename, TEXT("Import"), MaxDiffsToLog);
+		}
+
+		if (!IsExportMapIdentical(SourceLinker, DestLinker))
+		{
+			DumpTableDifferences<FObjectExport>(SourceLinker, DestLinker, SourceLinker->ExportMap, DestLinker->ExportMap, *AssetFilename, TEXT("Export"), MaxDiffsToLog);
+		}
+	}
+
+	if (SourceLinker)
+	{
+		UE::ArchiveStackTrace::ForceKillPackageAndLinker(SourceLinker);
+	}
+	if (DestLinker)
+	{
+		UE::ArchiveStackTrace::ForceKillPackageAndLinker(DestLinker);
+	}
+}
+
+static void DumpPackageHeaderDiffs_ZenPackage(
+	const FDiffWriterArchiveWriter::FPackageData& SourcePackage,
+	const FDiffWriterArchiveWriter::FPackageData& DestPackage,
+	const FString& AssetFilename,
+	const int32 MaxDiffsToLog)
+{
+	// TODO: Fill in detailed diffing of Zen Package Summary
+}
+
+void FDiffWriterArchiveWriter::DumpPackageHeaderDiffs(
+	const FPackageData& SourcePackage,
+	const FPackageData& DestPackage,
+	const FString& AssetFilename,
+	const int32 MaxDiffsToLog,
+	const EPackageHeaderFormat PackageHeaderFormat /* = EPackageHeaderFormat::PackageFileSummary */)
+{
+	switch (PackageHeaderFormat)
+	{
+	case EPackageHeaderFormat::PackageFileSummary:
+		DumpPackageHeaderDiffs_LinkerLoad(SourcePackage, DestPackage, AssetFilename, MaxDiffsToLog);
+		break;
+	case EPackageHeaderFormat::ZenPackageSummary:
+		DumpPackageHeaderDiffs_ZenPackage(SourcePackage, DestPackage, AssetFilename, MaxDiffsToLog);
+		break;
+	default:
+		unimplemented();
+	}
+}
diff --git a/Engine/Source/Editor/UnrealEd/Private/Cooker/DiffWriterArchive.h b/Engine/Source/Editor/UnrealEd/Private/Cooker/DiffWriterArchive.h
new file mode 100644
--- /dev/null
+++ b/Engine/Source/Editor/UnrealEd/Private/Cooker/DiffWriterArchive.h
@@ -0,0 +1,355 @@
+// Copyright Epic Games, Inc. All Rights Reserved.
+
+#pragma once
+
+#include "Containers/Array.h"
+#include "Containers/ArrayView.h"
+#include "Containers/Map.h"
+#include "Containers/UnrealString.h"
+#include "HAL/Platform.h"
+#include "Serialization/Archive.h"
+#include "Serialization/ArchiveProxy.h"
+#include "Serialization/ArchiveStackTrace.h"
+#include "Serialization/LargeMemoryWriter.h"
+#include "Serialization/PackageWriter.h"
+#include "Templates/RefCounting.h"
+#include "Templates/UniquePtr.h"
+#include "UObject/NameTypes.h"
+
+class FLinkerLoad;
+class FProperty;
+class FUObjectThreadContext;
+class UObject;
+struct FUObjectSerializeContext;
+
+
+struct FDiffWriterDiffInfo
+{
+	int64 Offset;
+	int64 Size;
+	FDiffWriterDiffInfo()
+		: Offset(0)
+		, Size(0)
+	{
+	}
+	FDiffWriterDiffInfo(int64 InOffset, int64 InSize)
+		: Offset(InOffset)
+		, Size(InSize)
+	{
+	}
+	bool operator==(const FDiffWriterDiffInfo& InOther) const
+	{
+		return Offset == InOther.Offset;
+	}
+	bool operator<(const FDiffWriterDiffInfo& InOther) const
+	{
+		return Offset < InOther.Offset;
+	}
+	friend FArchive& operator << (FArchive& Ar, FDiffWriterDiffInfo& InDiffInfo)
+	{
+		Ar << InDiffInfo.Offset;
+		Ar << InDiffInfo.Size;
+		return Ar;
+	}
+};
+
+class FDiffWriterDiffMap : public TArray<FDiffWriterDiffInfo>
+{
+public:
+	bool ContainsOffset(int64 Offset) const
+	{
+		for (const FDiffWriterDiffInfo& Diff : *this)
+		{
+			if (Diff.Offset <= Offset && Offset < (Diff.Offset + Diff.Size))
+			{
+				return true;
+			}
+		}
+
+		return false;
+	}
+};
+
+/** Holds offsets to captured callstacks. */
+class FDiffWriterCallstacks
+{
+public:
+	/** Offset and callstack pair */
+	struct FCallstackAtOffset
+	{
+		/** Offset of a Serialize call */
+		int64 Offset = -1;
+		/** Callstack CRC for the Serialize call */
+		uint32 Callstack = 0;
+		/** Collected inside of skip scope */
+		bool bIgnore = false;
+	};
+
+	/** Struct to hold the actual Serialize call callstack and any associated data */
+	struct FCallstackData
+	{
+		/** Full callstack */
+		TUniquePtr<ANSICHAR[]> Callstack;
+		/** Full name of the currently serialized object */
+		FString SerializedObjectName;
+		/** The currently serialized property */
+		FProperty* SerializedProp = nullptr;
+		/** Name of the currently serialized property */
+		FString SerializedPropertyName;
+
+		FCallstackData() = default;
+		FCallstackData(TUniquePtr<ANSICHAR[]>&& InCallstack, UObject* InSerializedObject, FProperty* InSerializedProperty);
+		FCallstackData(FCallstackData&&) = default;
+		FCallstackData(const FCallstackData&) = delete;
+
+		FCallstackData& operator=(FCallstackData&&) = default;
+		FCallstackData& operator=(const FCallstackData&) = delete;
+
+		/** Converts the callstack and associated data to human readable string */
+		FString ToString(const TCHAR* CallstackCutoffText) const;
+
+		/** Clone the callstack data */
+		FCallstackData Clone() const;
+	};
+	
+	explicit FDiffWriterCallstacks(UObject* InAsset);
+
+	/** Returns the asset class name. */
+	FName GetAssetClass() const;
+
+	/** Returns the total number of callstacks. */
+	int32 Num() const
+	{
+		return CallstackAtOffsetMap.Num();
+	}
+
+	/** Capture and append the current callstack. */
+	void Add(
+		int64 Offset,
+		int64 Length,
+		UObject* SerializedObject,
+		FProperty* SerializedProperty,
+		TArrayView<const FName> DebugDataStack,
+		bool bIsCollectingCallstacks,
+		bool bCollectCurrentCallstack,
+		int32 StackIgnoreCount);
+
+	/** Append other callstacks. */
+	void Append(const FDiffWriterCallstacks& Other, int64 Offset = 0);
+
+	/** Finds a callstack associated with data at the specified offset */
+	int32 GetCallstackIndexAtOffset(int64 Offset, int32 MinOffsetIndex = 0) const;
+
+	/** Finds a callstack associated with data at the specified offset */
+	const FCallstackAtOffset& GetCallstack(int32 CallstackIndex) const
+	{
+		return CallstackAtOffsetMap[CallstackIndex];
+	}
+	
+	const FCallstackData& GetCallstackData(const FCallstackAtOffset& CallstackOffset) const
+	{
+		return UniqueCallstacks[CallstackOffset.Callstack];
+	}
+
+	/** Returns the size of serialized data at the specified offset. */
+	int64 GetSerializedDataSizeForOffsetIndex(int32 InOffsetIndex) const
+	{
+		if (InOffsetIndex < CallstackAtOffsetMap.Num() - 1)
+		{
+			return CallstackAtOffsetMap[InOffsetIndex + 1].Offset - CallstackAtOffsetMap[InOffsetIndex].Offset;
+		}
+		else
+		{
+			return TotalSize - CallstackAtOffsetMap[InOffsetIndex].Offset;
+		}
+	}
+
+	/** Returns total serialized bytes. */
+	int64 TotalCapturedSize() const
+	{
+		return TotalSize;
+	}
+
+private:
+	/** Adds a unique callstack to UniqueCallstacks map */
+	ANSICHAR* AddUniqueCallstack(bool bIsCollectingCallstacks, UObject* SerializedObject, FProperty* SerializedProperty, uint32& OutCallstackCRC);
+
+	/** The asset being serialized */
+	UObject* Asset;
+	/** List of offsets and their respective callstacks */
+	TArray<FCallstackAtOffset> CallstackAtOffsetMap;
+	/** Contains all unique callstacks for all Serialize calls */
+	TMap<uint32, FCallstackData> UniqueCallstacks;
+	/** Optimizes callstack comparison. If true the current and the last callstack should be compared as it may have changed */
+	bool bCallstacksDirty;
+	/** Maximum size of the stack trace */
+	const SIZE_T StackTraceSize;
+	/** Buffer for getting the current stack trace */
+	TUniquePtr<ANSICHAR[]> StackTrace;
+	/** Callstack associated with the previous Serialize call */
+	ANSICHAR* LastSerializeCallstack;
+	/** Total serialized bytes */
+	int64 TotalSize;
+};
+
+/** Archive proxy that captures callstacks for each serialize call. */
+class FDiffWriterArchiveWriter
+	: public FArchiveProxy
+{
+public:
+	FDiffWriterArchiveWriter(
+		FArchive& InInner,
+		FDiffWriterCallstacks& InCallstacks,
+		const FDiffWriterDiffMap* InDiffMap = nullptr,
+		int64 InDiffMapStartOffset = 0);
+
+	virtual ~FDiffWriterArchiveWriter() override;
+
+	FORCENOINLINE virtual void Serialize(void* Data, int64 Length) override; // FORCENOINLINE so it can be counted during StackTrace
+	virtual void SetSerializeContext(FUObjectSerializeContext* Context) override;
+	virtual FUObjectSerializeContext* GetSerializeContext() override;
+	
+	virtual void PushDebugDataString(const FName& DebugData) override
+	{
+		DebugDataStack.Push(DebugData);
+	}
+
+	virtual void PopDebugDataString() override
+	{
+		DebugDataStack.Pop();
+	}
+
+	const FDiffWriterDiffMap& GetDiffMap() const
+	{
+		static FDiffWriterDiffMap Empty;
+		return DiffMap != nullptr ? *DiffMap : Empty;
+	}
+
+	void SetDisableInnerArchive(bool bDisable)
+	{
+		bInnerArchiveDisabled = bDisable;
+	}
+
+	int32 GetStackIgnoreCount() const { return StackIgnoreCount; }
+	void SetStackIgnoreCount(const int32 IgnoreCount) { StackIgnoreCount = IgnoreCount; }
+
+	using FPackageData = UE::ArchiveStackTrace::FPackageData;
+
+	using EPackageHeaderFormat = ICookedPackageWriter::EPackageHeaderFormat;
+
+	/** Compares two packages and logs the differences and calltacks. */
+	static void Compare(
+		const FPackageData& SourcePackage,
+		const FPackageData& DestPackage,
+		const FDiffWriterCallstacks& Callstacks,
+		const FDiffWriterDiffMap& DiffMap,
+		const TCHAR* AssetFilename,
+		const TCHAR* CallstackCutoffText,
+		const int64 MaxDiffsToLog,
+		int32& InOutDiffsLogged,
+		TMap<FName, FArchiveDiffStats>& OutStats,
+		bool bSuppressLogging = false);
+
+	/** Creates map with mismatching callstacks. */
+	static bool GenerateDiffMap(
+		const FPackageData& SourcePackage,
+		const FPackageData& DestPackage,
+		const FDiffWriterCallstacks& Callstacks,
+		int32 MaxDiffsToFind,
+		FDiffWriterDiffMap& OutDiffMap);
+
+	/** Logs any mismatching header data. */
+	static void DumpPackageHeaderDiffs(
+		const FPackageData& SourcePackage,
+		const FPackageData& DestPackage,
+		const FString& AssetFilename,
+		const int32 MaxDiffsToLog,
+		const EPackageHeaderFormat PackageHeaderFormat = EPackageHeaderFormat::PackageFileSummary);
+
+	/** Returns a new linker for loading the specified package. */
+	static FLinkerLoad* CreateLinkerForPackage(
+		FUObjectSerializeContext* LoadContext,
+		const FString& InPackageName,
+		const FString& InFilename,
+		const FPackageData& PackageData);
+
+private:
+	FDiffWriterCallstacks& Callstacks;
+	const FDiffWriterDiffMap* DiffMap;
+	TRefCountPtr<FUObjectSerializeContext> SerializeContext;
+	TArray<FName> DebugDataStack;
+	int64 DiffMapOffset;
+	int32 StackIgnoreCount = 2;
+	bool bInnerArchiveDisabled;
+};
+
+/**
+ * Memory backed stack trace writer.
+ */
+class FDiffWriterArchiveMemoryWriter final
+	: public FLargeMemoryWriter
+{
+public:
+	FDiffWriterArchiveMemoryWriter(
+		FDiffWriterCallstacks& Callstacks,
+		const FDiffWriterDiffMap* DiffMap = nullptr,
+		const int64 DiffMapOffset = 0,
+		const int64 PreAllocateBytes = 0,
+		bool bIsPersistent = false,
+		const TCHAR* Filename = nullptr);
+
+	FORCENOINLINE virtual void Serialize(void* Memory, int64 Length) override; // FORCENOINLINE so it can be counted during StackTrace
+	virtual void SetSerializeContext(FUObjectSerializeContext* Context) override;
+	virtual FUObjectSerializeContext* GetSerializeContext() override;
+
+private:
+	FDiffWriterArchiveWriter StackTraceWriter;
+};
+
+/**
+ * Archive that stores a callstack for each of the Serialize calls and has the ability to compare itself to an existing
+ * package on disk and dump all the differences to log.
+ */
+class FDiffWriterArchive
+	: public FLargeMemoryWriter
+{
+public:
+	using FPackageData = FDiffWriterArchiveWriter::FPackageData;
+
+	FDiffWriterArchive(UObject* InAsset, const TCHAR* InFilename, bool bInCollectCallstacks = true, const FDiffWriterDiffMap* InDiffMap = nullptr);
+	virtual ~FDiffWriterArchive();
+
+	FDiffWriterCallstacks& GetCallstacks()
+	{
+		return Callstacks;
+	}
+
+	const FDiffWriterCallstacks& GetCallstacks() const
+	{
+		return Callstacks;
+	}
+
+	FORCENOINLINE virtual void Serialize(void* Memory, int64 Length) override; // FORCENOINLINE so it can be counted during StackTrace
+	virtual void SetSerializeContext(FUObjectSerializeContext* Context) override;
+	virtual FUObjectSerializeContext* GetSerializeContext() override;
+
+	/** Compares this archive with the given bytes from disk or FPackageData. Dumps all differences to log. */
+	void CompareWith(const TCHAR* InFilename, const int64 TotalHeaderSize, const TCHAR* CallstackCutoffText,
+		const int32 MaxDiffsToLog, TMap<FName, FArchiveDiffStats>& OutStats);
+	void CompareWith(const FPackageData& SourcePackage, const TCHAR* FileDisplayName, const int64 TotalHeaderSize,
+		const TCHAR* CallstackCutoffText, const int32 MaxDiffsToLog, TMap<FName, FArchiveDiffStats>& OutStats,
+		const FDiffWriterArchiveWriter::EPackageHeaderFormat PackageHeaderFormat = FDiffWriterArchiveWriter::EPackageHeaderFormat::PackageFileSummary);
+
+	/** Generates a map of all differences between this archive and the given bytes from disk or FPackageData. */
+	bool GenerateDiffMap(const TCHAR* InFilename, int64 TotalHeaderSize, int32 MaxDiffsToFind, FDiffWriterDiffMap& OutDiffMap);
+	bool GenerateDiffMap(const FPackageData& SourcePackage, int64 TotalHeaderSize, int32 MaxDiffsToFind,
+		FDiffWriterDiffMap& OutDiffMap);
+
+	/** Compares the provided buffer with the given bytes from disk or FPackageData. */
+	static bool IsIdentical(const TCHAR* InFilename, int64 BufferSize, const uint8* BufferData);
+	static bool IsIdentical(const FPackageData& SourcePackage, int64 BufferSize, const uint8* BufferData);
+
+private:
+	FDiffWriterCallstacks Callstacks;
+	FDiffWriterArchiveWriter StackTraceWriter;
+};
diff --git a/Engine/Source/Editor/UnrealEd/Private/Cooker/LooseCookedPackageWriter.cpp b/Engine/Source/Editor/UnrealEd/Private/Cooker/LooseCookedPackageWriter.cpp
--- a/Engine/Source/Editor/UnrealEd/Private/Cooker/LooseCookedPackageWriter.cpp
+++ b/Engine/Source/Editor/UnrealEd/Private/Cooker/LooseCookedPackageWriter.cpp
@@ -522,8 +522,8 @@ void FLooseCookedPackageWriter::UpdateManifest(FRecord& Record)
 
 bool FLooseCookedPackageWriter::GetPreviousCookedBytes(const FPackageInfo& Info, FPreviousCookedBytesData& OutData)
 {
-	FArchiveStackTrace::FPackageData ExistingPackageData;
-	FArchiveStackTrace::LoadPackageIntoMemory(*Info.LooseFilePath, ExistingPackageData, OutData.Data);
+	UE::ArchiveStackTrace::FPackageData ExistingPackageData;
+	UE::ArchiveStackTrace::LoadPackageIntoMemory(*Info.LooseFilePath, ExistingPackageData, OutData.Data);
 	OutData.Size = ExistingPackageData.Size;
 	OutData.HeaderSize = ExistingPackageData.HeaderSize;
 	OutData.StartOffset = ExistingPackageData.StartOffset;
diff --git a/Engine/Source/Editor/UnrealEd/Private/Tests/DiffWriterArchiveTests.cpp b/Engine/Source/Editor/UnrealEd/Private/Tests/DiffWriterArchiveTests.cpp
new file mode 100644
--- /dev/null
+++ b/Engine/Source/Editor/UnrealEd/Private/Tests/DiffWriterArchiveTests.cpp
@@ -0,0 +1,189 @@
+// Copyright Epic Games, Inc. All Rights Reserved.
+
+#include "Misc/AutomationTest.h"
+
+#if WITH_AUTOMATION_TESTS
+
+#include "Cooker/DiffWriterArchive.h"
+
+#include "Memory/SharedBuffer.h"
+#include "Serialization/LargeMemoryWriter.h"
+#include "Serialization/MemoryReader.h"
+
+namespace UE::DiffWriterArchive::Tests
+{
+
+struct FBasicDiffState
+{
+	int32 FirstParam = 0;
+	int32 SecondParam = 0;
+	int32 ThirdParam = 0;
+
+	FORCENOINLINE void Serialize(FArchive& Ar)
+	{
+		Ar << FirstParam;
+		Ar << SecondParam;
+		Ar << ThirdParam;
+	}
+};
+
+FDiffWriterArchiveWriter::FPackageData ToPackageData(FLargeMemoryWriter& Ar)
+{
+	return FDiffWriterArchiveWriter::FPackageData
+	{
+		Ar.GetData(),
+		Ar.TotalSize()
+	};
+}
+
+} // namespace UE::DiffWriterArchive::Tests
+
+IMPLEMENT_SIMPLE_AUTOMATION_TEST(FDiffWriterArchiveTestsCallstacks, "System.Core.Cooker.DiffWriterArchive.Callstacks",
+	EAutomationTestFlags::EditorContext | EAutomationTestFlags::EngineFilter)
+
+bool FDiffWriterArchiveTestsCallstacks::RunTest(const FString& Parameters)
+{
+	using namespace UE::DiffWriterArchive::Tests;
+
+	const int32 MaxDiffsToLog = 1000;
+	FDiffWriterCallstacks Callstacks(nullptr);
+
+	FLargeMemoryWriter Memory;
+	FBasicDiffState State = FBasicDiffState {1,2,3};
+	FDiffWriterArchiveWriter StackTraceWriter(Memory, Callstacks, 0);
+	State.Serialize(StackTraceWriter);
+
+	TestTrueExpr(Callstacks.Num() == 3);
+	TestTrueExpr(Callstacks.TotalCapturedSize() == sizeof(FBasicDiffState));
+	TestTrueExpr(Callstacks.GetCallstack(0).Offset == 0);
+	TestTrueExpr(Callstacks.GetCallstack(0).Callstack == 0);
+	TestTrueExpr(Callstacks.GetCallstack(0).bIgnore == false);
+	TestTrueExpr(Callstacks.GetCallstack(1).Offset == 4);
+	TestTrueExpr(Callstacks.GetCallstack(1).Callstack == 0);
+	TestTrueExpr(Callstacks.GetCallstack(1).bIgnore == false);
+	TestTrueExpr(Callstacks.GetCallstack(2).Offset == 8);
+	TestTrueExpr(Callstacks.GetCallstack(2).Callstack == 0);
+	TestTrueExpr(Callstacks.GetCallstack(2).bIgnore == false);
+	return true;
+}
+
+IMPLEMENT_SIMPLE_AUTOMATION_TEST(FDiffWriterArchiveTestsBasic, "System.Core.Cooker.DiffWriterArchive.Basic",
+	EAutomationTestFlags::EditorContext | EAutomationTestFlags::EngineFilter)
+
+bool FDiffWriterArchiveTestsBasic::RunTest(const FString& Parameters)
+{
+	using namespace UE::DiffWriterArchive::Tests;
+
+	const int32 MaxDiffsToLog = 1000;
+
+	// SECTION("DiffMap - identical")
+	{
+		FDiffWriterCallstacks Callstacks(nullptr);
+		FDiffWriterDiffMap DiffMap;
+
+		FLargeMemoryWriter InitialState;
+		{
+			FBasicDiffState State = FBasicDiffState {1,2,3};
+			State.Serialize(InitialState);
+		}
+
+		FLargeMemoryWriter NewState;
+		{
+			FBasicDiffState State = FBasicDiffState {1,2,3};
+			FDiffWriterArchiveWriter StackTraceWriter(NewState, Callstacks);
+			State.Serialize(StackTraceWriter);
+		}
+
+		FDiffWriterArchiveWriter::FPackageData InitialData = ToPackageData(InitialState);
+		FDiffWriterArchiveWriter::FPackageData NewData = ToPackageData(NewState);
+		const bool bIsIdentical = FDiffWriterArchiveWriter::GenerateDiffMap(InitialData, NewData, Callstacks, MaxDiffsToLog, DiffMap);
+		TestTrueExpr(bIsIdentical);
+	}
+
+	// SECTION("DiffMap - mismatch")
+	{
+		FDiffWriterCallstacks Callstacks(nullptr);
+		FDiffWriterDiffMap DiffMap;
+
+		FLargeMemoryWriter InitialState;
+		{
+			FBasicDiffState State = FBasicDiffState {1,2,3};
+			State.Serialize(InitialState);
+		}
+
+		FLargeMemoryWriter NewState;
+		{
+			FBasicDiffState State = FBasicDiffState {1,2,4};
+			FDiffWriterArchiveWriter StackTraceWriter(NewState, Callstacks);
+			State.Serialize(StackTraceWriter);
+		}
+
+		FDiffWriterArchiveWriter::FPackageData InitialData = ToPackageData(InitialState);
+		FDiffWriterArchiveWriter::FPackageData NewData = ToPackageData(NewState);
+		const bool bIsIdentical = FDiffWriterArchiveWriter::GenerateDiffMap(InitialData, NewData, Callstacks, MaxDiffsToLog, DiffMap);
+		TestTrueExpr(bIsIdentical == false);
+	}
+
+	// SECTION("Compare - mismatch")
+	{
+		FDiffWriterCallstacks Callstacks(nullptr);
+		FDiffWriterCallstacks CallstacksWithStackTrace(nullptr);
+		FDiffWriterDiffMap DiffMap;
+
+		FBasicDiffState InitialState = FBasicDiffState {1,2,3};
+		FLargeMemoryWriter InitialMemory;
+		InitialMemory.SetByteSwapping(false);
+
+		InitialState.Serialize(InitialMemory);
+
+		FBasicDiffState NewState = FBasicDiffState {1,22,33};
+		TUniquePtr<FLargeMemoryWriter> NewMemory = MakeUnique<FLargeMemoryWriter>();
+		NewMemory->SetByteSwapping(false);
+
+		{
+			FDiffWriterArchiveWriter StackTraceWriter(*NewMemory, Callstacks);
+			NewState.Serialize(StackTraceWriter);
+		}
+
+		// Generate diff map
+		{
+			FDiffWriterArchiveWriter::FPackageData InitialData = ToPackageData(InitialMemory);
+			FDiffWriterArchiveWriter::FPackageData NewData = ToPackageData(*NewMemory);
+			const bool bIsIdentical = FDiffWriterArchiveWriter::GenerateDiffMap(InitialData, NewData, Callstacks, MaxDiffsToLog, DiffMap);
+			TestTrueExpr(bIsIdentical == false);
+		}
+		
+		NewMemory = MakeUnique<FLargeMemoryWriter>();
+		NewMemory->SetByteSwapping(false);
+		{
+			FDiffWriterArchiveWriter StackTraceWriter(*NewMemory, CallstacksWithStackTrace, &DiffMap);
+			NewState.Serialize(StackTraceWriter);
+		}
+
+		// Compare using callstacks and diffmap
+		{
+			FDiffWriterArchiveWriter::FPackageData InitialData = ToPackageData(InitialMemory);
+			FDiffWriterArchiveWriter::FPackageData NewData = ToPackageData(*NewMemory);
+			int32 DiffsLogged = 0;
+			TMap<FName, FArchiveDiffStats> DiffStats;
+			const bool bSuppressLogging = true; // Suppress warning logs when running tests on CI
+
+			FDiffWriterArchiveWriter::Compare(
+				InitialData,
+				NewData,
+				CallstacksWithStackTrace,
+				DiffMap,
+				TEXT("BasicDiffTest"),
+				nullptr,
+				MaxDiffsToLog,
+				DiffsLogged,
+				DiffStats,
+				bSuppressLogging);
+
+			TestTrueExpr(DiffStats[NAME_None].NumDiffs == 2);
+		}
+	}
+	return true;
+}
+
+#endif // WITH_AUTOMATION_TESTS
diff --git a/Engine/Source/Runtime/CoreUObject/Private/Serialization/ArchiveStackTrace.cpp b/Engine/Source/Runtime/CoreUObject/Private/Serialization/ArchiveStackTrace.cpp
--- a/Engine/Source/Runtime/CoreUObject/Private/Serialization/ArchiveStackTrace.cpp
+++ b/Engine/Source/Runtime/CoreUObject/Private/Serialization/ArchiveStackTrace.cpp
@@ -112,6 +112,82 @@
 	GIgnoreDiffManager.PopDisabled();
 }
 
+namespace UE::ArchiveStackTrace
+{
+
+bool LoadPackageIntoMemory(const TCHAR* InFilename, UE::ArchiveStackTrace::FPackageData& OutPackageData,
+	TUniquePtr<uint8>& OutLoadedBytes)
+{
+	TUniquePtr<FArchive> UAssetFileArchive(IFileManager::Get().CreateFileReader(InFilename));
+	if (!UAssetFileArchive || UAssetFileArchive->TotalSize() == 0)
+	{
+		// The package doesn't exist on disk
+		OutLoadedBytes.Reset();
+		OutPackageData.Data = nullptr;
+		OutPackageData.Size = 0;
+		OutPackageData.HeaderSize = 0;
+		OutPackageData.StartOffset = 0;
+		return false;
+	}
+	else
+	{
+		// Handle EDL packages (uexp files)
+		TUniquePtr<FArchive> ExpFileArchive = nullptr;
+		OutPackageData.Size = UAssetFileArchive->TotalSize();
+		{
+			FString UExpFilename = FPaths::ChangeExtension(InFilename, TEXT("uexp"));
+			ExpFileArchive.Reset(IFileManager::Get().CreateFileReader(*UExpFilename));
+			if (ExpFileArchive)
+			{
+				// The header size is the current package size
+				OutPackageData.HeaderSize = OutPackageData.Size;
+				// Grow the buffer size to append the uexp file contents
+				OutPackageData.Size += ExpFileArchive->TotalSize();
+			}
+		}
+		OutLoadedBytes.Reset(new uint8[OutPackageData.Size]);
+		OutPackageData.Data = OutLoadedBytes.Get();
+		UAssetFileArchive->Serialize(OutPackageData.Data, UAssetFileArchive->TotalSize());
+
+		if (ExpFileArchive)
+		{
+			// If uexp file is present, append its contents at the end of the buffer
+			ExpFileArchive->Serialize(OutPackageData.Data + OutPackageData.HeaderSize, ExpFileArchive->TotalSize());
+		}
+	}
+
+	return true;
+}
+
+void ForceKillPackageAndLinker(FLinkerLoad* Linker)
+{
+	UPackage* Package = Linker->LinkerRoot;
+	Linker->Detach();
+	FLinkerManager::Get().RemoveLinker(Linker);
+	if (Package)
+	{
+		Package->ClearPackageFlags(PKG_ContainsMapData | PKG_ContainsMap);
+		Package->MarkAsGarbage();
+	}
+}
+
+} // namespace UE::ArchiveStackTrace
+
+PRAGMA_DISABLE_DEPRECATION_WARNINGS;
+
+bool FArchiveDiffMap::ContainsOffset(int64 Offset) const
+{
+	for (const FArchiveDiffInfo& Diff : *this)
+	{
+		if (Diff.Offset <= Offset && Offset < (Diff.Offset + Diff.Size))
+		{
+			return true;
+		}
+	}
+
+	return false;
+}
+
 FArchiveCallstacks::FCallstackData::FCallstackData(TUniquePtr<ANSICHAR[]>&& InCallstack, UObject* InSerializedObject, FProperty* InSerializedProperty)
 	: Callstack(MoveTemp(InCallstack))
 	, SerializedProp(InSerializedProperty)
@@ -585,45 +661,13 @@
 
 bool FArchiveStackTrace::LoadPackageIntoMemory(const TCHAR* InFilename, FPackageData& OutPackageData, TUniquePtr<uint8>& OutLoadedBytes)
 {
-	TUniquePtr<FArchive> UAssetFileArchive(IFileManager::Get().CreateFileReader(InFilename));
-	if (!UAssetFileArchive || UAssetFileArchive->TotalSize() == 0)
-	{
-		// The package doesn't exist on disk
-		OutLoadedBytes.Reset();
-		OutPackageData.Data = nullptr;
-		OutPackageData.Size = 0;
-		OutPackageData.HeaderSize = 0;
-		OutPackageData.StartOffset = 0;
-		return false;
-	}
-	else
-	{
-		// Handle EDL packages (uexp files)
-		TUniquePtr<FArchive> ExpFileArchive = nullptr;
-		OutPackageData.Size = UAssetFileArchive->TotalSize();
-		{
-			FString UExpFilename = FPaths::ChangeExtension(InFilename, TEXT("uexp"));
-			ExpFileArchive.Reset(IFileManager::Get().CreateFileReader(*UExpFilename));
-			if (ExpFileArchive)
-			{				
-				// The header size is the current package size
-				OutPackageData.HeaderSize = OutPackageData.Size;
-				// Grow the buffer size to append the uexp file contents
-				OutPackageData.Size += ExpFileArchive->TotalSize();
-			}
-		}
-		OutLoadedBytes.Reset(new uint8[OutPackageData.Size]);
-		OutPackageData.Data = OutLoadedBytes.Get();
-		UAssetFileArchive->Serialize(OutPackageData.Data, UAssetFileArchive->TotalSize());
-
-		if (ExpFileArchive)
-		{
-			// If uexp file is present, append its contents at the end of the buffer
-			ExpFileArchive->Serialize(OutPackageData.Data + OutPackageData.HeaderSize, ExpFileArchive->TotalSize());
-		}
-	}
-
-	return true;
+	UE::ArchiveStackTrace::FPackageData PackageData;
+	bool bResult = UE::ArchiveStackTrace::LoadPackageIntoMemory(InFilename, PackageData, OutLoadedBytes);
+	OutPackageData.Data = PackageData.Data;
+	OutPackageData.Size = PackageData.Size;
+	OutPackageData.HeaderSize = PackageData.HeaderSize;
+	OutPackageData.StartOffset = PackageData.StartOffset;
+	return bResult;
 }
 
 namespace
@@ -1240,29 +1284,29 @@
 	}
 }
 
-bool ComparePackageIndices(FLinkerLoad* SourceLinker, FLinkerLoad* DestLinker, const FPackageIndex& SourceIndex, const FPackageIndex& DestIndex);
+static bool ComparePackageIndices(FLinkerLoad* SourceLinker, FLinkerLoad* DestLinker, const FPackageIndex& SourceIndex, const FPackageIndex& DestIndex);
 
-bool CompareTableItem(FLinkerLoad* SourceLinker, FLinkerLoad* DestLinker, const FName& SourceName, const FName& DestName)
+static bool CompareTableItem(FLinkerLoad* SourceLinker, FLinkerLoad* DestLinker, const FName& SourceName, const FName& DestName)
 {
 	return SourceName == DestName;
 }
 
-bool CompareTableItem(FLinkerLoad* SourceLinker, FLinkerLoad* DestLinker, FNameEntryId SourceName, FNameEntryId DestName)
+static bool CompareTableItem(FLinkerLoad* SourceLinker, FLinkerLoad* DestLinker, FNameEntryId SourceName, FNameEntryId DestName)
 {
 	return SourceName == DestName;
 }
 
-FString ConvertItemToText(const FName& Name, FLinkerLoad* Linker)
+static FString ConvertItemToText(const FName& Name, FLinkerLoad* Linker)
 {
 	return Name.ToString();
 }
 
-FString ConvertItemToText(FNameEntryId Id, FLinkerLoad* Linker)
+static FString ConvertItemToText(FNameEntryId Id, FLinkerLoad* Linker)
 {
 	return FName::GetEntry(Id)->GetPlainNameString();
 }
 
-bool CompareTableItem(FLinkerLoad* SourceLinker, FLinkerLoad* DestLinker, const FObjectImport& SourceImport, const FObjectImport& DestImport)
+static bool CompareTableItem(FLinkerLoad* SourceLinker, FLinkerLoad* DestLinker, const FObjectImport& SourceImport, const FObjectImport& DestImport)
 {
 	if (SourceImport.ObjectName != DestImport.ObjectName ||
 		SourceImport.ClassName != DestImport.ClassName ||
@@ -1277,7 +1321,7 @@
 	}
 }
 
-FString ConvertItemToText(const FObjectImport& Import, FLinkerLoad* Linker)
+static FString ConvertItemToText(const FObjectImport& Import, FLinkerLoad* Linker)
 {
 	return FString::Printf(
 		TEXT("%s ClassPackage: %s"),
@@ -1286,7 +1330,7 @@
 	);
 }
 
-bool CompareTableItem(FLinkerLoad* SourceLinker, FLinkerLoad* DestLinker, const FObjectExport& SourceExport, const FObjectExport& DestExport)
+static bool CompareTableItem(FLinkerLoad* SourceLinker, FLinkerLoad* DestLinker, const FObjectExport& SourceExport, const FObjectExport& DestExport)
 {
 	if (SourceExport.ObjectName != DestExport.ObjectName ||
 		SourceExport.PackageFlags != DestExport.PackageFlags ||
@@ -1329,7 +1373,7 @@
 	return bIdentical;
 }
 
-bool ComparePackageIndices(FLinkerLoad* SourceLinker, FLinkerLoad* DestLinker, const FPackageIndex& SourceIndex, const FPackageIndex& DestIndex)
+static bool ComparePackageIndices(FLinkerLoad* SourceLinker, FLinkerLoad* DestLinker, const FPackageIndex& SourceIndex, const FPackageIndex& DestIndex)
 {
 	if (SourceIndex.IsNull() && DestIndex.IsNull())
 	{
@@ -1379,7 +1423,7 @@
 	return false;
 }
 
-FString ConvertItemToText(const FObjectExport& Export, FLinkerLoad* Linker)
+static FString ConvertItemToText(const FObjectExport& Export, FLinkerLoad* Linker)
 {
 	FName ClassName = Export.ClassIndex.IsNull() ? FName(NAME_Class) : Linker->ImpExp(Export.ClassIndex).ObjectName;
 	return FString::Printf(TEXT("%s Super: %s, Template: %s, Flags: %d, Size: %lld, PackageFlags: %d, ForcedExport: %d, NotForClient: %d, NotForServer: %d, NotAlwaysLoadedForEditorGame: %d, IsAsset: %d, IsInheritedInstance: %d, GeneratePublicHash: %d"),
@@ -1415,18 +1459,6 @@
 	return bIdentical;
 }
 
-static void ForceKillPackageAndLinker(FLinkerLoad* Linker)
-{
-	UPackage* Package = Linker->LinkerRoot;
-	Linker->Detach();
-	FLinkerManager::Get().RemoveLinker(Linker);
-	if (Package)
-	{
-		Package->ClearPackageFlags(PKG_ContainsMapData | PKG_ContainsMap);
-		Package->MarkAsGarbage();
-	}
-}
-
 /** Structure that holds an item from the NameMap/ImportMap/ExportMap in a TSet for diffing */
 template <typename T>
 struct TTableItem
@@ -1613,11 +1645,11 @@
 
 	if (SourceLinker)
 	{
-		ForceKillPackageAndLinker(SourceLinker);
+		UE::ArchiveStackTrace::ForceKillPackageAndLinker(SourceLinker);
 	}
 	if (DestLinker)
 	{
-		ForceKillPackageAndLinker(DestLinker);
+		UE::ArchiveStackTrace::ForceKillPackageAndLinker(DestLinker);
 	}
 #endif // !NO_LOGGING
 }
@@ -1655,6 +1687,8 @@
 #endif // !NO_LOGGING
 }
 
+PRAGMA_ENABLE_DEPRECATION_WARNINGS;
+
 FArchiveStackTraceReader::FSerializeData::FSerializeData(int64 InOffset, int64 InSize, UObject* InObject, FProperty* InProperty)
 : Offset(InOffset)
 , Size(InSize)
@@ -1699,8 +1733,8 @@
 {
 	FArchiveStackTraceReader* Reader = nullptr;
 	TUniquePtr<uint8> PackageBytes;
-	FArchiveStackTrace::FPackageData PackageData;
-	if (FArchiveStackTrace::LoadPackageIntoMemory(InFilename, PackageData, PackageBytes))
+	UE::ArchiveStackTrace::FPackageData PackageData;
+	if (UE::ArchiveStackTrace::LoadPackageIntoMemory(InFilename, PackageData, PackageBytes))
 	{
 		Reader = new FArchiveStackTraceReader(InFilename, PackageBytes.Release(), PackageData.Size);
 	}
diff --git a/Engine/Source/Runtime/CoreUObject/Private/UObject/SavePackage.cpp b/Engine/Source/Runtime/CoreUObject/Private/UObject/SavePackage.cpp
--- a/Engine/Source/Runtime/CoreUObject/Private/UObject/SavePackage.cpp
+++ b/Engine/Source/Runtime/CoreUObject/Private/UObject/SavePackage.cpp
@@ -12,25 +12,6 @@ COREUOBJECT_API extern bool GOutputCookingWarnings;
 
 #endif
 
-FSavePackageResultStruct UPackage::Save(UPackage* InOuter, UObject* Base, EObjectFlags TopLevelFlags,
-	const TCHAR* Filename, FOutputDevice* Error, FLinkerNull* ConformNO, bool bForceByteSwapping,
-	bool bWarnOfLongFilename, uint32 SaveFlags, const ITargetPlatform* TargetPlatform,
-	const FDateTime& FinalTimeStamp, bool bSlowTask, FArchiveDiffMap* InOutDiffMap,
-	FSavePackageContext* SavePackageContext)
-{
-	// CookData should only be nonzero if we are cooking.
-	TOptional<FArchiveCookData> CookData;
-	FArchiveCookContext CookContext(InOuter, FArchiveCookContext::ECookTypeUnknown, FArchiveCookContext::ECookingDLCUnknown);
-	if (TargetPlatform != nullptr)
-	{		
-		CookData.Emplace(*TargetPlatform, CookContext);
-	}
-	
-	FSavePackageArgs SaveArgs = { nullptr /* deprecated target platform */, CookData.GetPtrOrNull(), TopLevelFlags, SaveFlags, bForceByteSwapping,
-		bWarnOfLongFilename, bSlowTask, FinalTimeStamp, Error, SavePackageContext };
-	return UPackage::Save(InOuter, Base, Filename, SaveArgs);
-}
-
 FSavePackageResultStruct UPackage::Save(UPackage* InOuter, UObject* InAsset, const TCHAR* Filename,
 	const FSavePackageArgs& SaveArgs)
 {
diff --git a/Engine/Source/Runtime/CoreUObject/Public/Serialization/ArchiveStackTrace.h b/Engine/Source/Runtime/CoreUObject/Public/Serialization/ArchiveStackTrace.h
--- a/Engine/Source/Runtime/CoreUObject/Public/Serialization/ArchiveStackTrace.h
+++ b/Engine/Source/Runtime/CoreUObject/Public/Serialization/ArchiveStackTrace.h
@@ -21,8 +21,6 @@ class FLinkerLoad;
 class FProperty;
 class FUObjectThreadContext;
 class UObject;
-struct FObjectExport;
-struct FObjectImport;
 struct FUObjectSerializeContext;
 
 /** Structure that holds stats from comparing two packages */
@@ -65,7 +63,26 @@ public:
 	COREUOBJECT_API ~FArchiveStackTraceDisabledScope();
 };
 
-struct FArchiveDiffInfo
+namespace UE::ArchiveStackTrace
+{
+
+struct FPackageData
+{
+	uint8* Data = nullptr;
+	int64 Size = 0;
+	int64 HeaderSize = 0;
+	int64 StartOffset = 0;
+};
+/** Helper function to load package contents into memory. Supports EDL packages. */
+COREUOBJECT_API bool LoadPackageIntoMemory(const TCHAR* InFilename,
+	FPackageData& OutPackageData, TUniquePtr<uint8>& OutLoadedBytes);
+COREUOBJECT_API void ForceKillPackageAndLinker(FLinkerLoad* Linker);
+
+} // namespace UE::ArchiveStackTrace
+
+struct
+// Deprecated: 5.3, "FArchiveStackTrace was only used by DiffPackageWriter, and has been moved into a private helper class. Contact Epic if you need this class for another reason.")
+FArchiveDiffInfo
 {
 	int64 Offset;
 	int64 Size;
@@ -95,25 +112,17 @@ struct FArchiveDiffInfo
 	}
 };
 
-class FArchiveDiffMap : public TArray<FArchiveDiffInfo>
+class 
+UE_DEPRECATED(5.3, "FArchiveStackTrace was only used by DiffPackageWriter, and has been moved into a private helper class. Contact Epic if you need this class for another reason.")
+FArchiveDiffMap : public TArray<FArchiveDiffInfo>
 {
 public:
-	bool ContainsOffset(int64 Offset) const
-	{
-		for (const FArchiveDiffInfo& Diff : *this)
-		{
-			if (Diff.Offset <= Offset && Offset < (Diff.Offset + Diff.Size))
-			{
-				return true;
-			}
-		}
-
-		return false;
-	}
+	bool ContainsOffset(int64 Offset) const;
 };
 
-/** Holds offsets to captured callstacks. */
-class FArchiveCallstacks
+class
+UE_DEPRECATED(5.3, "FArchiveStackTrace was only used by DiffPackageWriter, and has been moved into a private helper class. Contact Epic if you need this class for another reason.")
+FArchiveCallstacks
 {
 public:
 	/** Offset and callstack pair */
@@ -234,11 +243,14 @@ private:
 	int64 TotalSize;
 };
 
-/** Archive proxy that captures callstacks for each serialize call. */
-class FArchiveStackTraceWriter
+class
+UE_DEPRECATED(5.3, "FArchiveStackTrace was only used by DiffPackageWriter, and has been moved into a private helper class. Contact Epic if you need this class for another reason.")
+FArchiveStackTraceWriter
 	: public FArchiveProxy
 {
 public:
+	PRAGMA_DISABLE_DEPRECATION_WARNINGS;
+
 	COREUOBJECT_API FArchiveStackTraceWriter(
 		FArchive& InInner,
 		FArchiveCallstacks& InCallstacks,
@@ -333,15 +345,20 @@ private:
 	int64 DiffMapOffset;
 	int32 StackIgnoreCount = 2;
 	bool bInnerArchiveDisabled;
+
+PRAGMA_ENABLE_DEPRECATION_WARNINGS;
+
 };
 
-/**
- * Memory backed stack trace writer.
- */
-class FArchiveStackTraceMemoryWriter final
+class
+UE_DEPRECATED(5.3, "FArchiveStackTrace was only used by DiffPackageWriter, and has been moved into a private helper class. Contact Epic if you need this class for another reason.")
+FArchiveStackTraceMemoryWriter final
 	: public FLargeMemoryWriter
 {
 public:
+
+PRAGMA_DISABLE_DEPRECATION_WARNINGS;
+
 	COREUOBJECT_API FArchiveStackTraceMemoryWriter(
 		FArchiveCallstacks& Callstacks,
 		const FArchiveDiffMap* DiffMap = nullptr,
@@ -356,16 +373,22 @@ public:
 
 private:
 	FArchiveStackTraceWriter StackTraceWriter;
+
+PRAGMA_ENABLE_DEPRECATION_WARNINGS;
 };
 
 /**
  * Archive that stores a callstack for each of the Serialize calls and has the ability to compare itself to an existing
  * package on disk and dump all the differences to log.
  */
-class FArchiveStackTrace
+class
+UE_DEPRECATED(5.3, "FArchiveStackTrace was only used by DiffPackageWriter, and has been moved into a private helper class. Contact Epic if you need this class for another reason.")
+FArchiveStackTrace
 	: public FLargeMemoryWriter
 {
 public:
+	PRAGMA_DISABLE_DEPRECATION_WARNINGS;
+
 	using FPackageData = FArchiveStackTraceWriter::FPackageData;
 
 	COREUOBJECT_API FArchiveStackTrace(UObject* InAsset, const TCHAR* InFilename, bool bInCollectCallstacks = true, const FArchiveDiffMap* InDiffMap = nullptr);
@@ -408,6 +431,8 @@ public:
 private:
 	FArchiveCallstacks Callstacks;
 	FArchiveStackTraceWriter StackTraceWriter;
+
+PRAGMA_ENABLE_DEPRECATION_WARNINGS;
 };
 
 class FArchiveStackTraceReader : public FLargeMemoryReader
diff --git a/Engine/Source/Runtime/CoreUObject/Public/UObject/Package.h b/Engine/Source/Runtime/CoreUObject/Public/UObject/Package.h
--- a/Engine/Source/Runtime/CoreUObject/Public/UObject/Package.h
+++ b/Engine/Source/Runtime/CoreUObject/Public/UObject/Package.h
@@ -921,14 +921,6 @@ public:
 	*/
 	static FSavePackageResultStruct Save(UPackage* InOuter, UObject* InAsset, const TCHAR* Filename, const FSavePackageArgs& SaveArgs);
 
-	UE_DEPRECATED(5.0, "Pack the arguments into FSavePackageArgs and call the function overload that takes FSavePackageArgs. Note that Conform and InOutDiffMap are no longer implemented.")
-	static FSavePackageResultStruct Save(UPackage* InOuter, UObject* Base, EObjectFlags TopLevelFlags,
-		const TCHAR* Filename, FOutputDevice* Error = GError, FLinkerNull* Conform = nullptr,
-		bool bForceByteSwapping = false, bool bWarnOfLongFilename = true, uint32 SaveFlags = SAVE_None,
-		const ITargetPlatform* TargetPlatform = nullptr, const FDateTime& FinalTimeStamp = FDateTime::MinValue(),
-		bool bSlowTask = true, class FArchiveDiffMap* InOutDiffMap = nullptr,
-		FSavePackageContext* SavePackageContext = nullptr);
-
 	/**
 	 * Save a list of packages concurrently using Save2 mechanism
 	 * SaveConcurrent is currently experimental and shouldn't be used until it can safely replace Save.
diff --git a/Engine/Source/Runtime/CoreUObject/Public/UObject/SavePackage.h b/Engine/Source/Runtime/CoreUObject/Public/UObject/SavePackage.h
--- a/Engine/Source/Runtime/CoreUObject/Public/UObject/SavePackage.h
+++ b/Engine/Source/Runtime/CoreUObject/Public/UObject/SavePackage.h
@@ -31,7 +31,6 @@ class UObject;
 #endif
 
 class FArchive;
-class FArchiveDiffMap;
 class FCbFieldView;
 class FCbWriter;
 class FIoBuffer;
@@ -89,8 +88,6 @@ struct FSavePackageArgs
 	FOutputDevice* Error = GError;
 	/** Structure to hold longer-lifetime parameters that apply to multiple saves */
 	FSavePackageContext* SavePackageContext = nullptr;
-	UE_DEPRECATED(5.0, "FArchiveDiffMap is no longer used; it is now implemented by DiffPackageWriter.")
-	FArchiveDiffMap* DiffMap = nullptr;
 	UE_DEPRECATED(4.27, "UPackage::Guid has not been used by the engine for a long time and it will be removed.")
 	TOptional<FGuid> OutputPackageGuid;
 
diff --git a/Engine/Source/Runtime/CoreUObject/Tests/Serialization/ArchiveStackTraceTests.cpp b/Engine/Source/Runtime/CoreUObject/Tests/Serialization/ArchiveStackTraceTests.cpp
deleted file mode 100644
