Make the Game pseudo plugin depend on the Engine pseudo plugin.

diff --git a/Engine/Source/Editor/UnrealEd/Private/CookOnTheFlyServer.cpp b/Engine/Source/Editor/UnrealEd/Private/CookOnTheFlyServer.cpp
--- a/Engine/Source/Editor/UnrealEd/Private/CookOnTheFlyServer.cpp
+++ b/Engine/Source/Editor/UnrealEd/Private/CookOnTheFlyServer.cpp
@@ -9409,6 +9409,11 @@ void UCookOnTheFlyServer::WriteCookMetadata(const ITargetPlatform* InTargetPlatf
 			Entry.DependencyIndexEnd = PluginChildArray.Num();
 		}
 
+		// Also ensure Game depends on Engine.
+		PluginsToAdd[GamePluginIndex].DependencyIndexStart = PluginChildArray.Num();
+		PluginChildArray.Add(EnginePluginIndex);
+		PluginsToAdd[GamePluginIndex].DependencyIndexEnd = PluginChildArray.Num();
+
 		if (IntFitsIn<uint16>(PluginChildArray.Num()) == false)
 		{
 			UE_LOG(LogCook, Warning, TEXT("Number of child plugins exceeds 64k, unable to write cook metadata file (count = %d)"), PluginChildArray.Num());
