Turn off -SkipOnlyEditorOnly (so we now keep unsolicited packages that are otherwise unreferenced) to see if it is the cause of missing system-generated files.
#rnx

diff --git a/Engine/Source/Editor/UnrealEd/Private/CookOnTheFlyServer.cpp b/Engine/Source/Editor/UnrealEd/Private/CookOnTheFlyServer.cpp
--- a/Engine/Source/Editor/UnrealEd/Private/CookOnTheFlyServer.cpp
+++ b/Engine/Source/Editor/UnrealEd/Private/CookOnTheFlyServer.cpp
@@ -6611,6 +6611,7 @@ void UCookOnTheFlyServer::SetInitializeConfigSettings(UE::Cook::FInitializeConfi
 	// Debugging hidden dependencies
 	bOnlyEditorOnlyDebug = FParse::Param(FCommandLine::Get(), TEXT("OnlyEditorOnlyDebug"));
 	bSkipOnlyEditorOnly = false;
+#if 0
 	GConfig->GetBool(TEXT("CookSettings"), TEXT("SkipOnlyEditorOnly"), bSkipOnlyEditorOnly, GEditorIni);
 	FString ParamText;
 	if (FParse::Value(FCommandLine::Get(), TEXT("-SkipOnlyEditorOnly="), ParamText))
@@ -6621,6 +6622,7 @@ void UCookOnTheFlyServer::SetInitializeConfigSettings(UE::Cook::FInitializeConfi
 	{
 		bSkipOnlyEditorOnly = true;
 	}
+#endif
 	bSkipOnlyEditorOnly |= bOnlyEditorOnlyDebug;
 	if (bSkipOnlyEditorOnly)
 	{
