Backout the removal of -SkipOnlyEditorOnly; turn it back on. The problem of missing system-generated files still reproduced even with -SkipOnlyEditorOnly=false.
#rnx

diff --git a/Engine/Source/Editor/UnrealEd/Private/CookOnTheFlyServer.cpp b/Engine/Source/Editor/UnrealEd/Private/CookOnTheFlyServer.cpp
--- a/Engine/Source/Editor/UnrealEd/Private/CookOnTheFlyServer.cpp
+++ b/Engine/Source/Editor/UnrealEd/Private/CookOnTheFlyServer.cpp
@@ -6611,7 +6611,6 @@ void UCookOnTheFlyServer::SetInitializeConfigSettings(UE::Cook::FInitializeConfi
 	// Debugging hidden dependencies
 	bOnlyEditorOnlyDebug = FParse::Param(FCommandLine::Get(), TEXT("OnlyEditorOnlyDebug"));
 	bSkipOnlyEditorOnly = false;
-#if 0
 	GConfig->GetBool(TEXT("CookSettings"), TEXT("SkipOnlyEditorOnly"), bSkipOnlyEditorOnly, GEditorIni);
 	FString ParamText;
 	if (FParse::Value(FCommandLine::Get(), TEXT("-SkipOnlyEditorOnly="), ParamText))
@@ -6622,7 +6621,6 @@ void UCookOnTheFlyServer::SetInitializeConfigSettings(UE::Cook::FInitializeConfi
 	{
 		bSkipOnlyEditorOnly = true;
 	}
-#endif
 	bSkipOnlyEditorOnly |= bOnlyEditorOnlyDebug;
 	if (bSkipOnlyEditorOnly)
 	{
