Fix crash when physics asset thumbnail renderer try to render the thumbnail in same time the skeletal mesh is being compiled.

#rnx

diff --git a/Engine/Source/Editor/UnrealEd/Classes/ThumbnailRendering/PhysicsAssetThumbnailRenderer.h b/Engine/Source/Editor/UnrealEd/Classes/ThumbnailRendering/PhysicsAssetThumbnailRenderer.h
--- a/Engine/Source/Editor/UnrealEd/Classes/ThumbnailRendering/PhysicsAssetThumbnailRenderer.h
+++ b/Engine/Source/Editor/UnrealEd/Classes/ThumbnailRendering/PhysicsAssetThumbnailRenderer.h
@@ -15,6 +15,7 @@ class UPhysicsAssetThumbnailRenderer : public UDefaultSizedThumbnailRenderer
 {
 	GENERATED_UCLASS_BODY()
 
+	virtual bool CanVisualizeAsset(UObject* Object) override;
 
 	// Begin UThumbnailRenderer Object
 	UNREALED_API virtual void Draw(UObject* Object, int32 X, int32 Y, uint32 Width, uint32 Height, FRenderTarget*, FCanvas* Canvas, bool bAdditionalViewFamily) override;
diff --git a/Engine/Source/Editor/UnrealEd/Private/ThumbnailRendering/PhysicsAssetThumbnailRenderer.cpp b/Engine/Source/Editor/UnrealEd/Private/ThumbnailRendering/PhysicsAssetThumbnailRenderer.cpp
--- a/Engine/Source/Editor/UnrealEd/Private/ThumbnailRendering/PhysicsAssetThumbnailRenderer.cpp
+++ b/Engine/Source/Editor/UnrealEd/Private/ThumbnailRendering/PhysicsAssetThumbnailRenderer.cpp
@@ -13,18 +13,44 @@ UPhysicsAssetThumbnailRenderer::UPhysicsAssetThumbnailRenderer(const FObjectInit
 	ThumbnailScene = nullptr;
 }
 
+bool UPhysicsAssetThumbnailRenderer::CanVisualizeAsset(UObject* Object)
+{
+	if (!Super::CanVisualizeAsset(Object))
+	{
+		return false;
+	}
+	UPhysicsAsset* PhysicsAsset = Cast<UPhysicsAsset>(Object);
+	if (!PhysicsAsset)
+	{
+		return false;
+	}
+	USkeletalMesh* SkeletalMesh = PhysicsAsset->PreviewSkeletalMesh.LoadSynchronous();
+	const bool bValidRenderData = SkeletalMesh && !SkeletalMesh->IsCompiling() && SkeletalMesh->GetResourceForRendering();
+	if (!bValidRenderData)
+	{
+		return false;
+	}
+	return true;
+}
+
 void UPhysicsAssetThumbnailRenderer::Draw(UObject* Object, int32 X, int32 Y, uint32 Width, uint32 Height, FRenderTarget* RenderTarget, FCanvas* Canvas, bool bAdditionalViewFamily)
 {
 	UPhysicsAsset* PhysicsAsset = Cast<UPhysicsAsset>(Object);
-	if (PhysicsAsset != nullptr)
+	if (PhysicsAsset)
 	{
-		if ( ThumbnailScene == nullptr )
+		USkeletalMesh* SkeletalMesh = PhysicsAsset->PreviewSkeletalMesh.LoadSynchronous();
+		const bool bValidRenderData = SkeletalMesh && !SkeletalMesh->IsCompiling() && SkeletalMesh->GetResourceForRendering();
+		if (!bValidRenderData)
+		{
+			return;
+		}
+		if (!ThumbnailScene)
 		{
 			ThumbnailScene = new FPhysicsAssetThumbnailScene();
 		}
 
 		ThumbnailScene->SetPhysicsAsset(PhysicsAsset);
-		FSceneViewFamilyContext ViewFamily( FSceneViewFamily::ConstructionValues( RenderTarget, ThumbnailScene->GetScene(), FEngineShowFlags(ESFIM_Game) )
+		FSceneViewFamilyContext ViewFamily(FSceneViewFamily::ConstructionValues(RenderTarget, ThumbnailScene->GetScene(), FEngineShowFlags(ESFIM_Game))
 			.SetTime(UThumbnailRenderer::GetTime())
 			.SetAdditionalViewFamily(bAdditionalViewFamily));
 
