Create sparse data for SKEL_REINST classes to avoid crashes when CDOs are intentionnaly or accidently created for these types

diff --git a/Engine/Source/Editor/UnrealEd/Private/Kismet2/KismetReinstanceUtilities.cpp b/Engine/Source/Editor/UnrealEd/Private/Kismet2/KismetReinstanceUtilities.cpp
--- a/Engine/Source/Editor/UnrealEd/Private/Kismet2/KismetReinstanceUtilities.cpp
+++ b/Engine/Source/Editor/UnrealEd/Private/Kismet2/KismetReinstanceUtilities.cpp
@@ -1764,6 +1764,14 @@ void FBlueprintCompileReinstancer::MoveDependentSkelToReinst(UClass* const Owner
 		ReinstClass->Bind();
 		ReinstClass->StaticLink(true);
 
+		if (!ReinstClass->GetSparseClassDataStruct())
+		{
+			if (UScriptStruct* SparseClassDataStructArchetype = ReinstClass->GetSparseClassDataArchetypeStruct())
+			{
+				ReinstClass->SetSparseClassDataStruct(SparseClassDataStructArchetype);
+			}
+		}
+
 		// Map the old class to the new one
 		NewSkeletonToOldSkeleton.Add(CurClass, ReinstClass);
 
