DiffWriterArchive and FOptionalProperty: Add support for FOptionalProperty in the moved version of FDiffWriterArchive (previously FArchiveStackTrace).
#rnx

diff --git a/Engine/Source/Editor/UnrealEd/Private/Cooker/DiffWriterArchive.cpp b/Engine/Source/Editor/UnrealEd/Private/Cooker/DiffWriterArchive.cpp
--- a/Engine/Source/Editor/UnrealEd/Private/Cooker/DiffWriterArchive.cpp
+++ b/Engine/Source/Editor/UnrealEd/Private/Cooker/DiffWriterArchive.cpp
@@ -16,6 +16,7 @@
 #include "UObject/LinkerLoad.h"
 #include "UObject/ObjectMacros.h"
 #include "UObject/Package.h"
+#include "UObject/PropertyOptional.h"
 #include "UObject/PropertyTempVal.h"
 #include "UObject/UnrealType.h"
 #include "UObject/UObjectGlobals.h"
@@ -582,6 +583,11 @@ namespace
 			}
 		}
 
+		if (FOptionalProperty* OptionalProp = CastField<FOptionalProperty>(Prop))
+		{
+			return ShouldDumpPropertyValueState(OptionalProp->GetValueProperty());
+		}
+
 		return false;
 	}
 }
