SClassDialogPicker: Augmented the sizes of the buttons to prevent cutoffs

diff --git a/Engine/Source/Editor/UnrealEd/Private/SClassPickerDialog.cpp b/Engine/Source/Editor/UnrealEd/Private/SClassPickerDialog.cpp
--- a/Engine/Source/Editor/UnrealEd/Private/SClassPickerDialog.cpp
+++ b/Engine/Source/Editor/UnrealEd/Private/SClassPickerDialog.cpp
@@ -282,7 +282,7 @@ TSharedRef<ITableRow> SClassPickerDialog::GenerateListRow(TSharedPtr<FClassPicke
 			SNew(SHorizontalBox)
 			+SHorizontalBox::Slot()
 			.VAlign(VAlign_Center)
-			.FillWidth(0.3f)
+			.FillWidth(0.34f)
 			[
 				SNew(SButton)
 				.OnClicked(this, &SClassPickerDialog::OnDefaultClassPicked, ItemClass)
@@ -310,7 +310,7 @@ TSharedRef<ITableRow> SClassPickerDialog::GenerateListRow(TSharedPtr<FClassPicke
 			+SHorizontalBox::Slot()
 			.VAlign(VAlign_Center)
 			.Padding(10.0f, 0.0f, 4.0f, 0.0f)
-			.FillWidth(0.63f)
+			.FillWidth(0.59f)
 			[
 				SNew(STextBlock)
 				.Text(Obj->GetDescription())
