Cooker: DiffWriterArchive: Fix a bug after refactor with FArchiveStackTraceIgnoreScope no longer being respected. Change DiffWriterArchive to report all messages to a message callback rather than logging them directly. Future users of DiffWriterArchive will store the messages rather than logging them.
#rnx

diff --git a/Engine/Source/Editor/UnrealEd/Private/Cooker/DiffPackageWriter.cpp b/Engine/Source/Editor/UnrealEd/Private/Cooker/DiffPackageWriter.cpp
--- a/Engine/Source/Editor/UnrealEd/Private/Cooker/DiffPackageWriter.cpp
+++ b/Engine/Source/Editor/UnrealEd/Private/Cooker/DiffPackageWriter.cpp
@@ -7,12 +7,14 @@
 #include "CookOnTheSide/CookLog.h"
 #include "CoreGlobals.h"
 #include "Engine/Engine.h"
+#include "Logging/LogMacros.h"
 #include "Misc/CommandLine.h"
 #include "Misc/ConfigCacheIni.h"
 #include "Misc/CString.h"
 #include "Misc/FeedbackContext.h"
 #include "Misc/Parse.h"
 #include "Misc/OutputDevice.h"
+#include "Misc/OutputDeviceHelper.h"
 #include "ProfilingDebugging/CookStats.h"
 #include "Serialization/LargeMemoryWriter.h"
 #include "Templates/UniquePtr.h"
@@ -23,6 +25,8 @@
 #include "UObject/UObjectGlobals.h"
 #include "UObject/UObjectIterator.h"
 
+DEFINE_LOG_CATEGORY_STATIC(LogDiff, Log, All);
+
 FDiffPackageWriter::FDiffPackageWriter(TUniquePtr<ICookedPackageWriter>&& InInner)
 	: Inner(MoveTemp(InInner))
 {
@@ -45,6 +49,10 @@ FDiffPackageWriter::FDiffPackageWriter(TUniquePtr<ICookedPackageWriter>&& InInne
 	}
 
 	ParseCmds();
+
+	Indent = FCString::Spc(FOutputDeviceHelper::FormatLogLine(ELogVerbosity::Warning,
+		LogDiff.GetCategoryName(), TEXT(""), GPrintLogTimes).Len());
+	NewLine = TEXT("\n"); // OutputDevices are responsible for remapping to LINE_TERMINATOR if desired
 }
 
 void FDiffPackageWriter::ParseCmds()
@@ -188,6 +196,24 @@ void FDiffPackageWriter::WritePackageData(const FPackageInfo& Info, FLargeMemory
 	Inner->WritePackageData(LocalInfo, ExportsArchive, FileRegions);
 }
 
+UE::DiffWriterArchive::FMessageCallback FDiffPackageWriter::GetDiffWriterMessageCallback()
+{
+	return UE::DiffWriterArchive::FMessageCallback([this](ELogVerbosity::Type Verbosity, FStringView Message)
+		{
+			this->OnDiffWriterMessage(Verbosity, Message);
+		});
+}
+
+void FDiffPackageWriter::OnDiffWriterMessage(ELogVerbosity::Type Verbosity, FStringView Message)
+{
+	FString ResolvedText(Message);
+	check(this->Indent && this->NewLine);
+	ResolvedText.ReplaceInline(UE::DiffWriterArchive::IndentToken, this->Indent);
+	ResolvedText.ReplaceInline(UE::DiffWriterArchive::NewLineToken, this->NewLine);
+
+	FMsg::Logf(__FILE__, __LINE__, LogDiff.GetCategoryName(), Verbosity, TEXT("%s"), *ResolvedText);
+}
+
 TUniquePtr<FLargeMemoryWriter> FDiffPackageWriter::CreateLinkerArchive(FName PackageName, UObject* Asset, uint16 MultiOutputIndex)
 {
 	// The entire package will be serialized to memory and then compared against package on disk.
@@ -196,12 +222,12 @@ TUniquePtr<FLargeMemoryWriter> FDiffPackageWriter::CreateLinkerArchive(FName Pac
 		check(MultiOutputIndex < 2);
 		// Each difference will be logged with its Serialize call stack trace
 		return TUniquePtr<FLargeMemoryWriter>(new FDiffWriterArchive(Asset, *PackageName.ToString(),
-			true /* bInCollectCallstacks */, &DiffMap[MultiOutputIndex]));
+			GetDiffWriterMessageCallback(), true /* bInCollectCallstacks */, &DiffMap[MultiOutputIndex]));
 	}
 	else
 	{
 		return TUniquePtr<FLargeMemoryWriter>(new FDiffWriterArchive(Asset, *PackageName.ToString(),
-			false /* bInCollectCallstacks */));
+			GetDiffWriterMessageCallback(), false /* bInCollectCallstacks */));
 	}
 }
 
diff --git a/Engine/Source/Editor/UnrealEd/Private/Cooker/DiffPackageWriter.h b/Engine/Source/Editor/UnrealEd/Private/Cooker/DiffPackageWriter.h
--- a/Engine/Source/Editor/UnrealEd/Private/Cooker/DiffPackageWriter.h
+++ b/Engine/Source/Editor/UnrealEd/Private/Cooker/DiffPackageWriter.h
@@ -121,11 +121,15 @@ private:
 	bool FilterPackageName(const FString& InWildcard);
 	void ConditionallyDumpObjList();
 	void ConditionallyDumpObjects();
+	UE::DiffWriterArchive::FMessageCallback GetDiffWriterMessageCallback();
+	virtual void OnDiffWriterMessage(ELogVerbosity::Type Verbosity, FStringView Message);
 
 	FDiffWriterDiffMap DiffMap[2];
 	TUniquePtr<FDiffWriterCallstacks> ExportsCallstacks;
 	FBeginPackageInfo BeginInfo;
 	TUniquePtr<ICookedPackageWriter> Inner;
+	const TCHAR* Indent = nullptr;
+	const TCHAR* NewLine = nullptr;
 	FString DumpObjListParams;
 	FString PackageFilter;
 	int64 ExportsDiffMapOffset[2] = {0, 0};
diff --git a/Engine/Source/Editor/UnrealEd/Private/Cooker/DiffWriterArchive.cpp b/Engine/Source/Editor/UnrealEd/Private/Cooker/DiffWriterArchive.cpp
--- a/Engine/Source/Editor/UnrealEd/Private/Cooker/DiffWriterArchive.cpp
+++ b/Engine/Source/Editor/UnrealEd/Private/Cooker/DiffWriterArchive.cpp
@@ -22,64 +22,14 @@
 #include "UObject/UObjectGlobals.h"
 #include "UObject/UObjectThreadContext.h"
 
-DEFINE_LOG_CATEGORY_STATIC(LogDiffWriter, Log, All);
-
-/** Helper class that holds runtime generated constants for log output formatting */
-struct FDiffFormatHelper
-{
-public:
-	const TCHAR* const Indent;
-	const TCHAR* const LineTerminator;
-
-	FDiffFormatHelper()
-		: Indent(FCString::Spc(FOutputDeviceHelper::FormatLogLine(ELogVerbosity::Warning, LogDiffWriter.GetCategoryName(), TEXT(""), GPrintLogTimes).Len()))
-		, LineTerminator(TEXT("\n")) // because LINE_TERMINATOR doesn't work well wit EC
-	{}
-	static FDiffFormatHelper& Get()
-	{
-		static FDiffFormatHelper Instance;
-		return Instance;
-	}
-};
-
-class FIgnoreDiffManager
+static const ANSICHAR* DebugDataStackMarker = "\r\nDebugDataStack:\r\n";
+namespace UE::DiffWriterArchive
 {
-	int32 IgnoreCount = 0;
-	int32 DisabledCount = 0;
-
-public:
-	FIgnoreDiffManager() = default;
-	void Push()
-	{
-		IgnoreCount++;
-	}
-	void Pop()
-	{
-		IgnoreCount--;
-		check(IgnoreCount >= 0);
-	}
-	bool ShouldIgnoreDiff() const
-	{
-		return !!IgnoreCount;
-	}
-	void PushDisabled()
-	{
-		DisabledCount++;
-	}
-	void PopDisabled()
-	{
-		DisabledCount--;
-		check(DisabledCount >= 0);
-	}
-	bool ShouldBypassDiff() const
-	{
-		return !!DisabledCount;
-	}
-};
 
-static FIgnoreDiffManager GIgnoreDiffManager;
+const TCHAR* const IndentToken = TEXT("%DWA%    ");
+const TCHAR* const NewLineToken = TEXT("%DWA%\n");
 
-static const ANSICHAR* DebugDataStackMarker = "\r\nDebugDataStack:\r\n";
+}
 
 FDiffWriterCallstacks::FCallstackData::FCallstackData(TUniquePtr<ANSICHAR[]>&& InCallstack, UObject* InSerializedObject, FProperty* InSerializedProperty)
 	: Callstack(MoveTemp(InCallstack))
@@ -99,9 +49,6 @@ FString FDiffWriterCallstacks::FCallstackData::ToString(const TCHAR* CallstackCu
 {
 	FString HumanReadableString;
 
-	const TCHAR* const LineTerminator = FDiffFormatHelper::Get().LineTerminator;
-	const TCHAR* const Indent = FDiffFormatHelper::Get().Indent;
-	
 	FString StackTraceText = Callstack.Get();
 	if (CallstackCutoffText != nullptr)
 	{
@@ -129,29 +76,29 @@ FString FDiffWriterCallstacks::FCallstackData::ToString(const TCHAR* CallstackCu
 				StackLine.MidInline(CutoffIndex + 1, MAX_int32, false);
 			}
 		}
-		HumanReadableString += Indent;
+		HumanReadableString += UE::DiffWriterArchive::IndentToken;
 		HumanReadableString += StackLine;
-		HumanReadableString += LineTerminator;
+		HumanReadableString += UE::DiffWriterArchive::NewLineToken;
 	}
 
 	if (!SerializedObjectName.IsEmpty())
 	{
-		HumanReadableString += LineTerminator;
-		HumanReadableString += Indent;
+		HumanReadableString += UE::DiffWriterArchive::NewLineToken;
+		HumanReadableString += UE::DiffWriterArchive::IndentToken;
 		HumanReadableString += TEXT("Serialized Object: ");
 		HumanReadableString += SerializedObjectName;
-		HumanReadableString += LineTerminator;
+		HumanReadableString += UE::DiffWriterArchive::NewLineToken;
 	}
 	if (!SerializedPropertyName.IsEmpty())
 	{
 		if (SerializedObjectName.IsEmpty())
 		{
-			HumanReadableString += LineTerminator;
+			HumanReadableString += UE::DiffWriterArchive::NewLineToken;
 		}
-		HumanReadableString += Indent;
+		HumanReadableString += UE::DiffWriterArchive::IndentToken;
 		HumanReadableString += TEXT("Serialized Property: ");
 		HumanReadableString += SerializedPropertyName;
-		HumanReadableString += LineTerminator;
+		HumanReadableString += UE::DiffWriterArchive::NewLineToken;
 	}
 	return HumanReadableString;
 }
@@ -224,7 +171,7 @@ void FDiffWriterCallstacks::Add(
 	bool bCollectCurrentCallstack,
 	int32 StackIgnoreCount)
 {
-	if (GIgnoreDiffManager.ShouldBypassDiff())
+	if (UE::ArchiveStackTrace::ShouldBypassDiff())
 	{
 		return;
 	}
@@ -232,7 +179,7 @@ void FDiffWriterCallstacks::Add(
 	const int64 CurrentOffset = Offset;
 	TotalSize = FMath::Max(TotalSize, CurrentOffset + Length); 
 
-	const bool bShouldCollectCallstack = bIsCollectingCallstacks && bCollectCurrentCallstack && !GIgnoreDiffManager.ShouldIgnoreDiff();
+	const bool bShouldCollectCallstack = bIsCollectingCallstacks && bCollectCurrentCallstack && !UE::ArchiveStackTrace::ShouldIgnoreDiff();
 	if (bShouldCollectCallstack)
 	{
 		StackTrace[0] = '\0';
@@ -242,7 +189,7 @@ void FDiffWriterCallstacks::Add(
 		{
 			FCStringAnsi::Strcat(StackTrace.Get(), StackTraceSize, DebugDataStackMarker);
 
-			const FString SubIndent = FString(FDiffFormatHelper::Get().Indent) + FString(TEXT("    "));
+			const FString SubIndent = FString(UE::DiffWriterArchive::IndentToken) + FString(TEXT("    "));
 
 			bool bIsIndenting = true;
 			for (const auto& DebugData : DebugDataStack)
@@ -274,7 +221,7 @@ void FDiffWriterCallstacks::Add(
 		{
 			// New data serialized at the end of archive buffer
 			LastSerializeCallstack = AddUniqueCallstack(bIsCollectingCallstacks, SerializedObject, SerializedProperty, CallstackCRC);
-			CallstackAtOffsetMap.Add(FCallstackAtOffset {CurrentOffset, CallstackCRC, GIgnoreDiffManager.ShouldIgnoreDiff()});
+			CallstackAtOffsetMap.Add(FCallstackAtOffset {CurrentOffset, CallstackCRC, UE::ArchiveStackTrace::ShouldIgnoreDiff()});
 		}
 		else
 		{
@@ -291,7 +238,7 @@ void FDiffWriterCallstacks::Add(
 			{
 				// Insert a new callstack
 				check(CallstackToUpdate.Offset < CurrentOffset);
-				CallstackAtOffsetMap.Insert(FCallstackAtOffset {CurrentOffset, CallstackCRC, GIgnoreDiffManager.ShouldIgnoreDiff()}, CallstackToUpdateIndex + 1);
+				CallstackAtOffsetMap.Insert(FCallstackAtOffset {CurrentOffset, CallstackCRC, UE::ArchiveStackTrace::ShouldIgnoreDiff()}, CallstackToUpdateIndex + 1);
 			}
 		}
 		check(CallstackCRC != 0 || !bShouldCollectCallstack);
@@ -517,10 +464,12 @@ FUObjectSerializeContext* FDiffWriterArchiveMemoryWriter::GetSerializeContext()
 	return StackTraceWriter.GetSerializeContext();
 }
 
-FDiffWriterArchive::FDiffWriterArchive(UObject* InAsset, const TCHAR* InFilename, bool bInCollectCallstacks, const FDiffWriterDiffMap* InDiffMap)
+FDiffWriterArchive::FDiffWriterArchive(UObject* InAsset, const TCHAR* InFilename, UE::DiffWriterArchive::FMessageCallback&& InMessageCallback,
+	bool bInCollectCallstacks, const FDiffWriterDiffMap* InDiffMap)
 	: FLargeMemoryWriter(0, false, InFilename)
 	, Callstacks(InAsset)
 	, StackTraceWriter(*this, Callstacks, InDiffMap, 0)
+	, MessageCallback(MoveTemp(InMessageCallback))
 {
 	// Hack to prevent recursive calls to serialize when passing in this.
 	StackTraceWriter.SetDisableInnerArchive(true);
@@ -602,10 +551,9 @@ void FDiffWriterArchiveWriter::Compare(
 	const int64 MaxDiffsToLog,
 	int32& InOutDiffsLogged,
 	TMap<FName, FArchiveDiffStats>& OutStats,
+	const UE::DiffWriterArchive::FMessageCallback& MessageCallback,
 	bool bSuppressLogging)
 {
-	const TCHAR* const Indent = FDiffFormatHelper::Get().Indent;
-	const TCHAR* const LineTerminator = FDiffFormatHelper::Get().LineTerminator;
 	const int64 SourceSize = SourcePackage.Size - SourcePackage.StartOffset;
 	const int64 DestSize = DestPackage.Size - DestPackage.StartOffset;
 	const int64 SizeToCompare = FMath::Min(SourceSize, DestSize);
@@ -613,7 +561,11 @@ void FDiffWriterArchiveWriter::Compare(
 	
 	if (SourceSize != DestSize)
 	{
-		UE_CLOG(!bSuppressLogging, LogDiffWriter, Warning, TEXT("%s: Size mismatch: on disk: %lld vs memory: %lld"), AssetFilename, SourceSize, DestSize);
+		if (!bSuppressLogging)
+		{
+			MessageCallback(ELogVerbosity::Warning, FString::Printf(
+				TEXT("%s: Size mismatch: on disk: %lld vs memory: %lld"), AssetFilename, SourceSize, DestSize));
+		}
 		int64 SizeDiff = DestPackage.Size - SourcePackage.Size;
 		OutStats.FindOrAdd(AssetClass).DiffSize += SizeDiff;
 	}
@@ -656,7 +608,11 @@ void FDiffWriterArchiveWriter::Compare(
 
 			if (DifferenceCallstackoffsetIndex < 0)
 			{
-				UE_CLOG(!bSuppressLogging, LogDiffWriter, Warning, TEXT("%s: Difference at offset %lld (absolute offset: %lld), unknown callstack"), AssetFilename, LocalOffset, DestAbsoluteOffset);
+				if (!bSuppressLogging)
+				{
+					MessageCallback(ELogVerbosity::Warning, FString::Printf(
+						TEXT("%s: Difference at offset %lld (absolute offset: %lld), unknown callstack"), AssetFilename, LocalOffset, DestAbsoluteOffset));
+				}
 				continue;
 			}
 
@@ -725,7 +681,9 @@ void FDiffWriterArchiveWriter::Compare(
 				FString DiffValues;
 				if (BeforePropertyVal != AfterPropertyVal)
 				{
-					DiffValues = FString::Printf(TEXT("\r\n%sBefore: %s\r\n%sAfter:  %s"), Indent, *BeforePropertyVal, Indent, *AfterPropertyVal);
+					DiffValues = FString::Printf(TEXT("\r\n%sBefore: %s\r\n%sAfter:  %s"),
+						UE::DiffWriterArchive::IndentToken, *BeforePropertyVal,
+						UE::DiffWriterArchive::IndentToken, *AfterPropertyVal);
 				}
 
 				FString DebugDataStackText;
@@ -734,51 +692,53 @@ void FDiffWriterArchiveWriter::Compare(
 				int32 DebugDataIndex = FullStackText.Find(ANSI_TO_TCHAR(DebugDataStackMarker), ESearchCase::CaseSensitive);
 				if (DebugDataIndex > 0)
 				{
-					DebugDataStackText = FString::Printf(TEXT("\r\n%s"), FDiffFormatHelper::Get().Indent) + FullStackText.RightChop(DebugDataIndex + 2);
+					DebugDataStackText = FString::Printf(TEXT("\r\n%s"),
+						UE::DiffWriterArchive::IndentToken)
+						+ FullStackText.RightChop(DebugDataIndex + 2);
 				}
 
-				UE_CLOG(
-					!bSuppressLogging,
-					LogDiffWriter,
-					Warning,
-					TEXT("%s: Difference at offset %lld%s (absolute offset: %lld): byte %d on disk, byte %d in memory, callstack:%s%s%s%s%s"),
-					AssetFilename,
-					CallstackAtOffset.Offset - DestPackage.StartOffset,
-					DestAbsoluteOffset > CallstackAtOffset.Offset ? *FString::Printf(TEXT("(+%lld)"), DestAbsoluteOffset - CallstackAtOffset.Offset) : TEXT(""),
-					DestAbsoluteOffset,
-					SourceByte, DestByte,
-					LineTerminator,
-					LineTerminator,
-					*DifferenceCallstackDataText,
-					*DiffValues,
-					*DebugDataStackText
-				);
+				if (!bSuppressLogging)
+				{
+					MessageCallback(ELogVerbosity::Warning, FString::Printf(
+						TEXT("%s: Difference at offset %lld%s (absolute offset: %lld): byte %d on disk, byte %d in memory, callstack:%s%s%s%s%s"),
+						AssetFilename,
+						CallstackAtOffset.Offset - DestPackage.StartOffset,
+						DestAbsoluteOffset > CallstackAtOffset.Offset ? *FString::Printf(TEXT("(+%lld)"), DestAbsoluteOffset - CallstackAtOffset.Offset) : TEXT(""),
+						DestAbsoluteOffset,
+						SourceByte, DestByte,
+						UE::DiffWriterArchive::NewLineToken,
+						UE::DiffWriterArchive::NewLineToken,
+						*DifferenceCallstackDataText,
+						*DiffValues,
+						*DebugDataStackText
+					));
+				}
 
 				const int BytesToLog = 128;
-				UE_CLOG(
-					!bSuppressLogging,
-					LogDiffWriter,
-					Display,
-					TEXT("%s: Logging %d bytes around absolute offset: %lld (%016X) in the on disk (existing) package, (which corresponds to offset %lld (%016X) in the in-memory package)"),
-					AssetFilename,
-					BytesToLog,
-					SourceAbsoluteOffset,
-					SourceAbsoluteOffset,
-					DestAbsoluteOffset,
-					DestAbsoluteOffset
-				);
+				if (!bSuppressLogging)
+				{
+					MessageCallback(ELogVerbosity::Display, FString::Printf(
+						TEXT("%s: Logging %d bytes around absolute offset: %lld (%016X) in the on disk (existing) package, (which corresponds to offset %lld (%016X) in the in-memory package)"),
+						AssetFilename,
+						BytesToLog,
+						SourceAbsoluteOffset,
+						SourceAbsoluteOffset,
+						DestAbsoluteOffset,
+						DestAbsoluteOffset
+					));
+				}
 				FCompressionUtil::LogHexDump(SourcePackage.Data, SourcePackage.Size, SourceAbsoluteOffset - BytesToLog / 2, SourceAbsoluteOffset + BytesToLog / 2);
 
-				UE_CLOG(
-					!bSuppressLogging,
-					LogDiffWriter,
-					Display,
-					TEXT("%s: Logging %d bytes around absolute offset: %lld (%016X) in the in memory (new) package"),
-					AssetFilename,
-					BytesToLog,
-					DestAbsoluteOffset,
-					DestAbsoluteOffset
-				);
+				if (!bSuppressLogging)
+				{
+					MessageCallback(ELogVerbosity::Display, FString::Printf(
+						TEXT("%s: Logging %d bytes around absolute offset: %lld (%016X) in the in memory (new) package"),
+						AssetFilename,
+						BytesToLog,
+						DestAbsoluteOffset,
+						DestAbsoluteOffset
+					));
+				}
 				FCompressionUtil::LogHexDump(DestPackage.Data, DestPackage.Size, DestAbsoluteOffset - BytesToLog / 2, DestAbsoluteOffset + BytesToLog / 2);
 
 				bDifferenceLogged = true;
@@ -807,11 +767,20 @@ void FDiffWriterArchiveWriter::Compare(
 	{
 		if (FirstUnreportedDiffIndex != -1)
 		{
-			UE_CLOG(!bSuppressLogging, LogDiffWriter, Warning, TEXT("%s: %lld difference(s) not logged (first at absolute offset: %lld)."), AssetFilename, NumDiffsLocal - NumDiffsLoggedLocal, FirstUnreportedDiffIndex);
+			if (!bSuppressLogging)
+			{
+				MessageCallback(ELogVerbosity::Warning, FString::Printf(
+					TEXT("%s: %lld difference(s) not logged (first at absolute offset: %lld)."),
+					AssetFilename, NumDiffsLocal - NumDiffsLoggedLocal, FirstUnreportedDiffIndex));
+			}
 		}
 		else
 		{
-			UE_CLOG(!bSuppressLogging, LogDiffWriter, Warning, TEXT("%s: %lld difference(s) not logged."), AssetFilename, NumDiffsLocal - NumDiffsLoggedLocal);
+			if (!bSuppressLogging)
+			{
+				MessageCallback(ELogVerbosity::Warning, FString::Printf(
+					TEXT("%s: %lld difference(s) not logged."), AssetFilename, NumDiffsLocal - NumDiffsLoggedLocal));
+			}
 		}
 	}
 }
@@ -832,7 +801,7 @@ void FDiffWriterArchive::CompareWith(const FPackageData& SourcePackage, const TC
 	OutStats.FindOrAdd(AssetClass).NewFileTotalSize = TotalSize();
 	if (SourcePackage.Size == 0)
 	{
-		UE_LOG(LogDiffWriter, Warning, TEXT("New package: %s"), *GetArchiveName());
+		MessageCallback(ELogVerbosity::Warning, FString::Printf(TEXT("New package: %s"), *GetArchiveName()));
 		OutStats.FindOrAdd(AssetClass).DiffSize = OutStats.FindOrAdd(AssetClass).NewFileTotalSize;
 		return;
 	}
@@ -843,8 +812,8 @@ void FDiffWriterArchive::CompareWith(const FPackageData& SourcePackage, const TC
 	DestPackage.HeaderSize = TotalHeaderSize;
 	DestPackage.StartOffset = 0;
 
-	UE_LOG(LogDiffWriter, Display, TEXT("Comparing: %s"), *GetArchiveName());
-	UE_LOG(LogDiffWriter, Display, TEXT("Asset class: %s"), *AssetClass.ToString());
+	MessageCallback(ELogVerbosity::Display, FString::Printf(TEXT("Comparing: %s"), *GetArchiveName()));
+	MessageCallback(ELogVerbosity::Warning, FString::Printf(TEXT("Asset class: %s"), *AssetClass.ToString()));
 
 	int32 NumLoggedDiffs = 0;
 
@@ -858,11 +827,14 @@ void FDiffWriterArchive::CompareWith(const FPackageData& SourcePackage, const TC
 	DestPackageHeader.HeaderSize = 0;
 	DestPackageHeader.StartOffset = 0;
 
-	FDiffWriterArchiveWriter::Compare(SourcePackageHeader, DestPackageHeader, Callstacks, StackTraceWriter.GetDiffMap(), FileDisplayName, CallstackCutoffText, MaxDiffsToLog, NumLoggedDiffs, OutStats);
+	FDiffWriterArchiveWriter::Compare(SourcePackageHeader, DestPackageHeader, Callstacks,
+		StackTraceWriter.GetDiffMap(), FileDisplayName, CallstackCutoffText, MaxDiffsToLog, NumLoggedDiffs,
+		OutStats, MessageCallback);
 
 	if (TotalHeaderSize > 0 && OutStats.FindOrAdd(AssetClass).NumDiffs > 0)
 	{
-		FDiffWriterArchiveWriter::DumpPackageHeaderDiffs(SourcePackage, DestPackage, FileDisplayName, MaxDiffsToLog, PackageHeaderFormat);
+		FDiffWriterArchiveWriter::DumpPackageHeaderDiffs(SourcePackage, DestPackage, FileDisplayName, MaxDiffsToLog,
+			PackageHeaderFormat, MessageCallback);
 	}
 
 	FPackageData SourcePackageExports = SourcePackage;
@@ -883,7 +855,9 @@ void FDiffWriterArchive::CompareWith(const FPackageData& SourcePackage, const TC
 		AssetName = FileDisplayName;
 	}
 
-	FDiffWriterArchiveWriter::Compare(SourcePackageExports, DestPackageExports, Callstacks, StackTraceWriter.GetDiffMap(), *AssetName, CallstackCutoffText, MaxDiffsToLog, NumLoggedDiffs, OutStats);
+	FDiffWriterArchiveWriter::Compare(SourcePackageExports, DestPackageExports, Callstacks,
+		StackTraceWriter.GetDiffMap(), *AssetName, CallstackCutoffText, MaxDiffsToLog, NumLoggedDiffs,
+		OutStats, MessageCallback);
 
 	// Optionally save out any differences we detected.
 	const FArchiveDiffStats& Stats = OutStats.FindOrAdd(AssetClass);
@@ -939,7 +913,8 @@ void FDiffWriterArchive::CompareWith(const FPackageData& SourcePackage, const TC
 			}
 			else
 			{
-				UE_LOG(LogDiffWriter, Warning, TEXT("Package '%s' doesn't seem to be writing to the Saved directory - skipping writing diff"), *OutputFilename);
+				MessageCallback(ELogVerbosity::Warning,
+					FString::Printf(TEXT("Package '%s' doesn't seem to be writing to the Saved directory - skipping writing diff"), *OutputFilename));
 			}
 		}
 	}
@@ -1158,14 +1133,17 @@ static inline FString GetTableKeyForIndex(const FLinkerLoad* Linker, FPackageInd
 	}
 }
 
-static bool ComparePackageIndices(FLinkerLoad* SourceLinker, FLinkerLoad* DestLinker, const FPackageIndex& SourceIndex, const FPackageIndex& DestIndex);
+static bool ComparePackageIndices(FLinkerLoad* SourceLinker, FLinkerLoad* DestLinker, const FPackageIndex& SourceIndex,
+	const FPackageIndex& DestIndex, const UE::DiffWriterArchive::FMessageCallback& MessageCallback);
 
-static bool CompareTableItem(FLinkerLoad* SourceLinker, FLinkerLoad* DestLinker, const FName& SourceName, const FName& DestName)
+static bool CompareTableItem(FLinkerLoad* SourceLinker, FLinkerLoad* DestLinker,
+	const FName& SourceName, const FName& DestName, const UE::DiffWriterArchive::FMessageCallback& MessageCallback)
 {
 	return SourceName == DestName;
 }
 
-static bool CompareTableItem(FLinkerLoad* SourceLinker, FLinkerLoad* DestLinker, FNameEntryId SourceName, FNameEntryId DestName)
+static bool CompareTableItem(FLinkerLoad* SourceLinker, FLinkerLoad* DestLinker,
+	FNameEntryId SourceName, FNameEntryId DestName, const UE::DiffWriterArchive::FMessageCallback& MessageCallback)
 {
 	return SourceName == DestName;
 }
@@ -1180,12 +1158,14 @@ static FString ConvertItemToText(FNameEntryId Id, FLinkerLoad* Linker)
 	return FName::GetEntry(Id)->GetPlainNameString();
 }
 
-static bool CompareTableItem(FLinkerLoad* SourceLinker, FLinkerLoad* DestLinker, const FObjectImport& SourceImport, const FObjectImport& DestImport)
+static bool CompareTableItem(FLinkerLoad* SourceLinker, FLinkerLoad* DestLinker, const FObjectImport& SourceImport,
+	const FObjectImport& DestImport, const UE::DiffWriterArchive::FMessageCallback& MessageCallback)
 {
 	if (SourceImport.ObjectName != DestImport.ObjectName ||
 		SourceImport.ClassName != DestImport.ClassName ||
 		SourceImport.ClassPackage != DestImport.ClassPackage ||
-		!ComparePackageIndices(SourceLinker, DestLinker, SourceImport.OuterIndex, DestImport.OuterIndex))
+		!ComparePackageIndices(SourceLinker, DestLinker, SourceImport.OuterIndex,
+			DestImport.OuterIndex, MessageCallback))
 	{
 		return false;
 	}
@@ -1204,7 +1184,8 @@ static FString ConvertItemToText(const FObjectImport& Import, FLinkerLoad* Linke
 	);
 }
 
-static bool CompareTableItem(FLinkerLoad* SourceLinker, FLinkerLoad* DestLinker, const FObjectExport& SourceExport, const FObjectExport& DestExport)
+static bool CompareTableItem(FLinkerLoad* SourceLinker, FLinkerLoad* DestLinker, const FObjectExport& SourceExport,
+	const FObjectExport& DestExport, const UE::DiffWriterArchive::FMessageCallback& MessageCallback)
 {
 	if (SourceExport.ObjectName != DestExport.ObjectName ||
 		SourceExport.PackageFlags != DestExport.PackageFlags ||
@@ -1217,10 +1198,14 @@ static bool CompareTableItem(FLinkerLoad* SourceLinker, FLinkerLoad* DestLinker,
 		SourceExport.bIsAsset != DestExport.bIsAsset ||
 		SourceExport.bIsInheritedInstance != DestExport.bIsInheritedInstance ||
 		SourceExport.bGeneratePublicHash != DestExport.bGeneratePublicHash ||
-		!ComparePackageIndices(SourceLinker, DestLinker, SourceExport.TemplateIndex, DestExport.TemplateIndex) ||
-		!ComparePackageIndices(SourceLinker, DestLinker, SourceExport.OuterIndex, DestExport.OuterIndex) ||
-		!ComparePackageIndices(SourceLinker, DestLinker, SourceExport.ClassIndex, DestExport.ClassIndex) ||
-		!ComparePackageIndices(SourceLinker, DestLinker, SourceExport.SuperIndex, DestExport.SuperIndex))
+		!ComparePackageIndices(SourceLinker, DestLinker, SourceExport.TemplateIndex,
+			DestExport.TemplateIndex, MessageCallback) ||
+		!ComparePackageIndices(SourceLinker, DestLinker, SourceExport.OuterIndex,
+			DestExport.OuterIndex, MessageCallback) ||
+		!ComparePackageIndices(SourceLinker, DestLinker, SourceExport.ClassIndex,
+			DestExport.ClassIndex, MessageCallback) ||
+		!ComparePackageIndices(SourceLinker, DestLinker, SourceExport.SuperIndex, 
+			DestExport.SuperIndex, MessageCallback))
 	{
 		return false;
 	}
@@ -1230,14 +1215,16 @@ static bool CompareTableItem(FLinkerLoad* SourceLinker, FLinkerLoad* DestLinker,
 	}
 }
 
-static bool IsImportMapIdentical(FLinkerLoad* SourceLinker, FLinkerLoad* DestLinker)
+static bool IsImportMapIdentical(FLinkerLoad* SourceLinker, FLinkerLoad* DestLinker,
+	const UE::DiffWriterArchive::FMessageCallback& MessageCallback)
 {
 	bool bIdentical = (SourceLinker->ImportMap.Num() == DestLinker->ImportMap.Num());
 	if (bIdentical)
 	{
 		for (int32 ImportIndex = 0; ImportIndex < SourceLinker->ImportMap.Num(); ++ImportIndex)
 		{
-			if (!CompareTableItem(SourceLinker, DestLinker, SourceLinker->ImportMap[ImportIndex], DestLinker->ImportMap[ImportIndex]))
+			if (!CompareTableItem(SourceLinker, DestLinker, SourceLinker->ImportMap[ImportIndex],
+				DestLinker->ImportMap[ImportIndex], MessageCallback))
 			{
 				bIdentical = false;
 				break;
@@ -1247,7 +1234,8 @@ static bool IsImportMapIdentical(FLinkerLoad* SourceLinker, FLinkerLoad* DestLin
 	return bIdentical;
 }
 
-static bool ComparePackageIndices(FLinkerLoad* SourceLinker, FLinkerLoad* DestLinker, const FPackageIndex& SourceIndex, const FPackageIndex& DestIndex)
+static bool ComparePackageIndices(FLinkerLoad* SourceLinker, FLinkerLoad* DestLinker, const FPackageIndex& SourceIndex,
+	const FPackageIndex& DestIndex, const UE::DiffWriterArchive::FMessageCallback& MessageCallback)
 {
 	if (SourceIndex.IsNull() && DestIndex.IsNull())
 	{
@@ -1261,7 +1249,9 @@ static bool ComparePackageIndices(FLinkerLoad* SourceLinker, FLinkerLoad* DestLi
 
 		if (!SourceLinker->ExportMap.IsValidIndex(SourceArrayIndex) || !DestLinker->ExportMap.IsValidIndex(DestArrayIndex))
 		{
-			UE_LOG(LogDiffWriter, Warning, TEXT("Invalid export indices found, source: %d (of %d), dest: %d (of %d)"), SourceArrayIndex, SourceLinker->ExportMap.Num(), DestArrayIndex, DestLinker->ExportMap.Num());
+			MessageCallback(ELogVerbosity::Warning, FString::Printf(
+				TEXT("Invalid export indices found, source: %d (of %d), dest: %d (of %d)"),
+				SourceArrayIndex, SourceLinker->ExportMap.Num(), DestArrayIndex, DestLinker->ExportMap.Num()));
 			return false;
 		}
 
@@ -1281,7 +1271,9 @@ static bool ComparePackageIndices(FLinkerLoad* SourceLinker, FLinkerLoad* DestLi
 
 		if (!SourceLinker->ImportMap.IsValidIndex(SourceArrayIndex) || !DestLinker->ImportMap.IsValidIndex(DestArrayIndex))
 		{
-			UE_LOG(LogDiffWriter, Warning, TEXT("Invalid import indices found, source: %d (of %d), dest: %d (of %d)"), SourceArrayIndex, SourceLinker->ImportMap.Num(), DestArrayIndex, DestLinker->ImportMap.Num());
+			MessageCallback(ELogVerbosity::Warning, FString::Printf(
+				TEXT("Invalid import indices found, source: %d (of %d), dest: %d (of %d)"),
+				SourceArrayIndex, SourceLinker->ImportMap.Num(), DestArrayIndex, DestLinker->ImportMap.Num()));
 			return false;
 		}
 
@@ -1316,14 +1308,16 @@ static FString ConvertItemToText(const FObjectExport& Export, FLinkerLoad* Linke
 		Export.bGeneratePublicHash);
 }
 
-static bool IsExportMapIdentical(FLinkerLoad* SourceLinker, FLinkerLoad* DestLinker)
+static bool IsExportMapIdentical(FLinkerLoad* SourceLinker, FLinkerLoad* DestLinker,
+	const UE::DiffWriterArchive::FMessageCallback& MessageCallback)
 {
 	bool bIdentical = (SourceLinker->ExportMap.Num() == DestLinker->ExportMap.Num());
 	if (bIdentical)
 	{
 		for (int32 ExportIndex = 0; ExportIndex < SourceLinker->ExportMap.Num(); ++ExportIndex)
 		{
-			if (!CompareTableItem(SourceLinker, DestLinker, SourceLinker->ExportMap[ExportIndex], DestLinker->ExportMap[ExportIndex]))
+			if (!CompareTableItem(SourceLinker, DestLinker, SourceLinker->ExportMap[ExportIndex],
+				DestLinker->ExportMap[ExportIndex], MessageCallback))
 			{
 				bIdentical = false;
 				break;
@@ -1371,12 +1365,10 @@ static void DumpTableDifferences(
 	TArray<T>& DestTable,
 	const TCHAR* AssetFilename,
 	const TCHAR* ItemName,
-	const int32 MaxDiffsToLog
+	const int32 MaxDiffsToLog,
+	const UE::DiffWriterArchive::FMessageCallback& MessageCallback
 )
 {
-	const TCHAR* const LineTerminator = FDiffFormatHelper::Get().LineTerminator;
-	const TCHAR* const Indent = FDiffFormatHelper::Get().Indent;
-
 	FString HumanReadableString;
 	int32 LoggedDiffs = 0;
 	int32 NumDiffs = 0;
@@ -1407,7 +1399,8 @@ static void DumpTableDifferences(
 	{
 		if (const TTableItem<T>* ChangedDestItem = DestSet.Find(ChangedSourceItem))
 		{
-			if (!CompareTableItem(SourceLinker, DestLinker, *ChangedSourceItem.Item, *ChangedDestItem->Item))
+			if (!CompareTableItem(SourceLinker, DestLinker, *ChangedSourceItem.Item,
+				*ChangedDestItem->Item, MessageCallback))
 			{
 				RemovedItems.Add(ChangedSourceItem);
 				AddedItems  .Add(*ChangedDestItem);
@@ -1422,15 +1415,15 @@ static void DumpTableDifferences(
 	// Dump all changes
 	for (const TTableItem<T>& RemovedItem : RemovedItems)
 	{
-		HumanReadableString += Indent;
+		HumanReadableString += UE::DiffWriterArchive::IndentToken;
 		HumanReadableString += FString::Printf(TEXT("-[%d] %s"), RemovedItem.Index, *ConvertItemToText(*RemovedItem.Item, SourceLinker));
-		HumanReadableString += LineTerminator;
+		HumanReadableString += UE::DiffWriterArchive::NewLineToken;
 	}
 	for (const TTableItem<T>& AddedItem : AddedItems)
 	{
-		HumanReadableString += Indent;
+		HumanReadableString += UE::DiffWriterArchive::IndentToken;
 		HumanReadableString += FString::Printf(TEXT("+[%d] %s"), AddedItem.Index, *ConvertItemToText(*AddedItem.Item, DestLinker));
-		HumanReadableString += LineTerminator;
+		HumanReadableString += UE::DiffWriterArchive::NewLineToken;
 	}
 
 	// For now just log everything out. When this becomes too spammy, respect the MaxDiffsToLog parameter
@@ -1439,14 +1432,12 @@ static void DumpTableDifferences(
 
 	if (NumDiffs > LoggedDiffs)
 	{
-		HumanReadableString += Indent;
+		HumanReadableString += UE::DiffWriterArchive::IndentToken;
 		HumanReadableString += FString::Printf(TEXT("+ %d differences not logged."), (NumDiffs - LoggedDiffs));
-		HumanReadableString += LineTerminator;
+		HumanReadableString += UE::DiffWriterArchive::NewLineToken;
 	}
 
-	UE_LOG(
-		LogDiffWriter,
-		Warning,
+	MessageCallback(ELogVerbosity::Warning, FString::Printf(
 		TEXT("%s: %sMap is different (%d %ss in source package vs %d %ss in dest package):%s%s"),		
 		AssetFilename,
 		ItemName,
@@ -1454,15 +1445,16 @@ static void DumpTableDifferences(
 		ItemName,
 		DestTable.Num(),
 		ItemName,
-		LineTerminator,
-		*HumanReadableString);
+		UE::DiffWriterArchive::NewLineToken,
+		*HumanReadableString));
 }
 
 static void DumpPackageHeaderDiffs_LinkerLoad(
 	const FDiffWriterArchiveWriter::FPackageData& SourcePackage,
 	const FDiffWriterArchiveWriter::FPackageData& DestPackage,
 	const FString& AssetFilename,
-	const int32 MaxDiffsToLog)
+	const int32 MaxDiffsToLog,
+	const UE::DiffWriterArchive::FMessageCallback& MessageCallback)
 {
 	FString AssetPathName = FPaths::Combine(*FPaths::GetPath(AssetFilename.Mid(AssetFilename.Find(TEXT(":"), ESearchCase::CaseSensitive) + 1)), *FPaths::GetBaseFilename(AssetFilename));
 	// The root directory could have a period in it (d:/Release5.0/EngineTest/Saved/Cooked),
@@ -1500,17 +1492,20 @@ static void DumpPackageHeaderDiffs_LinkerLoad(
 	{
 		if (SourceLinker->NameMap != DestLinker->NameMap)
 		{
-			DumpTableDifferences<FNameEntryId>(SourceLinker, DestLinker, SourceLinker->NameMap, DestLinker->NameMap, *AssetFilename, TEXT("Name"), MaxDiffsToLog);
+			DumpTableDifferences<FNameEntryId>(SourceLinker, DestLinker, SourceLinker->NameMap, DestLinker->NameMap,
+				*AssetFilename, TEXT("Name"), MaxDiffsToLog, MessageCallback);
 		}
 
-		if (!IsImportMapIdentical(SourceLinker, DestLinker))
+		if (!IsImportMapIdentical(SourceLinker, DestLinker, MessageCallback))
 		{
-			DumpTableDifferences<FObjectImport>(SourceLinker, DestLinker, SourceLinker->ImportMap, DestLinker->ImportMap, *AssetFilename, TEXT("Import"), MaxDiffsToLog);
+			DumpTableDifferences<FObjectImport>(SourceLinker, DestLinker, SourceLinker->ImportMap, DestLinker->ImportMap,
+				*AssetFilename, TEXT("Import"), MaxDiffsToLog, MessageCallback);
 		}
 
-		if (!IsExportMapIdentical(SourceLinker, DestLinker))
+		if (!IsExportMapIdentical(SourceLinker, DestLinker, MessageCallback))
 		{
-			DumpTableDifferences<FObjectExport>(SourceLinker, DestLinker, SourceLinker->ExportMap, DestLinker->ExportMap, *AssetFilename, TEXT("Export"), MaxDiffsToLog);
+			DumpTableDifferences<FObjectExport>(SourceLinker, DestLinker, SourceLinker->ExportMap, DestLinker->ExportMap,
+				*AssetFilename, TEXT("Export"), MaxDiffsToLog, MessageCallback);
 		}
 	}
 
@@ -1528,7 +1523,8 @@ static void DumpPackageHeaderDiffs_ZenPackage(
 	const FDiffWriterArchiveWriter::FPackageData& SourcePackage,
 	const FDiffWriterArchiveWriter::FPackageData& DestPackage,
 	const FString& AssetFilename,
-	const int32 MaxDiffsToLog)
+	const int32 MaxDiffsToLog,
+	const UE::DiffWriterArchive::FMessageCallback& MessageCallback)
 {
 	// TODO: Fill in detailed diffing of Zen Package Summary
 }
@@ -1538,15 +1534,16 @@ void FDiffWriterArchiveWriter::DumpPackageHeaderDiffs(
 	const FPackageData& DestPackage,
 	const FString& AssetFilename,
 	const int32 MaxDiffsToLog,
-	const EPackageHeaderFormat PackageHeaderFormat /* = EPackageHeaderFormat::PackageFileSummary */)
+	const EPackageHeaderFormat PackageHeaderFormat,
+	const UE::DiffWriterArchive::FMessageCallback& MessageCallback)
 {
 	switch (PackageHeaderFormat)
 	{
 	case EPackageHeaderFormat::PackageFileSummary:
-		DumpPackageHeaderDiffs_LinkerLoad(SourcePackage, DestPackage, AssetFilename, MaxDiffsToLog);
+		DumpPackageHeaderDiffs_LinkerLoad(SourcePackage, DestPackage, AssetFilename, MaxDiffsToLog, MessageCallback);
 		break;
 	case EPackageHeaderFormat::ZenPackageSummary:
-		DumpPackageHeaderDiffs_ZenPackage(SourcePackage, DestPackage, AssetFilename, MaxDiffsToLog);
+		DumpPackageHeaderDiffs_ZenPackage(SourcePackage, DestPackage, AssetFilename, MaxDiffsToLog, MessageCallback);
 		break;
 	default:
 		unimplemented();
diff --git a/Engine/Source/Editor/UnrealEd/Private/Cooker/DiffWriterArchive.h b/Engine/Source/Editor/UnrealEd/Private/Cooker/DiffWriterArchive.h
--- a/Engine/Source/Editor/UnrealEd/Private/Cooker/DiffWriterArchive.h
+++ b/Engine/Source/Editor/UnrealEd/Private/Cooker/DiffWriterArchive.h
@@ -7,12 +7,14 @@
 #include "Containers/Map.h"
 #include "Containers/UnrealString.h"
 #include "HAL/Platform.h"
+#include "Logging/LogVerbosity.h"
 #include "Serialization/Archive.h"
 #include "Serialization/ArchiveProxy.h"
 #include "Serialization/ArchiveStackTrace.h"
 #include "Serialization/LargeMemoryWriter.h"
 #include "Serialization/PackageWriter.h"
 #include "Templates/RefCounting.h"
+#include "Templates/Function.h"
 #include "Templates/UniquePtr.h"
 #include "UObject/NameTypes.h"
 
@@ -23,6 +25,15 @@ class UObject;
 struct FUObjectSerializeContext;
 
 
+namespace UE::DiffWriterArchive
+{
+
+extern const TCHAR* const IndentToken;
+extern const TCHAR* const NewLineToken;
+typedef TUniqueFunction<void(ELogVerbosity::Type, FStringView)> FMessageCallback;
+
+}
+
 struct FDiffWriterDiffInfo
 {
 	int64 Offset;
@@ -248,6 +259,7 @@ public:
 		const int64 MaxDiffsToLog,
 		int32& InOutDiffsLogged,
 		TMap<FName, FArchiveDiffStats>& OutStats,
+		const UE::DiffWriterArchive::FMessageCallback& MessageCallback,
 		bool bSuppressLogging = false);
 
 	/** Creates map with mismatching callstacks. */
@@ -264,7 +276,8 @@ public:
 		const FPackageData& DestPackage,
 		const FString& AssetFilename,
 		const int32 MaxDiffsToLog,
-		const EPackageHeaderFormat PackageHeaderFormat = EPackageHeaderFormat::PackageFileSummary);
+		const EPackageHeaderFormat PackageHeaderFormat,
+		const UE::DiffWriterArchive::FMessageCallback& MessageCallback);
 
 	/** Returns a new linker for loading the specified package. */
 	static FLinkerLoad* CreateLinkerForPackage(
@@ -316,7 +329,8 @@ class FDiffWriterArchive
 public:
 	using FPackageData = FDiffWriterArchiveWriter::FPackageData;
 
-	FDiffWriterArchive(UObject* InAsset, const TCHAR* InFilename, bool bInCollectCallstacks = true, const FDiffWriterDiffMap* InDiffMap = nullptr);
+	FDiffWriterArchive(UObject* InAsset, const TCHAR* InFilename, UE::DiffWriterArchive::FMessageCallback&& InMessageCallback,
+		bool bInCollectCallstacks = true, const FDiffWriterDiffMap* InDiffMap = nullptr);
 	virtual ~FDiffWriterArchive();
 
 	FDiffWriterCallstacks& GetCallstacks()
@@ -352,4 +366,5 @@ public:
 private:
 	FDiffWriterCallstacks Callstacks;
 	FDiffWriterArchiveWriter StackTraceWriter;
+	UE::DiffWriterArchive::FMessageCallback MessageCallback;
 };
diff --git a/Engine/Source/Editor/UnrealEd/Private/Tests/DiffWriterArchiveTests.cpp b/Engine/Source/Editor/UnrealEd/Private/Tests/DiffWriterArchiveTests.cpp
--- a/Engine/Source/Editor/UnrealEd/Private/Tests/DiffWriterArchiveTests.cpp
+++ b/Engine/Source/Editor/UnrealEd/Private/Tests/DiffWriterArchiveTests.cpp
@@ -167,7 +167,9 @@ bool FDiffWriterArchiveTestsBasic::RunTest(const FString& Parameters)
 			int32 DiffsLogged = 0;
 			TMap<FName, FArchiveDiffStats> DiffStats;
 			const bool bSuppressLogging = true; // Suppress warning logs when running tests on CI
-
+			UE::DiffWriterArchive::FMessageCallback MessageCallback =
+				[](ELogVerbosity::Type Verbosity, FStringView Message)
+			{};
 			FDiffWriterArchiveWriter::Compare(
 				InitialData,
 				NewData,
@@ -178,6 +180,7 @@ bool FDiffWriterArchiveTestsBasic::RunTest(const FString& Parameters)
 				MaxDiffsToLog,
 				DiffsLogged,
 				DiffStats,
+				MessageCallback,
 				bSuppressLogging);
 
 			TestTrueExpr(DiffStats[NAME_None].NumDiffs == 2);
diff --git a/Engine/Source/Runtime/CoreUObject/Private/Serialization/ArchiveStackTrace.cpp b/Engine/Source/Runtime/CoreUObject/Private/Serialization/ArchiveStackTrace.cpp
--- a/Engine/Source/Runtime/CoreUObject/Private/Serialization/ArchiveStackTrace.cpp
+++ b/Engine/Source/Runtime/CoreUObject/Private/Serialization/ArchiveStackTrace.cpp
@@ -171,6 +171,15 @@ void ForceKillPackageAndLinker(FLinkerLoad* Linker)
 	}
 }
 
+bool ShouldIgnoreDiff()
+{
+	return GIgnoreDiffManager.ShouldIgnoreDiff();
+}
+bool ShouldBypassDiff()
+{
+	return GIgnoreDiffManager.ShouldBypassDiff();
+}
+
 } // namespace UE::ArchiveStackTrace
 
 PRAGMA_DISABLE_DEPRECATION_WARNINGS;
diff --git a/Engine/Source/Runtime/CoreUObject/Public/Serialization/ArchiveStackTrace.h b/Engine/Source/Runtime/CoreUObject/Public/Serialization/ArchiveStackTrace.h
--- a/Engine/Source/Runtime/CoreUObject/Public/Serialization/ArchiveStackTrace.h
+++ b/Engine/Source/Runtime/CoreUObject/Public/Serialization/ArchiveStackTrace.h
@@ -77,6 +77,8 @@ struct FPackageData
 COREUOBJECT_API bool LoadPackageIntoMemory(const TCHAR* InFilename,
 	FPackageData& OutPackageData, TUniquePtr<uint8>& OutLoadedBytes);
 COREUOBJECT_API void ForceKillPackageAndLinker(FLinkerLoad* Linker);
+COREUOBJECT_API bool ShouldIgnoreDiff();
+COREUOBJECT_API bool ShouldBypassDiff();
 
 } // namespace UE::ArchiveStackTrace
 
