EditorDomain: Add -noeditordomainwrite to turn off writes. We already could turn off reads, but thought we wouldn't need to turn off writes. But now we do, since they're currently broken but we still need the not-currently-broken  EditorDomain package digests.
#rnx

diff --git a/Engine/Source/Editor/UnrealEd/Private/EditorDomain/EditorDomain.cpp b/Engine/Source/Editor/UnrealEd/Private/EditorDomain/EditorDomain.cpp
--- a/Engine/Source/Editor/UnrealEd/Private/EditorDomain/EditorDomain.cpp
+++ b/Engine/Source/Editor/UnrealEd/Private/EditorDomain/EditorDomain.cpp
@@ -118,6 +118,7 @@ FEditorDomain::FEditorDomain()
 
 	bEditorDomainReadEnabled = !FParse::Param(FCommandLine::Get(), TEXT("noeditordomainread"))
 	 && !FParse::Param(FCommandLine::Get(), TEXT("testeditordomaindeterminism"));
+	bEditorDomainWriteEnabled = !FParse::Param(FCommandLine::Get(), TEXT("noeditordomainwrite"));
 
 	ELoadingPhase::Type CurrentPhase = IPluginManager::Get().GetLastCompletedLoadingPhase();
 	if (CurrentPhase == ELoadingPhase::None || CurrentPhase < ELoadingPhase::PostEngineInit)
@@ -610,7 +611,7 @@ void FEditorDomain::Tick(float DeltaTime)
 
 void FEditorDomain::OnEndLoadPackage(const FEndLoadPackageContext& Context)
 {
-	if (bExternalSave)
+	if (bExternalSave || !bEditorDomainWriteEnabled)
 	{
 		return;
 	}
@@ -640,7 +641,7 @@ void FEditorDomain::OnPostEngineInit()
 	{
 		FScopeLock ScopeLock(&Locks->Lock);
 		bHasPassedPostEngineInit = true;
-		if (bExternalSave)
+		if (bExternalSave || !bEditorDomainWriteEnabled)
 		{
 			return;
 		}
diff --git a/Engine/Source/Editor/UnrealEd/Private/EditorDomain/EditorDomain.h b/Engine/Source/Editor/UnrealEd/Private/EditorDomain/EditorDomain.h
--- a/Engine/Source/Editor/UnrealEd/Private/EditorDomain/EditorDomain.h
+++ b/Engine/Source/Editor/UnrealEd/Private/EditorDomain/EditorDomain.h
@@ -267,6 +267,8 @@ private:
 
 	/** True by default, set to false when reading is disabled for testing. */
 	bool bEditorDomainReadEnabled = true;
+	/** True by default, set to false when writing is disabled for testing. */
+	bool bEditorDomainWriteEnabled = true;
 	/** If true, use an out-of-process EditorDomainSaveServer for saves, else save in process in EndLoad */
 	bool bExternalSave = false;
 	/** Marker for whether our PostEngineInit callback has been called */
