Cooking: Iterative: Add -iterativevalidate mode that diffs iteratively unmodified previous packages against their resaved version.
#rnx

diff --git a/Engine/Source/Developer/IoStoreUtilities/Private/ZenStoreWriter.cpp b/Engine/Source/Developer/IoStoreUtilities/Private/ZenStoreWriter.cpp
--- a/Engine/Source/Developer/IoStoreUtilities/Private/ZenStoreWriter.cpp
+++ b/Engine/Source/Developer/IoStoreUtilities/Private/ZenStoreWriter.cpp
@@ -1210,32 +1210,28 @@ void FZenStoreWriter::RemoveCookedPackages()
 	PackageNameToIndex.Empty();
 }
 
-void FZenStoreWriter::MarkPackagesUpToDate(TArrayView<const FName> UpToDatePackages)
+void FZenStoreWriter::UpdatePackageModificationStatus(FName PackageName, bool bIterativelyUnmodified,
+	bool& bInOutShouldIterativelySkip)
 {
-	TRACE_CPUPROFILER_EVENT_SCOPE(FZenStoreWriter::MarkPackagesUpToDate);
+	TRACE_CPUPROFILER_EVENT_SCOPE(FZenStoreWriter::UpdatePackageModificationStatus);
 
 	IPackageStoreWriter::FMarkUpToDateEventArgs MarkUpToDateEventArgs;
 
-	MarkUpToDateEventArgs.PackageIndexes.Reserve(UpToDatePackages.Num());
-
 	{
 		FWriteScopeLock _(EntriesLock);
-		for (FName PackageName : UpToDatePackages)
+		int32* Index = PackageNameToIndex.Find(PackageName);
+		if (!Index)
 		{
-			int32* Index = PackageNameToIndex.Find(PackageName);
-			if (!Index)
+			if (!FPackageName::IsScriptPackage(WriteToString<128>(PackageName)))
 			{
-				if (!FPackageName::IsScriptPackage(WriteToString<128>(PackageName)))
-				{
-					UE_LOG(LogZenStoreWriter, Warning, TEXT("MarkPackagesUpToDate called with package %s that is not in the oplog."),
-						*PackageName.ToString());
-				}
-				continue;
+				UE_LOG(LogZenStoreWriter, Verbose, TEXT("UpdatePackageModificationStatus called with package %s that is not in the oplog."),
+					*PackageName.ToString());
 			}
-
-			MarkUpToDateEventArgs.PackageIndexes.Add(*Index);
-			CookedPackagesInfo[*Index].bUpToDate = true;
+			return;
 		}
+
+		MarkUpToDateEventArgs.PackageIndexes.Add(*Index);
+		CookedPackagesInfo[*Index].bUpToDate = true;
 	}
 	if (MarkUpToDateEventArgs.PackageIndexes.Num())
 	{
diff --git a/Engine/Source/Developer/IoStoreUtilities/Public/ZenStoreWriter.h b/Engine/Source/Developer/IoStoreUtilities/Public/ZenStoreWriter.h
--- a/Engine/Source/Developer/IoStoreUtilities/Public/ZenStoreWriter.h
+++ b/Engine/Source/Developer/IoStoreUtilities/Public/ZenStoreWriter.h
@@ -129,7 +129,8 @@ public:
 	IOSTOREUTILITIES_API virtual FCbObject GetOplogAttachment(FName PackageName, FUtf8StringView AttachmentKey) override;
 	IOSTOREUTILITIES_API virtual void RemoveCookedPackages(TArrayView<const FName> PackageNamesToRemove) override;
 	IOSTOREUTILITIES_API virtual void RemoveCookedPackages() override;
-	IOSTOREUTILITIES_API virtual void MarkPackagesUpToDate(TArrayView<const FName> UpToDatePackages) override;
+	IOSTOREUTILITIES_API virtual void UpdatePackageModificationStatus(FName PackageName, bool bIterativelyUnmodified,
+		bool& bInOutShouldIterativelySkip) override;
 	IOSTOREUTILITIES_API bool GetPreviousCookedBytes(const FPackageInfo& Info, FPreviousCookedBytesData& OutData) override;
 	IOSTOREUTILITIES_API void CompleteExportsArchiveForDiff(FPackageInfo& Info, FLargeMemoryWriter& ExportsArchive) override;
 	IOSTOREUTILITIES_API virtual EPackageWriterResult BeginCacheForCookedPlatformData(FBeginCacheForCookedPlatformDataInfo& Info) override;
diff --git a/Engine/Source/Editor/UnrealEd/Private/CookOnTheFlyServer.cpp b/Engine/Source/Editor/UnrealEd/Private/CookOnTheFlyServer.cpp
--- a/Engine/Source/Editor/UnrealEd/Private/CookOnTheFlyServer.cpp
+++ b/Engine/Source/Editor/UnrealEd/Private/CookOnTheFlyServer.cpp
@@ -34,6 +34,7 @@
 #include "Cooker/CookWorkerClient.h"
 #include "Cooker/DiffPackageWriter.h"
 #include "Cooker/IoStoreCookOnTheFlyRequestManager.h"
+#include "Cooker/IterativeValidatePackageWriter.h"
 #include "Cooker/LooseCookedPackageWriter.h"
 #include "Cooker/NetworkFileCookOnTheFlyRequestManager.h"
 #include "Cooker/PackageTracker.h"
@@ -2628,6 +2629,10 @@ void UCookOnTheFlyServer::DemoteToIdle(UE::Cook::FPackageData& PackageData, UE::
 
 			// Reachability: Suppress the diagnostic that were found via cookload reference traversal but are not reachable on the target platforms
 			bPrintDiagnostic &= (PackageData.HasInstigator() || Reason != ESuppressCookReason::OnlyEditorOnly);
+
+			// Iterative cooks: suppress the diagnostic for packages that were iteratively skipped
+			bPrintDiagnostic &= !PackageData.HasAllCookedPlatforms(PlatformManager->GetSessionPlatforms(), true /* bIncludeFailed */);
+
 			if (bPrintDiagnostic)
 			{
 				UE_CLOG((GCookProgressDisplay & (int32)ECookProgressDisplayMode::Instigators), LogCook, Display,
@@ -5678,7 +5683,7 @@ public:
 	void InitializePackageWriter(ICookedPackageWriter*& CookedPackageWriter)
 	{
 		Initialize();
-		if (!bDiffEnabled && !bLinkerDiffEnabled)
+		if (!bAnyDiffModeEnabled)
 		{
 			return;
 		}
@@ -5686,9 +5691,9 @@ public:
 		ICookedPackageWriter::FCookCapabilities Capabilities = CookedPackageWriter->GetCookCapabilities();
 		if (!Capabilities.bDiffModeSupported)
 		{
-			const TCHAR* CommandLineArg = bDiffEnabled ? TEXT("-DIFFONLY") : TEXT("-LINKERDIFF");
-			UE_LOG(LogCook, Fatal, TEXT("%s was enabled, but -iostore is also enabled and iostore PackageWriters do not support %s."),
-				CommandLineArg, CommandLineArg);
+			// All current PackageWriters support bDiffModeSupported; log a fatal error in case a new one is added.
+			UE_LOG(LogCook, Fatal, 
+				TEXT("A DiffMode was enabled, but the current PackageWriter has bDiffModeSupported=false."));
 		}
 
 		if (bDiffEnabled)
@@ -5696,11 +5701,15 @@ public:
 			// Wrap the incoming writer inside a FDiffPackageWriter
 			CookedPackageWriter = new FDiffPackageWriter(TUniquePtr<ICookedPackageWriter>(CookedPackageWriter));
 		}
-		else
+		else if (bLinkerDiffEnabled)
 		{
-			check(bLinkerDiffEnabled);
 			CookedPackageWriter = new FLinkerDiffPackageWriter(TUniquePtr<ICookedPackageWriter>(CookedPackageWriter));
 		}
+		else
+		{
+			check(bIterativeValidateEnabled);
+			CookedPackageWriter = new FIterativeValidatePackageWriter(TUniquePtr<ICookedPackageWriter>(CookedPackageWriter));
+		}
 	}
 
 private:
@@ -5711,20 +5720,26 @@ private:
 			return;
 		}
 
-		bDiffEnabled = FParse::Param(FCommandLine::Get(), TEXT("DIFFONLY"));
-		FString Value;
-		bLinkerDiffEnabled = FParse::Value(FCommandLine::Get(), TEXT("-LINKERDIFF="), Value);
-		if (bDiffEnabled && bLinkerDiffEnabled)
+		const TCHAR* CommandLine = FCommandLine::Get();
+		bDiffEnabled = FParse::Param(CommandLine, TEXT("DIFFONLY"));
+		bLinkerDiffEnabled = false;
+		FParse::Bool(CommandLine, TEXT("-LINKERDIFF="), bLinkerDiffEnabled);
+		bIterativeValidateEnabled = FParse::Param(CommandLine, TEXT("ITERATIVEVALIDATE"));
+		int32 NumEnabled = (bDiffEnabled ? 1 : 0) + (bLinkerDiffEnabled ? 1 : 0) + (bIterativeValidateEnabled ? 1 : 0);
+		if (NumEnabled > 1)
 		{
-			UE_LOG(LogCook, Fatal, TEXT("-DiffOnly and -LinkerDiff are mutually exclusive."));
+			UE_LOG(LogCook, Fatal, TEXT("-DiffOnly, -LinkerDiff, and -IterativeValidate are mutually exclusive."));
 		}
+		bAnyDiffModeEnabled = NumEnabled >= 1;
 
 		bInitialized = true;
 	}
 
 	bool bInitialized = false;
+	bool bAnyDiffModeEnabled = false;
 	bool bDiffEnabled = false;
 	bool bLinkerDiffEnabled = false;
+	bool bIterativeValidateEnabled = false;
 };
 
 #if OUTPUT_COOKTIMING
@@ -7971,39 +7986,63 @@ void UCookOnTheFlyServer::PopulateCookedPackages(TArrayView<const ITargetPlatfor
 			int32 ModifiedCookedNum = 0;
 			int32 RemovedCookedNum = 0;
 
-			auto AddCookedPackage =
-				[this, TargetPlatform](const FName PackageName, ECookResult CookResult, bool bRequireExists)
+			auto AddPlaceholderPackage =
+				[this, TargetPlatform](const FName PackageName, ECookResult CookResult)
 				{
-					FPackageData* PackageData = PackageDatas->TryAddPackageDataByPackageName(PackageName, bRequireExists);
+					FPackageData* PackageData = PackageDatas->TryAddPackageDataByPackageName(PackageName, true /* bRequireExists */);
 					if (PackageData)
 					{
 						PackageData->SetPlatformCooked(TargetPlatform, CookResult);
-						PackageData->FindOrAddPlatformData(TargetPlatform).SetIterativelySkipped(true);
 					}
 				};
-			auto AddIdenticalCookedPackage =
-				[&AddCookedPackage, &IdenticalCooked](const FName PackageName, bool bRequireExists)
-				{
-					IdenticalCooked.Add(PackageName);
-					AddCookedPackage(PackageName, ECookResult::Succeeded, bRequireExists);
-					// Declare the package to the EDLCookInfo verification so we don't warn about missing exports from it
-					UE::SavePackageUtilities::EDLCookInfoAddIterativelySkippedPackage(PackageName);
-				};
 			bool bCookByTheBook = IsCookByTheBookMode();
-			auto AddRequestForModifiedPackage = [this, TargetPlatform, bCookByTheBook](const FName PackageName)
+			auto UpdateCookedPackage = [this, TargetPlatform, bCookByTheBook, &Difference, &IdenticalCookedNum,
+				&ModifiedCookedNum, &PackageWriter, &bFirstPlatform, &PackagesToRemove]
+			(const FName PackageName, bool bRequireExists, bool bIterativelyUnmodified)
 			{
-				if (bCookByTheBook)
+				++(bIterativelyUnmodified ? IdenticalCookedNum : ModifiedCookedNum);
+				FPackageData* PackageData = PackageDatas->TryAddPackageDataByPackageName(PackageName, bRequireExists);
+				if (PackageData)
 				{
-					// cook on the fly will queue packages when it needs them, but for cook by the book we force cook the modified files
-					// so that the output set of packages is up to date (even if the user is currently cooking only a subset)
-					FPackageData* PackageData = PackageDatas->TryAddPackageDataByPackageName(PackageName);
-					if (PackageData)
+					if (bIterativelyUnmodified)
+					{
+						PackageData->FindOrAddPlatformData(TargetPlatform).SetIterativelyUnmodified(true);
+					}
+					bool bShouldIterativelySkip = bIterativelyUnmodified;
+					PackageWriter.UpdatePackageModificationStatus(PackageName, bIterativelyUnmodified,
+						bShouldIterativelySkip);
+					if (bShouldIterativelySkip && !bIterativelyUnmodified)
+					{
+						// Override the PackageWriter's request to iteratively skip a modified generator package, because we
+						// need to cook the generator packages to evaluate whether their generated packages should be skipped.
+						EDifference* GeneratorDifference = Difference.Packages.Find(PackageName);
+						if (GeneratorDifference)
+						{
+							bShouldIterativelySkip = false;
+						}
+					}
+					if (bShouldIterativelySkip)
 					{
-						WorkerRequests->AddStartCookByTheBookRequest(FFilePlatformRequest(PackageData->GetFileName(),
-							EInstigator::IterativeCook, TConstArrayView<const ITargetPlatform*>{ TargetPlatform }));
+						PackageData->SetPlatformCooked(TargetPlatform, ECookResult::Succeeded);
+						if (bFirstPlatform)
+						{
+							COOK_STAT(++DetailedCookStats::NumPackagesIterativelySkipped);
+						}
+						// Declare the package to the EDLCookInfo verification so we don't warn about missing exports from it
+						UE::SavePackageUtilities::EDLCookInfoAddIterativelySkippedPackage(PackageName);
+					}
+					else
+					{
+						if (bCookByTheBook)
+						{
+							// cook on the fly will queue packages when it needs them, but for cook by the book we force cook the modified files
+							// so that the output set of packages is up to date (even if the user is currently cooking only a subset)
+							WorkerRequests->AddStartCookByTheBookRequest(FFilePlatformRequest(PackageData->GetFileName(),
+								EInstigator::IterativeCook, TConstArrayView<const ITargetPlatform*>{ TargetPlatform }));
+						}
+						PackagesToRemove.Add(PackageName);
 					}
 				}
-
 			};
 
 			// Add CookedPackages for any identical packages, delete from disk any modified packages
@@ -8014,19 +8053,17 @@ void UCookOnTheFlyServer::PopulateCookedPackages(TArrayView<const ITargetPlatfor
 				switch (Pair.Value)
 				{
 				case EDifference::IdenticalCooked:
-					AddIdenticalCookedPackage(PackageName, true /* bRequireExists */);
+					UpdateCookedPackage(PackageName, true /* bRequireExists */, true /* bIterativelyUnmodified */);
 					break;
 				case EDifference::ModifiedCooked:
-					AddRequestForModifiedPackage(PackageName);
-					PackagesToRemove.Add(PackageName);
-					++ModifiedCookedNum;
+					UpdateCookedPackage(PackageName, true /* bRequireExists */, false /* bIterativelyUnmodified */);
 					break;
 				case EDifference::RemovedCooked:
 					PackagesToRemove.Add(PackageName);
 					++RemovedCookedNum;
 					break;
 				case EDifference::IdenticalUncooked:
-					AddCookedPackage(PackageName, ECookResult::Failed, true /* bRequireExists */);
+					AddPlaceholderPackage(PackageName, ECookResult::Failed);
 					PackagesToRemove.Add(PackageName);
 					break;
 				case EDifference::ModifiedUncooked:
@@ -8036,7 +8073,7 @@ void UCookOnTheFlyServer::PopulateCookedPackages(TArrayView<const ITargetPlatfor
 					PackagesToRemove.Add(PackageName);
 					break;
 				case EDifference::IdenticalNeverCookPlaceholder:
-					AddCookedPackage(PackageName, ECookResult::NeverCookPlaceholder, true /* bRequireExists */);
+					AddPlaceholderPackage(PackageName, ECookResult::NeverCookPlaceholder);
 					PackagesToRemove.Add(PackageName);
 					break;
 				case EDifference::ModifiedNeverCookPlaceholder:
@@ -8057,12 +8094,12 @@ void UCookOnTheFlyServer::PopulateCookedPackages(TArrayView<const ITargetPlatfor
 				Iter; ++Iter)
 			{
 				FName Generator = Iter->Key;
-				EDifference* GeneratorDifference = Difference.Packages.Find(Generator);
-				if (GeneratorDifference && *GeneratorDifference == EDifference::IdenticalCooked)
+				FPackageData* PackageData = PackageDatas->TryAddPackageDataByPackageName(Generator, false /* bRequireExists */);
+				if (PackageData && PackageData->FindOrAddPlatformData(TargetPlatform).IsCookAttempted())
 				{
 					for (const TPair<FName, FGuid>& GeneratedPair : Iter->Value.Generated)
 					{
-						AddIdenticalCookedPackage(GeneratedPair.Key, false /* bRequireExists */);
+						UpdateCookedPackage(GeneratedPair.Key, false /* bRequireExists */, true /* bIterativelyUnmodified */);
 					}
 					Iter.RemoveCurrent();
 				}
@@ -8075,14 +8112,9 @@ void UCookOnTheFlyServer::PopulateCookedPackages(TArrayView<const ITargetPlatfor
 			int32 CookedNum = IdenticalCooked.Num() + ModifiedCookedNum + RemovedCookedNum + DeferredEvaluateGeneratedNum;
 			UE_LOG(LogCook, Display, TEXT("Found %d cooked package(s) in package store."), CookedNum);
 			UE_LOG(LogCook, Display, TEXT("Keeping %d. Recooking %d. Removing %d. %d generated packages to be evaluated for iterative skipping later."),
-				IdenticalCooked.Num(), ModifiedCookedNum, RemovedCookedNum, DeferredEvaluateGeneratedNum);
-			if (bFirstPlatform)
-			{
-				COOK_STAT(DetailedCookStats::NumPackagesIterativelySkipped += IdenticalCooked.Num());
-				bFirstPlatform = false;
-			}
+				IdenticalCookedNum, ModifiedCookedNum, RemovedCookedNum, DeferredEvaluateGeneratedNum);
+			bFirstPlatform = false;
 
-			PackageWriter.MarkPackagesUpToDate(IdenticalCooked);
 			PackageWriter.RemoveCookedPackages(PackagesToRemove);
 		}
 
@@ -10237,7 +10269,7 @@ void UCookOnTheFlyServer::LoadBeginCookIterativeFlagsLocal(FBeginCookContext& Be
 		const ITargetPlatform* TargetPlatform = PlatformContext.TargetPlatform;
 		UE::Cook::FPlatformData* PlatformData = PlatformContext.PlatformData;
 		const ICookedPackageWriter* PackageWriterPtr = FindPackageWriter(TargetPlatform);
-		check(PackageWriterPtr); // PackageContexts should have been created by SelectSessoinPlatforms or by FindOrCreateSaveContexts in AddCookOnTheFlyPlatformFromGameThread
+		check(PackageWriterPtr); // PackageContexts should have been created by SelectSessionPlatforms or by FindOrCreateSaveContexts in AddCookOnTheFlyPlatformFromGameThread
 		const ICookedPackageWriter& PackageWriter(*PackageWriterPtr);
 		bool bIterateSharedBuild = false;
 		if (bIterative && bIsSharedIterativeCook && !PlatformData->bIsSandboxInitialized)
diff --git a/Engine/Source/Editor/UnrealEd/Private/Cooker/CookPackageData.cpp b/Engine/Source/Editor/UnrealEd/Private/Cooker/CookPackageData.cpp
--- a/Engine/Source/Editor/UnrealEd/Private/Cooker/CookPackageData.cpp
+++ b/Engine/Source/Editor/UnrealEd/Private/Cooker/CookPackageData.cpp
@@ -55,7 +55,7 @@ static FAutoConsoleVariableRef CVarPollAsyncPeriod(
 // FPackageData
 FPackagePlatformData::FPackagePlatformData()
 	: bReachable(0), bVisitedByCluster(0), bSaveTimedOut(0), bCookable(1), bExplorable(1), bExplorableOverride(0)
-	, bIterativelySkipped(0), bRegisteredForCachedObjectsInOuter(0), CookResults((uint8)ECookResult::NotAttempted)
+	, bIterativelyUnmodified(0), bRegisteredForCachedObjectsInOuter(0), CookResults((uint8)ECookResult::NotAttempted)
 {
 }
 
@@ -1590,10 +1590,10 @@ bool FGeneratorPackage::TryGenerateList(UObject* OwnerObject, FPackageDatas& Pac
 	OwnerPackageData.GetPlatformsNeedingCooking(PlatformsToCook);
 	bool bHybridIterativeEnabled = COTFS.bHybridIterativeEnabled;
 
+	int32 NumIterativeUnmodified = 0;
 	int32 NumIterativeModified = 0; 
 	int32 NumIterativeRemoved = 0;
 	int32 NumIterativePrevious = PreviousGeneratedPackages.Num();
-	TArray<FName> IdenticalGenerated;
 
 	for (ICookPackageSplitter::FGeneratedPackage& SplitterData : GeneratorDatas)
 	{
@@ -1657,16 +1657,9 @@ bool FGeneratorPackage::TryGenerateList(UObject* OwnerObject, FPackageDatas& Pac
 		FGuid PreviousGuid;
 		if (PreviousGeneratedPackages.RemoveAndCopyValue(PackageFName, PreviousGuid) && !bHybridIterativeEnabled)
 		{
-			bool bIdentical;
-			GeneratedInfo.IterativeCookValidateOrClear(*this, PlatformsToCook, PreviousGuid, bIdentical);
-			if (bIdentical)
-			{
-				IdenticalGenerated.Add(PackageFName);
-			}
-			else
-			{
-				++NumIterativeModified;
-			}
+			bool bIterativelyUnmodified;
+			GeneratedInfo.IterativeCookValidateOrClear(*this, PlatformsToCook, PreviousGuid, bIterativelyUnmodified);
+			++(bIterativelyUnmodified ? NumIterativeUnmodified : NumIterativeModified);
 		}
 	}
 	if (!PreviousGeneratedPackages.IsEmpty())
@@ -1686,13 +1679,7 @@ bool FGeneratorPackage::TryGenerateList(UObject* OwnerObject, FPackageDatas& Pac
 		UE_LOG(LogCook, Display, TEXT("Found %d cooked package(s) in package store for generator package %s."),
 			NumIterativePrevious, *WriteToString<256>(GetOwner().GetPackageName()));
 		UE_LOG(LogCook, Display, TEXT("Keeping %d. Recooking %d. Removing %d."),
-			IdenticalGenerated.Num(), NumIterativeModified, NumIterativeRemoved);
-		COOK_STAT(DetailedCookStats::NumPackagesIterativelySkipped += IdenticalGenerated.Num());
-
-		for (const ITargetPlatform* TargetPlatform : PlatformsToCook)
-		{
-			COTFS.FindOrCreatePackageWriter(TargetPlatform).MarkPackagesUpToDate(IdenticalGenerated);
-		}
+			NumIterativeUnmodified, NumIterativeModified, NumIterativeRemoved);
 	}
 
 	RemainingToPopulate = GeneratorDatas.Num() + 1; // GeneratedPackaged plus one for the Generator
@@ -2235,11 +2222,11 @@ void FCookGenerationInfo::CreateGuid()
 }
 
 void FCookGenerationInfo::IterativeCookValidateOrClear(FGeneratorPackage& Generator,
-	TConstArrayView<const ITargetPlatform*> RequestedPlatforms, const FGuid& PreviousGuid, bool& bOutIdentical)
+	TConstArrayView<const ITargetPlatform*> RequestedPlatforms, const FGuid& PreviousGuid, bool& bOutIterativelyUnmodified)
 {
 	UCookOnTheFlyServer& COTFS = Generator.GetOwner().GetPackageDatas().GetCookOnTheFlyServer();
-	bOutIdentical = PreviousGuid == this->Guid;
-	if (bOutIdentical)
+	bOutIterativelyUnmodified = PreviousGuid == this->Guid;
+	if (bOutIterativelyUnmodified)
 	{
 		// If not directly modified, mark it as indirectly modified if any of its dependencies
 		// were detected as modified during PopulateCookedPackages.
@@ -2248,39 +2235,51 @@ void FCookGenerationInfo::IterativeCookValidateOrClear(FGeneratorPackage& Genera
 			FPackageData* DependencyData = COTFS.PackageDatas->FindPackageDataByPackageName(Dependency.AssetId.PackageName);
 			if (!DependencyData)
 			{
-				bOutIdentical = false;
+				bOutIterativelyUnmodified = false;
 				break;
 			}
 			for (const ITargetPlatform* TargetPlatform : RequestedPlatforms)
 			{
 				FPackagePlatformData* DependencyPlatformData = DependencyData->FindPlatformData(TargetPlatform);
-				if (!DependencyPlatformData || !DependencyPlatformData->IsIterativelySkipped())
+				if (!DependencyPlatformData || !DependencyPlatformData->IsIterativelyUnmodified())
 				{
-					bOutIdentical = false;
+					bOutIterativelyUnmodified = false;
 					break;
 				}
 			}
-			if (!bOutIdentical)
+			if (!bOutIterativelyUnmodified)
 			{
 				break;
 			}
 		}
 	}
 
-	if (bOutIdentical)
+	bool bFirstPlatform = true;
+	for (const ITargetPlatform* TargetPlatform : RequestedPlatforms)
 	{
-		for (const ITargetPlatform* TargetPlatform : RequestedPlatforms)
+		if (bOutIterativelyUnmodified)
+		{
+			PackageData->FindOrAddPlatformData(TargetPlatform).SetIterativelyUnmodified(true);
+		}
+		bool bShouldIterativelySkip = bOutIterativelyUnmodified;
+		ICookedPackageWriter& PackageWriter = COTFS.FindOrCreatePackageWriter(TargetPlatform);
+		PackageWriter.UpdatePackageModificationStatus(PackageData->GetPackageName(), bOutIterativelyUnmodified,
+			bShouldIterativelySkip);
+		if (bShouldIterativelySkip)
 		{
 			PackageData->SetPlatformCooked(TargetPlatform, ECookResult::Succeeded);
-			PackageData->FindOrAddPlatformData(TargetPlatform).SetIterativelySkipped(true);
+			if (bFirstPlatform)
+			{
+				COOK_STAT(++DetailedCookStats::NumPackagesIterativelySkipped);
+			}
+			// Declare the package to the EDLCookInfo verification so we don't warn about missing exports from it
+			UE::SavePackageUtilities::EDLCookInfoAddIterativelySkippedPackage(PackageData->GetPackageName());
 		}
-	}
-	else
-	{
-		for (const ITargetPlatform* TargetPlatform : RequestedPlatforms)
+		else
 		{
 			COTFS.DeleteOutputForPackage(PackageData->GetPackageName(), TargetPlatform);
 		}
+		bFirstPlatform = false;
 	}
 }
 
diff --git a/Engine/Source/Editor/UnrealEd/Private/Cooker/CookPackageData.h b/Engine/Source/Editor/UnrealEd/Private/Cooker/CookPackageData.h
--- a/Engine/Source/Editor/UnrealEd/Private/Cooker/CookPackageData.h
+++ b/Engine/Source/Editor/UnrealEd/Private/Cooker/CookPackageData.h
@@ -175,8 +175,8 @@ struct FPackagePlatformData
 	void MarkCookableForWorker(FCookWorkerClient& CookWorkerClient);
 
 	/** The package was found to be unmodified in the current iterative cook. */
-	bool IsIterativelySkipped() const { return bIterativelySkipped != 0; }
-	void SetIterativelySkipped(bool bValue) { bIterativelySkipped = (uint32)bValue; }
+	bool IsIterativelyUnmodified() const { return bIterativelyUnmodified != 0; }
+	void SetIterativelyUnmodified(bool bValue) { bIterativelyUnmodified = (uint32)bValue; }
 
 	ECookResult GetCookResults() const { return (ECookResult)CookResults; }
 	bool IsCookAttempted() const { return CookResults != (uint32)ECookResult::NotAttempted; }
@@ -196,7 +196,7 @@ private:
 	uint32 bCookable : 1;
 	uint32 bExplorable : 1;
 	uint32 bExplorableOverride : 1;
-	uint32 bIterativelySkipped : 1;
+	uint32 bIterativelyUnmodified : 1;
 	uint32 bRegisteredForCachedObjectsInOuter : 1;
 	uint32 CookResults : (int)ECookResult::NumBits;
 };
@@ -814,7 +814,7 @@ public:
 	 * incremental cooks handle invalidation by querying the TargetDomainDigest during the RequestCluster.
 	 */
 	void IterativeCookValidateOrClear(FGeneratorPackage& Generator,
-		TConstArrayView<const ITargetPlatform*> RequestedPlatforms, const FGuid& PreviousGuid, bool& bOutIdentical);
+		TConstArrayView<const ITargetPlatform*> RequestedPlatforms, const FGuid& PreviousGuid, bool& bOutIterativelyUnmodified);
 
 	TConstArrayView<FAssetDependency> GetDependencies() const { return PackageDependencies; }
 
diff --git a/Engine/Source/Editor/UnrealEd/Private/Cooker/CookRequestCluster.cpp b/Engine/Source/Editor/UnrealEd/Private/Cooker/CookRequestCluster.cpp
--- a/Engine/Source/Editor/UnrealEd/Private/Cooker/CookRequestCluster.cpp
+++ b/Engine/Source/Editor/UnrealEd/Private/Cooker/CookRequestCluster.cpp
@@ -1095,21 +1095,16 @@ void FRequestCluster::FGraphSearch::ExploreVertexEdges(FVertexData& Vertex)
 
 			const FCookAttachments& PlatformAttachments = QueryPlatformData.CookAttachments;
 			bool bFoundBuildDefinitions = false;
+			bool bIterativelyUnmodified = false;
+			ICookedPackageWriter* PackageWriter = FetchPlatformData.Writer;
 			if (IsCookAttachmentsValid(PackageName, PlatformAttachments))
 			{
-				ICookedPackageWriter* PackageWriter = FetchPlatformData.Writer;
 				if (!Cluster.bFullBuild && Cluster.COTFS.bHybridIterativeEnabled)
 				{
 					if (IsIterativeEnabled(PackageName, Cluster.COTFS.bHybridIterativeAllowAllClasses))
 					{
-						if (PlatformIndex == FirstSessionPlatformIndex)
-						{
-							COOK_STAT(++DetailedCookStats::NumPackagesIterativelySkipped);
-						}
-						PackageData.SetPlatformCooked(TargetPlatform, ECookResult::Succeeded);
-						PackageWriter->MarkPackagesUpToDate({ PackageName });
-						// Declare the package to the EDLCookInfo verification so we don't warn about missing exports from it
-						UE::SavePackageUtilities::EDLCookInfoAddIterativelySkippedPackage(PackageName);
+						bIterativelyUnmodified = true;
+						PackageData.FindOrAddPlatformData(TargetPlatform).SetIterativelyUnmodified(true);
 					}
 					AddPlatformDependencyRange(PlatformAttachments.BuildDependencies, PlatformIndex,
 						true /* bHardDependency */);
@@ -1127,6 +1122,19 @@ void FRequestCluster::FGraphSearch::ExploreVertexEdges(FVertexData& Vertex)
 					}
 				}
 			}
+			bool bShouldIterativelySkip = bIterativelyUnmodified;
+			PackageWriter->UpdatePackageModificationStatus(PackageName, bIterativelyUnmodified, bShouldIterativelySkip);
+			if (bShouldIterativelySkip)
+			{
+				PackageData.SetPlatformCooked(TargetPlatform, ECookResult::Succeeded);
+				if (PlatformIndex == FirstSessionPlatformIndex)
+				{
+					COOK_STAT(++DetailedCookStats::NumPackagesIterativelySkipped);
+				}
+				// Declare the package to the EDLCookInfo verification so we don't warn about missing exports from it
+				UE::SavePackageUtilities::EDLCookInfoAddIterativelySkippedPackage(PackageName);
+			}
+
 			if (Cluster.bPreQueueBuildDefinitions && !bFoundBuildDefinitions)
 			{
 				if (PlatformAgnosticQueryPlatformData.bActive &&
diff --git a/Engine/Source/Editor/UnrealEd/Private/Cooker/DiffPackageWriter.h b/Engine/Source/Editor/UnrealEd/Private/Cooker/DiffPackageWriter.h
--- a/Engine/Source/Editor/UnrealEd/Private/Cooker/DiffPackageWriter.h
+++ b/Engine/Source/Editor/UnrealEd/Private/Cooker/DiffPackageWriter.h
@@ -91,9 +91,10 @@ public:
 	{
 		Inner->RemoveCookedPackages();
 	}
-	virtual void MarkPackagesUpToDate(TArrayView<const FName> UpToDatePackages) override
+	virtual void UpdatePackageModificationStatus(FName PackageName, bool bIterativelyUnmodified,
+		bool& bInOutShouldIterativelySkip) override
 	{
-		Inner->MarkPackagesUpToDate(UpToDatePackages);
+		Inner->UpdatePackageModificationStatus(PackageName, bIterativelyUnmodified, bInOutShouldIterativelySkip);
 	}
 	virtual EPackageWriterResult BeginCacheForCookedPlatformData(FBeginCacheForCookedPlatformDataInfo& Info) override
 	{
@@ -113,7 +114,8 @@ public:
 	{
 		return Inner->GetPackageHashes();
 	}
-private:
+
+protected:
 	void ParseCmds();
 	void ParseDumpObjList(FString InParams);
 	void ParseDumpObjects(FString InParams);
@@ -236,9 +238,10 @@ public:
 	{
 		Inner->RemoveCookedPackages();
 	}
-	virtual void MarkPackagesUpToDate(TArrayView<const FName> UpToDatePackages) override
+	virtual void UpdatePackageModificationStatus(FName PackageName, bool bIterativelyUnmodified,
+		bool& bInOutShouldIterativelySkip) override
 	{
-		Inner->MarkPackagesUpToDate(UpToDatePackages);
+		Inner->UpdatePackageModificationStatus(PackageName, bIterativelyUnmodified, bInOutShouldIterativelySkip);
 	}
 	virtual EPackageWriterResult BeginCacheForCookedPlatformData(FBeginCacheForCookedPlatformDataInfo& Info) override
 	{
diff --git a/Engine/Source/Editor/UnrealEd/Private/Cooker/IterativeValidatePackageWriter.cpp b/Engine/Source/Editor/UnrealEd/Private/Cooker/IterativeValidatePackageWriter.cpp
new file mode 100644
--- /dev/null
+++ b/Engine/Source/Editor/UnrealEd/Private/Cooker/IterativeValidatePackageWriter.cpp
@@ -0,0 +1,21 @@
+// Copyright Epic Games, Inc. All Rights Reserved.
+
+#include "Cooker/IterativeValidatePackageWriter.h"
+
+FIterativeValidatePackageWriter::FIterativeValidatePackageWriter(TUniquePtr<ICookedPackageWriter>&& InInner)
+	: FDiffPackageWriter(MoveTemp(InInner))
+{
+}
+
+void FIterativeValidatePackageWriter::UpdatePackageModificationStatus(FName PackageName, bool bIterativelyUnmodified,
+	bool& bInOutShouldIterativelySkip)
+{
+	bool bModified = !bIterativelyUnmodified;
+	bInOutShouldIterativelySkip = bModified;
+
+	bool bInnerIterativelyUnmodified = bModified;
+	bool bInnerInOutShouldIterativelySkip = bModified;
+	Inner->UpdatePackageModificationStatus(PackageName, bInnerIterativelyUnmodified, bInnerInOutShouldIterativelySkip);
+	checkf(bInnerInOutShouldIterativelySkip == bInnerIterativelyUnmodified,
+		TEXT("IterativeValidatePackageWriter is not supported with an Inner that modifies bInOutShouldIterativelySkip."));
+}
diff --git a/Engine/Source/Editor/UnrealEd/Private/Cooker/IterativeValidatePackageWriter.h b/Engine/Source/Editor/UnrealEd/Private/Cooker/IterativeValidatePackageWriter.h
new file mode 100644
--- /dev/null
+++ b/Engine/Source/Editor/UnrealEd/Private/Cooker/IterativeValidatePackageWriter.h
@@ -0,0 +1,19 @@
+// Copyright Epic Games, Inc. All Rights Reserved.
+
+#pragma once
+
+#include "DiffPackageWriter.h"
+
+/**
+ * A CookedPackageWriter that diffs the cook results of iteratively-unmodified packages between their last cook
+ * results and the current cook.
+ */
+class FIterativeValidatePackageWriter : public FDiffPackageWriter
+{
+public:
+	FIterativeValidatePackageWriter(TUniquePtr<ICookedPackageWriter>&& InInner);
+
+	// ICookedPackageWriter
+	virtual void UpdatePackageModificationStatus(FName PackageName, bool bIterativelyUnmodified,
+		bool& bInOutShouldIterativelySkip) override;
+};
diff --git a/Engine/Source/Editor/UnrealEd/Private/Cooker/LooseCookedPackageWriter.cpp b/Engine/Source/Editor/UnrealEd/Private/Cooker/LooseCookedPackageWriter.cpp
--- a/Engine/Source/Editor/UnrealEd/Private/Cooker/LooseCookedPackageWriter.cpp
+++ b/Engine/Source/Editor/UnrealEd/Private/Cooker/LooseCookedPackageWriter.cpp
@@ -842,7 +842,8 @@ void FLooseCookedPackageWriter::RemoveCookedPackagesByPackageName(TArrayView<con
 	}
 }
 
-void FLooseCookedPackageWriter::MarkPackagesUpToDate(TArrayView<const FName> UpToDatePackages)
+void FLooseCookedPackageWriter::UpdatePackageModificationStatus(FName PackageName,
+	bool bIterativelyUnmodified, bool& bInOutShouldIterativelySkip)
 {
 }
 
diff --git a/Engine/Source/Editor/UnrealEd/Private/Cooker/LooseCookedPackageWriter.h b/Engine/Source/Editor/UnrealEd/Private/Cooker/LooseCookedPackageWriter.h
--- a/Engine/Source/Editor/UnrealEd/Private/Cooker/LooseCookedPackageWriter.h
+++ b/Engine/Source/Editor/UnrealEd/Private/Cooker/LooseCookedPackageWriter.h
@@ -60,7 +60,8 @@ public:
 	virtual FCbObject GetOplogAttachment(FName PackageName, FUtf8StringView AttachmentKey) override;
 	virtual void RemoveCookedPackages(TArrayView<const FName> PackageNamesToRemove) override;
 	virtual void RemoveCookedPackages() override;
-	virtual void MarkPackagesUpToDate(TArrayView<const FName> UpToDatePackages) override;
+	virtual void UpdatePackageModificationStatus(FName PackageName, bool bIterativelyUnmodified,
+		bool& bInOutShouldIterativelySkip) override;
 	virtual TFuture<FCbObject> WriteMPCookMessageForPackage(FName PackageName) override;
 	virtual bool TryReadMPCookMessageForPackage(FName PackageName, FCbObjectView Message) override;
 	virtual bool GetPreviousCookedBytes(const FPackageInfo& Info, FPreviousCookedBytesData& OutData) override;
diff --git a/Engine/Source/Runtime/Core/Public/Serialization/PackageWriter.h b/Engine/Source/Runtime/Core/Public/Serialization/PackageWriter.h
--- a/Engine/Source/Runtime/Core/Public/Serialization/PackageWriter.h
+++ b/Engine/Source/Runtime/Core/Public/Serialization/PackageWriter.h
@@ -308,10 +308,14 @@ public:
 	 */
 	virtual void RemoveCookedPackages() = 0;
 
+	UE_DEPRECATED(5.4, "No longer called; override UpdatePackageModificationStatus instead")
+	virtual void MarkPackagesUpToDate(TArrayView<const FName> UpToDatePackages) {}
 	/**
-	 * Signal the given cooked package(s) have been checked for changes and have not been modified since the last cook.
+	 * During iterative cooking, signal the cooked package has been checked for changes,
+	 * and decide whether it should be iteratively skipped.
 	 */
-	virtual void MarkPackagesUpToDate(TArrayView<const FName> UpToDatePackages) = 0;
+	virtual void UpdatePackageModificationStatus(FName PackageName, bool bIterativelyUnmodified,
+		bool& bInOutShouldIterativelySkip) = 0;
 
 	struct FPreviousCookedBytesData
 	{
