List the data asset classes by their Default Name and not by their Class name for easier readability for users.

diff --git a/Engine/Source/Editor/UnrealEd/Private/Factories/EditorFactories.cpp b/Engine/Source/Editor/UnrealEd/Private/Factories/EditorFactories.cpp
--- a/Engine/Source/Editor/UnrealEd/Private/Factories/EditorFactories.cpp
+++ b/Engine/Source/Editor/UnrealEd/Private/Factories/EditorFactories.cpp
@@ -7454,7 +7454,7 @@ bool UDataAssetFactory::ConfigureProperties()
 	// Fill in options
 	FClassViewerInitializationOptions Options;
 	Options.Mode = EClassViewerMode::ClassPicker;
-
+	Options.NameTypeToDisplay = EClassViewerNameTypeToDisplay::DisplayName;
 	TSharedPtr<FAssetClassParentFilter> Filter = MakeShareable(new FAssetClassParentFilter);
 	Options.ClassFilters.Add(Filter.ToSharedRef());
 
