Cooking: Iterative: Fix UpdatePackageModificationStatus being incorrectly called during full builds and called twice in legacy iterative builds.
#rnx

diff --git a/Engine/Source/Editor/UnrealEd/Private/Cooker/CookRequestCluster.cpp b/Engine/Source/Editor/UnrealEd/Private/Cooker/CookRequestCluster.cpp
--- a/Engine/Source/Editor/UnrealEd/Private/Cooker/CookRequestCluster.cpp
+++ b/Engine/Source/Editor/UnrealEd/Private/Cooker/CookRequestCluster.cpp
@@ -1095,11 +1095,11 @@ void FRequestCluster::FGraphSearch::ExploreVertexEdges(FVertexData& Vertex)
 
 			const FCookAttachments& PlatformAttachments = QueryPlatformData.CookAttachments;
 			bool bFoundBuildDefinitions = false;
-			bool bIterativelyUnmodified = false;
 			ICookedPackageWriter* PackageWriter = FetchPlatformData.Writer;
-			if (IsCookAttachmentsValid(PackageName, PlatformAttachments))
+			if (!Cluster.bFullBuild && Cluster.COTFS.bHybridIterativeEnabled)
 			{
-				if (!Cluster.bFullBuild && Cluster.COTFS.bHybridIterativeEnabled)
+				bool bIterativelyUnmodified = false;
+				if (IsCookAttachmentsValid(PackageName, PlatformAttachments))
 				{
 					if (IsIterativeEnabled(PackageName, Cluster.COTFS.bHybridIterativeAllowAllClasses))
 					{
@@ -1121,18 +1121,18 @@ void FRequestCluster::FGraphSearch::ExploreVertexEdges(FVertexData& Vertex)
 							PlatformAttachments.BuildDefinitionList);
 					}
 				}
-			}
-			bool bShouldIterativelySkip = bIterativelyUnmodified;
-			PackageWriter->UpdatePackageModificationStatus(PackageName, bIterativelyUnmodified, bShouldIterativelySkip);
-			if (bShouldIterativelySkip)
-			{
-				PackageData.SetPlatformCooked(TargetPlatform, ECookResult::Succeeded);
-				if (PlatformIndex == FirstSessionPlatformIndex)
+				bool bShouldIterativelySkip = bIterativelyUnmodified;
+				PackageWriter->UpdatePackageModificationStatus(PackageName, bIterativelyUnmodified, bShouldIterativelySkip);
+				if (bShouldIterativelySkip)
 				{
-					COOK_STAT(++DetailedCookStats::NumPackagesIterativelySkipped);
+					PackageData.SetPlatformCooked(TargetPlatform, ECookResult::Succeeded);
+					if (PlatformIndex == FirstSessionPlatformIndex)
+					{
+						COOK_STAT(++DetailedCookStats::NumPackagesIterativelySkipped);
+					}
+					// Declare the package to the EDLCookInfo verification so we don't warn about missing exports from it
+					UE::SavePackageUtilities::EDLCookInfoAddIterativelySkippedPackage(PackageName);
 				}
-				// Declare the package to the EDLCookInfo verification so we don't warn about missing exports from it
-				UE::SavePackageUtilities::EDLCookInfoAddIterativelySkippedPackage(PackageName);
 			}
 
 			if (Cluster.bPreQueueBuildDefinitions && !bFoundBuildDefinitions)
