[Feature] Data-Only Blueprint merging

diff --git a/Engine/Plugins/Editor/EngineAssetDefinitions/Source/Private/AssetDefinition_DataAsset.cpp b/Engine/Plugins/Editor/EngineAssetDefinitions/Source/Private/AssetDefinition_DataAsset.cpp
--- a/Engine/Plugins/Editor/EngineAssetDefinitions/Source/Private/AssetDefinition_DataAsset.cpp
+++ b/Engine/Plugins/Editor/EngineAssetDefinitions/Source/Private/AssetDefinition_DataAsset.cpp
@@ -5,19 +5,14 @@
 #include "Editor.h"
 #include "ClassViewerFilter.h"
 #include "ContentBrowserMenuContexts.h"
-#include "ISourceControlModule.h"
-#include "ISourceControlProvider.h"
+#include "MergeUtils.h"
 #include "ObjectTools.h"
 #include "SDetailsDiff.h"
-#include "SourceControlHelpers.h"
-#include "SourceControlOperations.h"
 #include "ToolMenu.h"
 #include "ToolMenus.h"
 #include "Kismet2/SClassPickerDialog.h"
 #include "Engine/Engine.h"
 #include "ToolMenuSection.h"
-#include "Serialization/ObjectReader.h"
-#include "Serialization/ObjectWriter.h"
 #include "HAL/PlatformFileManager.h"
 #include "UObject/Linker.h"
 
@@ -193,327 +188,15 @@ bool UAssetDefinition_DataAsset::CanMerge() const
 	return true;
 }
 
-struct ScopedMergeResolveTransaction
-{
-	ScopedMergeResolveTransaction(UObject* InManagedObject, EMergeFlags InFlags)
-		: ManagedObject(InManagedObject)
-		, Flags(InFlags)
-	{
-		if (Flags & MF_HANDLE_SOURCE_CONTROL)
-		{
-			UndoHandler = NewObject<UUndoableResolveHandler>();
-			UndoHandler->SetFlags(RF_Transactional);
-			UndoHandler->SetManagedObject(ManagedObject);
-			
-			TransactionNum = GEditor->BeginTransaction(LOCTEXT("ResolveMerge", "ResolveAutoMerge"));
-			ensure(UndoHandler->Modify());
-			ensure(ManagedObject->Modify());
-		}
-	}
-
-	void Cancel()
-	{
-		bCanceled = true;
-	}
-	
-	~ScopedMergeResolveTransaction()
-	{
-		if (Flags & MF_HANDLE_SOURCE_CONTROL)
-		{
-			if (!bCanceled)
-			{
-				UndoHandler->MarkResolved();
-				GEditor->EndTransaction();
-			}
-			else
-			{
-				ManagedObject->GetPackage()->SetDirtyFlag(false);
-				GEditor->CancelTransaction(TransactionNum);
-			}
-		}
-	}
-
-	UObject* ManagedObject;
-	UUndoableResolveHandler* UndoHandler = nullptr;
-	EMergeFlags Flags;
-	int TransactionNum = 0;
-	bool bCanceled = false;
-};
-
-static UPackage* LoadMergePackage(const FString& SCFile, const FString& Revision, const UPackage* LocalPackage)
-{
-	const FString FileWithRevision = SCFile + TEXT("#") + Revision;
-	const TSharedRef<FDownloadFile, ESPMode::ThreadSafe> DownloadFileOperation = ISourceControlOperation::Create<FDownloadFile>(FPaths::DiffDir(), FDownloadFile::EVerbosity::Full);
-	ISourceControlModule::Get().GetProvider().Execute(DownloadFileOperation, FileWithRevision, EConcurrency::Synchronous);
-	const FString DownloadPath = FPaths::ConvertRelativePathToFull(FPaths::DiffDir() / FPaths::GetCleanFilename(FileWithRevision));
-	FString CopyPath = DownloadPath;
-	CopyPath.ReplaceInline(TEXT(".uasset"), TEXT(""));
-	CopyPath.ReplaceCharInline('#', '-');
-	CopyPath.ReplaceCharInline('.', '-');
-	CopyPath = FPaths::CreateTempFilename(*FPaths::GetPath(CopyPath), *FPaths::GetBaseFilename(CopyPath), TEXT(".uasset"));
-
-	if (FPlatformFileManager::Get().GetPlatformFile().CopyFile(*CopyPath, *DownloadPath))
-	{
-		return DiffUtils::LoadPackageForDiff(FPackagePath::FromLocalPath(CopyPath), LocalPackage->GetLoadedPath());
-	}
-	return nullptr;
-}
-
 EAssetCommandResult UAssetDefinition_DataAsset::Merge(const FAssetAutomaticMergeArgs& MergeArgs) const
 {
-	if (!ensure(MergeArgs.LocalAsset))
-	{
-		return EAssetCommandResult::Unhandled;
-	}
-	
-	FAssetManualMergeArgs ManualMergeArgs;
-	ManualMergeArgs.LocalAsset = MergeArgs.LocalAsset;
-	ManualMergeArgs.ResolutionCallback = MergeArgs.ResolutionCallback;
-	ManualMergeArgs.Flags = MergeArgs.Flags;
-	const UPackage* LocalPackage = ManualMergeArgs.LocalAsset->GetPackage();
-	
-	
-	ISourceControlProvider& SourceControlProvider = ISourceControlModule::Get().GetProvider();
-	const TSharedRef<FUpdateStatus, ESPMode::ThreadSafe> UpdateStatusOperation = ISourceControlOperation::Create<FUpdateStatus>();
-	UpdateStatusOperation->SetUpdateHistory(true);
-	SourceControlProvider.Execute(UpdateStatusOperation, LocalPackage);
-	
-	// Get the SCC state
-	const FSourceControlStatePtr SourceControlState = SourceControlProvider.GetState(LocalPackage, EStateCacheUsage::Use);
-
-	// If we have an asset and its in SCC..
-	if( SourceControlState.IsValid() && SourceControlState->IsSourceControlled() )
-	{
-		const ISourceControlState::FResolveInfo ResolveInfo = SourceControlState->GetResolveInfo();
-		check(ResolveInfo.IsValid());
-		
-		if(UPackage* TempPackage = LoadMergePackage(ResolveInfo.RemoteFile, ResolveInfo.RemoteRevision, LocalPackage))
-		{
-			// Grab the old asset from that old package
-			ManualMergeArgs.RemoteAsset = FindObject<UObject>(TempPackage, *ManualMergeArgs.LocalAsset->GetName());
-
-			// Recovery for package names that don't match
-			if (ManualMergeArgs.RemoteAsset == nullptr)
-			{
-				ManualMergeArgs.RemoteAsset = TempPackage->FindAssetInPackage();
-			}
-		}
-		
-		if(UPackage* TempPackage = LoadMergePackage(ResolveInfo.BaseFile, ResolveInfo.BaseRevision, LocalPackage))
-		{
-			// Grab the old asset from that old package
-			ManualMergeArgs.BaseAsset = FindObject<UObject>(TempPackage, *ManualMergeArgs.LocalAsset->GetName());
-
-			// Recovery for package names that don't match
-			if (ManualMergeArgs.BaseAsset == nullptr)
-			{
-				ManualMergeArgs.BaseAsset = TempPackage->FindAssetInPackage();
-			}
-		}
-	}
-
-	// single asset merging is only supported for assets in a conflicted state in source control
-	if (!ensure(ManualMergeArgs.BaseAsset && ManualMergeArgs.RemoteAsset && ManualMergeArgs.LocalAsset))
-	{
-		return Super::Merge(MergeArgs);
-	}
-	
-	return Merge(ManualMergeArgs);
+	return MergeUtils::Merge(MergeArgs);
 }
 
-
-	
-static TArray<uint8> SerializeToBinary(UObject* Object, UObject* Default = nullptr)
-{
-	TArray<uint8> Result;
-	FObjectWriter ObjectWriter(Result);
-	ObjectWriter.SetIsPersistent(true);
-	Default = Default ? ToRawPtr(Default) : ToRawPtr(Object->GetClass()->ClassDefaultObject);
-	Object->GetClass()->SerializeTaggedProperties( ObjectWriter, reinterpret_cast<uint8*>(Object), Default->GetClass(), reinterpret_cast<uint8*>(Default));
-	return Result;
-};
-
-static void DeserializeFromBinary(UObject* Object, const TArray<uint8>& Data, UObject* Default = nullptr)
-{
-	FObjectReader ObjectReader(Data);
-	ObjectReader.SetIsPersistent(true);
-	Default = Default ? ToRawPtr(Default) : ToRawPtr(Object->GetClass()->ClassDefaultObject);
-	Object->GetClass()->SerializeTaggedProperties(ObjectReader, reinterpret_cast<uint8*>(Object), Default->GetClass(), reinterpret_cast<uint8*>(Default));
-};
-
 EAssetCommandResult UAssetDefinition_DataAsset::Merge(const FAssetManualMergeArgs& MergeArgs) const
 {
-	auto NotifyResolution = [&MergeArgs](EAssetMergeResult Result)
-	{
-		FAssetMergeResults Results;
-		Results.Result = Result;
-		Results.MergedPackage = MergeArgs.LocalAsset->GetPackage();
-		MergeArgs.ResolutionCallback.ExecuteIfBound(Results);
-		return EAssetCommandResult::Handled;
-	};
-	
-	TArray<UObject*> SubObjects;
-	GetObjectsWithOuter(MergeArgs.LocalAsset, SubObjects);
-	if (!ensure(SubObjects.IsEmpty()))
-	{
-		TArray<FString> SubObjectNames;
-		for (const UObject* SubObject : SubObjects)
-		{
-			SubObjectNames.Add(SubObject->GetName());
-		}
-		UE_LOG(LogEngine, Warning, TEXT("Merge is not currently supported with sub-objects: [{0}]"), *FString::Join(SubObjectNames, TEXT(",")));
-		
-		FAssetMergeResults Results;
-		Results.Result = EAssetMergeResult::Cancelled;
-		Results.MergedPackage = MergeArgs.LocalAsset->GetPackage();
-		MergeArgs.ResolutionCallback.ExecuteIfBound(Results);
-		
-		return NotifyResolution(EAssetMergeResult::Cancelled);
-	}
-
-	auto AutoMerge = [](UObject* HeadRevision, UObject* BranchA, UObject* BranchB, FName Name) -> UObject*
-	{
-		// get delta of BranchB relative to HeadRevision
-		const TArray<uint8> Delta = SerializeToBinary(BranchB, HeadRevision);
-		
-		// apply delta on top of BranchA
-		Name = MakeUniqueObjectName(nullptr, BranchA->GetClass(), Name, EUniqueObjectNameOptions::GloballyUnique);
-		UObject* Merged = DuplicateObject(BranchA, nullptr, Name);
-		DeserializeFromBinary(Merged, Delta, BranchA);
-		return Merged;
-	};
-
-	// apply changes in different orders to come up with two possible merge options
-	UObject* FavorRemote = AutoMerge(MergeArgs.BaseAsset,MergeArgs.LocalAsset, MergeArgs.RemoteAsset, TEXT("FavorRemote"));
-	UObject* FavorLocal = AutoMerge(MergeArgs.BaseAsset,MergeArgs.RemoteAsset,MergeArgs.LocalAsset, TEXT("FavorLocal"));
-	const TArray<uint8> FavorRemoteBinary = SerializeToBinary(FavorRemote, MergeArgs.BaseAsset);
-	const TArray<uint8> FavorLocalBinary = SerializeToBinary(FavorLocal, MergeArgs.BaseAsset);
-
-	// if both merge options are the same, we have no conflicts
-	if (FavorRemoteBinary == FavorLocalBinary)
-	{
-		// copy changes over to the local asset
-		ScopedMergeResolveTransaction(MergeArgs.LocalAsset, MergeArgs.Flags);
-		DeserializeFromBinary(MergeArgs.LocalAsset, FavorLocalBinary);
-		return NotifyResolution(EAssetMergeResult::Completed);
-	}
-	
-	// conflicts detected. We need to ask the user to manually resolve them
-	
-	if (!(MergeArgs.Flags & MF_NO_GUI))
-	{
-		const TSharedRef<SDetailsDiff> DiffView = SDetailsDiff::CreateDiffWindow(FavorRemote, FavorLocal, {}, {}, UDataAsset::StaticClass());
-		DiffView->SetOutputObject(MergeArgs.LocalAsset);
-
-		FObjectReader ObjectReader(FavorLocalBinary);
-		DiffView->RequestModifications(ObjectReader);
-		DiffView->OnWindowClosedEvent.AddLambda([MergeArgs](TSharedRef<SDetailsDiff> DiffView)
-		{
-			// serialize the changes made by the user while the diff view was open
-			TArray<uint8> DiffViewModifications;
-			FObjectWriter ObjectWriter(DiffViewModifications);
-			DiffView->GetModifications(ObjectWriter);
-			
-			// begin undoable transaction
-			ScopedMergeResolveTransaction(MergeArgs.LocalAsset, MergeArgs.Flags);
-			
-			// apply manual changes made within the diff tool
-			DeserializeFromBinary(MergeArgs.LocalAsset, DiffViewModifications);
-			
-			// notify caller that merge is complete
-			FAssetMergeResults Results;
-			Results.Result = EAssetMergeResult::Completed;
-			Results.MergedPackage = MergeArgs.LocalAsset->GetPackage();
-			MergeArgs.ResolutionCallback.ExecuteIfBound(Results);
-		});
-		
-		return EAssetCommandResult::Handled;
-	}
-	
-	return NotifyResolution(EAssetMergeResult::Cancelled);
+	return MergeUtils::Merge(MergeArgs);
 }
 
-void UUndoableResolveHandler::SetManagedObject(UObject* Object)
-{
-	ManagedObject = Object;
-	const UPackage* Package = ManagedObject->GetPackage();
-	const FString Filepath = FPaths::ConvertRelativePathToFull(Package->GetLoadedPath().GetLocalFullPath());
-	
-	ISourceControlProvider& Provider = ISourceControlModule::Get().GetProvider();
-	const FSourceControlStatePtr SourceControlState = Provider.GetState(Package, EStateCacheUsage::Use);
-	const ISourceControlState::FResolveInfo ResolveInfo = SourceControlState->GetResolveInfo();
-	BaseRevisionNumber = SourceControlState->GetResolveInfo().BaseRevision;
-	if (const TSharedPtr<class ISourceControlRevision, ESPMode::ThreadSafe> CurrentRevision = SourceControlState->GetCurrentRevision())
-	{
-		CurrentRevisionNumber = FString::FromInt(CurrentRevision->GetRevisionNumber());
-	}
-	else
-	{
-		CurrentRevisionNumber = {};
-	}
-	CheckinIdentifier = SourceControlState->GetCheckInIdentifier();
-
-	// save package and copy the package to a temp file so it can be reverted
-	const FString BaseFilename = FPaths::GetBaseFilename(Filepath);
-	BackupFilepath = FPaths::CreateTempFilename(*(FPaths::ProjectSavedDir()/TEXT("Temp")), *BaseFilename.Left(32));
-	ensure(FPlatformFileManager::Get().GetPlatformFile().CopyFile(*BackupFilepath, *Filepath));
-}
-
-void UUndoableResolveHandler::MarkResolved()
-{
-	const UPackage* Package = ManagedObject->GetPackage();
-	const FString Filepath = FPaths::ConvertRelativePathToFull(Package->GetLoadedPath().GetLocalFullPath());
-	
-	ISourceControlProvider& Provider = ISourceControlModule::Get().GetProvider();
-	Provider.Execute(ISourceControlOperation::Create<FResolve>(), TArray{Filepath}, EConcurrency::Synchronous);
-	bShouldBeResolved = true;
-}
-
-void UUndoableResolveHandler::PostEditUndo()
-{
-	if (bShouldBeResolved) // redo resolution
-	{
-		MarkResolved();
-	}
-	else // undo resolution
-	{
-		UPackage* Package = ManagedObject->GetPackage();
-		const FString Filepath = FPaths::ConvertRelativePathToFull(Package->GetLoadedPath().GetLocalFullPath());
-		
-		if (BaseRevisionNumber.IsEmpty() || CurrentRevisionNumber.IsEmpty())
-		{
-			ensure(FPlatformFileManager::Get().GetPlatformFile().CopyFile(*Filepath, *BackupFilepath));
-			return;
-		}
-		
-		// to force the file to revert to it's pre-resolved state, we must revert, sync back to base revision,
-		// apply the conflicting changes, then sync forward again.
-		ISourceControlProvider& Provider = ISourceControlModule::Get().GetProvider();
-		{
-			const TSharedRef<FSync> SyncOperation = ISourceControlOperation::Create<FSync>();
-			SyncOperation->SetRevision(BaseRevisionNumber);
-			Provider.Execute(SyncOperation, Filepath, EConcurrency::Synchronous);
-		}
-		
-		ResetLoaders(Package);
-		Provider.Execute( ISourceControlOperation::Create<FRevert>(), Filepath, EConcurrency::Synchronous);
-
-		{
-			const TSharedRef<FCheckOut> CheckoutOperation = ISourceControlOperation::Create<FCheckOut>();
-			Provider.Execute(CheckoutOperation, CheckinIdentifier, {Filepath}, EConcurrency::Synchronous);
-		}
-
-		ensure(FPlatformFileManager::Get().GetPlatformFile().CopyFile(*Filepath, *BackupFilepath));
-		
-		{
-			const TSharedRef<FSync> SyncOperation = ISourceControlOperation::Create<FSync>();
-			SyncOperation->SetRevision(CurrentRevisionNumber);
-			Provider.Execute(SyncOperation, Filepath, EConcurrency::Synchronous);
-		}
-
-		Provider.Execute(ISourceControlOperation::Create<FUpdateStatus>(), {Filepath}, EConcurrency::Asynchronous);
-	}
-	UObject::PostEditUndo();
-}
 
 #undef LOCTEXT_NAMESPACE
diff --git a/Engine/Plugins/Editor/EngineAssetDefinitions/Source/Private/Script/AssetDefinition_Blueprint.cpp b/Engine/Plugins/Editor/EngineAssetDefinitions/Source/Private/Script/AssetDefinition_Blueprint.cpp
--- a/Engine/Plugins/Editor/EngineAssetDefinitions/Source/Private/Script/AssetDefinition_Blueprint.cpp
+++ b/Engine/Plugins/Editor/EngineAssetDefinitions/Source/Private/Script/AssetDefinition_Blueprint.cpp
@@ -19,6 +19,7 @@
 #include "ToolMenu.h"
 #include "ToolMenuSection.h"
 #include "BlueprintEditor.h"
+#include "MergeUtils.h"
 #include "SBlueprintDiff.h"
 #include "Kismet2/BlueprintEditorUtils.h"
 #include "ThumbnailRendering/SceneThumbnailInfo.h"
@@ -63,6 +64,11 @@ EAssetCommandResult UAssetDefinition_Blueprint::Merge(const FAssetAutomaticMerge
 {
 	UBlueprint* AsBlueprint = CastChecked<UBlueprint>(MergeArgs.LocalAsset);
 	
+	if (FBlueprintEditorUtils::IsDataOnlyBlueprint(AsBlueprint))
+	{
+		return MergeUtils::Merge(MergeArgs);
+	}
+	
 	// Kludge to get the merge panel in the blueprint editor to show up:
 	bool Success = GEditor->GetEditorSubsystem<UAssetEditorSubsystem>()->OpenEditorForAsset(AsBlueprint);
 	if (Success)
@@ -78,15 +84,40 @@ EAssetCommandResult UAssetDefinition_Blueprint::Merge(const FAssetAutomaticMerge
 
 EAssetCommandResult UAssetDefinition_Blueprint::Merge(const FAssetManualMergeArgs& MergeArgs) const
 {
-	UBlueprint* AsBlueprint = CastChecked<UBlueprint>(MergeArgs.LocalAsset);
+	UBlueprint* LocalBlueprint = Cast<UBlueprint>(MergeArgs.LocalAsset);
+	if(!ensureMsgf(LocalBlueprint, TEXT("Merge LocalAsset is not a Blueprint")))
+	{
+		return Super::Merge(MergeArgs);
+	}
+	
+	const UBlueprint* RemoteBlueprint = Cast<UBlueprint>(MergeArgs.RemoteAsset);
+	if(!ensureMsgf(RemoteBlueprint, TEXT("Merge RemoteAsset is not a Blueprint")))
+	{
+		return Super::Merge(MergeArgs);
+	}
+	
+	const UBlueprint* BaseBlueprint = Cast<UBlueprint>(MergeArgs.BaseAsset);
+	if(!ensureMsgf(BaseBlueprint, TEXT("Merge BaseAsset is not a Blueprint")))
+	{
+		return Super::Merge(MergeArgs);
+	}
+	
+	if ( // all assets are data only
+		FBlueprintEditorUtils::IsDataOnlyBlueprint(LocalBlueprint) &&
+		FBlueprintEditorUtils::IsDataOnlyBlueprint(RemoteBlueprint) &&
+		FBlueprintEditorUtils::IsDataOnlyBlueprint(BaseBlueprint)
+	)
+	{
+		return MergeUtils::Merge(MergeArgs);
+	}
 	check(MergeArgs.LocalAsset->GetClass() == MergeArgs.BaseAsset->GetClass());
 	check(MergeArgs.LocalAsset->GetClass() == MergeArgs.RemoteAsset->GetClass());
 
-	if (GEditor->GetEditorSubsystem<UAssetEditorSubsystem>()->OpenEditorForAsset(AsBlueprint))
+	if (GEditor->GetEditorSubsystem<UAssetEditorSubsystem>()->OpenEditorForAsset(LocalBlueprint))
 	{
 		FBlueprintEditorModule& BlueprintEditorModule = FModuleManager::LoadModuleChecked<FBlueprintEditorModule>("Kismet");
 
-		FBlueprintEditor* BlueprintEditor = static_cast<FBlueprintEditor*>(GEditor->GetEditorSubsystem<UAssetEditorSubsystem>()->FindEditorForAsset(AsBlueprint, /*bFocusIfOpen =*/false));
+		FBlueprintEditor* BlueprintEditor = static_cast<FBlueprintEditor*>(GEditor->GetEditorSubsystem<UAssetEditorSubsystem>()->FindEditorForAsset(LocalBlueprint, /*bFocusIfOpen =*/false));
 		BlueprintEditor->CreateMergeToolTab(Cast<UBlueprint>(MergeArgs.BaseAsset), Cast<UBlueprint>(MergeArgs.RemoteAsset), 
 			FOnMergeResolved::CreateLambda([ResolutionCallback = MergeArgs.ResolutionCallback](UPackage* MergedPackage, EMergeResult::Type Result)
 			{
@@ -98,7 +129,7 @@ EAssetCommandResult UAssetDefinition_Blueprint::Merge(const FAssetManualMergeArg
 		);		
 	}
 	
-	return EAssetCommandResult::Handled;
+	 return EAssetCommandResult::Handled;
 }
 
 EAssetCommandResult UAssetDefinition_Blueprint::PerformAssetDiff(const FAssetDiffArgs& DiffArgs) const
diff --git a/Engine/Plugins/Editor/EngineAssetDefinitions/Source/Public/AssetDefinition_DataAsset.h b/Engine/Plugins/Editor/EngineAssetDefinitions/Source/Public/AssetDefinition_DataAsset.h
--- a/Engine/Plugins/Editor/EngineAssetDefinitions/Source/Public/AssetDefinition_DataAsset.h
+++ b/Engine/Plugins/Editor/EngineAssetDefinitions/Source/Public/AssetDefinition_DataAsset.h
@@ -33,25 +33,3 @@ public:
 #if UE_ENABLE_INCLUDE_ORDER_DEPRECATED_IN_5_2
 #include "CoreMinimal.h"
 #endif
-
-
-UCLASS()
-class UUndoableResolveHandler : public UObject
-{
-public:
-	GENERATED_BODY()
-	void SetManagedObject(UObject* Object);
-	void MarkResolved();
-
-	virtual void PostEditUndo() override;
-
-private:
-	FString BaseRevisionNumber;
-	FString CurrentRevisionNumber;
-	FString BackupFilepath;
-	UObject* ManagedObject;
-	TSharedPtr<class ISourceControlChangelist> CheckinIdentifier;
-	
-	UPROPERTY()
-	bool bShouldBeResolved = false;
-};
\ No newline at end of file
diff --git a/Engine/Source/Editor/Kismet/Private/DetailsDiff.cpp b/Engine/Source/Editor/Kismet/Private/DetailsDiff.cpp
--- a/Engine/Source/Editor/Kismet/Private/DetailsDiff.cpp
+++ b/Engine/Source/Editor/Kismet/Private/DetailsDiff.cpp
@@ -3,10 +3,14 @@
 #include "DetailsDiff.h"
 
 #include "BlueprintDetailsCustomization.h"
+#include "BlueprintEditorLibrary.h"
+#include "DetailCategoryBuilder.h"
+#include "DetailLayoutBuilder.h"
 #include "Algo/Copy.h"
 #include "Algo/Transform.h"
 #include "Containers/Set.h"
 #include "DetailsViewArgs.h"
+#include "DetailWidgetRow.h"
 #include "HAL/PlatformCrt.h"
 #include "IDetailsView.h"
 #include "Misc/AssertionMacros.h"
@@ -16,10 +20,15 @@
 #include "PropertyPath.h"
 #include "UObject/WeakObjectPtrTemplates.h"
 #include "Engine/MemberReference.h"
+#include "Kismet2/BlueprintEditorUtils.h"
 #include "Widgets/Layout/LinkableScrollBar.h"
+#include "Widgets/Input/SComboButton.h"
+#include "Engine/BlueprintGeneratedClass.h"
 
 class UObject;
 
+#define LOCTEXT_NAMESPACE "DetailsDif"
+
 FDetailsDiff::FDetailsDiff(const UObject* InObject, FOnDisplayedPropertiesChanged InOnDisplayedPropertiesChanged, bool bScrollbarOnLeft )
 	: OnDisplayedPropertiesChanged( InOnDisplayedPropertiesChanged )
 	, DisplayedObject( InObject )
@@ -30,6 +39,108 @@ FDetailsDiff::FDetailsDiff(const UObject* InObject, FOnDisplayedPropertiesChange
 	DisplayedProperties = DetailsView->GetPropertiesInOrderDisplayed();
 }
 
+class DOBPDetailsCustomization : public IDetailCustomization
+{
+public:
+	DOBPDetailsCustomization(UBlueprint* InBlueprint) : Blueprint(InBlueprint)
+	{}
+	
+	static TSharedRef<IDetailCustomization> Make(UBlueprint* Blueprint)
+	{
+		return MakeShared<DOBPDetailsCustomization>(Blueprint);
+	}
+	
+	virtual void CustomizeDetails(IDetailLayoutBuilder& DetailLayout) override
+	{
+		if (!Blueprint.IsValid())
+		{
+			return;
+		}
+		
+		IDetailCategoryBuilder& Category = DetailLayout.EditCategory("ClassOptions", LOCTEXT("ClassOptions", "Class Options"));
+
+		// put the ClassOptions first
+		Category.SetSortOrder(0);
+		
+		// ParentClass is a hidden property so we have to add it to the property map manually to use it
+		const TSharedPtr<IPropertyHandle> ParentClassProperty = DetailLayout.AddObjectPropertyData({Blueprint.Get()}, TEXT("ParentClass"));
+
+		Category.AddCustomRow( LOCTEXT("ClassOptions", "Class Options") )
+		.NameContent()
+		[
+			SNew(STextBlock)
+			.Text(LOCTEXT("BlueprintDetails_ParentClass", "Parent Class"))
+			.Font(IDetailLayoutBuilder::GetDetailFont())
+		]
+		.ValueContent()
+		[
+			SNew(SHorizontalBox)
+			+SHorizontalBox::Slot()
+			.FillWidth(1.0f)
+			[
+				SAssignNew(ParentClassComboButton, SComboButton)
+				.IsEnabled(this, &DOBPDetailsCustomization::CanReparent)
+				.OnGetMenuContent(this, &DOBPDetailsCustomization::GetParentClassMenuContent)
+				.ButtonContent()
+				[
+					SNew(STextBlock)
+					.Text(this, &DOBPDetailsCustomization::GetParentClassName)
+					.Font(IDetailLayoutBuilder::GetDetailFont())
+				]
+			]
+		]
+		.PropertyHandleList({ParentClassProperty});
+	}
+
+	bool CanReparent() const
+	{
+		// Don't show the reparent option if it's an Interface or we're not in editing mode
+		return Blueprint.IsValid() && !FBlueprintEditorUtils::IsInterfaceBlueprint(Blueprint.Get()) && (BPTYPE_FunctionLibrary != Blueprint->BlueprintType);
+	}
+	
+	TSharedRef<SWidget> GetParentClassMenuContent() const
+	{
+		if (!Blueprint.IsValid())
+		{
+			return SNew(SBox);
+		}
+		
+		TArray<UBlueprint*> Blueprints;
+		Blueprints.Add(Blueprint.Get());
+		const TSharedRef<SWidget> ClassPicker = FBlueprintEditorUtils::ConstructBlueprintParentClassPicker(Blueprints, FOnClassPicked::CreateSP(this, &DOBPDetailsCustomization::OnClassPicked));
+
+		// Achieving fixed width by nesting items within a fixed width box.
+		return SNew(SBox)
+			.WidthOverride(350.0f)
+			[
+				SNew(SVerticalBox)
+				+SVerticalBox::Slot()
+				.MaxHeight(400.0f)
+				.AutoHeight()
+				[
+					ClassPicker
+				]
+			];
+	}
+	
+	FText GetParentClassName() const
+	{
+		const UClass* ParentClass = Blueprint.IsValid() ? Blueprint->ParentClass : nullptr;
+		return ParentClass ? ParentClass->GetDisplayNameText() : FText::FromName(NAME_None);
+	}
+
+	void OnClassPicked(UClass* PickedClass) const
+	{
+		ParentClassComboButton->SetIsOpen(false);
+		UBlueprintEditorLibrary::ReparentBlueprint(Blueprint.Get(), PickedClass);
+	}
+private:
+	TWeakObjectPtr<UBlueprint> Blueprint;
+
+	/** Combo button used to choose a parent class */
+	TSharedPtr<SComboButton> ParentClassComboButton;
+};
+
 TSharedRef<IDetailsView> FDetailsDiff::CreateDetailsView(const UObject* InObject, TSharedPtr<SScrollBar> ExternalScrollbar, bool bScrollbarOnLeft)
 {
 	FDetailsViewArgs DetailsViewArgs;
@@ -50,6 +161,19 @@ TSharedRef<IDetailsView> FDetailsDiff::CreateDetailsView(const UObject* InObject
 		);
 		DetailsView->RegisterInstancedCustomPropertyLayout(UBlueprint::StaticClass(), LayoutOptionDetails);
 	}
+	// if InObject is a BP-CDO, add a "Parent Class" Combo button
+	if (InObject && Cast<UObject>(InObject->GetClass())->IsA<UBlueprintGeneratedClass>() )
+	{
+		UBlueprintGeneratedClass* Class = Cast<UBlueprintGeneratedClass>(InObject->GetClass());
+		if (Class->GetDefaultObject() == InObject)
+		{
+			const FOnGetDetailCustomizationInstance LayoutOptionDetails = FOnGetDetailCustomizationInstance::CreateStatic(
+				&DOBPDetailsCustomization::Make,
+				Cast<UBlueprint>(Class->ClassGeneratedBy)
+			);
+			DetailsView->RegisterInstancedCustomPropertyLayout(Class, LayoutOptionDetails);
+		}
+	}
 	// Forcing all advanced properties to be displayed for now, the logic to show changes made to advance properties
 	// conditionally is fragile and low priority for now:
 	DetailsView->ShowAllAdvancedProperties();
@@ -255,3 +379,5 @@ void FDetailsDiff::LinkScrolling(FDetailsDiff& LeftPanel, FDetailsDiff& RightPan
 {
 	SLinkableScrollBar::LinkScrollBars(LeftPanel.ScrollBar.ToSharedRef(), RightPanel.ScrollBar.ToSharedRef(), ScrollRate);
 }
+
+#undef LOCTEXT_NAMESPACE
diff --git a/Engine/Source/Editor/Kismet/Private/SDetailsDiff.cpp b/Engine/Source/Editor/Kismet/Private/SDetailsDiff.cpp
--- a/Engine/Source/Editor/Kismet/Private/SDetailsDiff.cpp
+++ b/Engine/Source/Editor/Kismet/Private/SDetailsDiff.cpp
@@ -50,6 +50,7 @@ void SDetailsDiff::Construct( const FArguments& InArgs)
 
 		AssetEditorCloseDelegate = GEditor->GetEditorSubsystem<UAssetEditorSubsystem>()->OnAssetEditorRequestClose().AddSP(this, &SDetailsDiff::OnCloseAssetEditor);
 	}
+	FCoreUObjectDelegates::OnObjectsReplaced.AddSP(this, &SDetailsDiff::OnObjectReplaced);
 
 	FToolBarBuilder NavToolBarBuilder(TSharedPtr< const FUICommandList >(), FMultiBoxCustomization::None);
 	NavToolBarBuilder.AddToolBarButton(
@@ -204,6 +205,40 @@ void SDetailsDiff::OnCloseAssetEditor(UObject* Asset, EAssetEditorCloseReason Cl
 	}
 }
 
+void SDetailsDiff::OnObjectReplaced(const FCoreUObjectDelegates::FReplacementObjectMap& Replacements)
+{
+	bool bNeedsRegenerate = false;
+
+	auto Refresh = [](const UObject* Obj, const UObject* From, const UObject* To)
+	{
+		bool bNeedsRegenerate = false;
+		if (Obj == From)
+		{
+			Obj = To;
+			bNeedsRegenerate = true;
+		}
+		if (Obj == From->GetClass()->ClassGeneratedBy)
+		{
+			bNeedsRegenerate = true;
+		}
+		return bNeedsRegenerate;
+	};
+	
+	// if any of the objects being displayed were reinstanced, refresh them
+	for (const auto &[From, To] : Replacements)
+	{
+		bNeedsRegenerate |= Refresh(OutputObject, From, To);
+		bNeedsRegenerate |= Refresh(PanelNew.Object, From, To);
+		bNeedsRegenerate |= Refresh(PanelOld.Object, From, To);
+	}
+	
+	if (bNeedsRegenerate)
+	{
+		GenerateDifferencesList();
+		RefreshCurrentModePanel();
+	}
+}
+
 TSharedRef<SWidget> SDetailsDiff::DefaultEmptyPanel()
 {
 	return SNew(SHorizontalBox)
@@ -280,32 +315,25 @@ TSharedRef<SDetailsDiff> SDetailsDiff::CreateDiffWindow(const UObject* OldObject
 	return CreateDiffWindow(WindowTitle, OldObject, NewObject, OldRevision, NewRevision);
 }
 
-void SDetailsDiff::SetOutputObject(const UObject* InOutputObject)
+void SDetailsDiff::SetOutputObject(UObject* InOutputObject)
 {
-	OutputObjectUnmodified = InOutputObject;
-	OutputObjectModified = InOutputObject ? DuplicateObject(InOutputObject, nullptr) : nullptr;
+	OutputObject = InOutputObject;
 	OnOutputObjectSetEvent.Broadcast();
 }
 
-void SDetailsDiff::GetModifications(FArchive& Archive) const
+UObject* SDetailsDiff::GetOutputObject() const
 {
-	OutputObjectModified->GetClass()->SerializeTaggedProperties(
-		Archive, reinterpret_cast<uint8*>(OutputObjectModified),
-		OutputObjectUnmodified->GetClass(), (uint8*)OutputObjectUnmodified
-	);
+	return OutputObject;
 }
 
-void SDetailsDiff::RequestModifications(FArchive& Archive) const
+bool SDetailsDiff::IsOutputEnabled() const
 {
-	OutputObjectModified->GetClass()->SerializeTaggedProperties(
-		Archive, reinterpret_cast<uint8*>(OutputObjectModified),
-		OutputObjectModified->GetClass(), reinterpret_cast<uint8*>(OutputObjectModified->GetClass()->ClassDefaultObject.Get())
-	);
+	return OutputObject != nullptr;
 }
 
-bool SDetailsDiff::IsOutputEnabled() const
+void SDetailsDiff::ReportMergeConflicts(const TMap<FString, TMap<FPropertySoftPath, ETreeDiffResult>>& Conflicts)
 {
-	return OutputObjectUnmodified && OutputObjectModified;
+	MergeConflicts = Conflicts;
 }
 
 void SDetailsDiff::NextDiff()
@@ -339,50 +367,70 @@ void SDetailsDiff::GenerateDifferencesList()
 	PrimaryDifferencesList.Empty();
 	RealDifferences.Empty();
 	ModePanels.Empty();
+	OnOutputObjectSetEvent.Clear(); // will be repopulated by ModePanel generation methods
 
 	// Now that we have done the diffs, create the panel widgets
 	// (we're currently only generating the details panel but we can add more as needed)
-	ModePanels.Add(DetailsMode, GenerateDetailsPanel());
+
+	const auto GetBlueprintCDO = [](const UObject* Object)->const UObject*
+	{
+		return CastChecked<UBlueprint>(Object)->GeneratedClass->GetDefaultObject();
+	};
+	
+	TFunction<const UObject*(const UObject*)> Redirector;
+	if (PanelOld.Object->IsA<UBlueprint>() && PanelNew.Object->IsA<UBlueprint>())
+	{
+		// Blueprints diff their GeneratedClass CDO in the details panel instead
+		Redirector = GetBlueprintCDO;
+	}
+	
+	ModePanels.Add(DetailsMode, GenerateDetailsPanel(Redirector));
 
 	DifferencesTreeView->RebuildList();
 }
 
-SDetailsDiff::FDiffControl SDetailsDiff::GenerateDetailsPanel()
+
+SDetailsDiff::FDiffControl SDetailsDiff::GenerateDetailsPanel(const TFunction<const UObject*(const UObject*)>& Redirector)
 {
-	const TSharedPtr<FDetailsDiffControl> NewDiffControl = MakeShared<FDetailsDiffControl>(PanelOld.Object, PanelNew.Object, FOnDiffEntryFocused::CreateRaw(this, &SDetailsDiff::SetCurrentMode, DetailsMode), true);
+	const UObject* OldObject = Redirector ? Redirector(PanelOld.Object) : PanelOld.Object;
+	const UObject* NewObject = Redirector ? Redirector(PanelNew.Object) : PanelNew.Object;
+	
+	const TSharedPtr<FDetailsDiffControl> NewDiffControl = MakeShared<FDetailsDiffControl>(OldObject, NewObject, FOnDiffEntryFocused::CreateRaw(this, &SDetailsDiff::SetCurrentMode, DetailsMode), true);
 	NewDiffControl->EnableComments(DifferencesTreeView.ToWeakPtr());
 	NewDiffControl->GenerateTreeEntries(PrimaryDifferencesList, RealDifferences);
 	
 	const TSharedRef<SDetailsSplitter> Splitter = SNew(SDetailsSplitter);
-	if (PanelOld.Object)
+	if (OldObject)
 	{
 		Splitter->AddSlot(
 			SDetailsSplitter::Slot()
 			.Value(0.5f)
-			.DetailsView(NewDiffControl->GetDetailsWidget(PanelOld.Object))
-			.DifferencesWithRightPanel(NewDiffControl.ToSharedRef(), &FDetailsDiffControl::GetDifferencesWithRight, Cast<UObject>(PanelOld.Object))
+			.DetailsView(NewDiffControl->GetDetailsWidget(OldObject))
+			.DifferencesWithRightPanel(NewDiffControl.ToSharedRef(), &FDetailsDiffControl::GetDifferencesWithRight, Cast<UObject>(OldObject))
 		);
 	}
-	if (PanelNew.Object)
+	if (NewObject)
 	{
 		Splitter->AddSlot(
 			SDetailsSplitter::Slot()
 			.Value(0.5f)
-			.DetailsView(NewDiffControl->GetDetailsWidget(PanelNew.Object))
-			.DifferencesWithRightPanel(NewDiffControl.ToSharedRef(), &FDetailsDiffControl::GetDifferencesWithRight, Cast<UObject>(PanelNew.Object))
+			.DetailsView(NewDiffControl->GetDetailsWidget(NewObject))
+			.DifferencesWithRightPanel(NewDiffControl.ToSharedRef(), &FDetailsDiffControl::GetDifferencesWithRight, Cast<UObject>(NewObject))
 		);
 	}
 
 	
 	const TWeakPtr<SDetailsSplitter> WeakSplitter = Splitter;
 	const TWeakPtr<FDetailsDiffControl> WeakDiffControl = NewDiffControl;
-	OnOutputObjectSetEvent.AddLambda([WeakSplitter, WeakDiffControl, this]()
+	OnOutputObjectSetEvent.AddLambda([WeakSplitter, WeakDiffControl, this, Redirector]()
 	{
+		const UObject* CenterObject = Redirector ? Redirector(OutputObject) : OutputObject;
+		
 		const TSharedPtr<SDetailsSplitter> Splitter = WeakSplitter.Pin();
 		const TSharedPtr<FDetailsDiffControl> DiffControl = WeakDiffControl.Pin();
 		if (Splitter && DiffControl)
 		{
-			TSharedRef<IDetailsView> DetailsView = DiffControl->InsertObject(OutputObjectModified, false, 1);
+			TSharedRef<IDetailsView> DetailsView = DiffControl->InsertObject(CenterObject, false, 1);
 
 			// allow user to edit the output panel
 			DetailsView->SetIsPropertyEditingEnabledDelegate(FIsPropertyEditingEnabled::CreateStatic([]{return true; }));
@@ -390,14 +438,40 @@ SDetailsDiff::FDiffControl SDetailsDiff::GenerateDetailsPanel()
 			// insert the output object as a central panel
 			Splitter->AddSlot(
 				SDetailsSplitter::Slot()
-					.DetailsView(DiffControl->GetDetailsWidget(OutputObjectModified))
+					.DetailsView(DiffControl->GetDetailsWidget(CenterObject))
 					.Value(0.5f)
 					.IsReadonly(false)
-					.DifferencesWithRightPanel(DiffControl.ToSharedRef(), &FDetailsDiffControl::GetDifferencesWithRight, (const UObject*)OutputObjectModified),
+					.DifferencesWithRightPanel(DiffControl.ToSharedRef(), &FDetailsDiffControl::GetDifferencesWithRight, (const UObject*)CenterObject),
 				1 // insert between left and right panel (index 1)
 			);
+
+			// highlight merge conflicts
+			TMap<FString, TMap<FPropertySoftPath, FLinearColor>> Highlights;
+			for (auto& [ObjectPath, Properties] : MergeConflicts)
+			{
+				for (auto& [propertyPath, Diff] : Properties)
+				{
+					switch(Diff)
+					{
+					case ETreeDiffResult::MissingFromTree1: // fall through
+					case ETreeDiffResult::MissingFromTree2: // fall through
+					case ETreeDiffResult::DifferentValues:
+						// color is intentionally using values greater than 1 so that it stays very saturated
+						Highlights.FindOrAdd(ObjectPath).Add(propertyPath, FLinearColor(1.5f, 0.3f, 0.3f));
+						break;
+					
+					default:; // ignore identical and invalid
+					}
+				}
+			}
+			
+			Splitter->AddHighlights(Highlights);
 		}
 	});
+	if (OutputObject)
+	{
+		OnOutputObjectSetEvent.Broadcast();
+	}
 
 	SDetailsDiff::FDiffControl Ret;
 	Ret.DiffControl = NewDiffControl;
@@ -443,6 +517,23 @@ void SDetailsDiff::SetCurrentMode(FName NewMode)
 	OnModeChanged(NewMode);
 }
 
+void SDetailsDiff::RefreshCurrentModePanel()
+{
+
+	FDiffControl* FoundControl = ModePanels.Find(CurrentMode);
+
+	if (FoundControl)
+	{
+		ModeContents->SetContent(FoundControl->Widget.ToSharedRef());
+	}
+	else
+	{
+		ensureMsgf(false, TEXT("Diff panel does not support mode %s"), *CurrentMode.ToString() );
+	}
+
+	OnModeChanged(CurrentMode);
+}
+
 void SDetailsDiff::UpdateTopSectionVisibility(const FName& InNewViewMode) const
 {
 	SSplitter* TopRevisionInfoWidgetPtr = TopRevisionInfoWidget.Get();
diff --git a/Engine/Source/Editor/Kismet/Public/SDetailsDiff.h b/Engine/Source/Editor/Kismet/Public/SDetailsDiff.h
--- a/Engine/Source/Editor/Kismet/Public/SDetailsDiff.h
+++ b/Engine/Source/Editor/Kismet/Public/SDetailsDiff.h
@@ -8,6 +8,7 @@
 #include "Widgets/SWindow.h"
 #include "Widgets/SCompoundWidget.h"
 
+enum class ETreeDiffResult;
 class FSpawnTabArgs;
 class FTabManager;
 class IDiffControl;
@@ -65,17 +66,16 @@ public:
 	static TSharedRef<SDetailsDiff> CreateDiffWindow(const UObject* OldObject, const UObject* NewObject, const FRevisionInfo& OldRevision, const FRevisionInfo& NewRevision, const UClass* ObjectClass);
 
 	/** Enables actions like "Choose Left" and "Choose Right" which will modify the OutputObject */
-	void SetOutputObject(const UObject* OutputObject);
+	void SetOutputObject(UObject* OutputObject);
 	
 	/** Return a serialized buffer of change requests made by the user */
-	void GetModifications(FArchive& Archive) const;
-	
-	/** submit an archive with change requests to make to OutputObject */
-	void RequestModifications(FArchive& Archive) const;
+	UObject* GetOutputObject() const;
 
 	/** Returns whether SetOutputObject was called with a valid object */
 	bool IsOutputEnabled() const;
 
+	void ReportMergeConflicts(const TMap<FString, TMap<FPropertySoftPath, ETreeDiffResult>>& Conflicts);
+
 	DECLARE_MULTICAST_DELEGATE_OneParam(FOnWindowClosedEvent, TSharedRef<SDetailsDiff>)
 	FOnWindowClosedEvent OnWindowClosedEvent;
 protected:
@@ -95,6 +95,8 @@ protected:
 	/** Called when editor may need to be closed */
 	void OnCloseAssetEditor(UObject* Asset, EAssetEditorCloseReason CloseReason);
 
+	void OnObjectReplaced(const FCoreUObjectDelegates::FReplacementObjectMap& Replacements);
+
 	struct FDiffControl
 	{
 		FDiffControl()
@@ -107,12 +109,13 @@ protected:
 		TSharedPtr< class IDiffControl > DiffControl;
 	};
 
-	FDiffControl GenerateDetailsPanel();
+	FDiffControl GenerateDetailsPanel(const TFunction<const UObject*(const UObject*)>& Redirector = nullptr);
 
 	TSharedRef<SBox> GenerateRevisionInfoWidgetForPanel(TSharedPtr<SWidget>& OutGeneratedWidget,const FText& InRevisionText) const;
 
 	/** Accessor and event handler for toggling between diff view modes (defaults, components, graph view, interface, macro): */
 	void SetCurrentMode(FName NewMode);
+	void RefreshCurrentModePanel();
 	FName GetCurrentMode() const { return CurrentMode; }
 	void OnModeChanged(const FName& InNewViewMode) const;
 
@@ -125,8 +128,7 @@ protected:
 	FDetailsDiffPanel PanelNew;
 	
 	/** If set, actions like "Choose Left" and "Choose Right" will be enabled and modify the OutputObject */
-	const UObject* OutputObjectUnmodified = nullptr;
-	UObject* OutputObjectModified = nullptr; // clone that can be mutated by user actions
+	UObject* OutputObject = nullptr; // clone that can be mutated by user actions
 
 	DECLARE_MULTICAST_DELEGATE(FOnSetOutputObjectEvent)
 	FOnSetOutputObjectEvent OnOutputObjectSetEvent;
@@ -155,6 +157,8 @@ protected:
 	TWeakPtr<SWindow> WeakParentWindow;
 
 	FDelegateHandle AssetEditorCloseDelegate;
+
+	TMap<FString, TMap<FPropertySoftPath, ETreeDiffResult>> MergeConflicts;
 };
 
 
diff --git a/Engine/Source/Editor/PropertyEditor/Private/AsyncDetailViewDiff.cpp b/Engine/Source/Editor/PropertyEditor/Private/AsyncDetailViewDiff.cpp
--- a/Engine/Source/Editor/PropertyEditor/Private/AsyncDetailViewDiff.cpp
+++ b/Engine/Source/Editor/PropertyEditor/Private/AsyncDetailViewDiff.cpp
@@ -16,7 +16,7 @@ static const UObject* GetObject(const TSharedPtr<FDetailTreeNode>& TreeNode)
 
 static FResolvedProperty GetResolvedProperty(const TSharedPtr<FPropertyNode>& PropertyNode, const UObject* Object)
 {
-	if (PropertyNode)
+	if (PropertyNode && Object)
 	{
 		const TSharedRef<FPropertyPath> PropertyPath = FPropertyNode::CreatePropertyPath(PropertyNode.ToSharedRef());
 		if (PropertyPath->IsValid())
diff --git a/Engine/Source/Editor/PropertyEditor/Private/SDetailsSplitter.cpp b/Engine/Source/Editor/PropertyEditor/Private/SDetailsSplitter.cpp
--- a/Engine/Source/Editor/PropertyEditor/Private/SDetailsSplitter.cpp
+++ b/Engine/Source/Editor/PropertyEditor/Private/SDetailsSplitter.cpp
@@ -90,7 +90,38 @@ static bool TryGetDestinationContainer(TSharedPtr<FDetailTreeNode> DetailsNode,
 	return false;
 }
 
-static void CopyPropertyValueForInsert(TSharedPtr<FDetailTreeNode> SourceDetailsNode, TSharedPtr<FDetailTreeNode> DestinationDetailsNode)
+static void SendPropertyChangeNotification(FPropertyNode* PropertyNode, UObject* OuterObject, EPropertyChangeType::Type ChangeType)
+{
+	// find an object to send the FPropertyChangedEvent to
+	FProperty* ChangedProperty = PropertyNode->GetProperty();
+	PropertyNode = PropertyNode->GetParentNode();
+	while(PropertyNode && !CastField<FObjectProperty>(PropertyNode->GetProperty()))
+	{
+		if (FProperty* Property = PropertyNode->GetProperty())
+		{
+			ChangedProperty = Property;
+		}
+		PropertyNode = PropertyNode->GetParentNode();
+	}
+
+	// send the notification
+	FPropertyChangedEvent ChangeEvent(ChangedProperty, ChangeType);
+	if (PropertyNode)
+	{
+		const FObjectProperty* AsObjectProperty = CastFieldChecked<FObjectProperty>(PropertyNode->GetProperty());
+		const FPropertySoftPath SoftPropertyPath(FPropertyNode::CreatePropertyPath(PropertyNode->AsShared()).Get());
+		const FResolvedProperty Resolved = SoftPropertyPath.Resolve(OuterObject);
+		UObject* InnerObject = AsObjectProperty->GetObjectPropertyValue(Resolved.Property->ContainerPtrToValuePtr<void>(Resolved.Object));
+		InnerObject->PostEditChangeProperty(ChangeEvent);
+	}
+	else
+	{
+		// if an FObjectProperty wasn't found, this property must be owned by OuterObject directly
+		OuterObject->PostEditChangeProperty(ChangeEvent);
+	}
+}
+
+static void CopyPropertyValueForInsert(const TSharedPtr<FDetailTreeNode>& SourceDetailsNode, const TSharedPtr<FDetailTreeNode>& DestinationDetailsNode)
 {
 	TUniquePtr<FScriptArrayHelper> SourceArray;
 	FPropertyNode* SourceArrayPropertyNode;
@@ -107,6 +138,10 @@ static void CopyPropertyValueForInsert(TSharedPtr<FDetailTreeNode> SourceDetails
 			const FArrayProperty* ArrayProperty = CastFieldChecked<FArrayProperty>(DestinationArrayPropertyNode->GetProperty());
 			const FProperty* ElementProperty = ArrayProperty->Inner;
 			ElementProperty->CopySingleValue(DestinationData, SourceData);
+			SendPropertyChangeNotification(
+				DestinationArrayPropertyNode,
+				DestinationDetailsNode->GetDetailsView()->GetSelectedObjects()[0].Get(),
+				EPropertyChangeType::ArrayAdd);
 		}
 	}
 	
@@ -121,6 +156,11 @@ static void CopyPropertyValueForInsert(TSharedPtr<FDetailTreeNode> SourceDetails
 		{
 			const void* SourceData = SourceSet->FindNthElementPtr(SourceDetailsNode->GetPropertyNode()->GetArrayIndex());
 			DestinationSet->AddElement(SourceData);
+			
+			SendPropertyChangeNotification(
+				DestinationPropertyNode,
+				DestinationDetailsNode->GetDetailsView()->GetSelectedObjects()[0].Get(),
+				EPropertyChangeType::ArrayAdd);
 		}
 	}
 	
@@ -137,11 +177,16 @@ static void CopyPropertyValueForInsert(TSharedPtr<FDetailTreeNode> SourceDetails
 			const void* SourceKey = SourceMap->GetKeyPtr(Index);
 			const void* SourceVal = SourceMap->GetValuePtr(Index);
 			DestinationMap->AddPair(SourceKey, SourceVal);
+			
+			SendPropertyChangeNotification(
+				DestinationPropertyNode,
+				DestinationDetailsNode->GetDetailsView()->GetSelectedObjects()[0].Get(),
+				EPropertyChangeType::ArrayAdd);
 		}
 	}
 }
 
-static void CopyPropertyValue(TSharedPtr<FDetailTreeNode> SourceDetailsNode, TSharedPtr<FDetailTreeNode> DestinationDetailsNode, ETreeDiffResult Diff)
+static void CopyPropertyValue(const TSharedPtr<FDetailTreeNode>& SourceDetailsNode, const TSharedPtr<FDetailTreeNode>& DestinationDetailsNode, ETreeDiffResult Diff)
 {
 	switch(Diff)
 	{
@@ -193,6 +238,10 @@ static void CopyPropertyValue(TSharedPtr<FDetailTreeNode> SourceDetailsNode, TSh
 		{
 			DestinationResolved.Property->CopyCompleteValue(DestinationData, SourceData);
 		}
+		SendPropertyChangeNotification(
+			DestinationDetailsNode->GetPropertyNode().Get(),
+			DestinationDetailsNode->GetDetailsView()->GetSelectedObjects()[0].Get(),
+			EPropertyChangeType::ValueSet);
 	}
 }
 
@@ -296,6 +345,11 @@ void SDetailsSplitter::AddSlot(const FSlot::FSlotArguments& SlotArgs, int32 Inde
 	}, Index);
 }
 
+void SDetailsSplitter::AddHighlights(const TMap<FString, TMap<FPropertySoftPath, FLinearColor>>& Highlights)
+{
+	CustomHighlights = Highlights;
+}
+
 SDetailsSplitter::FSlot::FSlotArguments SDetailsSplitter::Slot()
 {
 	return FSlot::FSlotArguments(MakeUnique<FSlot>());
@@ -328,7 +382,7 @@ int32 SDetailsSplitter::OnPaint(const FPaintArgs& Args, const FGeometry& Allotte
 			
 			Diff->ForEachRow([&](const TUniquePtr<FAsyncDetailViewDiff::DiffNodeType>& DiffNode, int32, int32)->ETreeTraverseControl
 			{
-				const FLinearColor Color = FLinearColor(0.f,1.f,1.f);
+				const FLinearColor Color = GetHighlightColor(DiffNode);
 				
 				FSlateRect LeftPropertyRect;
 				if (const TSharedPtr<FDetailTreeNode> LeftDetailNode = DiffNode->ValueA.Pin())
@@ -372,11 +426,10 @@ int32 SDetailsSplitter::OnPaint(const FPaintArgs& Args, const FGeometry& Allotte
 					
 					if (LeftPropertyRect.IsValid() && RightPropertyRect.IsValid() && DiffNode->DiffResult != ETreeDiffResult::Identical)
 					{
-						FLinearColor FillColor = Color.Desaturate(.3f) * FLinearColor(0.04f,0.04f,0.04f);
-						FillColor.A = 0.3f;
+						FLinearColor FillColor = Color.Desaturate(.3f) * FLinearColor(0.053f,0.053f,0.053f);
+						FillColor.A = 0.43f;
 			
-						FLinearColor OutlineColor = Color;
-						OutlineColor.A = 0.7f;
+						const FLinearColor OutlineColor = Color;
 						PaintPropertyConnector(OutDrawElements, MaxLayerId, LeftPropertyRect, RightPropertyRect, FillColor, OutlineColor);
 						++MaxLayerId;
 
@@ -409,8 +462,8 @@ int32 SDetailsSplitter::OnPaint(const FPaintArgs& Args, const FGeometry& Allotte
 
 	for (const auto& [PropertyRect, Color] : RowHighlights)
 	{
-		FLinearColor FillColor = Color.Desaturate(.3) * FLinearColor(0.04f,0.04f,0.04f);
-		FillColor.A = 0.3f;
+		FLinearColor FillColor = Color.Desaturate(.3f) * FLinearColor(0.053f,0.053f,0.053f);
+		FillColor.A = 0.43f;
 		
 		FPaintGeometry Geometry(
 			PropertyRect.GetTopLeft() + FVector2D{0.f,2.f},
@@ -761,4 +814,54 @@ void SDetailsSplitter::PaintCopyPropertyButton(FSlateWindowElementList& OutDrawE
 		ESlateDrawEffect::None,
 		ButtonColor
 	);
-}
\ No newline at end of file
+}
+
+// find the inner-most object in the property path and shorten the path relative to that object
+static void ShortenToPathFromLastObject(const UObject*& InOutObject, FPropertyPath& InOutPath)
+{
+	TArray<FPropertyInfo> PathFromSubObjectReversed;
+	PathFromSubObjectReversed.Add(InOutPath.GetLeafMostProperty());
+	InOutPath = *InOutPath.TrimPath(1);
+	while (InOutPath.IsValid())
+	{
+		// Note that if we have perf issues, this could be made faster by writing a custom resolve function that stops
+		// at the last FObjectProperty. that way we wouldn't need to resolve multiple times
+		const FResolvedProperty Resolved = FPropertySoftPath(InOutPath).Resolve(InOutObject);
+		if (const FObjectProperty* ObjectProperty = CastField<FObjectProperty>(Resolved.Property))
+		{
+			InOutObject = ObjectProperty->GetObjectPropertyValue(ObjectProperty->ContainerPtrToValuePtr<void>(Resolved.Object));
+			break;
+		}
+		PathFromSubObjectReversed.Add(InOutPath.GetLeafMostProperty());
+		InOutPath = *InOutPath.TrimPath(1);
+	}
+
+	InOutPath = {};
+	while(!PathFromSubObjectReversed.IsEmpty())
+	{
+		InOutPath.AddProperty(PathFromSubObjectReversed.Pop());
+	}
+}
+
+FLinearColor SDetailsSplitter::GetHighlightColor(const TUniquePtr<TDiffNode<TWeakPtr<FDetailTreeNode>>>& DiffNode) const
+{
+	const TSharedPtr<FDetailTreeNode> DetailNode = DiffNode->ValueA.IsValid() ? DiffNode->ValueA.Pin() : DiffNode->ValueB.Pin();
+	
+	FPropertyPath Path = DetailNode->GetPropertyPath();
+	const UObject* OwningObject = DetailNode->GetDetailsView()->GetSelectedObjects()[0].Get();
+	if (Path.IsValid())
+	{
+		FPropertyPath PathFromSubObject;
+		ShortenToPathFromLastObject(OwningObject, Path);
+		
+		if (const TMap<FPropertySoftPath, FLinearColor>* Highlights = CustomHighlights.Find(OwningObject->GetPathName(OwningObject->GetPackage())))
+		{
+			if (const FLinearColor* Highlight = Highlights->Find(FPropertySoftPath(Path)))
+			{
+				return *Highlight;
+			}
+		}
+	}
+	
+	return FLinearColor(0.f, 1.f, 1.f, .7f);
+}
diff --git a/Engine/Source/Editor/PropertyEditor/Public/SDetailsSplitter.h b/Engine/Source/Editor/PropertyEditor/Public/SDetailsSplitter.h
--- a/Engine/Source/Editor/PropertyEditor/Public/SDetailsSplitter.h
+++ b/Engine/Source/Editor/PropertyEditor/Public/SDetailsSplitter.h
@@ -41,6 +41,8 @@ public:
 	
 	void AddSlot(const FSlot::FSlotArguments& SlotArgs, int32 Index = INDEX_NONE);
 
+	void AddHighlights(const TMap<FString, TMap<FPropertySoftPath, FLinearColor>>& Highlights);
+
 private:
 	
 	struct FPanel
@@ -70,8 +72,11 @@ private:
 	
 	void PaintCopyPropertyButton(FSlateWindowElementList& OutDrawElements, int32 LayerId, const TUniquePtr<TDiffNode<TWeakPtr<FDetailTreeNode>>>& DiffNode,
 		const FSlateRect& LeftPropertyRect, const FSlateRect& RightPropertyRect, EPropertyCopyDirection CopyDirection) const;
+
+	FLinearColor GetHighlightColor(const TUniquePtr<TDiffNode<TWeakPtr<FDetailTreeNode>>>& DiffNode) const;
 	
 	TSharedPtr<SSplitter> Splitter;
 	TArray<FPanel> Panels;
 	FCopyPropertyButton HoveredCopyButton;
+	TMap<FString, TMap<FPropertySoftPath, FLinearColor>> CustomHighlights;
 };
diff --git a/Engine/Source/Editor/UnrealEd/Internal/AsyncTreeDifferences.h b/Engine/Source/Editor/UnrealEd/Internal/AsyncTreeDifferences.h
--- a/Engine/Source/Editor/UnrealEd/Internal/AsyncTreeDifferences.h
+++ b/Engine/Source/Editor/UnrealEd/Internal/AsyncTreeDifferences.h
@@ -15,8 +15,11 @@ enum class ETreeDiffResult
 
 enum class ETreeTraverseOrder
 {
-	PreOrder, // parent than children
-	PostOrder // children then parent
+	PreOrder, // parent then children
+	PostOrder, // children then parent
+	
+	ReversePreOrder, // parent then children in reverse order
+	ReversePostOrder // children in reverse order then parent
 };
 
 enum class ETreeTraverseControl
@@ -136,6 +139,8 @@ private:
 	
 	static bool PreOrderRecursive(const TUniquePtr<DiffNodeType>& Node, const TFunction<ETreeTraverseControl(const TUniquePtr<DiffNodeType>&)>& Method);
 	static bool PostOrderRecursive(const TUniquePtr<DiffNodeType>& Node, const TFunction<ETreeTraverseControl(const TUniquePtr<DiffNodeType>&)>& Method);
+	static bool ReversePreOrderRecursive(const TUniquePtr<DiffNodeType>& Node, const TFunction<ETreeTraverseControl(const TUniquePtr<DiffNodeType>&)>& Method);
+	static bool ReversePostOrderRecursive(const TUniquePtr<DiffNodeType>& Node, const TFunction<ETreeTraverseControl(const TUniquePtr<DiffNodeType>&)>& Method);
 
 	TArray<DiffNodeType*> UpdateQueue;
 	TAttribute<TArray<ValueType>> RootValuesA;
@@ -220,6 +225,12 @@ void TAsyncTreeDifferences<InNodeType>::ForEach(ETreeTraverseOrder TraversalOrde
 	case ETreeTraverseOrder::PostOrder:
 		PostOrderRecursive(Head, Method);
 		break;
+	case ETreeTraverseOrder::ReversePreOrder:
+		ReversePreOrderRecursive(Head, Method);
+		break;
+	case ETreeTraverseOrder::ReversePostOrder:
+		ReversePostOrderRecursive(Head, Method);
+		break;
 	default: check(false);
 	}
 }
@@ -269,6 +280,53 @@ bool TAsyncTreeDifferences<InNodeType>::PostOrderRecursive(const TUniquePtr<Diff
 	return true;
 }
 
+template <typename InNodeType>
+bool TAsyncTreeDifferences<InNodeType>::ReversePreOrderRecursive(const TUniquePtr<DiffNodeType>& Node, const TFunction<ETreeTraverseControl(const TUniquePtr<DiffNodeType>&)>& Method)
+{
+	if (Node->DiffResult != ETreeDiffResult::Invalid)
+	{
+		switch (Method(Node))
+		{
+			case ETreeTraverseControl::Break: return false;
+			case ETreeTraverseControl::SkipChildren: return true;
+		}
+	}
+	
+	for (int32 I = Node->Children.Num() - 1; I >= 0; --I)
+	{
+		const TUniquePtr<DiffNodeType>& Child = Node->Children[I];
+		if (!ReversePreOrderRecursive(Child, Method))
+		{
+			return false;
+		}
+	}
+
+	return true;
+}
+
+template <typename InNodeType>
+bool TAsyncTreeDifferences<InNodeType>::ReversePostOrderRecursive(const TUniquePtr<DiffNodeType>& Node, const TFunction<ETreeTraverseControl(const TUniquePtr<DiffNodeType>&)>& Method)
+{
+	for (int32 I = Node->Children.Num() - 1; I >= 0; --I)
+	{
+		const TUniquePtr<DiffNodeType>& Child = Node->Children[I];
+		if (!ReversePostOrderRecursive(Child, Method))
+		{
+			return false;
+		}
+	}
+	
+	if (Node->DiffResult != ETreeDiffResult::Invalid)
+	{
+		switch (Method(Node))
+		{
+			case ETreeTraverseControl::Break: return false;
+		}
+	}
+
+	return true;
+}
+
 template <typename InNodeType>
 void TAsyncTreeDifferences<InNodeType>::ProcessTopOfQueue()
 {
diff --git a/Engine/Source/Editor/UnrealEd/Private/MergeUtils.cpp b/Engine/Source/Editor/UnrealEd/Private/MergeUtils.cpp
new file mode 100644
--- /dev/null
+++ b/Engine/Source/Editor/UnrealEd/Private/MergeUtils.cpp
@@ -0,0 +1,1239 @@
+// Copyright Epic Games, Inc. All Rights Reserved.
+
+#include "MergeUtils.h"
+
+#include "DiffUtils.h"
+#include "ISourceControlModule.h"
+#include "ISourceControlRevision.h"
+#include "SDetailsDiff.h"
+#include "SourceControlOperations.h"
+#include "AsyncTreeDifferences.h"
+#include "Algo/RandomShuffle.h"
+#include "Editor.h"
+#include "HAL/PlatformFileManager.h"
+#include "Engine/Blueprint.h"
+#include "Kismet2/KismetEditorUtilities.h"
+#include "Kismet2/BlueprintEditorUtils.h"
+#include "UObject/Linker.h"
+#include "AssetRegistry/AssetRegistryModule.h"
+
+#define LOCTEXT_NAMESPACE "MergeUtils"
+
+struct FScopedMergeResolveTransaction
+{
+	FScopedMergeResolveTransaction(UObject* InManagedObject, EMergeFlags InFlags)
+		: ManagedObject(InManagedObject)
+		, Flags(InFlags)
+	{
+		if (Flags & MF_HANDLE_SOURCE_CONTROL)
+		{
+			UndoHandler = NewObject<UUndoableResolveHandler>();
+			UndoHandler->SetFlags(RF_Transactional);
+			UndoHandler->SetManagedObject(InManagedObject);
+			
+			TransactionNum = GEditor->BeginTransaction(LOCTEXT("ResolveMerge", "ResolveAutoMerge"));
+			ensure(UndoHandler->Modify());
+			if (InManagedObject)
+			{
+				ensure(InManagedObject->Modify());
+			}
+		}
+	}
+
+	void Cancel()
+	{
+		bCanceled = true;
+	}
+	
+	~FScopedMergeResolveTransaction()
+	{
+		if (Flags & MF_HANDLE_SOURCE_CONTROL)
+		{
+			if (!bCanceled)
+			{
+				UndoHandler->MarkResolved();
+				GEditor->EndTransaction();
+			}
+			else if (ManagedObject.IsValid())
+			{
+				ManagedObject->GetPackage()->SetDirtyFlag(false);
+				GEditor->CancelTransaction(TransactionNum);
+			}
+		}
+	}
+
+	TWeakObjectPtr<UObject> ManagedObject;
+	UUndoableResolveHandler* UndoHandler = nullptr;
+	EMergeFlags Flags;
+	int TransactionNum = 0;
+	bool bCanceled = false;
+};
+
+void UUndoableResolveHandler::SetManagedObject(UObject* Object)
+{
+	ManagedObject = Object;
+	const UPackage* Package = ManagedObject->GetPackage();
+	const FString Filepath = FPaths::ConvertRelativePathToFull(Package->GetLoadedPath().GetLocalFullPath());
+	
+	ISourceControlProvider& Provider = ISourceControlModule::Get().GetProvider();
+	const FSourceControlStatePtr SourceControlState = Provider.GetState(Package, EStateCacheUsage::Use);
+	const ISourceControlState::FResolveInfo ResolveInfo = SourceControlState->GetResolveInfo();
+	BaseRevisionNumber = SourceControlState->GetResolveInfo().BaseRevision;
+	if (const TSharedPtr<class ISourceControlRevision, ESPMode::ThreadSafe> CurrentRevision = SourceControlState->GetCurrentRevision())
+	{
+		CurrentRevisionNumber = FString::FromInt(CurrentRevision->GetRevisionNumber());
+	}
+	else
+	{
+		CurrentRevisionNumber = {};
+	}
+	CheckinIdentifier = SourceControlState->GetCheckInIdentifier();
+
+	// save package and copy the package to a temp file so it can be reverted
+	const FString BaseFilename = FPaths::GetBaseFilename(Filepath);
+	BackupFilepath = FPaths::CreateTempFilename(*(FPaths::ProjectSavedDir()/TEXT("Temp")), *BaseFilename.Left(32));
+	ensure(FPlatformFileManager::Get().GetPlatformFile().CopyFile(*BackupFilepath, *Filepath));
+}
+
+void UUndoableResolveHandler::MarkResolved()
+{
+	if (ManagedObject.IsValid())
+	{
+		const UPackage* Package = ManagedObject->GetPackage();
+		const FString Filepath = FPaths::ConvertRelativePathToFull(Package->GetLoadedPath().GetLocalFullPath());
+		
+		ISourceControlProvider& Provider = ISourceControlModule::Get().GetProvider();
+		Provider.Execute(ISourceControlOperation::Create<FResolve>(), TArray{Filepath}, EConcurrency::Synchronous);
+		bShouldBeResolved = true;
+	}
+}
+
+void UUndoableResolveHandler::PostEditUndo()
+{
+	if (bShouldBeResolved) // redo resolution
+	{
+		MarkResolved();
+	}
+	else if (ManagedObject.IsValid())// undo resolution
+	{
+		UPackage* Package = ManagedObject->GetPackage();
+		const FString Filepath = FPaths::ConvertRelativePathToFull(Package->GetLoadedPath().GetLocalFullPath());
+		
+		if (BaseRevisionNumber.IsEmpty() || CurrentRevisionNumber.IsEmpty())
+		{
+			ensure(FPlatformFileManager::Get().GetPlatformFile().CopyFile(*Filepath, *BackupFilepath));
+			return;
+		}
+		
+		// to force the file to revert to it's pre-resolved state, we must revert, sync back to base revision,
+		// apply the conflicting changes, then sync forward again.
+		ISourceControlProvider& Provider = ISourceControlModule::Get().GetProvider();
+		{
+			const TSharedRef<FSync> SyncOperation = ISourceControlOperation::Create<FSync>();
+			SyncOperation->SetRevision(BaseRevisionNumber);
+			Provider.Execute(SyncOperation, Filepath, EConcurrency::Synchronous);
+		}
+		
+		ResetLoaders(Package);
+		Provider.Execute( ISourceControlOperation::Create<FRevert>(), Filepath, EConcurrency::Synchronous);
+
+		{
+			const TSharedRef<FCheckOut> CheckoutOperation = ISourceControlOperation::Create<FCheckOut>();
+			Provider.Execute(CheckoutOperation, CheckinIdentifier, {Filepath}, EConcurrency::Synchronous);
+		}
+
+		ensure(FPlatformFileManager::Get().GetPlatformFile().CopyFile(*Filepath, *BackupFilepath));
+		
+		{
+			const TSharedRef<FSync> SyncOperation = ISourceControlOperation::Create<FSync>();
+			SyncOperation->SetRevision(CurrentRevisionNumber);
+			Provider.Execute(SyncOperation, Filepath, EConcurrency::Synchronous);
+		}
+
+		Provider.Execute(ISourceControlOperation::Create<FUpdateStatus>(), {Filepath}, EConcurrency::Asynchronous);
+	}
+	UObject::PostEditUndo();
+}
+
+static UPackage* LoadMergePackage(const FString& SCFile, const FString& Revision, const UPackage* LocalPackage)
+{
+	const FString FileWithRevision = SCFile + TEXT("#") + Revision;
+	const TSharedRef<FDownloadFile, ESPMode::ThreadSafe> DownloadFileOperation = ISourceControlOperation::Create<FDownloadFile>(FPaths::DiffDir(), FDownloadFile::EVerbosity::Full);
+	ISourceControlModule::Get().GetProvider().Execute(DownloadFileOperation, FileWithRevision, EConcurrency::Synchronous);
+	const FString DownloadPath = FPaths::ConvertRelativePathToFull(FPaths::DiffDir() / FPaths::GetCleanFilename(FileWithRevision));
+	FString CopyPath = DownloadPath;
+	CopyPath.ReplaceInline(TEXT(".uasset"), TEXT(""));
+	CopyPath.ReplaceCharInline('#', '-');
+	CopyPath.ReplaceCharInline('.', '-');
+	CopyPath = FPaths::CreateTempFilename(*FPaths::GetPath(CopyPath), *FPaths::GetBaseFilename(CopyPath), TEXT(".uasset"));
+
+	if (FPlatformFileManager::Get().GetPlatformFile().CopyFile(*CopyPath, *DownloadPath))
+	{
+		return DiffUtils::LoadPackageForDiff(FPackagePath::FromLocalPath(CopyPath), LocalPackage->GetLoadedPath());
+	}
+	return nullptr;
+}
+
+struct FBPReferenceFinder : public FArchiveUObject
+{
+	FBPReferenceFinder(const UObject* Obj, TArray<const UObject*>& InReferences)
+		: FArchiveUObject()
+		, References(InReferences)
+		, OwningPackage(Obj->GetOutermost())
+	{
+		// Copying FPackageHarvester:
+		SetIsPersistent(true);
+		SetIsSaving(true);
+		ArIsObjectReferenceCollector = true;
+		ArShouldSkipBulkData = true;
+		if (Obj->HasAnyFlags(RF_ClassDefaultObject))
+		{
+			Obj->GetClass()->SerializeDefaultObject(const_cast<UObject*>(Obj), *this);
+		}
+		else
+		{
+			const_cast<UObject*>(Obj)->Serialize(*this);
+		}
+	}
+
+private:
+	virtual FArchive& operator<<(UObject*& ObjRef) override
+	{
+		if (ObjRef != nullptr &&
+			(!ObjRef->HasAnyFlags(RF_Transient) || ObjRef->IsNative()) &&
+			!ObjRef->IsIn(GetTransientPackage()) &&
+			ObjRef != OwningPackage)
+		{
+			// Set to null any pointer to an external asset
+			References.Add(ObjRef);
+		}
+
+		return *this;
+	}
+
+	TArray<const UObject*>& References;
+	UPackage* OwningPackage;
+};
+
+static TArray<FString> GatherExportPaths(const UObject* Object)
+{
+	TArray<FString> OutExports;
+	
+	// find everything roots references that is in package and put it into
+	// OutExports, unless it is in the disallow list.. puts other references
+	// into OutImports:
+	
+	// Object won't be modified. const_cast just simplifies code
+	TArray PendingRefs{const_cast<UObject*>(Object)};
+
+	TSet<UObject*> RefsProcessed;
+	RefsProcessed.Append(PendingRefs);
+
+	TArray<UObject*> ScratchRefs; // just keeping this allocation alive across iterations
+	UPackage* Package = Object->GetPackage();
+
+	while (PendingRefs.Num())
+	{
+		const UObject* Iter = PendingRefs.Pop();
+		if (ensure(Iter->GetPackage() == Package))
+		{
+			const FString PathName = Iter->GetPathName(Package);
+            if (ensure(Iter == FindObject<UObject>(Package, *PathName)))
+            {
+            	OutExports.Add(PathName);
+            }
+		}
+		
+		ScratchRefs.Add(Iter->GetClass());
+		FBPReferenceFinder ReferencedObjects(Iter, (TArray<const UObject*>&)ScratchRefs);
+		for (UObject* Obj : ScratchRefs)
+		{
+			if (RefsProcessed.Contains(Obj))
+			{
+				continue;
+			}
+			RefsProcessed.Add(Obj);
+
+			if (!Obj->IsIn(Package))
+			{
+				// found import. ignore.
+				continue;
+			}
+
+			PendingRefs.Add(Obj);
+		}
+		ScratchRefs.Reset();
+	}
+	
+	OutExports.Sort();
+	
+	return OutExports;
+}
+
+enum class EObjectPointerType : uint8
+{
+	Nullptr,
+	Import,
+	Export,
+};
+
+static UObject* DuplicateForMerge(const UObject* Source, UObject* DestinationOuter)
+{
+	FObjectDuplicationParameters DuplicateParams(const_cast<UObject*>(Source), DestinationOuter);
+	DuplicateParams.PortFlags |= PPF_DuplicateVerbatim;
+	DuplicateParams.bSkipPostLoad = true;
+	DuplicateParams.DestName = Source->GetFName();
+	return StaticDuplicateObjectEx(DuplicateParams);
+}
+
+struct FPropertyInstance
+{
+	explicit FPropertyInstance() = default;
+	FPropertyInstance(const void* InObject, const FProperty* InProperty)
+		: Object(InObject), KeyProperty(InProperty), ValProperty(InProperty)
+	{}
+	FPropertyInstance(const void* InObject, const FProperty* InValProperty, const FProperty* InKeyProperty)
+		: Object(InObject), KeyProperty(InKeyProperty), ValProperty(InValProperty)
+	{}
+
+	operator bool() const
+	{
+		return Object != nullptr && KeyProperty != nullptr && ValProperty != nullptr;
+	}
+
+	bool operator==(const FPropertyInstance& Other) const
+	{
+		if (!ValProperty || !Other.ValProperty)
+		{
+			return ValProperty == Other.ValProperty;
+		}
+		
+		if (ValProperty == Other.ValProperty)
+		{
+			const void* DataA = ValProperty->ContainerPtrToValuePtr<void*>(Object);
+			const void* DataB = Other.ValProperty->ContainerPtrToValuePtr<void*>(Other.Object);
+			check(ValProperty->ElementSize == Other.ValProperty->ElementSize);
+			check(ValProperty->GetClass() == Other.ValProperty->GetClass());
+			
+			return ValProperty->Identical(DataA, DataB, PPF_DeepComparison);
+		}
+		return false;
+	}
+
+	void GetChildren(TArray<FPropertyInstance>& OutChildren) const
+	{
+		auto TryAddChild = [&OutChildren](const void* Data, const FProperty* NewValProperty, const FProperty* NewKeyProperty = nullptr)
+		{
+			// ignore transient properties
+			if (!NewValProperty->HasAnyPropertyFlags(CPF_Transient | CPF_DuplicateTransient | CPF_NonPIEDuplicateTransient))
+			{
+				OutChildren.Emplace(Data, NewValProperty, NewKeyProperty ? NewKeyProperty : NewValProperty);
+			}
+		};
+		
+		if (const FStructProperty* StructProperty = CastField<FStructProperty>(ValProperty))
+		{
+			// treat FSoftObjectPath as a leaf
+			if (StructProperty->Struct != TBaseStructure<FSoftObjectPath>::Get())
+			{
+				const void* StructPtr = StructProperty->ContainerPtrToValuePtr<void*>(Object);
+                for (TFieldIterator<FProperty> PropertyIt(StructProperty->Struct); PropertyIt; ++PropertyIt)
+                {
+                	TryAddChild(StructPtr, *PropertyIt);
+                }
+			}
+		}
+		else if (const FArrayProperty* ArrayProperty = CastField<FArrayProperty>(ValProperty))
+		{
+			FScriptArrayHelper Helper(ArrayProperty, ArrayProperty->ContainerPtrToValuePtr<void>(Object));
+			for (int32 Index = 0; Index < Helper.Num(); ++Index)
+			{
+				TryAddChild(Helper.GetElementPtr(Index), ArrayProperty->Inner);
+			}
+		}
+		else if (const FSetProperty* SetProperty = CastField<FSetProperty>(ValProperty))
+		{
+			FScriptSetHelper Helper(SetProperty, SetProperty->ContainerPtrToValuePtr<void>(Object));
+			for (FScriptSetHelper::FIterator It = Helper.CreateIterator(); It; ++It)
+			{
+				TryAddChild(Helper.GetElementPtr(*It), SetProperty->ElementProp);
+			}
+		}
+		else if (const FMapProperty* MapProperty = CastField<FMapProperty>(ValProperty))
+		{
+			FScriptMapHelper MapHelper(MapProperty, MapProperty->ContainerPtrToValuePtr<void>(Object));
+			for (FScriptMapHelper::FIterator MapIt = MapHelper.CreateIterator(); MapIt; ++MapIt)
+			{
+				TryAddChild(MapHelper.GetPairPtr(*MapIt), MapProperty->ValueProp, MapProperty->KeyProp);
+			}
+		}
+
+		// UObject* is treated a leaf so we don't need to enter FObjectProperty
+	}
+
+	const void* Object = nullptr;
+	const FProperty* KeyProperty = nullptr; // equal to ValProperty unless this is a map element
+	const FProperty* ValProperty = nullptr;
+};
+
+// methods that make FPropertyInstance diffable
+namespace TreeDiffSpecification
+{
+	template<>
+	bool AreValuesEqual<FPropertyInstance>(const FPropertyInstance& TreeNodeA, const FPropertyInstance& TreeNodeB)
+	{
+		return TreeNodeA == TreeNodeB;
+	}
+	
+	template<>
+	bool AreMatching<FPropertyInstance>(const FPropertyInstance& TreeNodeA, const FPropertyInstance& TreeNodeB)
+	{
+		if (TreeNodeA.KeyProperty->GetName() == TreeNodeB.KeyProperty->GetName())
+		{
+			if (CastField<FSetProperty>(TreeNodeA.KeyProperty->Owner.ToField()) || CastField<FMapProperty>(TreeNodeA.KeyProperty->Owner.ToField()))
+			{
+				const void* DataA = TreeNodeA.KeyProperty->ContainerPtrToValuePtr<void*>(TreeNodeA.Object);
+				const void* DataB = TreeNodeB.KeyProperty->ContainerPtrToValuePtr<void*>(TreeNodeB.Object);
+				return TreeNodeA.KeyProperty->Identical(DataA, DataB, PPF_DeepComparison);
+			}
+			else
+			{
+				return true;
+			}
+		}
+		return false;
+	}
+	
+	template<>
+	void GetChildren<FPropertyInstance>(const FPropertyInstance& InParent, TArray<FPropertyInstance>& OutChildren)
+	{
+		return InParent.GetChildren(OutChildren);
+	}
+
+	template<>
+	bool ShouldMatchByValue<FPropertyInstance>(const FPropertyInstance& TreeNodeA)
+	{
+		// array elements should match by value
+		return CastField<FArrayProperty>(TreeNodeA.KeyProperty->Owner.ToField()) != nullptr;
+	}
+};
+
+static TAsyncTreeDifferences<FPropertyInstance> ObjectPropTreeDiff(const UObject* ObjectA, const UObject* ObjectB)
+{
+	TArray<FPropertyInstance> ChildrenA;
+	if (ObjectA)
+	{
+		for (TFieldIterator<FProperty> PropertyIt(ObjectA->GetClass()); PropertyIt; ++PropertyIt)
+		{
+			if (!PropertyIt->HasAnyPropertyFlags(CPF_Transient | CPF_DuplicateTransient | CPF_NonPIEDuplicateTransient))
+			{
+				ChildrenA.Emplace(ObjectA, *PropertyIt);
+			}
+		}
+	}
+	TArray<FPropertyInstance> ChildrenB;
+	if (ObjectB)
+	{
+		for (TFieldIterator<FProperty> PropertyIt(ObjectB->GetClass()); PropertyIt; ++PropertyIt)
+		{
+			if (!PropertyIt->HasAnyPropertyFlags(CPF_Transient | CPF_DuplicateTransient | CPF_NonPIEDuplicateTransient))
+			{
+				ChildrenB.Emplace(ObjectB, *PropertyIt);
+			}
+		}
+	}
+	TAsyncTreeDifferences<FPropertyInstance> TreeDifferences(ChildrenA, ChildrenB);
+	TreeDifferences.FlushQueue();
+	return TreeDifferences;
+}
+
+// converts a TAsyncTreeDifferences<FPropertyInstance>::DiffNodeType into a FPropertySoftPath
+using DiffNode = TAsyncTreeDifferences<FPropertyInstance>::DiffNodeType;
+static FPropertySoftPath ToPropertySoftPath(const TUniquePtr<DiffNode>& Node)
+{
+	FString Path;
+	TArray<FName> Chain;
+	for (const DiffNode* Current = Node.Get(); Current->Parent; Current = Current->Parent)
+	{
+		const FPropertyInstance& Instance = (Current->ValueA.ValProperty) ? Current->ValueA : Current->ValueB;
+		const FPropertyInstance& Parent = (Current->ValueA.ValProperty) ? Current->Parent->ValueA : Current->Parent->ValueB;
+		FName Name = Instance.ValProperty->GetFName();
+		if (Parent.ValProperty)
+		{
+			if (const FArrayProperty* ArrayProperty = CastField<FArrayProperty>(Parent.ValProperty))
+            {
+            	FScriptArrayHelper Helper(ArrayProperty, ArrayProperty->ContainerPtrToValuePtr<void>(Parent.Object));
+				// possible optimization: use pointer arithmetic to avoid loop?
+            	for (int32 I = 0; I < Helper.Num(); ++I)
+            	{
+            		if (Helper.GetElementPtr(I) == Instance.ValProperty->ContainerPtrToValuePtr<void>(Instance.Object))
+            		{
+            			Name = FName(FString::FromInt(I));
+            			break;
+            		}
+            	}
+            }
+            else if (const FMapProperty* MapProperty = CastField<FMapProperty>(Parent.ValProperty))
+            {
+            	FScriptMapHelper Helper(MapProperty, MapProperty->ContainerPtrToValuePtr<void>(Parent.Object));
+            	const int32 I = Helper.FindMapIndexWithKey(Instance.KeyProperty->ContainerPtrToValuePtr<void>(Instance.Object));
+            	Name = FName(FString::FromInt(I));
+            }
+            else if (const FSetProperty* SetProperty = CastField<FSetProperty>(Parent.ValProperty))
+            {
+	            FScriptSetHelper Helper(SetProperty, SetProperty->ContainerPtrToValuePtr<void>(Parent.Object));
+            	const int32 I = Helper.FindElementIndex(Instance.KeyProperty->ContainerPtrToValuePtr<void>(Instance.Object));
+            	Name = FName(FString::FromInt(I));
+            }
+		}
+		Chain.Add(Name);
+	}
+	Algo::Reverse(Chain);
+	return FPropertySoftPath(Chain);
+}
+
+template<typename CppType>
+static void TryRelinkProperty(void* Data, const TFObjectPropertyBase<CppType>* Property, const UPackage* FromPackage, UPackage* ToPackage)
+{
+	CppType* CppTypePtr = Property->GetPropertyValuePtr(Data);
+	const UObject* Object = CppTypePtr->Get();
+	if (Object && Object->IsIn(FromPackage))
+	{
+		if (UObject* Found = FindObjectChecked<UObject>(ToPackage, *Object->GetPathName(FromPackage)))
+		{
+			check(Found->IsIn(ToPackage));
+			*CppTypePtr = CppType(Found);
+		}
+	}
+}
+
+static void TryRelinkProperty(void* Data, const FStructProperty* Property, const UPackage* FromPackage, UPackage* ToPackage)
+{
+    if (Property->Struct == TBaseStructure<FSoftObjectPath>::Get())
+    {
+    	FSoftObjectPath* PathPtr = reinterpret_cast<FSoftObjectPath*>(Data);
+    	const UObject* Object = PathPtr->ResolveObject();
+    	if (Object && Object->IsIn(FromPackage))
+    	{
+    		if (UObject* Found = FindObjectChecked<UObject>(ToPackage, *Object->GetPathName(FromPackage)))
+    		{
+    			check(Found->IsIn(ToPackage));
+    			*PathPtr = FSoftObjectPath(Found);
+    		}
+    	}
+    }
+}
+
+static void TryRelinkProperty(void* Data, const FInterfaceProperty* Property, const UPackage* FromPackage, UPackage* ToPackage)
+{
+	FScriptInterface* Interface = Property->GetPropertyValuePtr(Data);
+	const UObject* Object = Interface->GetObject();
+    if (Object && Object->IsIn(FromPackage))
+    {
+    	if (UObject* Found = FindObjectChecked<UObject>(ToPackage, *Object->GetPathName(FromPackage)))
+    	{
+    		check(Found->IsIn(ToPackage));
+    		*Interface = FScriptInterface(Found, Found->GetInterfaceAddress(Property->InterfaceClass));
+    	}
+    }
+}
+
+// redirect pointers to UObjects in FromPackage to UObjects in ToPackage
+static void RelinkObjectProperties(const TArray<FPropertyInstance>& Props, UPackage* FromPackage, UPackage* ToPackage)
+{
+	check(FromPackage && ToPackage);
+	for (const FPropertyInstance& Instance : Props)
+	{
+		void *Data = const_cast<void*>(Instance.ValProperty->ContainerPtrToValuePtr<void>(Instance.Object));
+		
+		if (const FObjectProperty* ObjectProperty = CastField<FObjectProperty>(Instance.ValProperty))
+		{
+			TryRelinkProperty(Data, ObjectProperty, FromPackage, ToPackage);
+		}
+		else if (const FSoftObjectProperty* SoftObjectProperty = CastField<FSoftObjectProperty>(Instance.ValProperty))
+		{
+			TryRelinkProperty(Data, SoftObjectProperty, FromPackage, ToPackage);
+		}
+		else if (const FWeakObjectProperty* WeakObjectProperty = CastField<FWeakObjectProperty>(Instance.ValProperty))
+		{
+			TryRelinkProperty(Data, WeakObjectProperty, FromPackage, ToPackage);
+		}
+		else if (const FLazyObjectProperty* LazyObjectProperty = CastField<FLazyObjectProperty>(Instance.ValProperty))
+		{
+			TryRelinkProperty(Data, LazyObjectProperty, FromPackage, ToPackage);
+		}
+		else if (const FStructProperty* StructProperty = CastField<FStructProperty>(Instance.ValProperty))
+		{
+			TryRelinkProperty(Data, StructProperty, FromPackage, ToPackage);
+		}
+		else if (const FInterfaceProperty* InterfaceProperty = CastField<FInterfaceProperty>(Instance.ValProperty))
+		{
+			TryRelinkProperty(Data, InterfaceProperty, FromPackage, ToPackage);
+		}
+		
+		TArray<FPropertyInstance> Children;
+		TreeDiffSpecification::GetChildren(Instance, Children);
+		if (Children.Num())
+		{
+			RelinkObjectProperties(Children, FromPackage, ToPackage);
+		}
+	}
+}
+
+// redirect pointers to UObjects in FromPackage to UObjects in ToPackage
+static void RelinkObjectProperties(UObject* Object, UPackage* FromPackage, UPackage* ToPackage)
+{
+	TArray<FPropertyInstance> PropertyInstances;
+	for (TFieldIterator<FProperty> PropertyIt(Object->GetClass()); PropertyIt; ++PropertyIt)
+	{
+		PropertyInstances.Emplace(Object, *PropertyIt);
+	}
+
+	RelinkObjectProperties(PropertyInstances, FromPackage, ToPackage);
+}
+
+template<typename CppType>
+bool SoftCompareProperty(const TFObjectPropertyBase<CppType>* Property, void* DataA, void* DataB, const UPackage* PackageA, const UPackage* PackageB)
+{
+	CppType* CppTypePtrA = Property->GetPropertyValuePtr(DataA);
+	CppType* CppTypePtrB = Property->GetPropertyValuePtr(DataB);
+	const UObject* ObjectA = CppTypePtrA->Get();
+	const UObject* ObjectB = CppTypePtrB->Get();
+	if (ObjectA && ObjectA->IsIn(PackageA) && ObjectB && ObjectB->IsIn(PackageB))
+	{
+		return ObjectA->GetPathName(PackageA) == ObjectB->GetPathName(PackageB);
+	}
+	return ObjectA == ObjectB;
+}
+
+bool SoftCompareProperty(const FStructProperty* Property, void* DataA, void* DataB, const UPackage* PackageA, const UPackage* PackageB)
+{
+	if (Property->Struct == TBaseStructure<FSoftObjectPath>::Get())
+	{
+		const UObject* ObjectA = reinterpret_cast<FSoftObjectPath*>(DataA)->ResolveObject();
+		const UObject* ObjectB = reinterpret_cast<FSoftObjectPath*>(DataB)->ResolveObject();
+		if (ObjectA->IsIn(PackageA) && ObjectB->IsIn(PackageB))
+		{
+			return ObjectA->GetPathName(PackageA) == ObjectB->GetPathName(PackageB);
+		}
+		return ObjectA == ObjectB;
+	}
+	return Property->Identical(DataA, DataB, PPF_DeepComparison);
+}
+
+bool SoftCompareProperty(const FInterfaceProperty* Property, void* DataA, void* DataB, const UPackage* PackageA, const UPackage* PackageB)
+{
+	const UObject* ObjectA = Property->GetPropertyValuePtr(DataA)->GetObject();
+	const UObject* ObjectB = Property->GetPropertyValuePtr(DataB)->GetObject();
+	if (ObjectA->IsIn(PackageA) && ObjectB->IsIn(PackageB))
+	{
+		return ObjectA->GetPathName(PackageA) == ObjectB->GetPathName(PackageB);
+	}
+	return ObjectA == ObjectB;
+}
+
+bool SoftCompareProperty(const FFieldPathProperty* Property, void* DataA, void* DataB, const UPackage* PackageA, const UPackage* PackageB)
+{
+	const FFieldPath FieldPathA = Property->GetPropertyValue(DataA);
+	const FFieldPath FieldPathB = Property->GetPropertyValue(DataB);
+	FString PathAString = FieldPathA.ToString();
+	FString PathBString = FieldPathB.ToString();
+	if (PathAString.RemoveFromStart(PackageA->GetName() + TEXT(".")) && PathBString.RemoveFromStart(PackageB->GetName() + TEXT(".")))
+	{
+		// if the field path is within the merging packages, compare the path without the package
+		return PathAString == PathBString;
+	}
+	return FieldPathA == FieldPathB;
+}
+
+static bool SoftCompare(const FPropertyInstance& A, const FPropertyInstance& B, const UPackage* PackageA, const UPackage* PackageB)
+{
+	void *DataA = const_cast<void*>(A.ValProperty->ContainerPtrToValuePtr<void>(A.Object));
+	void *DataB = const_cast<void*>(B.ValProperty->ContainerPtrToValuePtr<void>(B.Object));
+	
+	if (const FObjectProperty* Prop = CastField<FObjectProperty>(A.ValProperty))
+	{
+		return SoftCompareProperty(Prop, DataA, DataB, PackageA, PackageB);
+	}
+	if (const FSoftObjectProperty* Prop = CastField<FSoftObjectProperty>(A.ValProperty))
+	{
+		return SoftCompareProperty(Prop, DataA, DataB, PackageA, PackageB);
+	}
+	if (const FWeakObjectProperty* Prop = CastField<FWeakObjectProperty>(A.ValProperty))
+	{
+		return SoftCompareProperty(Prop, DataA, DataB, PackageA, PackageB);
+	}
+	if (const FLazyObjectProperty* Prop = CastField<FLazyObjectProperty>(A.ValProperty))
+	{
+		return SoftCompareProperty(Prop, DataA, DataB, PackageA, PackageB);
+	}
+	if (const FStructProperty* Prop = CastField<FStructProperty>(A.ValProperty))
+	{
+		return SoftCompareProperty(Prop, DataA, DataB, PackageA, PackageB);
+	}
+	if (const FInterfaceProperty* Prop = CastField<FInterfaceProperty>(A.ValProperty))
+	{
+		return SoftCompareProperty(Prop, DataA, DataB, PackageA, PackageB);
+	}
+	if (const FFieldPathProperty* Prop = CastField<FFieldPathProperty>(A.ValProperty))
+	{
+		return SoftCompareProperty(Prop, DataA, DataB, PackageA, PackageB);
+	}
+	return A.ValProperty->Identical(DataA, DataB, PPF_DeepComparison);
+}
+
+// diff algorithm that gets all the property differences from every object in ExportPaths
+// note that FObjectProperties are shallow diffed
+static TMap<FString, TMap<FPropertySoftPath, ETreeDiffResult>> GetDifferences(const TArray<FString>& ExportPaths, UPackage* PackageA, UPackage* PackageB)
+{
+	TMap<FString, TMap<FPropertySoftPath, ETreeDiffResult>> Result;
+	for (const FString& Path : ExportPaths)
+	{
+		const UObject* ObjectA = PackageA ? FindObject<UObject>(PackageA, *Path) : nullptr;
+		const UObject* ObjectB =  PackageB ? FindObject<UObject>(PackageB, *Path) : nullptr;
+
+		const TAsyncTreeDifferences<FPropertyInstance> TreeDiff = ObjectPropTreeDiff(ObjectA, ObjectB);
+
+		// find every item that differs between source and default
+		TreeDiff.ForEach(ETreeTraverseOrder::PreOrder,
+			[&Result, &Path, PackageA, PackageB](const TUniquePtr<DiffNode>& Node)->ETreeTraverseControl
+			{
+				switch(Node->DiffResult)
+				{
+					case ETreeDiffResult::MissingFromTree1:
+						Result.FindOrAdd(Path).Add(ToPropertySoftPath(Node), Node->DiffResult);
+						break;
+					case ETreeDiffResult::MissingFromTree2:
+						Result.FindOrAdd(Path).Add(ToPropertySoftPath(Node), Node->DiffResult);
+						break;
+					case ETreeDiffResult::DifferentValues:
+						{
+							// if this isn't a leaf, continue to children to find more detailed diff info
+							if (!Node->Children.IsEmpty())
+							{
+								return ETreeTraverseControl::Continue;
+							}
+
+							if (!SoftCompare(Node->ValueA, Node->ValueB, PackageA, PackageB))
+							{
+								Result.FindOrAdd(Path).Add(ToPropertySoftPath(Node), Node->DiffResult);
+							}
+						}
+						break;
+					case ETreeDiffResult::Identical:
+						break;
+					default: check(false);
+				}
+				
+				return ETreeTraverseControl::Continue;
+			});
+	}
+	return Result;
+}
+
+// looks for subobjects of Template that aren't in Object's package and duplicates them
+static void DuplicateMissingSubobjects(UObject* Object, const UObject* Template)
+{
+	check(Object && Template);
+	UPackage* Package = Object->GetPackage();
+	UPackage* TemplatePackage = Template->GetPackage();
+	const TArray<FString> ExportPaths = GatherExportPaths(Template);
+	
+	for (const FString& ExportPath : ExportPaths)
+	{
+		if (!FindObject<UObject>(Package, *ExportPath))
+		{
+			const UObject* CurTemplate = FindObjectChecked<UObject>(TemplatePackage, *ExportPath);
+			UObject* Outer;
+			if (CurTemplate->GetOuter()->IsA<UPackage>())
+			{
+				Outer = Package;
+			}
+			else
+			{
+				FString OuterPath = CurTemplate->GetOuter()->GetPathName(TemplatePackage);
+				Outer = FindObjectChecked<UObject>(Package, *OuterPath);
+			}
+			check(Outer);
+			DuplicateForMerge(CurTemplate, Outer);
+		}
+	}
+}
+
+static void NotifyPropertyChange(UObject* Object, const TUniquePtr<DiffNode>& Node, EPropertyChangeType::Type Type)
+{			
+	if (Object)
+	{
+		const DiffNode* Itr = Node.Get();
+		while(Itr->Parent->ValueA.ValProperty != nullptr)
+		{
+			Itr = Itr->Parent;
+		}
+		
+		// notify property change
+		FPropertyChangedEvent Event(
+			const_cast<FProperty*>(Itr->ValueA.ValProperty),
+			Type
+		);
+		Object->PostEditChangeProperty(Event);
+	}
+};
+
+// assigns Node->ValueB to Node->ValueA
+static void HandleMergeAssign(UObject* Object, const TUniquePtr<DiffNode>& Node)
+{
+	check(Node->DiffResult == ETreeDiffResult::DifferentValues);
+	const void* Source = Node->ValueB.ValProperty->ContainerPtrToValuePtr<void>(Node->ValueB.Object);
+	void* Dest = const_cast<void*>(Node->ValueA.ValProperty->ContainerPtrToValuePtr<void>(Node->ValueA.Object));
+	Node->ValueA.ValProperty->CopyCompleteValue(Dest, Source);
+	NotifyPropertyChange(Object, Node, EPropertyChangeType::ValueSet);
+}
+
+// inserts Node->ValueB into Node->Parent->ValueA
+static void HandleMergeInsert(UObject* Object, const TUniquePtr<DiffNode>& Node)
+{
+	check(Node->DiffResult == ETreeDiffResult::MissingFromTree1);
+	if (Node->Parent->ValueA)
+	{
+        const FProperty* SourceProperty = Node->ValueB.ValProperty;
+        const void* SourceValue = Node->ValueB.ValProperty->ContainerPtrToValuePtr<void>(Node->ValueB.Object);
+        check(SourceProperty && SourceValue)
+        
+        const FProperty* ParentResultProperty = Node->Parent->ValueA.ValProperty;
+        const void* ParentResultObject = Node->Parent->ValueA.Object;
+        check(ParentResultProperty && ParentResultObject);
+        
+        if (const FArrayProperty* ArrayProperty = CastField<FArrayProperty>(ParentResultProperty))
+        {
+        	FScriptArrayHelper Helper(ArrayProperty, ArrayProperty->ContainerPtrToValuePtr<void>(ParentResultObject));
+    
+        	// count the number of valid result siblings before this node to get the insert index
+        	int32 InsertIndex = 0;
+        	for (const TUniquePtr<DiffNode>& Sibling : Node->Parent->Children)
+        	{
+        		if (Sibling == Node)
+        		{
+        			break;
+        		}
+        		if (Sibling->DiffResult != ETreeDiffResult::MissingFromTree1)
+        		{
+        			++InsertIndex;
+        		}
+        	}
+    
+        	Helper.InsertValues(InsertIndex, 1);
+    
+        	SourceProperty->CopyCompleteValue(Helper.GetElementPtr(InsertIndex), SourceValue);
+        	
+        	NotifyPropertyChange(Object, Node, EPropertyChangeType::ArrayAdd);
+        	return;
+        }
+        else if (const FSetProperty* SetProperty = CastField<FSetProperty>(ParentResultProperty))
+        {
+        	FScriptSetHelper Helper(SetProperty, SetProperty->ContainerPtrToValuePtr<void>(ParentResultObject));
+        	Helper.AddElement(SourceValue);
+        	NotifyPropertyChange(Object, Node, EPropertyChangeType::ArrayAdd);
+        	return;
+        }
+        else if (const FMapProperty* MapProperty = CastField<FMapProperty>(ParentResultProperty))
+        {
+        	FScriptMapHelper Helper(MapProperty, MapProperty->ContainerPtrToValuePtr<void>(ParentResultObject));
+        	const FProperty* SourceKeyProperty = Node->ValueB.KeyProperty;
+        	
+        	Helper.AddPair(
+        		SourceKeyProperty->ContainerPtrToValuePtr<void>(Node->ValueB.Object),
+        		SourceValue
+        	);
+        	NotifyPropertyChange(Object, Node, EPropertyChangeType::ArrayAdd);
+        	return;
+        }
+	}
+
+	// if you get this warning it's likely because of a type mismatch between the objects being merged.
+	UE_LOG(LogSourceControl, Warning, TEXT("Data loss in Merge: property: [%s] in object: %s"), *ToPropertySoftPath(Node).ToDisplayName(), *Object->GetName())
+}
+
+// removes Node->ValueA from Node->Parent->ValueA
+static void HandleMergeRemove(UObject* Object, const TUniquePtr<DiffNode>& Node)
+{
+	check(Node->DiffResult == ETreeDiffResult::MissingFromTree2);
+	if (Node->Parent->ValueA)
+	{
+		const FProperty* ResultProperty = Node->ValueA.ValProperty;
+		const void* ResultObject = Node->ValueA.Object;
+		check(ResultProperty && ResultObject)
+		
+		const FProperty* ParentResultProperty = Node->Parent->ValueA.ValProperty;
+		const void* ParentResultObject = Node->Parent->ValueA.Object;
+		check(ParentResultProperty && ParentResultObject);
+		
+		if (const FArrayProperty* ArrayProperty = CastField<FArrayProperty>(ParentResultProperty))
+        {
+            FScriptArrayHelper Helper(ArrayProperty, ArrayProperty->ContainerPtrToValuePtr<void>(ParentResultObject));
+            // possible optimization: use pointer arithmetic to avoid loop?
+            for (int I = 0; I < Helper.Num(); ++I)
+            {
+                if (Helper.GetElementPtr(I) == ResultProperty->ContainerPtrToValuePtr<void>(ResultObject))
+                {
+                    Helper.RemoveValues(I, 1);
+                    break;
+                }
+            }
+            NotifyPropertyChange(Object, Node, EPropertyChangeType::ArrayRemove);
+        	return;
+        }
+		else if (const FSetProperty* SetProperty = CastField<FSetProperty>(ParentResultProperty))
+        {
+            FScriptSetHelper Helper(SetProperty, SetProperty->ContainerPtrToValuePtr<void>(ParentResultObject));
+            Helper.RemoveElement(ResultProperty->ContainerPtrToValuePtr<void>(ResultObject));
+            NotifyPropertyChange(Object, Node, EPropertyChangeType::ArrayRemove);
+        	return;
+        }
+		else if (const FMapProperty* MapProperty = CastField<FMapProperty>(ParentResultProperty))
+        {
+            FScriptMapHelper Helper(MapProperty, MapProperty->ContainerPtrToValuePtr<void>(ParentResultObject));
+            const FProperty* ResultKeyProperty = Node->ValueA.KeyProperty;
+            Helper.RemovePair(ResultKeyProperty->ContainerPtrToValuePtr<void>(ResultObject));
+            NotifyPropertyChange(Object, Node, EPropertyChangeType::ArrayRemove);
+        	return;
+        }
+	}
+	
+	// if you get this warning it's likely because of a type mismatch between the objects being merged.
+	UE_LOG(LogSourceControl, Warning, TEXT("Data loss in Merge: property: [%s] in object: %s"), *ToPropertySoftPath(Node).ToDisplayName(), *Object->GetName())
+}
+
+ /**
+  *	PseudoCode:\n
+  *	foreach (CurResultObject, CurSourceObject, CurDefaultObject) in (ResultObject, SourceObject, DefaultObject).SubObjects():
+  *		foreach (ResultProperty, SourceProperty, DefaultProperty) in (CurResultObject, CurSourceObject, CurDefaultObject).Properties():
+  *			if (SourceProperty != DefaultProperty && ResultProperty != SourceProperty):
+  *				ResultProperty = SourceProperty
+  */
+static void CopyDeltaProperties(UObject* ResultObject, const UObject* SourceObject, const UObject* DefaultObject = nullptr)
+{
+	check(ResultObject && SourceObject);
+	UPackage* ResultPackage = ResultObject->GetPackage();
+	UPackage* SourcePackage = SourceObject->GetPackage();
+	UPackage* DefaultPackage = DefaultObject ? DefaultObject->GetPackage() : nullptr;
+	
+	// if SourceObject has extra exports, create them
+	DuplicateMissingSubobjects(ResultObject, SourceObject);
+
+	TArray<FString> ExportPaths = GatherExportPaths(SourceObject);
+	TMap<FString, TMap<FPropertySoftPath, ETreeDiffResult>> SourceDefaultDiff = GetDifferences(ExportPaths, SourcePackage, DefaultPackage);
+	
+	for (const FString& Path : ExportPaths)
+	{
+		const UObject* CurSource = FindObjectChecked<UObject>(SourcePackage, *Path);
+		UObject* CurResult = FindObject<UObject>(ResultPackage, *Path);
+		if (!CurResult)
+		{
+			UObject* Outer = ResultPackage;
+			if (CurSource->GetOuter() != SourcePackage)
+			{
+				Outer = FindObjectChecked<UObject>(ResultPackage, *CurSource->GetOuter()->GetPathName(SourcePackage));
+			}
+			CurResult = DuplicateForMerge(CurSource, Outer);
+		}
+		
+		const TMap<FPropertySoftPath, ETreeDiffResult>* CurSourceDefaultDiff = SourceDefaultDiff.Find(Path);
+		
+		check(CurResult);
+		if (CurSource && CurSourceDefaultDiff)
+		{
+			const TAsyncTreeDifferences<FPropertyInstance> ResultSourceTreeDiff = ObjectPropTreeDiff(CurResult, CurSource);
+		
+			ResultSourceTreeDiff.ForEach(ETreeTraverseOrder::ReversePreOrder,
+				[CurSourceDefaultDiff, CurResult, SourcePackage, ResultPackage] (const TUniquePtr<DiffNode>& Node)->ETreeTraverseControl
+				{
+					// only copy leaf nodes
+					if (!Node->Children.IsEmpty())
+					{
+						// even if insert/delete has children, treat it as a leaf.
+						if (Node->DiffResult != ETreeDiffResult::MissingFromTree1 && Node->DiffResult != ETreeDiffResult::MissingFromTree2)
+						{
+							return ETreeTraverseControl::Continue;
+						}
+					}
+
+					// find cached diff result between SourceObject and DefaultObject for this property
+					ETreeDiffResult SourceDefaultDiffResult = ETreeDiffResult::Identical;
+					if (const ETreeDiffResult* Found = CurSourceDefaultDiff->Find(ToPropertySoftPath(Node)))
+					{
+						SourceDefaultDiffResult = *Found;
+					}
+					
+					// for every item that differs between source and default, copy source to result
+					if (SourceDefaultDiffResult != ETreeDiffResult::Identical)
+					{
+						switch(Node->DiffResult)
+						{
+						case ETreeDiffResult::MissingFromTree1:
+							HandleMergeInsert(CurResult, Node);
+							break;
+						case ETreeDiffResult::MissingFromTree2:
+							HandleMergeRemove(CurResult, Node);
+							break;
+						case ETreeDiffResult::DifferentValues:
+							{
+								if (!SoftCompare(Node->ValueA, Node->ValueB, ResultPackage, SourcePackage))
+								{
+                            		HandleMergeAssign(CurResult, Node);
+								}
+							}
+							
+							break;
+						case ETreeDiffResult::Identical: break;
+						default: check(false);
+						}
+					}
+
+					return ETreeTraverseControl::SkipChildren;
+				});
+		}
+
+		// redirect pointers to UObjects in SourcePackage to UObjects in ResultPackage
+		RelinkObjectProperties(CurResult, SourcePackage, ResultPackage);
+    }
+}
+
+// Copies BranchA's changes before BranchB's changes so that BranchB always has precedence in the returned object
+static UObject* AutoMerge(const UObject* BaseRevision, const UObject* BranchA, const UObject* BranchB, FName PackageName)
+{
+	// apply changes from BranchA first by simply duplicating it over to the result
+	PackageName = MakeUniqueObjectName(nullptr, BranchA->GetClass(), PackageName, EUniqueObjectNameOptions::GloballyUnique);
+	UPackage* MergedPackage = CreatePackage(*(TEXT("/Temp/") + PackageName.ToString()));
+
+	// Base revision into MergedPackage
+	UObject* Merged = DuplicateForMerge(BaseRevision, MergedPackage);
+	
+	// deep copy BranchA's changes into Merged Object
+	CopyDeltaProperties(Merged, BranchA, BaseRevision);
+	
+	// deep copy BranchB's changes into Merged Object
+	CopyDeltaProperties(Merged, BranchB, BaseRevision);
+	
+	return Merged;
+}
+
+EAssetCommandResult MergeUtils::Merge(const FAssetAutomaticMergeArgs& MergeArgs)
+{
+	if (!ensure(MergeArgs.LocalAsset))
+	{
+		return EAssetCommandResult::Unhandled;
+	}
+	
+	FAssetManualMergeArgs ManualMergeArgs;
+	ManualMergeArgs.LocalAsset = MergeArgs.LocalAsset;
+	ManualMergeArgs.ResolutionCallback = MergeArgs.ResolutionCallback;
+	ManualMergeArgs.Flags = MergeArgs.Flags;
+	const UPackage* LocalPackage = ManualMergeArgs.LocalAsset->GetPackage();
+	
+	
+	ISourceControlProvider& SourceControlProvider = ISourceControlModule::Get().GetProvider();
+	const TSharedRef<FUpdateStatus, ESPMode::ThreadSafe> UpdateStatusOperation = ISourceControlOperation::Create<FUpdateStatus>();
+	UpdateStatusOperation->SetUpdateHistory(true);
+	SourceControlProvider.Execute(UpdateStatusOperation, LocalPackage);
+	
+	// Get the SCC state
+	const FSourceControlStatePtr SourceControlState = SourceControlProvider.GetState(LocalPackage, EStateCacheUsage::Use);
+
+	// If we have an asset and its in SCC..
+	if( SourceControlState.IsValid() && SourceControlState->IsSourceControlled() )
+	{
+		const ISourceControlState::FResolveInfo ResolveInfo = SourceControlState->GetResolveInfo();
+		check(ResolveInfo.IsValid());
+		
+		if(UPackage* TempPackage = LoadMergePackage(ResolveInfo.RemoteFile, ResolveInfo.RemoteRevision, LocalPackage))
+		{
+			// Grab the old asset from that old package
+			ManualMergeArgs.RemoteAsset = FindObject<UObject>(TempPackage, *ManualMergeArgs.LocalAsset->GetName());
+
+			// Recovery for package names that don't match
+			if (ManualMergeArgs.RemoteAsset == nullptr)
+			{
+				ManualMergeArgs.RemoteAsset = TempPackage->FindAssetInPackage();
+			}
+		}
+		
+		if(UPackage* TempPackage = LoadMergePackage(ResolveInfo.BaseFile, ResolveInfo.BaseRevision, LocalPackage))
+		{
+			// Grab the old asset from that old package
+			ManualMergeArgs.BaseAsset = FindObject<UObject>(TempPackage, *ManualMergeArgs.LocalAsset->GetName());
+
+			// Recovery for package names that don't match
+			if (ManualMergeArgs.BaseAsset == nullptr)
+			{
+				ManualMergeArgs.BaseAsset = TempPackage->FindAssetInPackage();
+			}
+		}
+	}
+
+	// single asset merging is only supported for assets in a conflicted state in source control
+	if (!ensure(ManualMergeArgs.BaseAsset && ManualMergeArgs.RemoteAsset && ManualMergeArgs.LocalAsset))
+	{
+		return EAssetCommandResult::Unhandled;
+	}
+	
+	return Merge(ManualMergeArgs);
+}
+
+EAssetCommandResult MergeUtils::Merge(const FAssetManualMergeArgs& MergeArgs)
+{
+	auto NotifyResolution = [MergeArgs](EAssetMergeResult Result)
+	{
+		FAssetMergeResults Results;
+		Results.Result = Result;
+		Results.MergedPackage = MergeArgs.LocalAsset->GetPackage();
+		MergeArgs.ResolutionCallback.ExecuteIfBound(Results);
+		return EAssetCommandResult::Handled;
+	};
+
+	// apply changes in different orders to come up with two possible merge options
+	const UObject* FavorRemote = AutoMerge(MergeArgs.BaseAsset,MergeArgs.LocalAsset, MergeArgs.RemoteAsset, TEXT("FavorRemote"));
+	const UObject* FavorLocal = AutoMerge(MergeArgs.BaseAsset,MergeArgs.RemoteAsset,MergeArgs.LocalAsset, TEXT("FavorLocal"));
+	
+	const TArray<FString> FavorRemoteExports = GatherExportPaths(FavorRemote);
+	const TArray<FString> FavorLocalExports = GatherExportPaths(FavorLocal);
+	ensureMsgf(FavorRemoteExports == FavorLocalExports, TEXT("Merge Result has a non-deterministic export list. This is likely because of object names with guids in them."));
+	check(FavorRemoteExports.Num() == FavorLocalExports.Num());
+	
+	TMap<FString, TMap<FPropertySoftPath, ETreeDiffResult>> Conflicts = GetDifferences(
+		FavorLocalExports, FavorLocal->GetPackage(), FavorRemote->GetPackage());
+
+	if (!Conflicts.IsEmpty())
+	{
+		// some properties like uuids mutate on duplicate so this will catch all of those properties and we can ignore them
+		const UObject* FavorLocal2 = AutoMerge(MergeArgs.BaseAsset,MergeArgs.RemoteAsset,MergeArgs.LocalAsset, TEXT("FavorLocal"));
+	
+		const TMap<FString, TMap<FPropertySoftPath, ETreeDiffResult>> FalsePositives = GetDifferences(
+			FavorLocalExports, FavorLocal->GetPackage(), FavorLocal2->GetPackage());
+		
+		for (const auto& [ObjectPath, Differences] : FalsePositives)
+		{
+			for (auto& [Path, Diff] : Differences)
+			{
+				if (TMap<FPropertySoftPath, ETreeDiffResult>* Found = Conflicts.Find(ObjectPath))
+				{
+					Found->Remove(Path);
+					if (Found->IsEmpty())
+					{
+						Conflicts.Remove(ObjectPath);
+					}
+				}
+			}
+		}
+	}
+
+	
+	// if both merge options are the same, we have no conflicts
+	if (Conflicts.IsEmpty())
+	{
+		// auto-merge successful!
+		FScopedMergeResolveTransaction UndoHandler(MergeArgs.LocalAsset, MergeArgs.Flags);
+		
+		// copy changes over to the local asset
+		CopyDeltaProperties(MergeArgs.LocalAsset, FavorLocal);
+		
+		return NotifyResolution(EAssetMergeResult::Completed);
+	}
+	
+	// conflicts detected. We need to ask the user to manually resolve them
+	if (!(MergeArgs.Flags & MF_NO_GUI))
+	{
+		const TSharedRef<SDetailsDiff> DiffView = SDetailsDiff::CreateDiffWindow(MergeArgs.RemoteAsset, MergeArgs.LocalAsset, {}, {}, FavorRemote->GetClass());
+
+		// construct center panel object and copy FavorLocal into it
+		const FName ResultPackageName = MakeUniqueObjectName(nullptr, FavorLocal->GetClass(), TEXT("MergeResult"), EUniqueObjectNameOptions::GloballyUnique);
+		UPackage* ResultPackage = CreatePackage(*(TEXT("/Temp/") + ResultPackageName.ToString()));
+		UObject* ResultObject = DuplicateForMerge(FavorLocal, ResultPackage);
+		
+		DiffView->ReportMergeConflicts(Conflicts);
+		DiffView->SetOutputObject(ResultObject);
+		
+		DiffView->OnWindowClosedEvent.AddLambda([MergeArgs, NotifyResolution](const TSharedRef<SDetailsDiff>& DiffView)
+		{
+			FScopedMergeResolveTransaction UndoHandler(MergeArgs.LocalAsset, MergeArgs.Flags);
+			
+			// copy changes over to the local asset
+			CopyDeltaProperties(MergeArgs.LocalAsset, DiffView->GetOutputObject());
+
+			if (UBlueprint* AsBlueprint = Cast<UBlueprint>(MergeArgs.LocalAsset))
+			{
+				// because merging may have changed the parent class, we should recompile
+				// what we really need from this is to get the CDO up to date so it's type
+				// matches the UBlueprint::ParentClass
+				EBlueprintCompileOptions CompileOptions
+                {
+                		EBlueprintCompileOptions::SkipSave
+                	|	EBlueprintCompileOptions::UseDeltaSerializationDuringReinstancing
+                	|	EBlueprintCompileOptions::SkipNewVariableDefaultsDetection
+                };
+    
+                FKismetEditorUtilities::CompileBlueprint(AsBlueprint, CompileOptions);
+			}
+			
+			NotifyResolution(EAssetMergeResult::Completed);
+		});
+		
+		return EAssetCommandResult::Handled;
+	}
+	
+	return NotifyResolution(EAssetMergeResult::Cancelled);
+}
+
+#if WITH_DEV_AUTOMATION_TESTS
+#include "Misc/AutomationTest.h"
+namespace UE::MergeUtilsTests
+{
+
+	IMPLEMENT_COMPLEX_AUTOMATION_TEST(FMergeWithSelfTests, "ReviewDiffMerge.MergeWithSelf", EAutomationTestFlags::ApplicationContextMask | EAutomationTestFlags::EngineFilter);
+	void FMergeWithSelfTests::GetTests(TArray<FString>& OutBeautifiedNames, TArray<FString>& OutTestCommands) const
+	{
+		FAssetRegistryModule& AssetRegistryModule = FModuleManager::LoadModuleChecked<FAssetRegistryModule>("AssetRegistry");
+		TArray<FAssetData> Assets;
+		AssetRegistryModule.Get().GetAssetsByClass(UBlueprint::StaticClass()->GetClassPathName(), Assets, true);
+		Algo::RandomShuffle(Assets);
+		for (const FAssetData& Asset : Assets)
+		{
+			UBlueprint* BP = CastChecked<UBlueprint>(Asset.GetAsset());
+			if (FBlueprintEditorUtils::IsDataOnlyBlueprint(BP))
+			{
+				OutBeautifiedNames.Add(Asset.AssetName.ToString());
+				OutTestCommands.Add(Asset.PackageName.ToString());
+				if (OutTestCommands.Num() >= 5)
+				{
+					break;
+				}
+			}
+		}
+	}
+
+	bool FMergeWithSelfTests::RunTest(const FString& PackageName)
+	{
+		const FPackagePath Path = FPackagePath::FromPackageNameChecked(PackageName);
+		if (const UPackage* Package = DiffUtils::LoadPackageForDiff(Path, {}))
+		{
+			UBlueprint* BP = CastChecked<UBlueprint>(Package->FindAssetInPackage());
+         	FAssetManualMergeArgs Args;
+         	Args.BaseAsset = BP;
+         	Args.LocalAsset = BP;
+         	Args.RemoteAsset = BP;
+         	Args.Flags = MF_NO_GUI;
+         	bool bSuccess = false;
+         	Args.ResolutionCallback = FOnAssetMergeResolved::CreateLambda(
+         		[&bSuccess](const FAssetMergeResults& Results)
+         		{
+         			bSuccess = Results.Result == EAssetMergeResult::Completed;
+         		});
+         	FString CheckName = FString::Format(TEXT("Self merge: {0}"), {BP->GetFriendlyName()});
+         	MergeUtils::Merge(Args);
+         	UTEST_TRUE(CheckName, bSuccess);
+		}
+        return true;
+	}
+}
+#endif
+
+#undef LOCTEXT_NAMESPACE
\ No newline at end of file
diff --git a/Engine/Source/Editor/UnrealEd/Public/MergeUtils.h b/Engine/Source/Editor/UnrealEd/Public/MergeUtils.h
new file mode 100644
--- /dev/null
+++ b/Engine/Source/Editor/UnrealEd/Public/MergeUtils.h
@@ -0,0 +1,32 @@
+// Copyright Epic Games, Inc. All Rights Reserved.
+
+#pragma once
+#include "AssetDefinition.h"
+#include "MergeUtils.generated.h"
+
+namespace MergeUtils
+{
+	UNREALED_API EAssetCommandResult Merge(const FAssetAutomaticMergeArgs& MergeArgs);
+	UNREALED_API EAssetCommandResult Merge(const FAssetManualMergeArgs& MergeArgs);
+}
+
+UCLASS()
+class UUndoableResolveHandler : public UObject
+{
+public:
+	GENERATED_BODY()
+	void SetManagedObject(UObject* Object);
+	void MarkResolved();
+
+	virtual void PostEditUndo() override;
+
+private:
+	FString BaseRevisionNumber;
+	FString CurrentRevisionNumber;
+	FString BackupFilepath;
+	TWeakObjectPtr<UObject> ManagedObject;
+	TSharedPtr<class ISourceControlChangelist> CheckinIdentifier;
+	
+	UPROPERTY()
+	bool bShouldBeResolved = false;
+};
\ No newline at end of file
