[CrashFix] Missing null checks in MergeUtils.cpp SoftCompareProperty

diff --git a/Engine/Source/Editor/UnrealEd/Private/MergeUtils.cpp b/Engine/Source/Editor/UnrealEd/Private/MergeUtils.cpp
--- a/Engine/Source/Editor/UnrealEd/Private/MergeUtils.cpp
+++ b/Engine/Source/Editor/UnrealEd/Private/MergeUtils.cpp
@@ -612,7 +612,7 @@ bool SoftCompareProperty(const FStructProperty* Property, void* DataA, void* Dat
 	{
 		const UObject* ObjectA = reinterpret_cast<FSoftObjectPath*>(DataA)->ResolveObject();
 		const UObject* ObjectB = reinterpret_cast<FSoftObjectPath*>(DataB)->ResolveObject();
-		if (ObjectA->IsIn(PackageA) && ObjectB->IsIn(PackageB))
+		if (ObjectA && ObjectA->IsIn(PackageA) && ObjectB && ObjectB->IsIn(PackageB))
 		{
 			return ObjectA->GetPathName(PackageA) == ObjectB->GetPathName(PackageB);
 		}
@@ -625,7 +625,7 @@ bool SoftCompareProperty(const FInterfaceProperty* Property, void* DataA, void*
 {
 	const UObject* ObjectA = Property->GetPropertyValuePtr(DataA)->GetObject();
 	const UObject* ObjectB = Property->GetPropertyValuePtr(DataB)->GetObject();
-	if (ObjectA->IsIn(PackageA) && ObjectB->IsIn(PackageB))
+	if (ObjectA && ObjectA->IsIn(PackageA) && ObjectB && ObjectB->IsIn(PackageB))
 	{
 		return ObjectA->GetPathName(PackageA) == ObjectB->GetPathName(PackageB);
 	}
