Cooker IterativeValidate: Fix diff output from DiffWriterArchive that is being logged directly rather than being sent to the MessageCallback.
#rnx

diff --git a/Engine/Source/Editor/UnrealEd/Private/Cooker/DiffWriterArchive.cpp b/Engine/Source/Editor/UnrealEd/Private/Cooker/DiffWriterArchive.cpp
--- a/Engine/Source/Editor/UnrealEd/Private/Cooker/DiffWriterArchive.cpp
+++ b/Engine/Source/Editor/UnrealEd/Private/Cooker/DiffWriterArchive.cpp
@@ -551,8 +551,7 @@ void FDiffWriterArchiveWriter::Compare(
 	const int64 MaxDiffsToLog,
 	int32& InOutDiffsLogged,
 	TMap<FName, FArchiveDiffStats>& OutStats,
-	const UE::DiffWriterArchive::FMessageCallback& MessageCallback,
-	bool bSuppressLogging)
+	const UE::DiffWriterArchive::FMessageCallback& MessageCallback)
 {
 	const int64 SourceSize = SourcePackage.Size - SourcePackage.StartOffset;
 	const int64 DestSize = DestPackage.Size - DestPackage.StartOffset;
@@ -561,11 +560,8 @@ void FDiffWriterArchiveWriter::Compare(
 	
 	if (SourceSize != DestSize)
 	{
-		if (!bSuppressLogging)
-		{
-			MessageCallback(ELogVerbosity::Warning, FString::Printf(
-				TEXT("%s: Size mismatch: on disk: %lld vs memory: %lld"), AssetFilename, SourceSize, DestSize));
-		}
+		MessageCallback(ELogVerbosity::Warning, FString::Printf(
+			TEXT("%s: Size mismatch: on disk: %lld vs memory: %lld"), AssetFilename, SourceSize, DestSize));
 		int64 SizeDiff = DestPackage.Size - SourcePackage.Size;
 		OutStats.FindOrAdd(AssetClass).DiffSize += SizeDiff;
 	}
@@ -608,11 +604,8 @@ void FDiffWriterArchiveWriter::Compare(
 
 			if (DifferenceCallstackoffsetIndex < 0)
 			{
-				if (!bSuppressLogging)
-				{
-					MessageCallback(ELogVerbosity::Warning, FString::Printf(
-						TEXT("%s: Difference at offset %lld (absolute offset: %lld), unknown callstack"), AssetFilename, LocalOffset, DestAbsoluteOffset));
-				}
+				MessageCallback(ELogVerbosity::Warning, FString::Printf(
+					TEXT("%s: Difference at offset %lld (absolute offset: %lld), unknown callstack"), AssetFilename, LocalOffset, DestAbsoluteOffset));
 				continue;
 			}
 
@@ -697,49 +690,50 @@ void FDiffWriterArchiveWriter::Compare(
 						+ FullStackText.RightChop(DebugDataIndex + 2);
 				}
 
-				if (!bSuppressLogging)
-				{
-					MessageCallback(ELogVerbosity::Warning, FString::Printf(
-						TEXT("%s: Difference at offset %lld%s (absolute offset: %lld): byte %d on disk, byte %d in memory, callstack:%s%s%s%s%s"),
-						AssetFilename,
-						CallstackAtOffset.Offset - DestPackage.StartOffset,
-						DestAbsoluteOffset > CallstackAtOffset.Offset ? *FString::Printf(TEXT("(+%lld)"), DestAbsoluteOffset - CallstackAtOffset.Offset) : TEXT(""),
-						DestAbsoluteOffset,
-						SourceByte, DestByte,
-						UE::DiffWriterArchive::NewLineToken,
-						UE::DiffWriterArchive::NewLineToken,
-						*DifferenceCallstackDataText,
-						*DiffValues,
-						*DebugDataStackText
-					));
-				}
+				MessageCallback(ELogVerbosity::Warning, FString::Printf(
+					TEXT("%s: Difference at offset %lld%s (absolute offset: %lld): byte %d on disk, byte %d in memory, callstack:%s%s%s%s%s"),
+					AssetFilename,
+					CallstackAtOffset.Offset - DestPackage.StartOffset,
+					DestAbsoluteOffset > CallstackAtOffset.Offset ? *FString::Printf(TEXT("(+%lld)"), DestAbsoluteOffset - CallstackAtOffset.Offset) : TEXT(""),
+					DestAbsoluteOffset,
+					SourceByte, DestByte,
+					UE::DiffWriterArchive::NewLineToken,
+					UE::DiffWriterArchive::NewLineToken,
+					*DifferenceCallstackDataText,
+					*DiffValues,
+					*DebugDataStackText
+				));
 
 				const int BytesToLog = 128;
-				if (!bSuppressLogging)
+				MessageCallback(ELogVerbosity::Display, FString::Printf(
+					TEXT("%s: Logging %d bytes around absolute offset: %lld (%016X) in the on disk (existing) package, (which corresponds to offset %lld (%016X) in the in-memory package)"),
+					AssetFilename,
+					BytesToLog,
+					SourceAbsoluteOffset,
+					SourceAbsoluteOffset,
+					DestAbsoluteOffset,
+					DestAbsoluteOffset
+				));
+				TArray<FString> HexDumpLines = FCompressionUtil::HexDumpLines(SourcePackage.Data, SourcePackage.Size,
+					SourceAbsoluteOffset - BytesToLog / 2, SourceAbsoluteOffset + BytesToLog / 2);
+				for (FString& Line : HexDumpLines)
 				{
-					MessageCallback(ELogVerbosity::Display, FString::Printf(
-						TEXT("%s: Logging %d bytes around absolute offset: %lld (%016X) in the on disk (existing) package, (which corresponds to offset %lld (%016X) in the in-memory package)"),
-						AssetFilename,
-						BytesToLog,
-						SourceAbsoluteOffset,
-						SourceAbsoluteOffset,
-						DestAbsoluteOffset,
-						DestAbsoluteOffset
-					));
+					MessageCallback(ELogVerbosity::Display, Line);
 				}
-				FCompressionUtil::LogHexDump(SourcePackage.Data, SourcePackage.Size, SourceAbsoluteOffset - BytesToLog / 2, SourceAbsoluteOffset + BytesToLog / 2);
 
-				if (!bSuppressLogging)
+				MessageCallback(ELogVerbosity::Display, FString::Printf(
+					TEXT("%s: Logging %d bytes around absolute offset: %lld (%016X) in the in memory (new) package"),
+					AssetFilename,
+					BytesToLog,
+					DestAbsoluteOffset,
+					DestAbsoluteOffset
+				));
+				HexDumpLines = FCompressionUtil::HexDumpLines(DestPackage.Data, DestPackage.Size, DestAbsoluteOffset - BytesToLog / 2,
+					DestAbsoluteOffset + BytesToLog / 2);
+				for (FString& Line : HexDumpLines)
 				{
-					MessageCallback(ELogVerbosity::Display, FString::Printf(
-						TEXT("%s: Logging %d bytes around absolute offset: %lld (%016X) in the in memory (new) package"),
-						AssetFilename,
-						BytesToLog,
-						DestAbsoluteOffset,
-						DestAbsoluteOffset
-					));
+					MessageCallback(ELogVerbosity::Display, Line);
 				}
-				FCompressionUtil::LogHexDump(DestPackage.Data, DestPackage.Size, DestAbsoluteOffset - BytesToLog / 2, DestAbsoluteOffset + BytesToLog / 2);
 
 				bDifferenceLogged = true;
 			}
@@ -767,20 +761,14 @@ void FDiffWriterArchiveWriter::Compare(
 	{
 		if (FirstUnreportedDiffIndex != -1)
 		{
-			if (!bSuppressLogging)
-			{
-				MessageCallback(ELogVerbosity::Warning, FString::Printf(
-					TEXT("%s: %lld difference(s) not logged (first at absolute offset: %lld)."),
-					AssetFilename, NumDiffsLocal - NumDiffsLoggedLocal, FirstUnreportedDiffIndex));
-			}
+			MessageCallback(ELogVerbosity::Warning, FString::Printf(
+				TEXT("%s: %lld difference(s) not logged (first at absolute offset: %lld)."),
+				AssetFilename, NumDiffsLocal - NumDiffsLoggedLocal, FirstUnreportedDiffIndex));
 		}
 		else
 		{
-			if (!bSuppressLogging)
-			{
-				MessageCallback(ELogVerbosity::Warning, FString::Printf(
-					TEXT("%s: %lld difference(s) not logged."), AssetFilename, NumDiffsLocal - NumDiffsLoggedLocal));
-			}
+			MessageCallback(ELogVerbosity::Warning, FString::Printf(
+				TEXT("%s: %lld difference(s) not logged."), AssetFilename, NumDiffsLocal - NumDiffsLoggedLocal));
 		}
 	}
 }
diff --git a/Engine/Source/Editor/UnrealEd/Private/Cooker/DiffWriterArchive.h b/Engine/Source/Editor/UnrealEd/Private/Cooker/DiffWriterArchive.h
--- a/Engine/Source/Editor/UnrealEd/Private/Cooker/DiffWriterArchive.h
+++ b/Engine/Source/Editor/UnrealEd/Private/Cooker/DiffWriterArchive.h
@@ -259,8 +259,7 @@ public:
 		const int64 MaxDiffsToLog,
 		int32& InOutDiffsLogged,
 		TMap<FName, FArchiveDiffStats>& OutStats,
-		const UE::DiffWriterArchive::FMessageCallback& MessageCallback,
-		bool bSuppressLogging = false);
+		const UE::DiffWriterArchive::FMessageCallback& MessageCallback);
 
 	/** Creates map with mismatching callstacks. */
 	static bool GenerateDiffMap(
diff --git a/Engine/Source/Editor/UnrealEd/Private/Tests/DiffWriterArchiveTests.cpp b/Engine/Source/Editor/UnrealEd/Private/Tests/DiffWriterArchiveTests.cpp
--- a/Engine/Source/Editor/UnrealEd/Private/Tests/DiffWriterArchiveTests.cpp
+++ b/Engine/Source/Editor/UnrealEd/Private/Tests/DiffWriterArchiveTests.cpp
@@ -166,7 +166,6 @@ bool FDiffWriterArchiveTestsBasic::RunTest(const FString& Parameters)
 			FDiffWriterArchiveWriter::FPackageData NewData = ToPackageData(*NewMemory);
 			int32 DiffsLogged = 0;
 			TMap<FName, FArchiveDiffStats> DiffStats;
-			const bool bSuppressLogging = true; // Suppress warning logs when running tests on CI
 			UE::DiffWriterArchive::FMessageCallback MessageCallback =
 				[](ELogVerbosity::Type Verbosity, FStringView Message)
 			{};
@@ -180,8 +179,7 @@ bool FDiffWriterArchiveTestsBasic::RunTest(const FString& Parameters)
 				MaxDiffsToLog,
 				DiffsLogged,
 				DiffStats,
-				MessageCallback,
-				bSuppressLogging);
+				MessageCallback);
 
 			TestTrueExpr(DiffStats[NAME_None].NumDiffs == 2);
 		}
diff --git a/Engine/Source/Runtime/Core/Private/Compression/CompressionUtil.cpp b/Engine/Source/Runtime/Core/Private/Compression/CompressionUtil.cpp
--- a/Engine/Source/Runtime/Core/Private/Compression/CompressionUtil.cpp
+++ b/Engine/Source/Runtime/Core/Private/Compression/CompressionUtil.cpp
@@ -94,26 +94,37 @@ namespace FCompressionUtil
 		}
 	}
 
-	void CORE_API LogHexDump(const uint8* Bytes, int64 BytesNum, int64 OffsetStart, int64 OffsetEnd, int32 BytesPerLine)
+	TArray<FString> CORE_API HexDumpLines(const uint8* Bytes, int64 BytesNum, int64 OffsetStart, int64 OffsetEnd, int32 BytesPerLine)
 	{
-		if (UE_LOG_ACTIVE(LogCore, Display))
-		{
-			OffsetStart = FMath::Max(0ll, OffsetStart);
-			OffsetEnd = FMath::Min(BytesNum, OffsetEnd);
+		OffsetStart = FMath::Max(0ll, OffsetStart);
+		OffsetEnd = FMath::Min(BytesNum, OffsetEnd);
+		TArray<FString> Results;
 
-			for (int64 Idx = OffsetStart; Idx < OffsetEnd;)
+		for (int64 Idx = OffsetStart; Idx < OffsetEnd;)
+		{
+			int64 LineOffset = OffsetStart;
+			FString HexString;
+			for (int64 Idx2 = 0; Idx2 < BytesPerLine && Idx < OffsetEnd; ++Idx, ++Idx2, ++OffsetStart)
 			{
-				int64 LineOffset = OffsetStart;
-				FString HexString;
-				for (int64 Idx2 = 0; Idx2 < BytesPerLine && Idx < OffsetEnd; ++Idx, ++Idx2, ++OffsetStart)
+				HexString += FString::Printf(TEXT("%02X "), Bytes[Idx]);
+				if ((Idx2 & 7) == 7)
 				{
-					HexString += FString::Printf(TEXT("%02X "), Bytes[Idx]);
-					if ((Idx2 & 7) == 7)
-					{
-						HexString += TEXT(" ");
-					}
+					HexString += TEXT(" ");
 				}
-				UE_LOG(LogCore, Display, TEXT("%016X: %s"), LineOffset, *HexString);
+			}
+			Results.Add(FString::Printf(TEXT("%016X: %s"), LineOffset, *HexString));
+		}
+		return Results;
+	}
+
+	void CORE_API LogHexDump(const uint8* Bytes, int64 BytesNum, int64 OffsetStart, int64 OffsetEnd, int32 BytesPerLine)
+	{
+		if (UE_LOG_ACTIVE(LogCore, Display))
+		{
+			TArray<FString> Results = HexDumpLines(Bytes, BytesNum, OffsetStart, OffsetEnd, BytesPerLine);
+			for (const FString& Result : Results)
+			{
+				UE_LOG(LogCore, Display, TEXT("%s"), *Result);
 			}
 		}
 	}
diff --git a/Engine/Source/Runtime/Core/Public/Compression/CompressionUtil.h b/Engine/Source/Runtime/Core/Public/Compression/CompressionUtil.h
--- a/Engine/Source/Runtime/Core/Public/Compression/CompressionUtil.h
+++ b/Engine/Source/Runtime/Core/Public/Compression/CompressionUtil.h
@@ -3,6 +3,8 @@
 #pragma once
 
 #include "CoreTypes.h"
+#include "Containers/Array.h"
+#include "Containers/UnrealString.h"
 
 class FArchive;
 class FCustomVersionContainer;
@@ -16,7 +18,8 @@ namespace FCompressionUtil
 	void CORE_API SerializeCompressorName(FArchive & Archive,FName & Compressor);
 
 
-	/** Logs out a memory region as a hex dump (log category: LogCore, level: Display).
+	/** Logs or writes to an array of strings a memory region as a hex dump
+	 * For logging, log category: LogCore, level: Display.
 	 *  Bytes, BytesNum are used to designate the overall memory region of interest (e.g. a file in memory), but only bytes from OffsetStart to OffsetEnd will be logged out,
 	 * and current offset will be printed at each line's start.
 	 * 
@@ -26,5 +29,6 @@ namespace FCompressionUtil
 	 * @param OffsetEnd offset to the end of log window (pass BytesNum if you want to log out whole region)
 	 * @param BytesPerLine how many bytes to print per line of hex dump
 	 */
+	TArray<FString> CORE_API HexDumpLines(const uint8* Bytes, int64 BytesNum, int64 OffsetStart, int64 OffsetEnd, int32 BytesPerLine = 32);
 	void CORE_API LogHexDump(const uint8* Bytes, int64 BytesNum, int64 OffsetStart, int64 OffsetEnd, int32 BytesPerLine = 32);
 };
\ No newline at end of file
