MPCook and SkipOnlyEditorOnly: Reject packages that are unreachable due to SkipOnlyEditorOnly before they are sent to CookWorkers. Sending them to CookWorkers makes them get cooked anyway, and can also cause an assert in IoStore if they get retracted back to the Director.
#rnx

diff --git a/Engine/Source/Editor/UnrealEd/Private/CookOnTheFlyServer.cpp b/Engine/Source/Editor/UnrealEd/Private/CookOnTheFlyServer.cpp
--- a/Engine/Source/Editor/UnrealEd/Private/CookOnTheFlyServer.cpp
+++ b/Engine/Source/Editor/UnrealEd/Private/CookOnTheFlyServer.cpp
@@ -2486,6 +2486,21 @@ void UCookOnTheFlyServer::PumpRequests(UE::Cook::FTickStackData& StackData, int3
 				TArray<TPair<FPackageData*, ESuppressCookReason>> RequestsToDemote;
 				TMap<FPackageData*, TArray<FPackageData*>> RequestGraph;
 				RequestCluster.ClearAndDetachOwnedPackageDatas(RequestsToLoad, RequestsToDemote, RequestGraph);
+				// Some packages might be reachable only on the CookerLoadingPlatform, or on previous cooked packages,
+				// because they are OnlyEditorOnly or are excluded for the newly reachable platform.
+				// Demote any packages that do not have any platforms needing cooking.
+				for (FPackageData*& PackageData : RequestsToLoad)
+				{
+					check(PackageData->GetState() == EPackageState::Request);
+					if (PackageData->GetPlatformsNeedingCookingNum() == 0)
+					{
+						ESuppressCookReason SuppressCookReason = PackageData->HasAnyCookedPlatform() ? ESuppressCookReason::AlreadyCooked :
+							ESuppressCookReason::OnlyEditorOnly;
+						RequestsToDemote.Add(TPair<FPackageData*,ESuppressCookReason>(PackageData, SuppressCookReason));
+						PackageData = nullptr; // Do not RemoveSwap; need to maintain order of the other elements in the list
+					}
+				}
+				RequestsToLoad.Remove(nullptr);
 				AssignRequests(RequestsToLoad, RequestQueue, MoveTemp(RequestGraph));
 				for (TPair<FPackageData*, ESuppressCookReason>& Pair : RequestsToDemote)
 				{
@@ -2534,16 +2549,8 @@ void UCookOnTheFlyServer::PumpRequests(UE::Cook::FTickStackData& StackData, int3
 		}
 		if (PackageData->GetPlatformsNeedingCookingNum() == 0)
 		{
-			TArray<const ITargetPlatform*> ReachablePlatforms;
-			PackageData->GetReachablePlatforms(ReachablePlatforms);
-			if (ReachablePlatforms.IsEmpty())
-			{
-				UE_LOG(LogCook, Warning, TEXT("Package %s was reported as a RequestToLoad by a requestcluster, but it has no reachable platforms."),
-					*WriteToString<256>(PackageData->GetPackageName()));
-			}
-			ReachablePlatforms.Remove(CookerLoadingPlatformKey);
-			ESuppressCookReason SuppressCookReason = ReachablePlatforms.IsEmpty() ?
-				ESuppressCookReason::OnlyEditorOnly : ESuppressCookReason::AlreadyCooked;
+			ESuppressCookReason SuppressCookReason = PackageData->HasAnyCookedPlatform() ? ESuppressCookReason::AlreadyCooked :
+				ESuppressCookReason::OnlyEditorOnly;
 			DemoteToIdle(*PackageData, ESendFlags::QueueAdd, SuppressCookReason);
 			continue;
 		}
