Cooker IterativeValidate: Add the second pass to cull indeterminism from the IterativeFalseNegative results.
#rnx

diff --git a/Engine/Source/Editor/UnrealEd/Private/CookOnTheFlyServer.cpp b/Engine/Source/Editor/UnrealEd/Private/CookOnTheFlyServer.cpp
--- a/Engine/Source/Editor/UnrealEd/Private/CookOnTheFlyServer.cpp
+++ b/Engine/Source/Editor/UnrealEd/Private/CookOnTheFlyServer.cpp
@@ -5687,10 +5687,19 @@ void UCookOnTheFlyServer::TickRecompileShaderRequestsPrivate()
 class FDiffModeCookServerUtils
 {
 public:
-	void InitializePackageWriter(ICookedPackageWriter*& CookedPackageWriter)
+	enum class EDiffMode
+	{
+		None,
+		DiffOnly,
+		LinkerDiff,
+		IterativeValidateFirstCook,
+		IterativeValidateFinalCook,
+	};
+
+	void InitializePackageWriter(ICookedPackageWriter*& CookedPackageWriter, const FString& ResolvedMetadataPath)
 	{
 		Initialize();
-		if (!bAnyDiffModeEnabled)
+		if (DiffMode == EDiffMode::None)
 		{
 			return;
 		}
@@ -5703,19 +5712,28 @@ public:
 				TEXT("A DiffMode was enabled, but the current PackageWriter has bDiffModeSupported=false."));
 		}
 
-		if (bDiffEnabled)
+		// Wrap the incoming writer inside the feature-specific-functionality writer
+		switch (DiffMode)
 		{
-			// Wrap the incoming writer inside a FDiffPackageWriter
+		case EDiffMode::DiffOnly:
 			CookedPackageWriter = new FDiffPackageWriter(TUniquePtr<ICookedPackageWriter>(CookedPackageWriter));
-		}
-		else if (bLinkerDiffEnabled)
-		{
+			break;
+		case EDiffMode::LinkerDiff:
 			CookedPackageWriter = new FLinkerDiffPackageWriter(TUniquePtr<ICookedPackageWriter>(CookedPackageWriter));
-		}
-		else
-		{
-			check(bIterativeValidateEnabled);
-			CookedPackageWriter = new FIterativeValidatePackageWriter(TUniquePtr<ICookedPackageWriter>(CookedPackageWriter));
+			break;
+		case EDiffMode::IterativeValidateFirstCook:
+			CookedPackageWriter = new FIterativeValidatePackageWriter(
+				TUniquePtr<ICookedPackageWriter>(CookedPackageWriter), FIterativeValidatePackageWriter::EPhase::FirstCook,
+				ResolvedMetadataPath);
+			break;
+		case EDiffMode::IterativeValidateFinalCook:
+			CookedPackageWriter = new FIterativeValidatePackageWriter(
+				TUniquePtr<ICookedPackageWriter>(CookedPackageWriter), FIterativeValidatePackageWriter::EPhase::FinalCook,
+				ResolvedMetadataPath);
+			break;
+		default:
+			checkNoEntry();
+			break;
 		}
 	}
 
@@ -5727,26 +5745,42 @@ private:
 			return;
 		}
 
+		DiffMode = EDiffMode::None;
 		const TCHAR* CommandLine = FCommandLine::Get();
-		bDiffEnabled = FParse::Param(CommandLine, TEXT("DIFFONLY"));
-		bLinkerDiffEnabled = false;
-		FParse::Bool(CommandLine, TEXT("-LINKERDIFF="), bLinkerDiffEnabled);
-		bIterativeValidateEnabled = FParse::Param(CommandLine, TEXT("ITERATIVEVALIDATE"));
-		int32 NumEnabled = (bDiffEnabled ? 1 : 0) + (bLinkerDiffEnabled ? 1 : 0) + (bIterativeValidateEnabled ? 1 : 0);
-		if (NumEnabled > 1)
+		auto EnsureMutualExclusion = [this]()
 		{
-			UE_LOG(LogCook, Fatal, TEXT("-DiffOnly, -LinkerDiff, and -IterativeValidate are mutually exclusive."));
-		}
-		bAnyDiffModeEnabled = NumEnabled >= 1;
+			if (DiffMode != EDiffMode::None)
+			{
+				UE_LOG(LogCook, Fatal, TEXT("-DiffOnly, -LinkerDiff, and -IterativeValidate* are mutually exclusive."));
+			}
+		};
 
+		bool bValue;
+		if (FParse::Param(CommandLine, TEXT("DIFFONLY")))
+		{
+			EnsureMutualExclusion();
+			DiffMode = EDiffMode::DiffOnly;
+		}
+		if (FParse::Bool(CommandLine, TEXT("-LINKERDIFF="), bValue) && bValue)
+		{
+			EnsureMutualExclusion();
+			DiffMode = EDiffMode::LinkerDiff;
+		}
+		if (FParse::Param(CommandLine, TEXT("IterativeValidatePrePass")))
+		{
+			EnsureMutualExclusion();
+			DiffMode = EDiffMode::IterativeValidateFirstCook;
+		}
+		if (FParse::Param(CommandLine, TEXT("IterativeValidate")))
+		{
+			EnsureMutualExclusion();
+			DiffMode = EDiffMode::IterativeValidateFinalCook;
+		}
 		bInitialized = true;
 	}
 
 	bool bInitialized = false;
-	bool bAnyDiffModeEnabled = false;
-	bool bDiffEnabled = false;
-	bool bLinkerDiffEnabled = false;
-	bool bIterativeValidateEnabled = false;
+	EDiffMode DiffMode = EDiffMode::None;
 };
 
 #if OUTPUT_COOKTIMING
@@ -10469,7 +10503,7 @@ UE::Cook::FCookSavePackageContext* UCookOnTheFlyServer::CreateSaveContext(const
 		WriterDebugName = TEXT("LooseCookedPackageWriter");
 	}
 
-	DiffModeHelper->InitializePackageWriter(PackageWriter);
+	DiffModeHelper->InitializePackageWriter(PackageWriter, ResolvedMetadataPath);
 
 	// Setup save package settings (i.e. validation)
 	FSavePackageSettings SavePackageSettings = FSavePackageSettings::GetDefaultSettings();
diff --git a/Engine/Source/Editor/UnrealEd/Private/Cooker/DiffPackageWriter.cpp b/Engine/Source/Editor/UnrealEd/Private/Cooker/DiffPackageWriter.cpp
--- a/Engine/Source/Editor/UnrealEd/Private/Cooker/DiffPackageWriter.cpp
+++ b/Engine/Source/Editor/UnrealEd/Private/Cooker/DiffPackageWriter.cpp
@@ -122,6 +122,7 @@ void FDiffPackageWriter::RemoveParam(FString& InOutParams, const TCHAR* InParamT
 void FDiffPackageWriter::BeginPackage(const FBeginPackageInfo& Info)
 {
 	bIsDifferent = false;
+	bNewPackage = false;
 	bDiffCallstack = false;
 	bHasStartedSecondSave = false;
 	DiffMap[0].Reset();
@@ -190,6 +191,7 @@ void FDiffPackageWriter::WritePackageData(const FPackageInfo& Info, FLargeMemory
 		check(LocalInfo.MultiOutputIndex < 2);
 		ExportsDiffMapOffset[LocalInfo.MultiOutputIndex] = LocalInfo.HeaderSize;
 
+		bNewPackage = PreviousPackageData.Size == 0;
 		bIsDifferent = !Writer.GenerateDiffMap(PreviousPackageData, LocalInfo.HeaderSize, MaxDiffsToLog, DiffMap[LocalInfo.MultiOutputIndex]);
 	}
 
@@ -205,13 +207,17 @@ UE::DiffWriterArchive::FMessageCallback FDiffPackageWriter::GetDiffWriterMessage
 }
 
 void FDiffPackageWriter::OnDiffWriterMessage(ELogVerbosity::Type Verbosity, FStringView Message)
+{
+	FMsg::Logf(__FILE__, __LINE__, LogDiff.GetCategoryName(), Verbosity, TEXT("%s"), *ResolveText(Message));
+}
+
+FString FDiffPackageWriter::ResolveText(FStringView Message)
 {
 	FString ResolvedText(Message);
 	check(this->Indent && this->NewLine);
 	ResolvedText.ReplaceInline(UE::DiffWriterArchive::IndentToken, this->Indent);
 	ResolvedText.ReplaceInline(UE::DiffWriterArchive::NewLineToken, this->NewLine);
-
-	FMsg::Logf(__FILE__, __LINE__, LogDiff.GetCategoryName(), Verbosity, TEXT("%s"), *ResolvedText);
+	return ResolvedText;
 }
 
 TUniquePtr<FLargeMemoryWriter> FDiffPackageWriter::CreateLinkerArchive(FName PackageName, UObject* Asset, uint16 MultiOutputIndex)
@@ -299,7 +305,7 @@ bool FDiffPackageWriter::IsAnotherSaveNeeded(FSavePackageResultStruct& PreviousR
 	if (!bHasStartedSecondSave)
 	{
 		bHasStartedSecondSave = true;
-		if (PreviousResult.Result == ESavePackageResult::Success && bIsDifferent)
+		if (PreviousResult.Result == ESavePackageResult::Success && bIsDifferent && !bNewPackage)
 		{
 			bDiffCallstack = true;
 
diff --git a/Engine/Source/Editor/UnrealEd/Private/Cooker/DiffPackageWriter.h b/Engine/Source/Editor/UnrealEd/Private/Cooker/DiffPackageWriter.h
--- a/Engine/Source/Editor/UnrealEd/Private/Cooker/DiffPackageWriter.h
+++ b/Engine/Source/Editor/UnrealEd/Private/Cooker/DiffPackageWriter.h
@@ -125,6 +125,7 @@ protected:
 	void ConditionallyDumpObjects();
 	UE::DiffWriterArchive::FMessageCallback GetDiffWriterMessageCallback();
 	virtual void OnDiffWriterMessage(ELogVerbosity::Type Verbosity, FStringView Message);
+	FString ResolveText(FStringView Message);
 
 	FDiffWriterDiffMap DiffMap[2];
 	TUniquePtr<FDiffWriterCallstacks> ExportsCallstacks;
@@ -140,6 +141,7 @@ protected:
 	bool bDiffOptional = false;
 	bool bIgnoreHeaderDiffs = false;
 	bool bIsDifferent = false;
+	bool bNewPackage = false;
 	bool bDiffCallstack = false;
 	bool bHasStartedSecondSave = false;
 	bool bDumpObjList = false;
diff --git a/Engine/Source/Editor/UnrealEd/Private/Cooker/IterativeValidatePackageWriter.cpp b/Engine/Source/Editor/UnrealEd/Private/Cooker/IterativeValidatePackageWriter.cpp
--- a/Engine/Source/Editor/UnrealEd/Private/Cooker/IterativeValidatePackageWriter.cpp
+++ b/Engine/Source/Editor/UnrealEd/Private/Cooker/IterativeValidatePackageWriter.cpp
@@ -2,20 +2,402 @@
 
 #include "Cooker/IterativeValidatePackageWriter.h"
 
-FIterativeValidatePackageWriter::FIterativeValidatePackageWriter(TUniquePtr<ICookedPackageWriter>&& InInner)
+#include "CookOnTheSide/CookLog.h"
+#include "HAL/FileManager.h"
+#include "Logging/LogMacros.h"
+#include "Misc/OutputDeviceHelper.h"
+#include "Misc/Paths.h"
+#include "Serialization/NameAsStringProxyArchive.h"
+#include "Templates/UniquePtr.h"
+
+DEFINE_LOG_CATEGORY_STATIC(LogIterativeValidate, Log, All);
+
+constexpr FStringView IterativeValidateFilename(TEXTVIEW("IterativeValidate.bin"));
+
+FIterativeValidatePackageWriter::FIterativeValidatePackageWriter(TUniquePtr<ICookedPackageWriter>&& InInner,
+	EPhase InPhase, const FString& ResolvedMetadataPath)
 	: FDiffPackageWriter(MoveTemp(InInner))
+	, MetadataPath(ResolvedMetadataPath)
+	, Phase(InPhase)
+{
+	Indent = FCString::Spc(FOutputDeviceHelper::FormatLogLine(ELogVerbosity::Warning,
+		LogIterativeValidate.GetCategoryName(), TEXT(""), GPrintLogTimes).Len());
+}
+
+void FIterativeValidatePackageWriter::BeginPackage(const FBeginPackageInfo& Info)
+{
+	bPackageFirstPass = true;
+	switch (Phase)
+	{
+	case EPhase::FirstCook:
+		bPackageSaveToDiskPass = false;
+		Super::BeginPackage(Info);
+		break;
+	case EPhase::FinalCook:
+		if (IterativeFailed.Contains(Info.PackageName))
+		{
+			bPackageSaveToDiskPass = false;
+			Super::BeginPackage(Info);
+		}
+		else
+		{
+			// This is not an IterativeValidated package (because we don't save those; UpdatePackageModificationStatus
+			// prevents saving it). And it is not an IterativeFailed package (checked above), so it is an
+			// IterativeModified package. It was found during the FirstCook phase to be modified and would in a
+			// normal iterative cook be resaved rather than iteratively skipped. Resave it as normal.
+			bPackageSaveToDiskPass = true;
+			Inner->BeginPackage(Info);
+		}
+		break;
+	default:
+		checkNoEntry();
+		break;
+	}
+}
+
+void FIterativeValidatePackageWriter::CommitPackage(FCommitPackageInfo&& Info)
+{
+	if (bPackageSaveToDiskPass)
+	{
+		Inner->CommitPackage(MoveTemp(Info));
+	}
+	else
+	{
+		Super::CommitPackage(MoveTemp(Info));
+	}
+}
+
+void FIterativeValidatePackageWriter::WritePackageData(const FPackageInfo& Info, FLargeMemoryWriter& ExportsArchive,
+	const TArray<FFileRegion>& FileRegions)
 {
+	if (bPackageSaveToDiskPass)
+	{
+		Inner->WritePackageData(Info, ExportsArchive, FileRegions);
+	}
+	else
+	{
+		Super::WritePackageData(Info, ExportsArchive, FileRegions);
+	}
+}
+
+TUniquePtr<FLargeMemoryWriter> FIterativeValidatePackageWriter::CreateLinkerArchive(FName PackageName, UObject* Asset, uint16 MultiOutputIndex)
+{
+	if (bPackageSaveToDiskPass)
+	{
+		return Inner->CreateLinkerArchive(PackageName, Asset, MultiOutputIndex);
+	}
+	else
+	{
+		return Super::CreateLinkerArchive(PackageName, Asset, MultiOutputIndex);
+	}
+}
+
+TUniquePtr<FLargeMemoryWriter> FIterativeValidatePackageWriter::CreateLinkerExportsArchive(FName PackageName, UObject* Asset, uint16 MultiOutputIndex)
+{
+	if (bPackageSaveToDiskPass)
+	{
+		return Inner->CreateLinkerExportsArchive(PackageName, Asset, MultiOutputIndex);
+	}
+	else
+	{
+		return Super::CreateLinkerExportsArchive(PackageName, Asset, MultiOutputIndex);
+	}
+}
+
+bool FIterativeValidatePackageWriter::IsPreSaveCompleted() const
+{
+	return !bPackageFirstPass;
+}
+
+void FIterativeValidatePackageWriter::Initialize(const FCookInfo& CookInfo)
+{
+	switch (Phase)
+	{
+	case EPhase::FirstCook:
+		if (CookInfo.bFullBuild)
+		{
+			UE_LOG(LogIterativeValidate, Display,
+				TEXT("The cook is running non-iteratively. All packages are reported \"modified\" and will be resaved during the final IterativeValidate phase."));
+		}
+		break;
+	case EPhase::FinalCook:
+		if (CookInfo.bFullBuild)
+		{
+			UE_LOG(LogIterativeValidate, Display,
+				TEXT("The cook is running non-iteratively. Packages that were iteratively skipped and found valid will be resaved anyway."));
+		}
+		break;
+	default:
+		checkNoEntry();
+		break;
+	}
+	Super::Initialize(CookInfo);
 }
 
 void FIterativeValidatePackageWriter::UpdatePackageModificationStatus(FName PackageName, bool bIterativelyUnmodified,
 	bool& bInOutShouldIterativelySkip)
 {
-	bool bModified = !bIterativelyUnmodified;
-	bInOutShouldIterativelySkip = bModified;
+	switch (Phase)
+	{
+	case EPhase::FirstCook:
+		// Invert what gets skipped: save the iteratively skipped files to record their diffs, but skip the regular files
+		bInOutShouldIterativelySkip = !bIterativelyUnmodified;
+		break;
+	case EPhase::FinalCook:
+		// Ignore the Unmodified flag from this cook phase. Skip only the packages that were found to 
+		// be IterativeValidated from the FirstCook phase.
+		bInOutShouldIterativelySkip = IterativeValidated.Contains(PackageName);
+		break;
+	default:
+		checkNoEntry();
+		break;
+	}
 
-	bool bInnerIterativelyUnmodified = bModified;
-	bool bInnerInOutShouldIterativelySkip = bModified;
+	bool bInnerIterativelyUnmodified = bInOutShouldIterativelySkip;
+	bool bInnerInOutShouldIterativelySkip = bInOutShouldIterativelySkip;
 	Inner->UpdatePackageModificationStatus(PackageName, bInnerIterativelyUnmodified, bInnerInOutShouldIterativelySkip);
 	checkf(bInnerInOutShouldIterativelySkip == bInnerIterativelyUnmodified,
 		TEXT("IterativeValidatePackageWriter is not supported with an Inner that modifies bInOutShouldIterativelySkip."));
 }
+
+void FIterativeValidatePackageWriter::BeginCook(const FCookInfo& Info)
+{
+	switch (Phase)
+	{
+	case EPhase::FirstCook:
+		break;
+	case EPhase::FinalCook:
+		Load();
+		break;
+	default:
+		checkNoEntry();
+		break;
+	}
+	Super::BeginCook(Info);
+}
+
+void FIterativeValidatePackageWriter::EndCook(const FCookInfo& Info)
+{
+	Super::EndCook(Info);
+	switch (Phase)
+	{
+	case EPhase::FirstCook:
+		UE_LOG(LogIterativeValidate, Display,
+			TEXT("DetectedUnmodified: %d. ValidatedUnmodified: %d. IterativeFalseNegativeOrIndeterminism: %d."),
+			IterativeValidated.Num() + IterativeFailed.Num(), IterativeValidated.Num(), IterativeFailed.Num());
+		Save();
+		break;
+	case EPhase::FinalCook:
+	{
+		UE_LOG(LogIterativeValidate, Display,
+			TEXT("DetectedUnmodified: %d. ValidatedUnmodified: %d. Indeterminism: %d."),
+			IterativeValidated.Num() + IterativeFailed.Num(), IterativeValidated.Num(), IndeterminismFailed.Num());
+		FString Message = FString::Printf(TEXT("IterativeFalseNegative: %d."), IterativeFalseNegative.Num());
+		if (IterativeFalseNegative.Num())
+		{
+			UE_LOG(LogIterativeValidate, Error, TEXT("%s"), *Message);
+		}
+		else
+		{
+			UE_LOG(LogIterativeValidate, Display, TEXT("%s"), *Message);
+		}
+		break;
+	}
+	default:
+		checkNoEntry();
+		break;
+	}
+}
+
+void FIterativeValidatePackageWriter::UpdateSaveArguments(FSavePackageArgs& SaveArgs)
+{
+	if (bPackageSaveToDiskPass)
+	{
+		return Inner->UpdateSaveArguments(SaveArgs);
+	}
+	else
+	{
+		return Super::UpdateSaveArguments(SaveArgs);
+	}
+}
+
+bool FIterativeValidatePackageWriter::IsAnotherSaveNeeded(FSavePackageResultStruct& PreviousResult, FSavePackageArgs& SaveArgs)
+{
+	bPackageFirstPass = false;
+	checkf(!Inner->IsAnotherSaveNeeded(PreviousResult, SaveArgs),
+		TEXT("IterativeValidatePackageWriter does not support an Inner that needs multiple saves."));
+	if (PreviousResult == ESavePackageResult::Timeout)
+	{
+		return false;
+	}
+	if (bPackageSaveToDiskPass)
+	{
+		// The packagesavetodiskpass, if present, is always the last pass no matter which phase
+		return false;
+	}
+
+	switch (Phase)
+	{
+	case EPhase::FirstCook:
+		if (Super::IsAnotherSaveNeeded(PreviousResult, SaveArgs))
+		{
+			return true;
+		}
+		else if (bIsDifferent && !bNewPackage)
+		{
+			// If our superclass FDiffPackageWriter found differences, when it finishes the saves it wants to do,
+			// Finish it off and start a SaveToDisk pass
+			FCommitPackageInfo CommitInfo;
+			CommitInfo.Status = IPackageWriter::ECommitStatus::Success;
+			CommitInfo.PackageName = BeginInfo.PackageName;
+			CommitInfo.WriteOptions = EWriteOptions::None;
+			Super::CommitPackage(MoveTemp(CommitInfo));
+			Inner->BeginPackage(BeginInfo);
+			bPackageSaveToDiskPass = true;
+
+			// Mark that the iterative validation failed if it was not already marked by log or warning messages.
+			// We need to record it for an indeterminism test
+			IterativeFailed.FindOrAdd(BeginInfo.PackageName);
+			return true;
+		}
+		else if (!bNewPackage)
+		{
+			// No differences found, so finish off the superclass's save during CommitPackage, without doing a
+			// savetodisk pass
+			// Mark that the iterative validation passed
+			IterativeValidated.Add(BeginInfo.PackageName);
+			TArray<FMessage> Messages;
+			IterativeFailed.RemoveAndCopyValue(BeginInfo.PackageName, Messages);
+			for (FMessage& Message : Messages)
+			{
+				// If no differences were detected, we should not have logged any warning or error messages
+				check(Message.Verbosity > ELogVerbosity::Warning);
+			}
+			return false;
+		}
+		else
+		{
+			// New packages need to be resaved in the FinalCook phase; for our purposes they are equivalent
+			// to a package that iteration detected as modified.
+			// Do not add an entry for it in our results for iterative packages, and do not resave it in this pass
+			return false;
+		}
+	case EPhase::FinalCook:
+		LogIterativeDifferences();
+		// No need to do the Super's second diff pass to find callstacks - just knowing whether differences exist is enough.
+		// No need to save package to disk; for these packages (packages found to be iteratively unmodified during
+		// firstcook phase) they were already resaved during the FirstCook phase
+		return false;
+	default:
+		checkNoEntry();
+		return false;
+	}
+}
+
+void FIterativeValidatePackageWriter::OnDiffWriterMessage(ELogVerbosity::Type Verbosity, FStringView Message)
+{
+	switch (Phase)
+	{
+	case EPhase::FirstCook:
+		IterativeFailed.FindOrAdd(BeginInfo.PackageName).Add(FMessage{ FString(Message), Verbosity });
+		break;
+	case EPhase::FinalCook:
+		break;
+	default:
+		checkNoEntry();
+		break;
+	}
+}
+
+void FIterativeValidatePackageWriter::LogIterativeDifferences()
+{
+	// This function is called during the FinalCook phase for packages that had differences during the FirstCook phase.
+	// It is called immediately after first-pass package save that is run by our super class FDiffPackageWriter, which
+	// compared it against the version that was written to disk during the FirstCook phase. If there are differences
+	// now from FirstCook phase, then this package has a determinism issue. We log that information at display rather
+	// than Warning because this cookmode only logs warnings for IterativeFalseNegatives.
+	bool bHasDeterminismIssue = bIsDifferent;
+	if (bHasDeterminismIssue)
+	{
+		UE_LOG(LogIterativeValidate, Display, TEXT("Could not validate %s because it has a non-deterministic save."),
+			*BeginInfo.PackageName.ToString());
+		IndeterminismFailed.Add(BeginInfo.PackageName);
+		return;
+	}
+
+	// Otherwise, no determinism issues, so the differences indicate a bug in Diff Package
+	IterativeFalseNegative.Add(BeginInfo.PackageName);
+	FMsg::Logf(__FILE__, __LINE__, LogIterativeValidate.GetCategoryName(), ELogVerbosity::Warning,
+		TEXT("IterativeFalseNegative package %s."), *BeginInfo.PackageName.ToString());
+	TArray<FMessage>& Messages = IterativeFailed.FindOrAdd(BeginInfo.PackageName);
+	for (const FMessage& Message : Messages)
+	{
+		FMsg::Logf(__FILE__, __LINE__, LogIterativeValidate.GetCategoryName(), Message.Verbosity,
+			TEXT("%s"), *ResolveText(Message.Text));
+	}
+}
+
+void FIterativeValidatePackageWriter::Save()
+{
+	FString IterativeValidatePath = GetIterativeValidatePath();
+	TUniquePtr<FArchive> DiskArchive(IFileManager::Get().CreateFileWriter(*IterativeValidatePath));
+
+	if (!DiskArchive)
+	{
+		UE_LOG(LogIterativeValidate, Error,
+			TEXT("Could not write to file %s. This file is needed to store results for the -IterativeValidate cook."),
+			*IterativeValidatePath);
+		return;
+	}
+	FNameAsStringProxyArchive Ar(*DiskArchive);
+	Serialize(Ar);
+}
+
+void FIterativeValidatePackageWriter::Load()
+{
+	FString IterativeValidatePath = GetIterativeValidatePath();
+	TUniquePtr<FArchive> DiskArchive(IFileManager::Get().CreateFileReader(*IterativeValidatePath));
+	if (!DiskArchive)
+	{
+		UE_LOG(LogIterativeValidate, Fatal,
+			TEXT("Could not load file %s. This file is required and should have been written by the -IterativeValidatePrePass cook."),
+			*IterativeValidatePath);
+		return;
+	}
+	FNameAsStringProxyArchive Ar(*DiskArchive);
+	Serialize(Ar);
+	if (Ar.IsError())
+	{
+		UE_LOG(LogIterativeValidate, Fatal, TEXT("Corrupt file %s"), *IterativeValidatePath);
+	}
+}
+
+void FIterativeValidatePackageWriter::Serialize(FArchive& Ar)
+{
+	constexpr int32 LatestVersion = 0;
+	int32 Version = LatestVersion;
+	Ar << Version;
+	if (Ar.IsLoading() && Version != LatestVersion)
+	{
+		Ar.SetError();
+		return;
+	}
+	Ar << IterativeValidated;
+	Ar << IterativeFailed;
+}
+
+FArchive& operator<<(FArchive& Ar, FIterativeValidatePackageWriter::FMessage& Message)
+{
+	uint8 Verbosity = static_cast<uint8>(Message.Verbosity);
+	Ar << Verbosity << Message.Text;
+	if (Ar.IsLoading())
+	{
+		Message.Verbosity = static_cast<ELogVerbosity::Type>(Verbosity);
+	}
+	return Ar;
+}
+
+FString FIterativeValidatePackageWriter::GetIterativeValidatePath() const
+{
+	return FPaths::Combine(MetadataPath, FString(IterativeValidateFilename));
+}
diff --git a/Engine/Source/Editor/UnrealEd/Private/Cooker/IterativeValidatePackageWriter.h b/Engine/Source/Editor/UnrealEd/Private/Cooker/IterativeValidatePackageWriter.h
--- a/Engine/Source/Editor/UnrealEd/Private/Cooker/IterativeValidatePackageWriter.h
+++ b/Engine/Source/Editor/UnrealEd/Private/Cooker/IterativeValidatePackageWriter.h
@@ -11,9 +11,55 @@
 class FIterativeValidatePackageWriter : public FDiffPackageWriter
 {
 public:
-	FIterativeValidatePackageWriter(TUniquePtr<ICookedPackageWriter>&& InInner);
+	using Super = FDiffPackageWriter;
+	enum class EPhase
+	{
+		FirstCook,
+		FinalCook,
+	};
+	FIterativeValidatePackageWriter(TUniquePtr<ICookedPackageWriter>&& InInner, EPhase InPhase,
+		const FString& ResolvedMetadataPath);
+
+	// IPackageWriter
+	virtual void BeginPackage(const FBeginPackageInfo& Info) override;
+	virtual void CommitPackage(FCommitPackageInfo&& Info) override;
+	virtual void WritePackageData(const FPackageInfo& Info, FLargeMemoryWriter& ExportsArchive,
+		const TArray<FFileRegion>& FileRegions) override;
+	virtual TUniquePtr<FLargeMemoryWriter> CreateLinkerArchive(FName PackageName, UObject* Asset, uint16 MultiOutputIndex) override;
+	virtual TUniquePtr<FLargeMemoryWriter> CreateLinkerExportsArchive(FName PackageName, UObject* Asset, uint16 MultiOutputIndex) override;
+	virtual bool IsPreSaveCompleted() const override;
 
 	// ICookedPackageWriter
+	virtual void Initialize(const FCookInfo& CookInfo) override;
 	virtual void UpdatePackageModificationStatus(FName PackageName, bool bIterativelyUnmodified,
 		bool& bInOutShouldIterativelySkip) override;
+	virtual void BeginCook(const FCookInfo& Info) override;
+	virtual void EndCook(const FCookInfo& Info) override;
+	virtual void UpdateSaveArguments(FSavePackageArgs& SaveArgs) override;
+	virtual bool IsAnotherSaveNeeded(FSavePackageResultStruct& PreviousResult, FSavePackageArgs& SaveArgs) override;
+
+protected:
+	virtual void OnDiffWriterMessage(ELogVerbosity::Type Verbosity, FStringView Message) override;
+	void LogIterativeDifferences();
+	void Save();
+	void Load();
+	void Serialize(FArchive& Ar);
+	FString GetIterativeValidatePath() const;
+
+	struct FMessage
+	{
+		FString Text;
+		ELogVerbosity::Type Verbosity;
+	};
+	friend FArchive& operator<<(FArchive& Ar, FMessage& Message);
+
+	TSet<FName> IterativeValidated;
+	TMap<FName, TArray<FMessage>> IterativeFailed;
+	TSet<FName> IndeterminismFailed;
+	TSet<FName> IterativeFalseNegative;
+
+	FString MetadataPath;
+	EPhase Phase = EPhase::FirstCook;
+	bool bPackageSaveToDiskPass = false;
+	bool bPackageFirstPass = false;
 };
