Better support of reinstantiation of subobjects in CDOs
Also reuse same pattern on instances to fix more subobject nesting issues

diff --git a/Engine/Plugins/Runtime/RigVM/Source/RigVM/Private/RigVMCore/RigVMMemoryStorage.cpp b/Engine/Plugins/Runtime/RigVM/Source/RigVM/Private/RigVMCore/RigVMMemoryStorage.cpp
--- a/Engine/Plugins/Runtime/RigVM/Source/RigVM/Private/RigVMCore/RigVMMemoryStorage.cpp
+++ b/Engine/Plugins/Runtime/RigVM/Source/RigVM/Private/RigVMCore/RigVMMemoryStorage.cpp
@@ -646,6 +646,7 @@ bool URigVMMemoryStorageGeneratorClass::RemoveStorageClass(UObject* InOuter, ERi
 	URigVMMemoryStorageGeneratorClass* OldClass = FindObject<URigVMMemoryStorageGeneratorClass>(Package, *ClassName);
 	if(OldClass)
 	{
+		OldClass->ClassFlags |= CLASS_NewerVersionExists;
 		OldClass->Rename(nullptr, GetTransientPackage(), REN_ForceNoResetLoaders | REN_DoNotDirty | REN_DontCreateRedirectors | REN_NonTransactional);
 		return true;
 	}
diff --git a/Engine/Source/Editor/Kismet/Private/BlueprintCompilationManager.cpp b/Engine/Source/Editor/Kismet/Private/BlueprintCompilationManager.cpp
--- a/Engine/Source/Editor/Kismet/Private/BlueprintCompilationManager.cpp
+++ b/Engine/Source/Editor/Kismet/Private/BlueprintCompilationManager.cpp
@@ -515,6 +515,11 @@ struct FReinstancingJob
 
 	// always set:
 	TPair<UClass*, UClass*> OldToNew;
+
+	// Sorting dependencies
+	TSet<const UClass*> BPClassDependencies;
+	bool DependsOn(const FReinstancingJob& OtherJob) const;
+	void CalculateBPClassDependencies();
 };
 
 FReinstancingJob::FReinstancingJob(TSharedPtr<FBlueprintCompileReinstancer> InReinstancer)
@@ -544,6 +549,32 @@ FReinstancingJob::FReinstancingJob(TPair<UClass*, UClass*> InOldToNew)
 {
 }
 
+bool FReinstancingJob::DependsOn(const FReinstancingJob& OtherJob) const
+{
+	return OtherJob.Reinstancer.IsValid() && OtherJob.OldToNew.Key && BPClassDependencies.Contains(OtherJob.OldToNew.Key);
+}
+
+void FReinstancingJob::CalculateBPClassDependencies()
+{
+	if (Reinstancer.IsValid() && OldToNew.Key)
+	{
+		const UObject* OldCDO = OldToNew.Key->ClassDefaultObject;
+
+		// Gather subobjects on old CDO and remember depends BP classes
+		TArray<UObject*> ContainedOldObjects;
+		GetObjectsWithOuter(OldCDO, ContainedOldObjects);
+		for (const UObject* OldObject : ContainedOldObjects)
+		{
+			const UClass* DependentClass = OldObject->GetClass();
+			while (DependentClass && UBlueprint::GetBlueprintFromClass(DependentClass))
+			{
+				BPClassDependencies.Add(DependentClass);
+				DependentClass = DependentClass->GetSuperClass();
+			}
+		}
+	}
+}
+
 namespace UE::Kismet::BlueprintCompilationManager::Private
 {
 	namespace ConsoleVariables
@@ -2215,6 +2246,8 @@ void FBlueprintCompilationManagerImpl::BuildDSOMap(UObject* OldObject, UObject*
 
 void FBlueprintCompilationManagerImpl::ReinstanceBatch(TArray<FReinstancingJob>& Reinstancers, TMap< UClass*, UClass* >& InOutOldToNewClassMap, FUObjectSerializeContext* InLoadContext)
 {
+	TGuardValue<bool> ReinstancingGuard(GIsReinstancing, true);
+
 	const auto FilterOutOfDateClasses = [](TArray<UClass*>& ClassList)
 	{
 		// Old versions of classes can be abandoned, classes without CDOs have no instances and don't require reinstancing
@@ -2350,7 +2383,7 @@ void FBlueprintCompilationManagerImpl::ReinstanceBatch(TArray<FReinstancingJob>&
 	{
 		TSet<UBlueprint*> DependentBPs;
 		TMap<FFieldVariant, FFieldVariant> FieldMappings;
-		for (const FReinstancingJob& ReinstancingJob : Reinstancers)
+		for (FReinstancingJob& ReinstancingJob : Reinstancers)
 		{
 			if (ReinstancingJob.OldToNew.Key)
 			{
@@ -2361,6 +2394,7 @@ void FBlueprintCompilationManagerImpl::ReinstanceBatch(TArray<FReinstancingJob>&
 			{
 				UBlueprint* CompiledBlueprint = UBlueprint::GetBlueprintFromClass(ReinstancingJob.OldToNew.Value);
 				ReinstancingJob.Reinstancer->UpdateBytecodeReferences(DependentBPs, FieldMappings);
+				ReinstancingJob.CalculateBPClassDependencies();
 			}
 		}
 
@@ -2377,10 +2411,11 @@ void FBlueprintCompilationManagerImpl::ReinstanceBatch(TArray<FReinstancingJob>&
 	Reinstancers.Sort(
 		[](const FReinstancingJob& ReinstancingDataA, const FReinstancingJob& ReinstancingDataB)
 		{
-			return FBlueprintCompileReinstancer::ReinstancerOrderingFunction(
-				ReinstancingDataA.OldToNew.Value, 
-				ReinstancingDataB.OldToNew.Value
-			);
+			return !ReinstancingDataA.DependsOn(ReinstancingDataB) && 
+				FBlueprintCompileReinstancer::ReinstancerOrderingFunction(
+					ReinstancingDataA.OldToNew.Value, 
+					ReinstancingDataB.OldToNew.Value
+				);
 		}
 	);
 
@@ -2404,7 +2439,15 @@ void FBlueprintCompilationManagerImpl::ReinstanceBatch(TArray<FReinstancingJob>&
 					FBlueprintSupport::RepairDeferredDependenciesInObject(OldCDO);
 				}
 
-				FBlueprintCompileReinstancer::CopyPropertiesForUnrelatedObjects(OldCDO, NewCDO, true, bUseDeltaSerialization);
+				TMap<UObject*, UObject*> CreatedInstanceMap;
+				FBlueprintCompileReinstancer::PreCreateSubObjectsForReinstantiation(InOutOldToNewClassMap, OldCDO, NewCDO, CreatedInstanceMap);
+
+				// We only need to copy properties of the pre-created instances, the rest of the default sub object is done inside the UEditorEngine::CopyPropertiesForUnrelatedObjects
+				TMap<UObject*, UObject*> OldToNewInstanceMap(CreatedInstanceMap);
+				for(const auto& Pair : CreatedInstanceMap)
+				{
+					FBlueprintCompileReinstancer::CopyPropertiesForUnrelatedObjects(Pair.Key, Pair.Value, /*bClearExternalReferences*/true, bUseDeltaSerialization, /*bOnlyHandleDirectSubObjects*/true, &OldToNewInstanceMap);
+				}
 
 				if (ReinstancingJob.Compiler.IsValid())
 				{
@@ -2472,13 +2515,13 @@ void FBlueprintCompilationManagerImpl::ReinstanceBatch(TArray<FReinstancingJob>&
 		if(OldClass)
 		{
 			UClass* NewClass = ReinstancingJob.OldToNew.Value;
-            if (NewClass && 
+			if (NewClass && 
 				OldClass->ClassDefaultObject && 
 				NewClass->ClassDefaultObject &&
 				OldClass->ClassDefaultObject != NewClass->ClassDefaultObject)
-            {
-                OldArchetypeToNewArchetype.Add(OldClass->ClassDefaultObject, NewClass->ClassDefaultObject);
-            }
+			{
+				OldArchetypeToNewArchetype.Add(OldClass->ClassDefaultObject, NewClass->ClassDefaultObject);
+			}
 
 			TArray<UObject*> ArchetypeObjects;
 			GetObjectsOfClass(OldClass, ArchetypeObjects, false);
diff --git a/Engine/Source/Editor/UnrealEd/Private/Kismet2/KismetReinstanceUtilities.cpp b/Engine/Source/Editor/UnrealEd/Private/Kismet2/KismetReinstanceUtilities.cpp
--- a/Engine/Source/Editor/UnrealEd/Private/Kismet2/KismetReinstanceUtilities.cpp
+++ b/Engine/Source/Editor/UnrealEd/Private/Kismet2/KismetReinstanceUtilities.cpp
@@ -58,6 +58,11 @@ static FAutoConsoleVariableRef CVarUseLegacyAnimInstanceReinstancingBehavior(
 	TEXT("Use the legacy re-instancing behavior for anim instances where the instance is destroyed and re-created.")
 );
 
+namespace UE::ReinstanceUtils
+{
+	const EObjectFlags FlagMask = RF_Public | RF_ArchetypeObject | RF_Transactional | RF_Transient | RF_TextExportTransient | RF_InheritableComponentTemplate | RF_Standalone; //TODO: what about RF_RootSet?
+}
+
 struct FReplaceReferenceHelper
 {
 	static void ValidateReplacementMappings(const TMap<UObject*, UObject*>& OldToNewInstanceMap)
@@ -1896,9 +1901,8 @@ bool FBlueprintCompileReinstancer::IsReinstClass(const UClass* Class)
 	return Class && Class->GetFName().ToString().StartsWith(ReinstPrefix);
 }
 
-static void ReplaceObjectHelper(UObject*& OldObject, UClass* OldClass, UObject*& NewUObject, UClass* NewClass, TMap<UObject*, UObject*>& OldToNewInstanceMap, TMap<UObject*, FName>& OldToNewNameMap, int32 OldObjIndex, TArray<UObject*>& ObjectsToReplace, TArray<UObject*>& PotentialEditorsForRefreshing, TSet<AActor*>& OwnersToRerunConstructionScript, TFunctionRef<TArray<TObjectPtr<USceneComponent>>&(USceneComponent*)> GetAttachChildrenArray, bool bIsComponent, bool bArchetypesAreUpToDate)
+static void ReplaceObjectHelper(UObject*& OldObject, UClass* OldClass, UObject*& NewUObject, UClass* NewClass, TMap<UObject*, UObject*>& OldToNewInstanceMap, const TMap<UClass*, UClass*>& OldToNewClassMap, TMap<UObject*, FName>& OldToNewNameMap, int32 OldObjIndex, TArray<UObject*>& ObjectsToReplace, TArray<UObject*>& PotentialEditorsForRefreshing, TSet<AActor*>& OwnersToRerunConstructionScript, TFunctionRef<TArray<TObjectPtr<USceneComponent>>&(USceneComponent*)> GetAttachChildrenArray, bool bIsComponent, bool bArchetypesAreUpToDate)
 {
-	const EObjectFlags FlagMask = RF_Public | RF_ArchetypeObject | RF_Transactional | RF_Transient | RF_TextExportTransient | RF_InheritableComponentTemplate | RF_Standalone; //TODO: what about RF_RootSet?
 	// If the old object was spawned from an archetype (i.e. not the CDO), we must use the new version of that archetype as the template object when constructing the new instance.
 	UObject* NewArchetype = nullptr;
 	if(bArchetypesAreUpToDate)
@@ -1909,7 +1913,7 @@ static void ReplaceObjectHelper(UObject*& OldObject, UClass* OldClass, UObject*&
 			// Otherwise, just use the old object's current name.
 			NewName = OldObject->GetFName();
 		}
-		NewArchetype = UObject::GetArchetypeFromRequiredInfo(NewClass, OldObject->GetOuter(), NewName, OldObject->GetFlags() & FlagMask);
+		NewArchetype = UObject::GetArchetypeFromRequiredInfo(NewClass, OldObject->GetOuter(), NewName, OldObject->GetFlags() & UE::ReinstanceUtils::FlagMask);
 	}
 	else
 	{
@@ -1954,8 +1958,6 @@ static void ReplaceObjectHelper(UObject*& OldObject, UClass* OldClass, UObject*&
 
 	FName OldName(OldObject->GetFName());
 
-	UEngine::FCopyPropertiesForUnrelatedObjectsParams Options;
-
 	// If the old object is in this table, we've already renamed it away in a previous iteration. Don't rename it again!
 	if (!OldToNewNameMap.Contains(OldObject))
 	{
@@ -1977,7 +1979,6 @@ static void ReplaceObjectHelper(UObject*& OldObject, UClass* OldClass, UObject*&
 		}
 		else
 		{
-			Options.SourceObjectArchetype = OldObject->GetArchetype(); // Cache the old object's archetype before renaming
 			OldObject->Rename(nullptr, OldObject->GetOuter(), REN_DoNotDirty | REN_DontCreateRedirectors | REN_ForceNoResetLoaders);
 		}
 	}
@@ -2018,13 +2019,17 @@ static void ReplaceObjectHelper(UObject*& OldObject, UClass* OldClass, UObject*&
 
 	check(NewUObject != nullptr);
 
-	NewUObject->SetFlags(OldFlags & FlagMask);
+	NewUObject->SetFlags(OldFlags & UE::ReinstanceUtils::FlagMask);
+
+	TMap<UObject*, UObject*> CreatedInstanceMap;
+	FBlueprintCompileReinstancer::PreCreateSubObjectsForReinstantiation(OldToNewClassMap, OldObject, NewUObject, CreatedInstanceMap, &OldToNewInstanceMap);
+	OldToNewInstanceMap.Append(CreatedInstanceMap);
 
-	InstancedPropertyUtils::FInstancedPropertyMap InstancedPropertyMap;
-	InstancedPropertyUtils::FArchiveInstancedSubObjCollector  InstancedSubObjCollector(OldObject, InstancedPropertyMap);
 	// Copy property values
+	UEngine::FCopyPropertiesForUnrelatedObjectsParams Options;
 	Options.bNotifyObjectReplacement = true;
 	Options.bSkipCompilerGeneratedDefaults = true;
+	Options.bOnlyHandleDirectSubObjects = true;
 	Options.OptionalReplacementMappings = &OldToNewInstanceMap;
 	// this currently happens because of some misguided logic in UBlueprintGeneratedClass::FindArchetype that
 	// points us to a mismatched archetype, in which case delta serialization becomes unsafe.. without
@@ -2033,9 +2038,11 @@ static void ReplaceObjectHelper(UObject*& OldObject, UClass* OldClass, UObject*&
 	{
 		Options.bDoDelta = false;
 	}
-	UEditorEngine::CopyPropertiesForUnrelatedObjects(OldObject, NewUObject, Options);
-	// Generate new subobjects
-	InstancedPropertyUtils::FArchiveInsertInstancedSubObjects InstancedSubObjSpawner(NewUObject, InstancedPropertyMap);
+	// We only need to copy properties of the pre-created instances, the rest of the default sub object is done inside the UEditorEngine::CopyPropertiesForUnrelatedObjects
+	for (const auto& Pair : CreatedInstanceMap)
+	{
+		UEditorEngine::CopyPropertiesForUnrelatedObjects(Pair.Key, Pair.Value, Options);
+	}
 
 	UWorld* RegisteredWorld = nullptr;
 	bool bWasRegistered = false;
@@ -2144,8 +2151,6 @@ static void ReplaceActorHelper(AActor* OldActor, UClass* OldClass, UObject*& New
 	UClass* const* MappedClass = InOldToNewClassMap.Find(OldActor->GetClass());
 	UClass*  SpawnClass = MappedClass ? *MappedClass : NewClass;
 
-	const EObjectFlags FlagMask = RF_Public | RF_ArchetypeObject | RF_Transactional | RF_Transient | RF_TextExportTransient | RF_InheritableComponentTemplate | RF_Standalone;
-
 	FActorSpawnParameters SpawnInfo;
 	SpawnInfo.OverrideLevel = ActorLevel;
 	SpawnInfo.Owner = OldActor->GetOwner();
@@ -2154,7 +2159,7 @@ static void ReplaceActorHelper(AActor* OldActor, UClass* OldClass, UObject*& New
 	SpawnInfo.SpawnCollisionHandlingOverride = ESpawnActorCollisionHandlingMethod::AlwaysSpawn;
 	SpawnInfo.bDeferConstruction = true;
 	SpawnInfo.Name = OldActor->GetFName();
-	SpawnInfo.ObjectFlags |= OldActor->GetFlags() & FlagMask;
+	SpawnInfo.ObjectFlags |= OldActor->GetFlags() & UE::ReinstanceUtils::FlagMask;
 
 	if (!OldActor->IsListedInSceneOutliner())
 	{
@@ -2450,45 +2455,6 @@ void FBlueprintCompileReinstancer::ReplaceInstancesOfClass_Inner(const TMap<UCla
 					AActor* OldActor = Cast<AActor>(OldObject);
 					bool bIsValid = IsValid(OldObject);
 
-					// If the object is not valid, it might been replaced by the IntancedReferenceSubobjectHelper code.
-					// Anyhow we still need to replace the object class as it only duplicated the object, 
-					// So check if it has replacement and process it now if the IntancedReferenceSubobjectHelper reinstanciated the class being replaced.
-					UObject* PrevOldObject = nullptr;
-					if (!bIsValid)
-					{
-						// Was it replaced because it owner got re-instantiated?
-						if (UObject** ReplacedOldObject = OldToNewInstanceMap.Find(OldObject))
-						{
-							// If the replaced object does not have the right class
-							if ((*ReplacedOldObject)->GetClass() != NewClass)
-							{
-								// Remove the old entry as it points to the wrong instance to replace
-								OldToNewInstanceMap.Remove(OldObject);
-								ObjectRemappingHelper.ReplacedObjects.Remove(OldObject);
-
-								// Make sure the replaced object has not already been re-instantiated
-								if (UObject** NewUObject = OldToNewInstanceMap.Find(*ReplacedOldObject))
-								{
-									// Add missing entry in the map
-									checkf(IsValid(*NewUObject), TEXT("Expecting the new object to already be valid"));
-									OldToNewInstanceMap.Add(OldObject, *NewUObject);
-									ObjectRemappingHelper.ReplacedObjects.Add(OldObject, *NewUObject);
-								}
-								else
-								{
-									// Remember old object and re-instantiate the replaced object
-									PrevOldObject = OldObject;
-									OldObject = *ReplacedOldObject;
-									bIsValid = IsValid(OldObject);
-								}
-							}
-							else
-							{
-								checkf(IsValid(*ReplacedOldObject), TEXT("Expecting the replaced object to already be valid"));
-							}
-						}
-					}
-
 					// Skip archetype instances, EXCEPT for component templates and child actor templates
 					const bool bIsChildActorTemplate = OldActor && OldActor->GetOuter()->IsA<UChildActorComponent>();
 					if ((!bIsValid && !bIsScriptComponent) || // @todo: why do we need to replace PendingKill script components?
@@ -2504,20 +2470,9 @@ void FBlueprintCompileReinstancer::ReplaceInstancesOfClass_Inner(const TMap<UCla
 						CheckAndSaveOuterPackageToCleanList(OldObject);
 
 						UObject* NewUObject = nullptr;
-						ReplaceObjectHelper(OldObject, OldClass, NewUObject, NewClass, OldToNewInstanceMap, OldToNewNameMap, OldObjIndex, ObjectsToReplace, PotentialEditorsForRefreshing, OwnersToRerunConstructionScript, &FDirectAttachChildrenAccessor::Get, bIsComponent, bArchetypesAreUpToDate);
+						ReplaceObjectHelper(OldObject, OldClass, NewUObject, NewClass, OldToNewInstanceMap, InOldToNewClassMap, OldToNewNameMap, OldObjIndex, ObjectsToReplace, PotentialEditorsForRefreshing, OwnersToRerunConstructionScript, &FDirectAttachChildrenAccessor::Get, bIsComponent, bArchetypesAreUpToDate);
 						UpdateObjectBeingDebugged(OldObject, NewUObject);
 
-						if (PrevOldObject)
-						{
-							OldToNewInstanceMap.Add(PrevOldObject, NewUObject);
-							ObjectRemappingHelper.ReplacedObjects.Add(PrevOldObject, NewUObject);
-							ObjectsReplaced.Add(PrevOldObject);
-						}
-						else
-						{
-							ObjectsReplaced.Add(OldObject);
-						}
-
 						if (bLogConversions)
 						{
 							UE_LOG(LogBlueprint, Log, TEXT("Converted instance '%s' to '%s'"), *GetPathNameSafe(OldObject), *GetPathNameSafe(NewUObject));
@@ -2612,7 +2567,7 @@ void FBlueprintCompileReinstancer::ReplaceInstancesOfClass_Inner(const TMap<UCla
 						else
 						{
 							// Actors that are not in a level cannot be reconstructed, sequencer team decided to reinstance these as normal objects:
-							ReplaceObjectHelper(OldObject, OldClass, NewUObject, NewClass, OldToNewInstanceMap, OldToNewNameMap, OldObjIndex, ObjectsToReplace, PotentialEditorsForRefreshing, OwnersToRerunConstructionScript, &FDirectAttachChildrenAccessor::Get, false, bArchetypesAreUpToDate);
+							ReplaceObjectHelper(OldObject, OldClass, NewUObject, NewClass, OldToNewInstanceMap, InOldToNewClassMap, OldToNewNameMap, OldObjIndex, ObjectsToReplace, PotentialEditorsForRefreshing, OwnersToRerunConstructionScript, &FDirectAttachChildrenAccessor::Get, false, bArchetypesAreUpToDate);
 						}
 						UpdateObjectBeingDebugged(OldObject, NewUObject);
 						ObjectsReplaced.Add(OldObject);
@@ -2902,11 +2857,8 @@ void FBlueprintCompileReinstancer::ReparentChild(UClass* ChildClass)
 	ChildClass->StaticLink(true);
 }
 
-void FBlueprintCompileReinstancer::CopyPropertiesForUnrelatedObjects(UObject* OldObject, UObject* NewObject, bool bClearExternalReferences, bool bForceDeltaSerialization /* = false */)
+void FBlueprintCompileReinstancer::CopyPropertiesForUnrelatedObjects(UObject* OldObject, UObject* NewObject, bool bClearExternalReferences, bool bForceDeltaSerialization /* = false */, bool bOnlyHandleDirectSubObjects/* = false */, TMap<UObject*, UObject*>* OldToNewInstanceMap /*=nullptr*/)
 {
-	InstancedPropertyUtils::FInstancedPropertyMap InstancedPropertyMap;
-	InstancedPropertyUtils::FArchiveInstancedSubObjCollector  InstancedSubObjCollector(OldObject, InstancedPropertyMap);
-
 	UEngine::FCopyPropertiesForUnrelatedObjectsParams Params;
 	// During a blueprint reparent, delta serialization must be enabled to correctly copy all properties
 	Params.bDoDelta = bForceDeltaSerialization || !OldObject->HasAnyFlags(RF_ClassDefaultObject);
@@ -2914,9 +2866,71 @@ void FBlueprintCompileReinstancer::CopyPropertiesForUnrelatedObjects(UObject* Ol
 	Params.bSkipCompilerGeneratedDefaults = true;
 	Params.bClearReferences = bClearExternalReferences;
 	Params.bNotifyObjectReplacement = true;
+	Params.OptionalReplacementMappings = OldToNewInstanceMap;
+	Params.bOnlyHandleDirectSubObjects = bOnlyHandleDirectSubObjects;
 	UEngine::CopyPropertiesForUnrelatedObjects(OldObject, NewObject, Params);
+}
+
+void FBlueprintCompileReinstancer::PreCreateSubObjectsForReinstantiation(const TMap<UClass*, UClass*>& OldToNewClassMap, UObject* OldObject, UObject* NewUObject, TMap<UObject*, UObject*>& CreatedInstanceMap, const TMap<UObject*, UObject*>* OldToNewInstanceMap/* = nullptr*/)
+{
+	// Make sure the mapping from the old to the new exist...
+	CreatedInstanceMap.Add(OldObject, NewUObject);
+
+	// Gather subobjects on old object and pre-create them if needed
+	TArray<UObject*> ContainedOldSubObjects;
+	GetObjectsWithOuter(OldObject, ContainedOldSubObjects, /*bIncludeNestedObjects*/false);
 
-	InstancedPropertyUtils::FArchiveInsertInstancedSubObjects InstancedSubObjSpawner(NewObject, InstancedPropertyMap);
+	// Pre-create all non default subobjects to prevent re-instancing object them an old classes 
+	TMap<UObject*, UObject*> ReferenceReplacementMap;
+	for (int32 i = 0; i < ContainedOldSubObjects.Num(); ++i)
+	{
+		UObject* OldSubObject = ContainedOldSubObjects[i];
+
+		UClass* OldSubObjectClass = OldSubObject->GetClass();
+		UObject* OldSubObjectOuter = OldSubObject->GetOuter();
+		FName SubObjectName = OldSubObject->GetFName();
+		EObjectFlags SubObjectFlags = OldSubObject->GetFlags() & UE::ReinstanceUtils::FlagMask;
+
+		// Only re-create objects that are created from a default object,
+		// skip any default subobjects are they will be created by their parent object CDO.
+		UObject* Archetype = UObject::GetArchetypeFromRequiredInfo(OldSubObjectClass, OldSubObjectOuter, SubObjectName, SubObjectFlags);
+		if (Archetype->HasAnyFlags(RF_ClassDefaultObject))
+		{
+			UObject* NewSubObjectOuter = CreatedInstanceMap.FindChecked(OldSubObjectOuter);
+
+			// Was it already re-instantiated
+			if (UObject* const* AlreadyCreatedSubObject = OldToNewInstanceMap ? OldToNewInstanceMap->Find(OldSubObject) : nullptr)
+			{
+				int32 AlreadyCreatedSubObjectIndex = ContainedOldSubObjects.Find(*AlreadyCreatedSubObject);
+				checkf(AlreadyCreatedSubObjectIndex > i, TEXT("Expecting the already created subobject to be in the old subobject list after this current old sub object"));
+				ContainedOldSubObjects.RemoveAt(AlreadyCreatedSubObjectIndex);
+				(*AlreadyCreatedSubObject)->Rename(nullptr, NewSubObjectOuter, REN_DoNotDirty | REN_DontCreateRedirectors | REN_ForceNoResetLoaders | REN_NonTransactional);
+			}
+			else
+			{
+				checkf( !OldToNewInstanceMap || !OldToNewInstanceMap->FindKey(OldSubObject), TEXT("For performance reason, let's assume the any old sub object will be before its replacement in the contained subobject list"));
+
+				// No need to handled invalid objects from this point on
+				if (!IsValid(OldSubObject))
+				{
+					continue;
+				}
+
+				UClass* SubObjectClass = OldSubObjectClass;
+				if (UClass* const* NewSubObjectClass = OldToNewClassMap.Find(OldSubObjectClass))
+				{
+					SubObjectClass = *NewSubObjectClass;
+				}
+
+				// Only pre-create object where the class does not have newer version of the it
+				if(!SubObjectClass->HasAnyClassFlags(CLASS_NewerVersionExists))
+				{
+					UObject* NewSubObject = NewObject<UObject>(NewSubObjectOuter, SubObjectClass, SubObjectName, SubObjectFlags);
+					PreCreateSubObjectsForReinstantiation(OldToNewClassMap, OldSubObject, NewSubObject, CreatedInstanceMap, OldToNewInstanceMap);
+				}
+			}
+		}
+	}
 }
 
 FRecreateUberGraphFrameScope::FRecreateUberGraphFrameScope(UClass* InClass, bool bRecreate)
diff --git a/Engine/Source/Editor/UnrealEd/Public/Kismet2/KismetReinstanceUtilities.h b/Engine/Source/Editor/UnrealEd/Public/Kismet2/KismetReinstanceUtilities.h
--- a/Engine/Source/Editor/UnrealEd/Public/Kismet2/KismetReinstanceUtilities.h
+++ b/Engine/Source/Editor/UnrealEd/Public/Kismet2/KismetReinstanceUtilities.h
@@ -296,7 +296,20 @@ protected:
 	* @param bClearExternalReferences		If true then attempt to replace references to old classes and instances on this object with the corresponding new ones
 	* @param bForceDeltaSerialization		If true the delta serialization will be used when copying
 	*/
-	static UNREALED_API void CopyPropertiesForUnrelatedObjects(UObject* OldObject, UObject* NewObject, bool bClearExternalReferences, bool bForceDeltaSerialization = false);
+	static UNREALED_API void CopyPropertiesForUnrelatedObjects(UObject* OldObject, UObject* NewObject, bool bClearExternalReferences, bool bForceDeltaSerialization = false, bool bOnlyHandleDirectSubObjects = false, TMap<UObject*, UObject*>* OldToNewInstanceMap =nullptr);
+
+public:
+	/**
+	 * This method will precreate all non-default sub object needed for a re-instantiation, what is left is to CopyPropertiesForUnrelatedObjects on the created instances map to finish the re-instancing
+	 * If the re-intancing is done in a big batch and part to the sub object might already be re-instantiated, you will need to provide those via the OldToNewInstanceMap
+	 * 
+	 * @param OldToNewClassMap in case there are subobject that will need new class type
+	 * @param OldObject to pre-create its non-default sub objects
+	 * @param NewUObject where to store those precreated sub objects
+	 * @param CreatedInstanceMap in/out of the result of all of the pre-created objects
+	 * @param OldToNewInstanceMap optional parameter of the possible re-instanced sub objects if any
+	 */
+	static UNREALED_API void PreCreateSubObjectsForReinstantiation(const TMap<UClass*, UClass*>& OldToNewClassMap, UObject* OldObject, UObject* NewUObject, TMap<UObject*, UObject*>& CreatedInstanceMap, const TMap<UObject*, UObject*>* OldToNewInstanceMap = nullptr);
 
 private:
 	/** Handles the work of ReplaceInstancesOfClass, handling both normal replacement of instances and batch */
diff --git a/Engine/Source/Runtime/CoreUObject/Private/UObject/UObjectArchetype.cpp b/Engine/Source/Runtime/CoreUObject/Private/UObject/UObjectArchetype.cpp
--- a/Engine/Source/Runtime/CoreUObject/Private/UObject/UObjectArchetype.cpp
+++ b/Engine/Source/Runtime/CoreUObject/Private/UObject/UObjectArchetype.cpp
@@ -172,11 +172,6 @@ void CacheArchetypeForObject(UObject* Object, UObject* Archetype)
 UObject* UObject::GetArchetypeFromRequiredInfo(const UClass* Class, const UObject* Outer, FName Name, EObjectFlags ObjectFlags)
 {
 	bool bUseUpToDateClass = false;
-#if WITH_EDITOR
-	// While compiling we just want to use whatever is in the object hierarchy,
-	// as some instances within the hierarchy may also be compiling:
-	bUseUpToDateClass = GIsReinstancing && Class->GetAuthoritativeClass() == Class;
-#endif
 	return GetArchetypeFromRequiredInfoImpl(Class, Outer, Name, ObjectFlags, bUseUpToDateClass);
 }
 
diff --git a/Engine/Source/Runtime/Engine/Classes/Engine/Engine.h b/Engine/Source/Runtime/Engine/Classes/Engine/Engine.h
--- a/Engine/Source/Runtime/Engine/Classes/Engine/Engine.h
+++ b/Engine/Source/Runtime/Engine/Classes/Engine/Engine.h
@@ -2939,6 +2939,7 @@ public:
 		bool bCopyDeprecatedProperties;
 		bool bPreserveRootComponent;
 		bool bPerformDuplication;
+		bool bOnlyHandleDirectSubObjects;
 
 		/** Skips copying properties with BlueprintCompilerGeneratedDefaults metadata */
 		bool bSkipCompilerGeneratedDefaults;
diff --git a/Engine/Source/Runtime/Engine/Private/UnrealEngine.cpp b/Engine/Source/Runtime/Engine/Private/UnrealEngine.cpp
--- a/Engine/Source/Runtime/Engine/Private/UnrealEngine.cpp
+++ b/Engine/Source/Runtime/Engine/Private/UnrealEngine.cpp
@@ -16733,6 +16733,7 @@ UEngine::FCopyPropertiesForUnrelatedObjectsParams::FCopyPropertiesForUnrelatedOb
 	, bCopyDeprecatedProperties(false)
 	, bPreserveRootComponent(true)
 	, bPerformDuplication(false)
+	, bOnlyHandleDirectSubObjects(false)
 	, bSkipCompilerGeneratedDefaults(false)
 	, bNotifyObjectReplacement(false)
 	, bClearReferences(true)
@@ -16791,9 +16792,9 @@ void UEngine::CopyPropertiesForUnrelatedObjects(UObject* OldObject, UObject* New
 	// includes any subobjects that have a matching name on the archetype. This function returns all subobjects that have
 	// a matching instancing on Object's archetype and any objects tagged as RF_DefaultSubObject. Non-DSO instanced subobjects
 	// may be missing, but testing will determine that:
-	const auto CollectAllSubobjects = [](UObject* Object, TArray<UObject*>& OutSubobjectArray)
+	const auto CollectAllSubobjects = [&Params](UObject* Object, TArray<UObject*>& OutSubobjectArray)
 	{
-		const bool bIncludedNestedObjects = true;
+		const bool bIncludedNestedObjects = !Params.bOnlyHandleDirectSubObjects;
 		GetObjectsWithOuter(Object, OutSubobjectArray, bIncludedNestedObjects);
 
 		// Remove contained objects that are not subobjects.
@@ -16804,6 +16805,11 @@ void UEngine::CopyPropertiesForUnrelatedObjects(UObject* OldObject, UObject* New
 			{
 				OutSubobjectArray.RemoveAtSwap(ComponentIndex--);
 			}
+			else if(Params.bOnlyHandleDirectSubObjects)
+			{
+				// If it is a default sub object, get its default
+				GetObjectsWithOuter(PotentialComponent, OutSubobjectArray, false);
+			}
 		}
 	};
 
@@ -16866,12 +16872,18 @@ void UEngine::CopyPropertiesForUnrelatedObjects(UObject* OldObject, UObject* New
 	// Gather references to old instances or objects that need to be replaced after we serialize in saved data
 	TMap<UObject*, UObject*> ReferenceReplacementMap;
 	ReferenceReplacementMap.Add(OldObject, NewObject);
-	ReferenceReplacementMap.Add(OldObject->GetArchetype(), NewObject->GetArchetype());
-	if (Params.bReplaceObjectClassReferences)
+	if (OldObject->GetArchetype() != NewObject->GetArchetype())
+	{
+		ReferenceReplacementMap.Add(OldObject->GetArchetype(), NewObject->GetArchetype());
+	}
+	if (OldObject->GetClass() != NewObject->GetClass())
 	{
-		ReferenceReplacementMap.Add(OldObject->GetClass(), NewObject->GetClass());
+		if (Params.bReplaceObjectClassReferences)
+		{
+			ReferenceReplacementMap.Add(OldObject->GetClass(), NewObject->GetClass());
+		}
+		ReferenceReplacementMap.Add(OldObject->GetClass()->GetDefaultObject(), NewObject->GetClass()->GetDefaultObject());
 	}
-	ReferenceReplacementMap.Add(OldObject->GetClass()->GetDefaultObject(), NewObject->GetClass()->GetDefaultObject());
 
 	TArray<UObject*> ComponentsOnNewObject;
 	TMap<UObject*, UObject*> AggressiveReplaceReferences;
@@ -16983,7 +16995,7 @@ void UEngine::CopyPropertiesForUnrelatedObjects(UObject* OldObject, UObject* New
 	// Also, if requested, leave pointers to instances of renewed classes intact, assuming we are in the midst of reinstancing, and those will be reinstanced
 	ForEachObjectWithOuter(OldObject, [&ReferenceReplacementMap, &Params](UObject* ObjectInOuter)
 	{
-		if (!ReferenceReplacementMap.Contains(ObjectInOuter))
+		if (!ReferenceReplacementMap.Contains(ObjectInOuter) && (!Params.OptionalReplacementMappings || !Params.OptionalReplacementMappings->Contains(ObjectInOuter)))
 		{
 			ReferenceReplacementMap.Add(ObjectInOuter, nullptr);
 		}
