Nullptr check to prevent CanPasteHere from crashing.

diff --git a/Engine/Source/Editor/UnrealEd/Private/MaterialGraphNode.cpp b/Engine/Source/Editor/UnrealEd/Private/MaterialGraphNode.cpp
--- a/Engine/Source/Editor/UnrealEd/Private/MaterialGraphNode.cpp
+++ b/Engine/Source/Editor/UnrealEd/Private/MaterialGraphNode.cpp
@@ -139,7 +139,7 @@ void UMaterialGraphNode::PostDuplicate(bool bDuplicateForPIE)
 
 bool UMaterialGraphNode::CanPasteHere(const UEdGraph* TargetGraph) const
 {
-	if (Super::CanPasteHere(TargetGraph))
+	if (Super::CanPasteHere(TargetGraph) && MaterialExpression)
 	{
 		if (MaterialExpression->IsA(UMaterialExpressionStrataLegacyConversion::StaticClass()))
 		{
