Made LoadPackageForMerge public

diff --git a/Engine/Source/Editor/UnrealEd/Private/MergeUtils.cpp b/Engine/Source/Editor/UnrealEd/Private/MergeUtils.cpp
--- a/Engine/Source/Editor/UnrealEd/Private/MergeUtils.cpp
+++ b/Engine/Source/Editor/UnrealEd/Private/MergeUtils.cpp
@@ -69,6 +69,32 @@ struct FScopedMergeResolveTransaction
 	bool bCanceled = false;
 };
 
+
+UPackage* MergeUtils::LoadPackageForMerge(const FString& SCFile, const FString& Revision, const UPackage* LocalPackage)
+{
+	return DiffUtils::LoadPackageForDiff(FPackagePath::FromLocalPath(LoadSCFileForMerge(SCFile, Revision)), LocalPackage->GetLoadedPath());
+}
+
+FString MergeUtils::LoadSCFileForMerge(const FString& SCFile, const FString& Revision)
+{
+	const FString FileWithRevision = SCFile + TEXT("#") + Revision;
+	const TSharedRef<FDownloadFile, ESPMode::ThreadSafe> DownloadFileOperation = ISourceControlOperation::Create<FDownloadFile>(FPaths::DiffDir(), FDownloadFile::EVerbosity::Full);
+	ISourceControlModule::Get().GetProvider().Execute(DownloadFileOperation, FileWithRevision, EConcurrency::Synchronous);
+	const FString DownloadPath = FPaths::ConvertRelativePathToFull(FPaths::DiffDir() / FPaths::GetCleanFilename(FileWithRevision));
+
+	// move downloaded file to renamed path so it meets ue asset name requirements
+	FString ResultPath = DownloadPath;
+	ResultPath.ReplaceInline(TEXT(".uasset"), TEXT(""));
+	ResultPath.ReplaceCharInline('#', '-');
+	ResultPath.ReplaceCharInline('.', '-');
+	ResultPath = FPaths::CreateTempFilename(*FPaths::GetPath(ResultPath), *FPaths::GetBaseFilename(ResultPath), TEXT(".uasset"));
+	if (ensure(FPlatformFileManager::Get().GetPlatformFile().MoveFile(*ResultPath, *DownloadPath)))
+	{
+		return ResultPath;
+	}
+	return {};
+}
+
 void UUndoableResolveHandler::SetManagedObject(UObject* Object)
 {
 	ManagedObject = Object;
@@ -155,25 +181,6 @@ void UUndoableResolveHandler::PostEditUndo()
 	UObject::PostEditUndo();
 }
 
-static UPackage* LoadMergePackage(const FString& SCFile, const FString& Revision, const UPackage* LocalPackage)
-{
-	const FString FileWithRevision = SCFile + TEXT("#") + Revision;
-	const TSharedRef<FDownloadFile, ESPMode::ThreadSafe> DownloadFileOperation = ISourceControlOperation::Create<FDownloadFile>(FPaths::DiffDir(), FDownloadFile::EVerbosity::Full);
-	ISourceControlModule::Get().GetProvider().Execute(DownloadFileOperation, FileWithRevision, EConcurrency::Synchronous);
-	const FString DownloadPath = FPaths::ConvertRelativePathToFull(FPaths::DiffDir() / FPaths::GetCleanFilename(FileWithRevision));
-	FString CopyPath = DownloadPath;
-	CopyPath.ReplaceInline(TEXT(".uasset"), TEXT(""));
-	CopyPath.ReplaceCharInline('#', '-');
-	CopyPath.ReplaceCharInline('.', '-');
-	CopyPath = FPaths::CreateTempFilename(*FPaths::GetPath(CopyPath), *FPaths::GetBaseFilename(CopyPath), TEXT(".uasset"));
-
-	if (FPlatformFileManager::Get().GetPlatformFile().CopyFile(*CopyPath, *DownloadPath))
-	{
-		return DiffUtils::LoadPackageForDiff(FPackagePath::FromLocalPath(CopyPath), LocalPackage->GetLoadedPath());
-	}
-	return nullptr;
-}
-
 struct FBPReferenceFinder : public FArchiveUObject
 {
 	FBPReferenceFinder(const UObject* Obj, TArray<const UObject*>& InReferences)
@@ -1045,7 +1052,7 @@ EAssetCommandResult MergeUtils::Merge(const FAssetAutomaticMergeArgs& MergeArgs)
 		const ISourceControlState::FResolveInfo ResolveInfo = SourceControlState->GetResolveInfo();
 		check(ResolveInfo.IsValid());
 		
-		if(UPackage* TempPackage = LoadMergePackage(ResolveInfo.RemoteFile, ResolveInfo.RemoteRevision, LocalPackage))
+		if(UPackage* TempPackage = LoadPackageForMerge(ResolveInfo.RemoteFile, ResolveInfo.RemoteRevision, LocalPackage))
 		{
 			// Grab the old asset from that old package
 			ManualMergeArgs.RemoteAsset = FindObject<UObject>(TempPackage, *ManualMergeArgs.LocalAsset->GetName());
@@ -1057,7 +1064,7 @@ EAssetCommandResult MergeUtils::Merge(const FAssetAutomaticMergeArgs& MergeArgs)
 			}
 		}
 		
-		if(UPackage* TempPackage = LoadMergePackage(ResolveInfo.BaseFile, ResolveInfo.BaseRevision, LocalPackage))
+		if(UPackage* TempPackage = LoadPackageForMerge(ResolveInfo.BaseFile, ResolveInfo.BaseRevision, LocalPackage))
 		{
 			// Grab the old asset from that old package
 			ManualMergeArgs.BaseAsset = FindObject<UObject>(TempPackage, *ManualMergeArgs.LocalAsset->GetName());
diff --git a/Engine/Source/Editor/UnrealEd/Public/MergeUtils.h b/Engine/Source/Editor/UnrealEd/Public/MergeUtils.h
--- a/Engine/Source/Editor/UnrealEd/Public/MergeUtils.h
+++ b/Engine/Source/Editor/UnrealEd/Public/MergeUtils.h
@@ -8,6 +8,11 @@ namespace MergeUtils
 {
 	UNREALED_API EAssetCommandResult Merge(const FAssetAutomaticMergeArgs& MergeArgs);
 	UNREALED_API EAssetCommandResult Merge(const FAssetManualMergeArgs& MergeArgs);
+
+	// download and load a temp package at a source control file path
+	UNREALED_API UPackage* LoadPackageForMerge(const FString& SCFile, const FString& Revision, const UPackage* LocalPackage);
+	// download a source control file path and return it's local path
+	UNREALED_API FString LoadSCFileForMerge(const FString& SCFile, const FString& Revision);
 }
 
 UCLASS()
