Cooker: IterativeValidate: Improve logging feedback
#rnx

diff --git a/Engine/Source/Editor/UnrealEd/Private/Cooker/IterativeValidatePackageWriter.cpp b/Engine/Source/Editor/UnrealEd/Private/Cooker/IterativeValidatePackageWriter.cpp
--- a/Engine/Source/Editor/UnrealEd/Private/Cooker/IterativeValidatePackageWriter.cpp
+++ b/Engine/Source/Editor/UnrealEd/Private/Cooker/IterativeValidatePackageWriter.cpp
@@ -142,6 +142,10 @@ void FIterativeValidatePackageWriter::UpdatePackageModificationStatus(FName Pack
 	case EPhase::FirstCook:
 		// Invert what gets skipped: save the iteratively skipped files to record their diffs, but skip the regular files
 		bInOutShouldIterativelySkip = !bIterativelyUnmodified;
+		if (!bIterativelyUnmodified)
+		{
+			++ModifiedCount;
+		}
 		break;
 	case EPhase::FinalCook:
 		// Ignore the Unmodified flag from this cook phase. Skip only the packages that were found to 
@@ -165,9 +169,17 @@ void FIterativeValidatePackageWriter::BeginCook(const FCookInfo& Info)
 	switch (Phase)
 	{
 	case EPhase::FirstCook:
+		UE_LOG(LogIterativeValidate, Display,
+			TEXT("Phase IterativeValidatePrePass: running -diffonly and a resave on all packages discovered to be iteratively unmodified."));
 		break;
 	case EPhase::FinalCook:
 		Load();
+		UE_LOG(LogIterativeValidate, Display,
+			TEXT("Phase IterativeValidate: %d packages were found during PrePass to be iteratively unmodified but had differences. Running -diffonly on them again to check whether the differences are due to indeterminism or to IterativeFalsePositives."),
+			IterativeFailed.Num());
+		UE_LOG(LogIterativeValidate, Display,
+			TEXT("Phase IterativeValidate: %d packages were found during PrePass to be modified or new and will be resaved."),
+			ModifiedCount);
 		break;
 	default:
 		checkNoEntry();
@@ -183,15 +195,15 @@ void FIterativeValidatePackageWriter::EndCook(const FCookInfo& Info)
 	{
 	case EPhase::FirstCook:
 		UE_LOG(LogIterativeValidate, Display,
-			TEXT("DetectedUnmodified: %d. ValidatedUnmodified: %d. IterativeFalseNegativeOrIndeterminism: %d."),
-			IterativeValidated.Num() + IterativeFailed.Num(), IterativeValidated.Num(), IterativeFailed.Num());
+			TEXT("Modified: %d. DetectedUnmodified: %d. ValidatedUnmodified: %d. IterativeFalseNegativeOrIndeterminism: %d."),
+			ModifiedCount, IterativeValidated.Num() + IterativeFailed.Num(), IterativeValidated.Num(), IterativeFailed.Num());
 		Save();
 		break;
 	case EPhase::FinalCook:
 	{
 		UE_LOG(LogIterativeValidate, Display,
-			TEXT("DetectedUnmodified: %d. ValidatedUnmodified: %d. Indeterminism: %d."),
-			IterativeValidated.Num() + IterativeFailed.Num(), IterativeValidated.Num(), IndeterminismFailed.Num());
+			TEXT("Modified: %d. DetectedUnmodified: %d. ValidatedUnmodified: %d. Indeterminism: %d."),
+			ModifiedCount, IterativeValidated.Num() + IterativeFailed.Num(), IterativeValidated.Num(), IndeterminismFailed.Num());
 		FString Message = FString::Printf(TEXT("IterativeFalseNegative: %d."), IterativeFalseNegative.Num());
 		if (IterativeFalseNegative.Num())
 		{
@@ -280,6 +292,7 @@ bool FIterativeValidatePackageWriter::IsAnotherSaveNeeded(FSavePackageResultStru
 			// New packages need to be resaved in the FinalCook phase; for our purposes they are equivalent
 			// to a package that iteration detected as modified.
 			// Do not add an entry for it in our results for iterative packages, and do not resave it in this pass
+			++ModifiedCount;
 			return false;
 		}
 	case EPhase::FinalCook:
@@ -384,6 +397,7 @@ void FIterativeValidatePackageWriter::Serialize(FArchive& Ar)
 	}
 	Ar << IterativeValidated;
 	Ar << IterativeFailed;
+	Ar << ModifiedCount;
 }
 
 FArchive& operator<<(FArchive& Ar, FIterativeValidatePackageWriter::FMessage& Message)
diff --git a/Engine/Source/Editor/UnrealEd/Private/Cooker/IterativeValidatePackageWriter.h b/Engine/Source/Editor/UnrealEd/Private/Cooker/IterativeValidatePackageWriter.h
--- a/Engine/Source/Editor/UnrealEd/Private/Cooker/IterativeValidatePackageWriter.h
+++ b/Engine/Source/Editor/UnrealEd/Private/Cooker/IterativeValidatePackageWriter.h
@@ -59,6 +59,7 @@ protected:
 	TSet<FName> IterativeFalseNegative;
 
 	FString MetadataPath;
+	int32 ModifiedCount = 0;
 	EPhase Phase = EPhase::FirstCook;
 	bool bPackageSaveToDiskPass = false;
 	bool bPackageFirstPass = false;
