Refactor SummarizeTraceCommandlet.cpp

Moved a significant portion of code from `SummarizeTraceCommandlet.cpp` to `SummarizeTraceUtils.cpp`.

Grouped the CSV utility functions under the `CsvUtil` namespace for improved organization and clarity.

Modified the `FSummarizeScope` struct to include two double arrays for recording time interval data required for editor workflows test.

diff --git a/Engine/Source/Editor/UnrealEd/Classes/Commandlets/SummarizeTraceCommandlet.h b/Engine/Source/Editor/UnrealEd/Classes/Commandlets/SummarizeTraceCommandlet.h
--- a/Engine/Source/Editor/UnrealEd/Classes/Commandlets/SummarizeTraceCommandlet.h
+++ b/Engine/Source/Editor/UnrealEd/Classes/Commandlets/SummarizeTraceCommandlet.h
@@ -7,6 +7,7 @@
 
 #pragma once
 #include "Commandlets/Commandlet.h"
+#include "Commandlets/SummarizeTraceUtils.h"
 #include "TraceServices/Model/Counters.h"
 #include "SummarizeTraceCommandlet.generated.h"
 
@@ -22,9 +23,8 @@ public:
 
 	//~ End UCommandlet Interface
 private:
-	void WriteAsUTF8String(IFileHandle* Handle, const FString& String);
 	TUniquePtr<IFileHandle> OpenCSVFile(const FString& Name);
-	bool GenerateScopesCSV(const TArray<struct FSummarizeScope>& SortedScopes);
+	bool GenerateScopesCSV(const TArray<FSummarizeScope>& SortedScopes);
 	bool GenerateCountersCSV(const class FSummarizeCountersProvider& CountersProvider);
 	bool GenerateBookmarksCSV(const class FSummarizeBookmarksProvider& BookmarksProvider);
 	bool GenerateTelemetryCSV(const FString& TestName, bool bAllTelemetry, const TArray<FSummarizeScope>& SortedScopes, const
diff --git a/Engine/Source/Editor/UnrealEd/Private/Commandlets/SummarizeTraceCommandlet.cpp b/Engine/Source/Editor/UnrealEd/Private/Commandlets/SummarizeTraceCommandlet.cpp
--- a/Engine/Source/Editor/UnrealEd/Private/Commandlets/SummarizeTraceCommandlet.cpp
+++ b/Engine/Source/Editor/UnrealEd/Private/Commandlets/SummarizeTraceCommandlet.cpp
@@ -6,803 +6,19 @@
 
 #include "Commandlets/SummarizeTraceCommandlet.h"
 
-#include "Containers/Array.h"
-#include "Containers/ArrayView.h"
-#include "Containers/Map.h"
-#include "Containers/Set.h"
-#include "Containers/StringConv.h"
-#include "Containers/StringFwd.h"
-#include "Containers/StringView.h"
-#include "GenericPlatform/GenericPlatformAffinity.h"
-#include "GenericPlatform/GenericPlatformFile.h"
 #include "HAL/PlatformFileManager.h"
-#include "Logging/LogCategory.h"
-#include "Logging/LogMacros.h"
-#include "Math/NumericLimits.h"
-#include "Math/UnrealMathSSE.h"
-#include "Misc/AssertionMacros.h"
-#include "Misc/CString.h"
-#include "Misc/Crc.h"
-#include "Misc/FileHelper.h"
-#include "Misc/Optional.h"
-#include "Misc/Parse.h"
 #include "Misc/Paths.h"
-#include "Misc/Timespan.h"
-#include "Model/MonotonicTimeline.h"
+#include "Misc/FileHelper.h"
 #include "ProfilingDebugging/CountersTrace.h"
 #include "String/ParseTokens.h"
-#include "Templates/ChooseClass.h"
-#include "Templates/Function.h"
-#include "Templates/SharedPointer.h"
-#include "Templates/Tuple.h"
-#include "Templates/UniquePtr.h"
-#include "Templates/UnrealTemplate.h"
+#include "Trace/Detail/Channel.h"
 #include "Trace/Analysis.h"
 #include "Trace/DataStream.h"
-#include "Trace/Detail/Channel.h"
 #include "TraceServices/AnalyzerFactories.h"
 #include "TraceServices/Model/Bookmarks.h"
 #include "TraceServices/Model/Counters.h"
-#include "TraceServices/Model/Threads.h"
-#include "TraceServices/Model/TimingProfiler.h"
 #include "TraceServices/Utils.h"
 
-namespace TraceServices { class IAnalysisSession; }
-namespace UE::Trace { class IAnalyzer; }
-
-/**
- * Base class to extend for CPU scope analysis. Derived class instances are meant to be registered as
- * delegates to the FSummarizeCpuProfilerProvider to handle scope events and perform meaningful analysis.
- */
-class FSummarizeCpuScopeAnalyzer
-{
-public:
-	enum class EScopeEventType : uint32
-	{
-		Enter,
-		Exit
-	};
-
-	struct FScopeEvent
-	{
-		EScopeEventType ScopeEventType;
-		uint32 ScopeId;
-		uint32 ThreadId;
-		double Timestamp; // As Seconds
-	};
-
-	struct FScope
-	{
-		uint32 ScopeId;
-		uint32 ThreadId;
-		double EnterTimestamp; // As Seconds
-		double ExitTimestamp;  // As Seconds
-	};
-
-public:
-	virtual ~FSummarizeCpuScopeAnalyzer() = default;
-
-	/** Invoked when a CPU scope is discovered. This function is always invoked first when a CPU scope is encountered for the first time.*/
-	virtual void OnCpuScopeDiscovered(uint32 ScopeId) {}
-
-	/** Invoked when CPU scope specification is encountered in the trace stream. */
-	virtual void OnCpuScopeName(uint32 ScopeId, const FStringView& ScopeName) {};
-
-	/** Invoked when a scope is entered. The scope name might not be known yet. */
-	virtual void OnCpuScopeEnter(const FScopeEvent& ScopeEnter, const FString* ScopeName) {};
-
-	/** Invoked when a scope is exited. The scope name might not be known yet. */
-	virtual void OnCpuScopeExit(const FScope& Scope, const FString* ScopeName) {};
-
-	/** Invoked when a root event on the specified thread along with all child events down to the leaves are known. */
-	virtual void OnCpuScopeTree(uint32 ThreadId, const TArray64<FSummarizeCpuScopeAnalyzer::FScopeEvent>& ScopeEvents, const TFunction<const FString*(uint32)>& ScopeLookupNameFn) {};
-
-	/** Invoked when the trace stream has been fully consumed/processed. */
-	virtual void OnCpuScopeAnalysisEnd() {};
-
-	static constexpr uint32 CoroutineSpecId        = (1u << 31u) - 1u;
-	static constexpr uint32 CoroutineUnknownSpecId = (1u << 31u) - 2u;
-};
-
-/**
- * Decodes, format and routes CPU scope events embedded into a trace stream to specialized CPU scope analyzers. This processor
- * decodes the low level events, keeps a small state and publishes higher level events to registered analyzers. The purpose
- * is to decode the stream only once and let several registered analyzers reuse the small state built up by this processor. This
- * matters when processing very large traces with possibly billions of scope events.
- */
-class FSummarizeCpuProfilerProvider
-	: public TraceServices::IEditableThreadProvider
-	, public TraceServices::IEditableTimingProfilerProvider
-{
-public:
-	FSummarizeCpuProfilerProvider();
-
-	/** Register an analyzer with this processor. The processor decodes the trace stream and invokes the registered analyzers when a CPU scope event occurs.*/
-	void AddCpuScopeAnalyzer(TSharedPtr<FSummarizeCpuScopeAnalyzer> Analyzer);
-
-	void AnalysisComplete();
-
-private:
-	struct FScopeEnter
-	{
-		uint32 ScopeId;
-		double Timestamp; // As Seconds
-	};
-
-	// Contains scope events for a root scope and its children along with extra info to analyze that tree at once.
-	struct FScopeTreeInfo
-	{
-		// Records the current root scope and its children to run analysis that needs to know the parent/child relationship.
-		TArray64<FSummarizeCpuScopeAnalyzer::FScopeEvent> ScopeEvents;
-
-		// Indicates if one of the scope in the current hierarchy is nameless. (Its names specs hasn't been received yet).
-		bool bHasNamelessScopes = false;
-
-		void Reset()
-		{
-			ScopeEvents.Reset();
-			bHasNamelessScopes = false;
-		}
-	};
-
-	// For each thread we track what the stack of scopes are, for matching end-to-start
-	struct FThread
-		: public TraceServices::IEditableTimeline<TraceServices::FTimingProfilerEvent>
-	{
-		FThread(uint32 InThreadId, FSummarizeCpuProfilerProvider* InProvider)
-			: ThreadId(InThreadId)
-			, Provider(InProvider)
-		{
-
-		}
-
-		virtual void AppendBeginEvent(double StartTime, const TraceServices::FTimingProfilerEvent& Event) override;
-		virtual void AppendEndEvent(double EndTime) override;
-
-		// The ThreadId of this thread
-		uint32 ThreadId;
-
-		// The provider to forward calls to
-		FSummarizeCpuProfilerProvider* Provider;
-
-		// The current stack state
-		TArray<FScopeEnter> ScopeStack;
-
-		// The events recorded for the current root scope and its children to run analysis that needs to know the parent/child relationship, for example to compute time including childreen and time
-		// excluding childreen.
-		FScopeTreeInfo ScopeTreeInfo;
-
-		// Scope trees for which as least one scope name was unknown. Some analysis need scope names, but scope names/event can be emitted out of order by the engine depending on thread scheduling.
-		// Some scope tree cannot be analyzed right away and need to be delayed until all scope names are discovered.
-		TArray<FScopeTreeInfo> DelayedScopeTreeInfo;
-	};
-
-private:
-	// TraceServices::IEditableThreadProvider
-	virtual void					AddThread(uint32 Id, const TCHAR* Name, EThreadPriority Priority) override;
-
-	// TraceServices::IEditableTimingProfilerProvider
-	virtual uint32					AddCpuTimer(FStringView Name, const TCHAR* File, uint32 Line) override;
-	virtual void					SetTimerNameAndLocation(uint32 TimerId, FStringView Name, const TCHAR* File, uint32 Line) override;
-	virtual uint32					AddMetadata(uint32 MasterTimerId, TArray<uint8>&& Metadata) override;
-	virtual TArrayView<const uint8> GetMetadata(uint32 TimerId) const override;
-	virtual TraceServices::IEditableTimeline<TraceServices::FTimingProfilerEvent>& GetCpuThreadEditableTimeline(uint32 ThreadId) override;
-
-	// Callbacks from FThread to forward calls to Analyzers
-	virtual void AppendBeginEvent(const FSummarizeCpuScopeAnalyzer::FScopeEvent& ScopeEvent);
-	virtual void AppendEndEvent(const FSummarizeCpuScopeAnalyzer::FScope& ScopeEvent, const FString* ScopeName);
-
-	// Processing the ScopeTree once it's complete
-	void OnCpuScopeTree(uint32 ThreadId, const FScopeTreeInfo& ScopeTreeInfo);
-
-	// Resolve a ScopeId to string
-	const FString* LookupScopeName(uint32 ScopeId);
-
-private:
-	// The state at any moment of the threads
-	TMap<uint32, TUniquePtr<FThread>> Threads;
-
-	// The scope names, the array index correspond to the scope Id. If the optional is not set, the scope hasn't been encountered yet.
-	TArray<TOptional<FString>> ScopeNames;
-
-	// List of analyzers to invoke when a scope event is decoded.
-	TArray<TSharedPtr<FSummarizeCpuScopeAnalyzer>> ScopeAnalyzers;
-
-	// Scope name lookup function, cached for efficiency.
-	TFunction<const FString*(uint32 ScopeId)> LookupScopeNameFn;
-
-	struct FMetadata
-	{
-		TArray<uint8> Payload;
-		uint32 TimerId;
-	};
-
-	TArray<FMetadata> Metadatas;
-};
-
-FSummarizeCpuProfilerProvider::FSummarizeCpuProfilerProvider()
- : LookupScopeNameFn([this](uint32 ScopeId) { return LookupScopeName(ScopeId); })
-{
-}
-
-void FSummarizeCpuProfilerProvider::AddCpuScopeAnalyzer(TSharedPtr<FSummarizeCpuScopeAnalyzer> Analyzer)
-{
-	ScopeAnalyzers.Add(MoveTemp(Analyzer));
-}
-
-void FSummarizeCpuProfilerProvider::AnalysisComplete()
-{
-	// Analyze scope trees that contained 'nameless' context when they were captured. Unless the trace was truncated,
-	// all scope names should be known now.
-	for (const auto& Thread : Threads)
-	{
-		for (FScopeTreeInfo& DelayedScopeTree : Thread.Value->DelayedScopeTreeInfo)
-		{
-			// Run summary analysis for this delayed hierarchy.
-			OnCpuScopeTree(Thread.Key, DelayedScopeTree);
-		}
-	}
-
-	for (TSharedPtr<FSummarizeCpuScopeAnalyzer>& Analyzer : ScopeAnalyzers)
-	{
-		Analyzer->OnCpuScopeAnalysisEnd();
-	}
-}
-
-void FSummarizeCpuProfilerProvider::AddThread(uint32 Id, const TCHAR* Name, EThreadPriority Priority)
-{
-	TUniquePtr<FThread>* Found = Threads.Find(Id);
-	if (!Found)
-	{
-		Threads.Add(Id, MakeUnique<FThread>(Id, this));
-	}
-}
-
-uint32 FSummarizeCpuProfilerProvider::AddCpuTimer(FStringView Name, const TCHAR* File, uint32 Line)
-{
-	TOptional<FString> ScopeName;
-
-	if (!Name.IsEmpty())
-	{
-		ScopeName.Emplace(FString(Name));
-	}
-	
-	uint32 TimerId = ScopeNames.Add(ScopeName);
-
-	// Notify the analyzers.
-	for (TSharedPtr<FSummarizeCpuScopeAnalyzer>& Analyzer : ScopeAnalyzers)
-	{
-		Analyzer->OnCpuScopeDiscovered(TimerId);
-
-		if (!Name.IsEmpty())
-		{
-			Analyzer->OnCpuScopeName(TimerId, Name);
-		}
-	}
-
-	return TimerId;
-}
-
-void FSummarizeCpuProfilerProvider::SetTimerNameAndLocation(uint32 TimerId, FStringView Name, const TCHAR* File, uint32 Line)
-{
-	check(TimerId < uint32(ScopeNames.Num()));
-	check(!Name.IsEmpty());
-
-	ScopeNames[TimerId].Emplace(FString(Name));
-
-	// Notify the registered scope analyzers.
-	for (TSharedPtr<FSummarizeCpuScopeAnalyzer>& Analyzer : ScopeAnalyzers)
-	{
-		Analyzer->OnCpuScopeName(TimerId, Name);
-	}
-}
-
-uint32 FSummarizeCpuProfilerProvider::AddMetadata(uint32 MasterTimerId, TArray<uint8>&& Metadata)
-{
-	uint32 MetadataId = Metadatas.Num();
-	Metadatas.Add({ MoveTemp(Metadata), MasterTimerId });
-	return ~MetadataId;
-}
-
-TArrayView<const uint8> FSummarizeCpuProfilerProvider::GetMetadata(uint32 TimerId) const
-{
-	if (int32(TimerId) >= 0)
-	{
-		return TArrayView<const uint8>();
-	}
-
-	TimerId = ~TimerId;
-	if (TimerId >= uint32(Metadatas.Num()))
-	{
-		return TArrayView<const uint8>();
-	}
-
-	const FMetadata& Metadata = Metadatas[TimerId];
-	return Metadata.Payload;
-}
-
-TraceServices::IEditableTimeline<TraceServices::FTimingProfilerEvent>& FSummarizeCpuProfilerProvider::GetCpuThreadEditableTimeline(uint32 ThreadId)
-{
-	TUniquePtr<FThread>* Found = Threads.Find(ThreadId);
-	if (Found)
-	{
-		return *(Found->Get());
-	}
-
-	return *Threads.Add(ThreadId, MakeUnique<FThread>(ThreadId, this));
-}
-
-void FSummarizeCpuProfilerProvider::FThread::AppendBeginEvent(double StartTime, const TraceServices::FTimingProfilerEvent& Event)
-{
-	FScopeEnter ScopeEnter{ Event.TimerIndex, StartTime };
-	ScopeStack.Add(ScopeEnter);
-
-	FSummarizeCpuScopeAnalyzer::FScopeEvent ScopeEvent{ FSummarizeCpuScopeAnalyzer::EScopeEventType::Enter, Event.TimerIndex, ThreadId, StartTime };
-	ScopeTreeInfo.ScopeEvents.Add(ScopeEvent);
-
-	Provider->AppendBeginEvent(ScopeEvent);
-}
-
-void FSummarizeCpuProfilerProvider::FThread::AppendEndEvent(double EndTime)
-{
-	if (ScopeStack.IsEmpty())
-	{
-		return;
-	}
-
-	FScopeEnter ScopeEnter = ScopeStack.Pop();
-
-	FSummarizeCpuScopeAnalyzer::FScopeEvent ScopeEvent{ FSummarizeCpuScopeAnalyzer::EScopeEventType::Exit, ScopeEnter.ScopeId, ThreadId, EndTime };
-	ScopeTreeInfo.ScopeEvents.Add(ScopeEvent);
-
-	// Check if at this point if the scope has a name
-	const FString* ScopeName = Provider->LookupScopeName(ScopeEnter.ScopeId);
-	ScopeTreeInfo.bHasNamelessScopes |= ScopeName == nullptr || ScopeName->IsEmpty();
-
-	FSummarizeCpuScopeAnalyzer::FScope Scope{ ScopeEnter.ScopeId, ThreadId, ScopeEnter.Timestamp, EndTime };
-	Provider->AppendEndEvent(Scope, ScopeName);
-
-	// The root scope on this thread just popped out.
-	if (ScopeStack.IsEmpty())
-	{
-		if (ScopeTreeInfo.bHasNamelessScopes)
-		{
-			// Delay the analysis until all the scope names are known.
-			DelayedScopeTreeInfo.Add(MoveTemp(ScopeTreeInfo));
-		}
-		else
-		{
-			// Run analysis for this scope tree.
-			Provider->OnCpuScopeTree(ThreadId, ScopeTreeInfo);
-		}
-
-		ScopeTreeInfo.Reset();
-	}
-}
-
-void FSummarizeCpuProfilerProvider::AppendBeginEvent(const FSummarizeCpuScopeAnalyzer::FScopeEvent& ScopeEvent)
-{
-	// Notify the registered scope analyzers.
-	for (TSharedPtr<FSummarizeCpuScopeAnalyzer>& Analyzer : ScopeAnalyzers)
-	{
-		Analyzer->OnCpuScopeEnter(ScopeEvent, LookupScopeName(ScopeEvent.ScopeId));
-	}
-}
-
-void FSummarizeCpuProfilerProvider::AppendEndEvent(const FSummarizeCpuScopeAnalyzer::FScope& Scope, const FString* ScopeName)
-{
-	// Notify the registered scope analyzers.
-	for (TSharedPtr<FSummarizeCpuScopeAnalyzer>& Analyzer : ScopeAnalyzers)
-	{
-		Analyzer->OnCpuScopeExit(Scope, ScopeName);
-	}
-}
-
-void FSummarizeCpuProfilerProvider::OnCpuScopeTree(uint32 ThreadId, const FScopeTreeInfo& ScopeTreeInfo)
-{
-	// Notify the registered scope analyzers.
-	for (TSharedPtr<FSummarizeCpuScopeAnalyzer>& Analyzer : ScopeAnalyzers)
-	{
-		Analyzer->OnCpuScopeTree(ThreadId, ScopeTreeInfo.ScopeEvents, LookupScopeNameFn);
-	}
-}
-
-const FString* FSummarizeCpuProfilerProvider::LookupScopeName(uint32 ScopeId)
-{
-	if (int32(ScopeId) < 0)
-	{
-		ScopeId = Metadatas[~ScopeId].TimerId;
-	}
-
-	if (ScopeId < static_cast<uint32>(ScopeNames.Num()) && ScopeNames[ScopeId])
-	{
-		return &ScopeNames[ScopeId].GetValue();
-	}
-	
-	return nullptr;
-}
-
-
-/*
-* 
-* Helper classes for the SummarizeTrace commandlet. Aggregates statistics about a trace.
-* Code preceding this comment should eventually be relocated to the Trace/Insights modules.
-* Everything succeeding this comment should stay in this commandlet.
-* 
-*/
-
-class FIncrementalVariance
-{
-public:
-	FIncrementalVariance()
-		: Count(0)
-		, Mean(0.0)
-		, VarianceAccumulator(0.0)
-	{
-
-	}
-
-	uint64 GetCount() const
-	{
-		return Count;
-	}
-
-	double GetMean() const
-	{
-		return Mean;
-	}
-
-	/**
-	* Compute the variance given Welford's accumulator and the overall count
-	*
-	* @return The variance in sample units squared
-	*/
-	double Variance() const
-	{
-		double Result = 0.0;
-
-		if (Count > 1)
-		{
-			// Welford's final step, dependent on sample count
-			Result = VarianceAccumulator / double(Count - 1);
-		}
-
-		return Result;
-	}
-
-	/**
-	* Compute the standard deviation given Welford's accumulator and the overall count
-	*
-	* @return The standard deviation in sample units
-	*/
-	double Deviation() const
-	{
-		double Result = 0.0;
-
-		if (Count > 1)
-		{
-			// Welford's final step, dependent on sample count
-			double DeviationSqrd = VarianceAccumulator / double(Count - 1);
-
-			// stddev is sqrt of variance, to restore to units (vs. units squared)
-			Result = sqrt(DeviationSqrd);
-		}
-
-		return Result;
-	}
-
-	/**
-	* Perform an increment of work for Welford's variance, from which we can compute variation and standard deviation
-	*
-	* @param InSample	The new sample value to operate on
-	*/
-	void Increment(const double InSample)
-	{
-		Count++;
-		const double OldMean = Mean;
-		Mean += ((InSample - Mean) / double(Count));
-		VarianceAccumulator += ((InSample - Mean) * (InSample - OldMean));
-	}
-
-	/**
-	* Merge with another IncrementalVariance series in progress
-	*
-	* @param Other	The other variance incremented from another mutually exclusive population of analogous data.
-	*/
-	void Merge(const FIncrementalVariance& Other)
-	{
-		// empty other, nothing to do
-		if (Other.Count == 0)
-		{
-			return;
-		}
-
-		// empty this, just copy other
-		if (Count == 0)
-		{
-			Count = Other.Count;
-			Mean = Other.Mean;
-			VarianceAccumulator = Other.VarianceAccumulator;
-			return;
-		}
-
-		const double TotalPopulation = static_cast<double>(Count + Other.Count);
-		const double MeanDifference = Mean - Other.Mean;
-		const double A = ((Count - 1) * Variance()) + ((Other.Count - 1) * Other.Variance());
-		const double B = (MeanDifference) * (MeanDifference) * (Count * Other.Count / TotalPopulation);
-		const double MergedVariance = (A + B) / (TotalPopulation - 1);
-
-		const uint64 NewCount = Count + Other.Count;
-		const double NewMean = ((Mean * double(Count)) + (Other.Mean * double(Other.Count))) / double(NewCount);
-		const double NewVarianceAccumulator = MergedVariance * (NewCount - 1);
-
-		Count = NewCount;
-		Mean = NewMean;
-		VarianceAccumulator = NewVarianceAccumulator;
-	}
-
-	/**
-	* Reset state back to initialized.
-	*/
-	void Reset()
-	{
-		Count = 0;
-		Mean = 0.0;
-		VarianceAccumulator = 0.0;
-	}
-
-private:
-	uint64 Count;
-	double Mean;
-	double VarianceAccumulator;
-};
-
-//
-// Helpers for Cpu channel events
-//
-
-struct FSummarizeScope
-{
-	FString Name;
-	FIncrementalVariance DurationVariance;
-	double TotalDurationSeconds = 0.0;
-
-	double FirstStartSeconds = 0.0;
-	double FirstFinishSeconds = 0.0;
-	double FirstDurationSeconds = 0.0;
-
-	double LastStartSeconds = 0.0;
-	double LastFinishSeconds = 0.0;
-	double LastDurationSeconds = 0.0;
-
-	double MinDurationSeconds = TNumericLimits<double>::Max();
-	double MaxDurationSeconds = TNumericLimits<double>::Min();
-
-	void AddDuration(double StartSeconds, double FinishSeconds)
-	{
-		// compute the duration
-		double DurationSeconds = FinishSeconds - StartSeconds;
-
-		DurationVariance.Increment(DurationSeconds);
-
-		// only set first for the first sample
-		if (DurationVariance.GetCount() == 1)
-		{
-			FirstStartSeconds = StartSeconds;
-			FirstFinishSeconds = FinishSeconds;
-			FirstDurationSeconds = DurationSeconds;
-		}
-
-		LastStartSeconds = StartSeconds;
-		LastFinishSeconds = FinishSeconds;
-		LastDurationSeconds = DurationSeconds;
-
-		// set duration statistics
-		TotalDurationSeconds += DurationSeconds;
-		MinDurationSeconds = FMath::Min(MinDurationSeconds, DurationSeconds);
-		MaxDurationSeconds = FMath::Max(MaxDurationSeconds, DurationSeconds);
-	}
-	
-	uint64 GetCount() const
-	{
-		return DurationVariance.GetCount();
-	}
-
-	double GetMeanDurationSeconds() const
-	{
-		return DurationVariance.GetMean();
-	}
-
-	double GetDeviationDurationSeconds() const
-	{
-		return DurationVariance.Deviation();
-	}
-
-	double GetCountPerSecond() const
-	{
-		double CountPerSecond = 0.0;
-		const uint64 Count = DurationVariance.GetCount();
-		if (Count)
-		{
-			CountPerSecond = Count / (LastFinishSeconds - FirstStartSeconds);
-		}
-		return CountPerSecond;
-	}
-
-	void Merge(const FSummarizeScope& Other)
-	{
-		check(Name == Other.Name);
-		DurationVariance.Merge(Other.DurationVariance);
-
-		if (FirstStartSeconds > Other.FirstStartSeconds)
-		{
-			FirstStartSeconds = Other.FirstStartSeconds;
-			FirstFinishSeconds = Other.FirstFinishSeconds;
-			FirstDurationSeconds = Other.FirstDurationSeconds;
-		}
-		
-		if (LastStartSeconds < Other.LastStartSeconds)
-		{
-			LastStartSeconds = Other.LastStartSeconds;
-			LastFinishSeconds = Other.LastFinishSeconds;
-			LastDurationSeconds = Other.LastDurationSeconds;
-		}
-		
-		TotalDurationSeconds += Other.TotalDurationSeconds;
-		MinDurationSeconds = FMath::Min(MinDurationSeconds, Other.MinDurationSeconds);
-		MaxDurationSeconds = FMath::Max(MaxDurationSeconds, Other.MaxDurationSeconds);
-	}
-
-	FString GetValue(const FStringView& Statistic) const
-	{
-		if (Statistic == TEXT("Name"))
-		{
-			return Name;
-		}
-		else if (Statistic == TEXT("Count"))
-		{
-			return FString::Printf(TEXT("%llu"), DurationVariance.GetCount());
-		}
-		else if (Statistic == TEXT("TotalDurationSeconds"))
-		{
-			return FString::Printf(TEXT("%f"), TotalDurationSeconds);
-		}
-		else if (Statistic == TEXT("FirstStartSeconds"))
-		{
-			return FString::Printf(TEXT("%f"), FirstStartSeconds);
-		}
-		else if (Statistic == TEXT("FirstFinishSeconds"))
-		{
-			return FString::Printf(TEXT("%f"), FirstFinishSeconds);
-		}
-		else if (Statistic == TEXT("FirstDurationSeconds"))
-		{
-			return FString::Printf(TEXT("%f"), FirstDurationSeconds);
-		}
-		else if (Statistic == TEXT("LastStartSeconds"))
-		{
-			return FString::Printf(TEXT("%f"), LastStartSeconds);
-		}
-		else if (Statistic == TEXT("LastFinishSeconds"))
-		{
-			return FString::Printf(TEXT("%f"), LastFinishSeconds);
-		}
-		else if (Statistic == TEXT("LastDurationSeconds"))
-		{
-			return FString::Printf(TEXT("%f"), LastDurationSeconds);
-		}
-		else if (Statistic == TEXT("MinDurationSeconds"))
-		{
-			return FString::Printf(TEXT("%f"), MinDurationSeconds);
-		}
-		else if (Statistic == TEXT("MaxDurationSeconds"))
-		{
-			return FString::Printf(TEXT("%f"), MaxDurationSeconds);
-		}
-		else if (Statistic == TEXT("MeanDurationSeconds"))
-		{
-			return FString::Printf(TEXT("%f"), DurationVariance.GetMean());
-		}
-		else if (Statistic == TEXT("DeviationDurationSeconds"))
-		{
-			return FString::Printf(TEXT("%f"), DurationVariance.Deviation());
-		}
-		else if (Statistic == TEXT("CountPerSecond"))
-		{
-			return FString::Printf(TEXT("%f"), GetCountPerSecond());
-		}
-		return FString();
-	}
-
-	// for deduplication
-	bool operator==(const FSummarizeScope& Scope) const
-	{
-		return Name == Scope.Name;
-	}
-
-	// for sorting descending
-	bool operator<(const FSummarizeScope& Scope) const
-	{
-		return TotalDurationSeconds > Scope.TotalDurationSeconds;
-	}
-};
-
-static uint32 GetTypeHash(const FSummarizeScope& Scope)
-{
-	return FCrc::StrCrc32(*Scope.Name);
-}
-
-/**
- * This analyzer aggregates CPU scopes having the same name together, counting the number of occurrences, total duration,
- * average occurrence duration, for each scope name.
- */
-class FSummarizeCpuScopeDurationAnalyzer
-	: public FSummarizeCpuScopeAnalyzer
-{
-public:
-	FSummarizeCpuScopeDurationAnalyzer(TFunction<void(const TArray<FSummarizeScope>&)> InPublishFn);
-
-protected:
-	virtual void OnCpuScopeDiscovered(uint32 ScopeId) override;
-	virtual void OnCpuScopeName(uint32 ScopeId, const FStringView& ScopeName) override;
-	virtual void OnCpuScopeExit(const FScope& Scope, const FString* ScopeName) override;
-	virtual void OnCpuScopeAnalysisEnd() override;
-
-private:
-	TMap<uint32, FSummarizeScope> Scopes;
-
-	// Function invoked to publish the list of scopes.
-	TFunction<void(const TArray<FSummarizeScope>&)> PublishFn;
-};
-
-FSummarizeCpuScopeDurationAnalyzer::FSummarizeCpuScopeDurationAnalyzer(TFunction<void(const TArray<FSummarizeScope>&)> InPublishFn)
-	: PublishFn(MoveTemp(InPublishFn))
-{
-	OnCpuScopeDiscovered(FSummarizeCpuScopeAnalyzer::CoroutineSpecId);
-	OnCpuScopeDiscovered(FSummarizeCpuScopeAnalyzer::CoroutineUnknownSpecId);
-	OnCpuScopeName(FSummarizeCpuScopeAnalyzer::CoroutineSpecId, TEXT("Coroutine"));
-	OnCpuScopeName(FSummarizeCpuScopeAnalyzer::CoroutineUnknownSpecId, TEXT("<unknown>"));
-}
-
-void FSummarizeCpuScopeDurationAnalyzer::OnCpuScopeDiscovered(uint32 ScopeId)
-{
-	if (!Scopes.Find(ScopeId))
-	{
-		Scopes.Add(ScopeId, FSummarizeScope());		
-	}
-}
-
-void FSummarizeCpuScopeDurationAnalyzer::OnCpuScopeName(uint32 ScopeId, const FStringView& ScopeName)
-{
-	Scopes[ScopeId].Name = ScopeName;
-}
-
-void FSummarizeCpuScopeDurationAnalyzer::OnCpuScopeExit(const FScope& Scope, const FString* ScopeName)
-{
-	// This can miss if we are given a metadata timer's index, which is negative signed
-	FSummarizeScope* Found = Scopes.Find(Scope.ScopeId);
-	if (Found)
-	{
-		Found->AddDuration(Scope.EnterTimestamp, Scope.ExitTimestamp);
-	}
-}
-
-void FSummarizeCpuScopeDurationAnalyzer::OnCpuScopeAnalysisEnd()
-{
-	TArray<FSummarizeScope> LocalScopes;
-	// Eliminates scopes that don't have a name. (On scope discovery, the array is expended to creates blank scopes that may never be filled).
-	for (TMap<uint32, FSummarizeScope>::TIterator Iter = Scopes.CreateIterator(); Iter; ++Iter)
-	{
-		if(!Iter.Value().Name.IsEmpty())
-		{
-			LocalScopes.Add(Iter.Value());
-		}
-	}
-	
-	// Publish the scopes.
-	PublishFn(LocalScopes);
-}
-
 /**
  * Summarizes matched CPU scopes, excluding time consumed by immediate children if any. The analyzer uses pattern matching to
  * selects the scopes of interest and detects parent/child relationship. Once a parent/child relationship is established in a
@@ -1256,64 +472,6 @@ public:
 	TArray<TUniquePtr<FSummarizeCounter>> Counters;
 };
 
-//
-// FSummarizeBookmarksProvider - Tally Bookmarks from bookmark events
-//
-
-struct FSummarizeBookmark
-{
-	FString Name;
-	uint64 Count = 0;
-
-	double FirstSeconds = 0.0;
-	double LastSeconds = 0.0;
-
-	void AddTimestamp(double Seconds)
-	{
-		Count += 1;
-
-		// only set first for the first sample, compare exact zero
-		if (FirstSeconds == 0.0)
-		{
-			FirstSeconds = Seconds;
-		}
-
-		LastSeconds = Seconds;
-	}
-
-	FString GetValue(const FStringView& Statistic) const
-	{
-		if (Statistic == TEXT("Name"))
-		{
-			return Name;
-		}
-		else if (Statistic == TEXT("Count"))
-		{
-			return FString::Printf(TEXT("%llu"), Count);
-		}
-		else if (Statistic == TEXT("FirstSeconds"))
-		{
-			return FString::Printf(TEXT("%f"), FirstSeconds);
-		}
-		else if (Statistic == TEXT("LastSeconds"))
-		{
-			return FString::Printf(TEXT("%f"), LastSeconds);
-		}
-		return FString();
-	}
-
-	// for deduplication
-	bool operator==(const FSummarizeBookmark& Bookmark) const
-	{
-		return Name == Bookmark.Name;
-	}
-};
-
-static uint32 GetTypeHash(const FSummarizeBookmark& Bookmark)
-{
-	return FCrc::StrCrc32(*Bookmark.Name);
-}
-
 class FSummarizeBookmarksProvider
 	: public TraceServices::IEditableBookmarkProvider
 {
@@ -1424,9 +582,6 @@ FSummarizeBookmark* FSummarizeBookmarksProvider::FindStartBookmarkForEndBookmark
 * Begin SummarizeTrace commandlet implementation
 */
 
-// Defined here to prevent adding logs in the code above, which will likely be moved elsewhere.
-DEFINE_LOG_CATEGORY_STATIC(LogSummarizeTrace, Log, All);
-
 /*
 * Helpers for the csv files
 */
@@ -2031,12 +1186,6 @@ int32 USummarizeTraceCommandlet::Main(const FString& CmdLineParams)
 	return 0;
 }
 
-void USummarizeTraceCommandlet::WriteAsUTF8String(IFileHandle* Handle, const FString& String)
-{
-	const auto& UTF8String = StringCast<ANSICHAR>(*String);
-	Handle->Write(reinterpret_cast<const uint8*>(UTF8String.Get()), UTF8String.Length());
-};
-
 TUniquePtr<IFileHandle> USummarizeTraceCommandlet::OpenCSVFile(const FString& Name)
 {
 	FString CsvFileName = TraceFileBasename + Name;
@@ -2061,16 +1210,16 @@ bool USummarizeTraceCommandlet::GenerateScopesCSV(const TArray<FSummarizeScope>&
 	}
 	
 	// no newline, see row printfs
-	WriteAsUTF8String(CsvHandle.Get(), FString::Printf(TEXT("Name,Count,CountPerSecond,TotalDurationSeconds,FirstStartSeconds,FirstFinishSeconds,FirstDurationSeconds,LastStartSeconds,LastFinishSeconds,LastDurationSeconds,MinDurationSeconds,MaxDurationSeconds,MeanDurationSeconds,DeviationDurationSeconds,")));
+	CsvUtils::WriteAsUTF8String(CsvHandle.Get(), FString::Printf(TEXT("Name,Count,CountPerSecond,TotalDurationSeconds,FirstStartSeconds,FirstFinishSeconds,FirstDurationSeconds,LastStartSeconds,LastFinishSeconds,LastDurationSeconds,MinDurationSeconds,MaxDurationSeconds,MeanDurationSeconds,DeviationDurationSeconds,")));
 	for (const FSummarizeScope& Scope : SortedScopes)
 	{
-		if (!IsCsvSafeString(Scope.Name))
+		if (!CsvUtils::IsCsvSafeString(Scope.Name))
 		{
 			continue;
 		}
 
 		// note newline is at the front of every data line to prevent final extraneous newline, per customary for csv
-		WriteAsUTF8String(CsvHandle.Get(), FString::Printf(TEXT("\n%s,%llu,%f,%f,%f,%f,%f,%f,%f,%f,%f,%f,%f,%f,"),
+		CsvUtils::WriteAsUTF8String(CsvHandle.Get(), FString::Printf(TEXT("\n%s,%llu,%f,%f,%f,%f,%f,%f,%f,%f,%f,%f,%f,%f,"),
 			*Scope.Name, Scope.GetCount(), Scope.GetCountPerSecond(),
 			Scope.TotalDurationSeconds,
 			Scope.FirstStartSeconds, Scope.FirstFinishSeconds, Scope.FirstDurationSeconds,
@@ -2092,16 +1241,16 @@ bool USummarizeTraceCommandlet::GenerateCountersCSV(const FSummarizeCountersProv
 	}
 	
 	// no newline, see row printfs
-	WriteAsUTF8String(CsvHandle.Get(), FString::Printf(TEXT("Name,Count,CountPerSecond,First,FirstSeconds,Last,LastSeconds,Minimum,Maximum,Mean,Deviation,")));
+	CsvUtils::WriteAsUTF8String(CsvHandle.Get(), FString::Printf(TEXT("Name,Count,CountPerSecond,First,FirstSeconds,Last,LastSeconds,Minimum,Maximum,Mean,Deviation,")));
 	for (const TUniquePtr<FSummarizeCounter>& Counter : CountersProvider.Counters)
 	{
-		if (!IsCsvSafeString(Counter->Name))
+		if (!CsvUtils::IsCsvSafeString(Counter->Name))
 		{
 			continue;
 		}
 
 		// note newline is at the front of every data line to prevent final extraneous newline, per customary for csv
-		WriteAsUTF8String(CsvHandle.Get(), FString::Printf(TEXT("\n%s,%llu,%f,%s,%f,%s,%f,%s,%s,%f,%f,"),
+		CsvUtils::WriteAsUTF8String(CsvHandle.Get(), FString::Printf(TEXT("\n%s,%llu,%f,%s,%f,%s,%f,%s,%s,%f,%f,"),
 			*Counter->Name, Counter->GetCount(), Counter->GetCountPerSecond(),
 			*Counter->GetValue(TEXT("First")), Counter->FirstSeconds,
 			*Counter->GetValue(TEXT("Last")), Counter->LastSeconds,
@@ -2122,16 +1271,16 @@ bool USummarizeTraceCommandlet::GenerateBookmarksCSV(const FSummarizeBookmarksPr
 	}
 	
 	// no newline, see row printfs
-	WriteAsUTF8String(CsvHandle.Get(), FString::Printf(TEXT("Name,Count,FirstSeconds,LastSeconds,")));
+	CsvUtils::WriteAsUTF8String(CsvHandle.Get(), FString::Printf(TEXT("Name,Count,FirstSeconds,LastSeconds,")));
 	for (const TMap<FString, FSummarizeBookmark>::ElementType& Bookmark : BookmarksProvider.Bookmarks)
 	{
-		if (!IsCsvSafeString(Bookmark.Value.Name))
+		if (!CsvUtils::IsCsvSafeString(Bookmark.Value.Name))
 		{
 			continue;
 		}
 
 		// note newline is at the front of every data line to prevent final extraneous newline, per customary for csv
-		WriteAsUTF8String(CsvHandle.Get(), FString::Printf(TEXT("\n%s,%d,%f,%f,"),
+		CsvUtils::WriteAsUTF8String(CsvHandle.Get(), FString::Printf(TEXT("\n%s,%d,%f,%f,"),
 			*Bookmark.Value.Name, Bookmark.Value.Count,
 			Bookmark.Value.FirstSeconds, Bookmark.Value.LastSeconds));
 	}
@@ -2178,7 +1327,7 @@ bool USummarizeTraceCommandlet::GenerateTelemetryCSV(const FString& TestName,
 		// resolve scopes to telemetry
 		for (const FSummarizeScope& Scope : SortedScopes)
 		{
-			if (!IsCsvSafeString(Scope.Name))
+			if (!CsvUtils::IsCsvSafeString(Scope.Name))
 			{
 				continue;
 			}
@@ -2228,7 +1377,7 @@ bool USummarizeTraceCommandlet::GenerateTelemetryCSV(const FString& TestName,
 		// resolve counters to telemetry
 		for (const TUniquePtr<FSummarizeCounter>& Counter : CountersProvider.Counters)
 		{
-			if (!IsCsvSafeString(Counter->Name))
+			if (!CsvUtils::IsCsvSafeString(Counter->Name))
 			{
 				continue;
 			}
@@ -2260,7 +1409,7 @@ bool USummarizeTraceCommandlet::GenerateTelemetryCSV(const FString& TestName,
 		// resolve bookmarks to telemetry
 		for (const TMap<FString, FSummarizeBookmark>::ElementType& Bookmark : BookmarksProvider.Bookmarks)
 		{
-			if (!IsCsvSafeString(Bookmark.Value.Name))
+			if (!CsvUtils::IsCsvSafeString(Bookmark.Value.Name))
 			{
 				continue;
 			}
@@ -2413,11 +1562,11 @@ bool USummarizeTraceCommandlet::GenerateTelemetryCSV(const FString& TestName,
 	}
 	
 	// no newline, see row printfs
-	WriteAsUTF8String(TelemetryCsvHandle.Get(), FString::Printf(TEXT("TestName,Context,DataPoint,Unit,Measurement,Baseline,")));
+	CsvUtils::WriteAsUTF8String(TelemetryCsvHandle.Get(), FString::Printf(TEXT("TestName,Context,DataPoint,Unit,Measurement,Baseline,")));
 	for (const TelemetryDefinition& Telemetry : TelemetryData)
 	{
 		// note newline is at the front of every data line to prevent final extraneous newline, per customary for csv
-		WriteAsUTF8String(TelemetryCsvHandle.Get(), FString::Printf(TEXT("\n%s,%s,%s,%s,%s,%s,"), *Telemetry.TestName, *Telemetry.Context, *Telemetry.DataPoint, *Telemetry.Unit, *Telemetry.Measurement, *Telemetry.Baseline));
+		CsvUtils::WriteAsUTF8String(TelemetryCsvHandle.Get(), FString::Printf(TEXT("\n%s,%s,%s,%s,%s,%s,"), *Telemetry.TestName, *Telemetry.Context, *Telemetry.DataPoint, *Telemetry.Unit, *Telemetry.Measurement, *Telemetry.Baseline));
 	}
 
 	TelemetryCsvHandle->Flush();
diff --git a/Engine/Source/Editor/UnrealEd/Private/Commandlets/SummarizeTraceUtils.cpp b/Engine/Source/Editor/UnrealEd/Private/Commandlets/SummarizeTraceUtils.cpp
new file mode 100644
--- /dev/null
+++ b/Engine/Source/Editor/UnrealEd/Private/Commandlets/SummarizeTraceUtils.cpp
@@ -0,0 +1,597 @@
+// Copyright Epic Games, Inc. All Rights Reserved.
+
+#include "SummarizeTraceUtils.h"
+
+#include "Containers/Array.h"
+#include "Containers/ArrayView.h"
+#include "Containers/Map.h"
+#include "Containers/Set.h"
+#include "Containers/StringConv.h"
+#include "Containers/StringFwd.h"
+#include "Containers/StringView.h"
+#include "GenericPlatform/GenericPlatformAffinity.h"
+#include "GenericPlatform/GenericPlatformFile.h"
+#include "Logging/LogCategory.h"
+#include "Logging/LogMacros.h"
+#include "Misc/AssertionMacros.h"
+#include "Misc/CString.h"
+#include "Misc/Crc.h"
+#include "Misc/Optional.h"
+#include "Misc/Parse.h"
+#include "Misc/Timespan.h"
+#include "Model/MonotonicTimeline.h"
+#include "Templates/ChooseClass.h"
+#include "Templates/Function.h"
+#include "Templates/SharedPointer.h"
+#include "Templates/Tuple.h"
+#include "Templates/UniquePtr.h"
+#include "Templates/UnrealTemplate.h"
+#include "Math/NumericLimits.h"
+#include "Math/UnrealMathSSE.h"
+
+uint64 FIncrementalVariance::GetCount() const
+{
+	return Count;
+}
+
+double FIncrementalVariance::GetMean() const
+{
+	return Mean;
+}
+
+double FIncrementalVariance::Variance() const
+{
+	double Result = 0.0;
+
+	if (Count > 1)
+	{
+		// Welford's final step, dependent on sample count
+		Result = VarianceAccumulator / double(Count - 1);
+	}
+
+	return Result;
+}
+
+double FIncrementalVariance::Deviation() const
+{
+	double Result = 0.0;
+
+	if (Count > 1)
+	{
+		// Welford's final step, dependent on sample count
+		double DeviationSqrd = VarianceAccumulator / double(Count - 1);
+
+		// stddev is sqrt of variance, to restore to units (vs. units squared)
+		Result = sqrt(DeviationSqrd);
+	}
+
+	return Result;
+}
+
+void FIncrementalVariance::Increment(const double InSample)
+{
+	Count++;
+	const double OldMean = Mean;
+	Mean += ((InSample - Mean) / double(Count));
+	VarianceAccumulator += ((InSample - Mean) * (InSample - OldMean));
+}
+
+void FIncrementalVariance::Merge(const FIncrementalVariance& Other)
+{
+	// empty other, nothing to do
+	if (Other.Count == 0)
+	{
+		return;
+	}
+
+	// empty this, just copy other
+	if (Count == 0)
+	{
+		Count = Other.Count;
+		Mean = Other.Mean;
+		VarianceAccumulator = Other.VarianceAccumulator;
+		return;
+	}
+
+	const double TotalPopulation = static_cast<double>(Count + Other.Count);
+	const double MeanDifference = Mean - Other.Mean;
+	const double A = ((Count - 1) * Variance()) + ((Other.Count - 1) * Other.Variance());
+	const double B = (MeanDifference) * (MeanDifference) * (Count * Other.Count / TotalPopulation);
+	const double MergedVariance = (A + B) / (TotalPopulation - 1);
+
+	const uint64 NewCount = Count + Other.Count;
+	const double NewMean = ((Mean * double(Count)) + (Other.Mean * double(Other.Count))) / double(NewCount);
+	const double NewVarianceAccumulator = MergedVariance * (NewCount - 1);
+
+	Count = NewCount;
+	Mean = NewMean;
+	VarianceAccumulator = NewVarianceAccumulator;
+}
+
+void FIncrementalVariance::Reset()
+{
+	Count = 0;
+	Mean = 0.0;
+	VarianceAccumulator = 0.0;
+}
+
+void FSummarizeScope::AddDuration(double StartSeconds, double FinishSeconds)
+{
+	// compute the duration
+	double DurationSeconds = FinishSeconds - StartSeconds;
+
+	DurationVariance.Increment(DurationSeconds);
+
+	// only set first for the first sample
+	if (DurationVariance.GetCount() == 1)
+	{
+		FirstStartSeconds = StartSeconds;
+		FirstFinishSeconds = FinishSeconds;
+		FirstDurationSeconds = DurationSeconds;
+	}
+
+	LastStartSeconds = StartSeconds;
+	LastFinishSeconds = FinishSeconds;
+	LastDurationSeconds = DurationSeconds;
+
+	// set duration statistics
+	TotalDurationSeconds += DurationSeconds;
+	MinDurationSeconds = FMath::Min(MinDurationSeconds, DurationSeconds);
+	MaxDurationSeconds = FMath::Max(MaxDurationSeconds, DurationSeconds);
+
+	BeginTimeArray.Add(StartSeconds);
+	EndTimeArray.Add(FinishSeconds);
+}
+
+uint64 FSummarizeScope::GetCount() const
+{
+	return DurationVariance.GetCount();
+}
+
+double FSummarizeScope::GetMeanDurationSeconds() const
+{
+	return DurationVariance.GetMean();
+}
+
+double FSummarizeScope::GetDeviationDurationSeconds() const
+{
+	return DurationVariance.Deviation();
+}
+
+double FSummarizeScope::GetCountPerSecond() const
+{
+	double CountPerSecond = 0.0;
+	const uint64 Count = DurationVariance.GetCount();
+	if (Count)
+	{
+		CountPerSecond = Count / (LastFinishSeconds - FirstStartSeconds);
+	}
+	return CountPerSecond;
+}
+
+void FSummarizeScope::Merge(const FSummarizeScope& Other)
+{
+	check(Name == Other.Name);
+	DurationVariance.Merge(Other.DurationVariance);
+
+	if (FirstStartSeconds > Other.FirstStartSeconds)
+	{
+		FirstStartSeconds = Other.FirstStartSeconds;
+		FirstFinishSeconds = Other.FirstFinishSeconds;
+		FirstDurationSeconds = Other.FirstDurationSeconds;
+	}
+
+	if (LastStartSeconds < Other.LastStartSeconds)
+	{
+		LastStartSeconds = Other.LastStartSeconds;
+		LastFinishSeconds = Other.LastFinishSeconds;
+		LastDurationSeconds = Other.LastDurationSeconds;
+	}
+
+	TotalDurationSeconds += Other.TotalDurationSeconds;
+	MinDurationSeconds = FMath::Min(MinDurationSeconds, Other.MinDurationSeconds);
+	MaxDurationSeconds = FMath::Max(MaxDurationSeconds, Other.MaxDurationSeconds);
+}
+
+FString FSummarizeScope::GetValue(const FStringView& Statistic) const
+{
+	if (Statistic == TEXT("Name"))
+	{
+		return Name;
+	}
+	else if (Statistic == TEXT("Count"))
+	{
+		return FString::Printf(TEXT("%llu"), DurationVariance.GetCount());
+	}
+	else if (Statistic == TEXT("TotalDurationSeconds"))
+	{
+		return FString::Printf(TEXT("%f"), TotalDurationSeconds);
+	}
+	else if (Statistic == TEXT("FirstStartSeconds"))
+	{
+		return FString::Printf(TEXT("%f"), FirstStartSeconds);
+	}
+	else if (Statistic == TEXT("FirstFinishSeconds"))
+	{
+		return FString::Printf(TEXT("%f"), FirstFinishSeconds);
+	}
+	else if (Statistic == TEXT("FirstDurationSeconds"))
+	{
+		return FString::Printf(TEXT("%f"), FirstDurationSeconds);
+	}
+	else if (Statistic == TEXT("LastStartSeconds"))
+	{
+		return FString::Printf(TEXT("%f"), LastStartSeconds);
+	}
+	else if (Statistic == TEXT("LastFinishSeconds"))
+	{
+		return FString::Printf(TEXT("%f"), LastFinishSeconds);
+	}
+	else if (Statistic == TEXT("LastDurationSeconds"))
+	{
+		return FString::Printf(TEXT("%f"), LastDurationSeconds);
+	}
+	else if (Statistic == TEXT("MinDurationSeconds"))
+	{
+		return FString::Printf(TEXT("%f"), MinDurationSeconds);
+	}
+	else if (Statistic == TEXT("MaxDurationSeconds"))
+	{
+		return FString::Printf(TEXT("%f"), MaxDurationSeconds);
+	}
+	else if (Statistic == TEXT("MeanDurationSeconds"))
+	{
+		return FString::Printf(TEXT("%f"), DurationVariance.GetMean());
+	}
+	else if (Statistic == TEXT("DeviationDurationSeconds"))
+	{
+		return FString::Printf(TEXT("%f"), DurationVariance.Deviation());
+	}
+	else if (Statistic == TEXT("CountPerSecond"))
+	{
+		return FString::Printf(TEXT("%f"), GetCountPerSecond());
+	}
+	return FString();
+}
+
+// for deduplication
+bool FSummarizeScope::operator==(const FSummarizeScope& Scope) const
+{
+	return Name == Scope.Name;
+}
+
+// for sorting descending
+bool FSummarizeScope::operator<(const FSummarizeScope& Scope) const
+{
+	return TotalDurationSeconds > Scope.TotalDurationSeconds;
+}
+
+FSummarizeCpuProfilerProvider::FSummarizeCpuProfilerProvider()
+	: LookupScopeNameFn([this](uint32 ScopeId) { return LookupScopeName(ScopeId); })
+{
+}
+
+void FSummarizeCpuProfilerProvider::AddCpuScopeAnalyzer(TSharedPtr<FSummarizeCpuScopeAnalyzer> Analyzer)
+{
+	ScopeAnalyzers.Add(MoveTemp(Analyzer));
+}
+
+void FSummarizeCpuProfilerProvider::AnalysisComplete()
+{
+	// Analyze scope trees that contained 'nameless' context when they were captured. Unless the trace was truncated,
+	// all scope names should be known now.
+	for (const auto& Thread : Threads)
+	{
+		for (FScopeTreeInfo& DelayedScopeTree : Thread.Value->DelayedScopeTreeInfo)
+		{
+			// Run summary analysis for this delayed hierarchy.
+			OnCpuScopeTree(Thread.Key, DelayedScopeTree);
+		}
+	}
+
+	for (TSharedPtr<FSummarizeCpuScopeAnalyzer>& Analyzer : ScopeAnalyzers)
+	{
+		Analyzer->OnCpuScopeAnalysisEnd();
+	}
+}
+
+void FSummarizeCpuProfilerProvider::AddThread(uint32 Id, const TCHAR* Name, EThreadPriority Priority)
+{
+	TUniquePtr<FThread>* Found = Threads.Find(Id);
+	if (!Found)
+	{
+		Threads.Add(Id, MakeUnique<FThread>(Id, this));
+	}
+}
+
+uint32 FSummarizeCpuProfilerProvider::AddCpuTimer(FStringView Name, const TCHAR* File, uint32 Line)
+{
+	TOptional<FString> ScopeName;
+
+	if (!Name.IsEmpty())
+	{
+		ScopeName.Emplace(FString(Name));
+	}
+
+	uint32 TimerId = ScopeNames.Add(ScopeName);
+
+	// Notify the analyzers.
+	for (TSharedPtr<FSummarizeCpuScopeAnalyzer>& Analyzer : ScopeAnalyzers)
+	{
+		Analyzer->OnCpuScopeDiscovered(TimerId);
+
+		if (!Name.IsEmpty())
+		{
+			Analyzer->OnCpuScopeName(TimerId, Name);
+		}
+	}
+
+	return TimerId;
+}
+
+void FSummarizeCpuProfilerProvider::SetTimerNameAndLocation(uint32 TimerId, FStringView Name, const TCHAR* File, uint32 Line)
+{
+	check(TimerId < uint32(ScopeNames.Num()));
+	check(!Name.IsEmpty());
+
+	ScopeNames[TimerId].Emplace(FString(Name));
+
+	// Notify the registered scope analyzers.
+	for (TSharedPtr<FSummarizeCpuScopeAnalyzer>& Analyzer : ScopeAnalyzers)
+	{
+		Analyzer->OnCpuScopeName(TimerId, Name);
+	}
+}
+
+uint32 FSummarizeCpuProfilerProvider::AddMetadata(uint32 MasterTimerId, TArray<uint8>&& Metadata)
+{
+	uint32 MetadataId = Metadatas.Num();
+	Metadatas.Add({ MoveTemp(Metadata), MasterTimerId });
+	return ~MetadataId;
+}
+
+TArrayView<const uint8> FSummarizeCpuProfilerProvider::GetMetadata(uint32 TimerId) const
+{
+	if (int32(TimerId) >= 0)
+	{
+		return TArrayView<const uint8>();
+	}
+
+	TimerId = ~TimerId;
+	if (TimerId >= uint32(Metadatas.Num()))
+	{
+		return TArrayView<const uint8>();
+	}
+
+	const FMetadata& Metadata = Metadatas[TimerId];
+	return Metadata.Payload;
+}
+
+TraceServices::IEditableTimeline<TraceServices::FTimingProfilerEvent>& FSummarizeCpuProfilerProvider::GetCpuThreadEditableTimeline(uint32 ThreadId)
+{
+	TUniquePtr<FThread>* Found = Threads.Find(ThreadId);
+	if (Found)
+	{
+		return *(Found->Get());
+	}
+
+	return *Threads.Add(ThreadId, MakeUnique<FThread>(ThreadId, this));
+}
+
+void FSummarizeCpuProfilerProvider::FThread::AppendBeginEvent(double StartTime, const TraceServices::FTimingProfilerEvent& Event)
+{
+	FScopeEnter ScopeEnter{ Event.TimerIndex, StartTime };
+	ScopeStack.Add(ScopeEnter);
+
+	FSummarizeCpuScopeAnalyzer::FScopeEvent ScopeEvent{ FSummarizeCpuScopeAnalyzer::EScopeEventType::Enter, Event.TimerIndex, ThreadId, StartTime };
+	ScopeTreeInfo.ScopeEvents.Add(ScopeEvent);
+
+	Provider->AppendBeginEvent(ScopeEvent);
+}
+
+void FSummarizeCpuProfilerProvider::FThread::AppendEndEvent(double EndTime)
+{
+	if (ScopeStack.IsEmpty())
+	{
+		return;
+	}
+
+	FScopeEnter ScopeEnter = ScopeStack.Pop();
+
+	FSummarizeCpuScopeAnalyzer::FScopeEvent ScopeEvent{ FSummarizeCpuScopeAnalyzer::EScopeEventType::Exit, ScopeEnter.ScopeId, ThreadId, EndTime };
+	ScopeTreeInfo.ScopeEvents.Add(ScopeEvent);
+
+	// Check if at this point if the scope has a name
+	const FString* ScopeName = Provider->LookupScopeName(ScopeEnter.ScopeId);
+	ScopeTreeInfo.bHasNamelessScopes |= ScopeName == nullptr || ScopeName->IsEmpty();
+
+	FSummarizeCpuScopeAnalyzer::FScope Scope{ ScopeEnter.ScopeId, ThreadId, ScopeEnter.Timestamp, EndTime };
+	Provider->AppendEndEvent(Scope, ScopeName);
+
+	// The root scope on this thread just popped out.
+	if (ScopeStack.IsEmpty())
+	{
+		if (ScopeTreeInfo.bHasNamelessScopes)
+		{
+			// Delay the analysis until all the scope names are known.
+			DelayedScopeTreeInfo.Add(MoveTemp(ScopeTreeInfo));
+		}
+		else
+		{
+			// Run analysis for this scope tree.
+			Provider->OnCpuScopeTree(ThreadId, ScopeTreeInfo);
+		}
+
+		ScopeTreeInfo.Reset();
+	}
+}
+
+void FSummarizeCpuProfilerProvider::AppendBeginEvent(const FSummarizeCpuScopeAnalyzer::FScopeEvent& ScopeEvent)
+{
+	// Notify the registered scope analyzers.
+	for (TSharedPtr<FSummarizeCpuScopeAnalyzer>& Analyzer : ScopeAnalyzers)
+	{
+		Analyzer->OnCpuScopeEnter(ScopeEvent, LookupScopeName(ScopeEvent.ScopeId));
+	}
+}
+
+void FSummarizeCpuProfilerProvider::AppendEndEvent(const FSummarizeCpuScopeAnalyzer::FScope& Scope, const FString* ScopeName)
+{
+	// Notify the registered scope analyzers.
+	for (TSharedPtr<FSummarizeCpuScopeAnalyzer>& Analyzer : ScopeAnalyzers)
+	{
+		Analyzer->OnCpuScopeExit(Scope, ScopeName);
+	}
+}
+
+void FSummarizeCpuProfilerProvider::OnCpuScopeTree(uint32 ThreadId, const FScopeTreeInfo& ScopeTreeInfo)
+{
+	// Notify the registered scope analyzers.
+	for (TSharedPtr<FSummarizeCpuScopeAnalyzer>& Analyzer : ScopeAnalyzers)
+	{
+		Analyzer->OnCpuScopeTree(ThreadId, ScopeTreeInfo.ScopeEvents, LookupScopeNameFn);
+	}
+}
+
+const FString* FSummarizeCpuProfilerProvider::LookupScopeName(uint32 ScopeId)
+{
+	if (int32(ScopeId) < 0)
+	{
+		ScopeId = Metadatas[~ScopeId].TimerId;
+	}
+
+	if (ScopeId < static_cast<uint32>(ScopeNames.Num()) && ScopeNames[ScopeId])
+	{
+		return &ScopeNames[ScopeId].GetValue();
+	}
+
+	return nullptr;
+}
+
+FSummarizeCpuScopeDurationAnalyzer::FSummarizeCpuScopeDurationAnalyzer(TFunction<void(const TArray<FSummarizeScope>&)> InPublishFn)
+	: PublishFn(MoveTemp(InPublishFn))
+{
+	OnCpuScopeDiscovered(FSummarizeCpuScopeAnalyzer::CoroutineSpecId);
+	OnCpuScopeDiscovered(FSummarizeCpuScopeAnalyzer::CoroutineUnknownSpecId);
+	OnCpuScopeName(FSummarizeCpuScopeAnalyzer::CoroutineSpecId, TEXT("Coroutine"));
+	OnCpuScopeName(FSummarizeCpuScopeAnalyzer::CoroutineUnknownSpecId, TEXT("<unknown>"));
+}
+
+void FSummarizeCpuScopeDurationAnalyzer::OnCpuScopeDiscovered(uint32 ScopeId)
+{
+	if (!Scopes.Find(ScopeId))
+	{
+		Scopes.Add(ScopeId, FSummarizeScope());
+	}
+}
+
+void FSummarizeCpuScopeDurationAnalyzer::OnCpuScopeName(uint32 ScopeId, const FStringView& ScopeName)
+{
+	Scopes[ScopeId].Name = ScopeName;
+}
+
+void FSummarizeCpuScopeDurationAnalyzer::OnCpuScopeExit(const FScope& Scope, const FString* ScopeName)
+{
+	// This can miss if we are given a metadata timer's index, which is negative signed
+	FSummarizeScope* Found = Scopes.Find(Scope.ScopeId);
+	if (Found)
+	{
+		Found->AddDuration(Scope.EnterTimestamp, Scope.ExitTimestamp);
+	}
+}
+
+void FSummarizeCpuScopeDurationAnalyzer::OnCpuScopeAnalysisEnd()
+{
+	TArray<FSummarizeScope> LocalScopes;
+	// Eliminates scopes that don't have a name. (On scope discovery, the array is expended to creates blank scopes that may never be filled).
+	for (TMap<uint32, FSummarizeScope>::TIterator Iter = Scopes.CreateIterator(); Iter; ++Iter)
+	{
+		if (!Iter.Value().Name.IsEmpty())
+		{
+			LocalScopes.Add(Iter.Value());
+		}
+	}
+
+	// Publish the scopes.
+	PublishFn(LocalScopes);
+}
+
+FString FSummarizeBookmark::GetValue(const FStringView& Statistic) const
+{
+	if (Statistic == TEXT("Name"))
+	{
+		return Name;
+	}
+	else if (Statistic == TEXT("Count"))
+	{
+		return FString::Printf(TEXT("%llu"), Count);
+	}
+	else if (Statistic == TEXT("FirstSeconds"))
+	{
+		return FString::Printf(TEXT("%f"), FirstSeconds);
+	}
+	else if (Statistic == TEXT("LastSeconds"))
+	{
+		return FString::Printf(TEXT("%f"), LastSeconds);
+	}
+	return FString();
+}
+
+void FSummarizeBookmark::AddTimestamp(double Seconds)
+{
+	Count += 1;
+
+	// only set first for the first sample, compare exact zero
+	if (FirstSeconds == 0.0)
+	{
+		FirstSeconds = Seconds;
+	}
+
+	LastSeconds = Seconds;
+}
+
+bool CsvUtils::IsCsvSafeString(const FString& String)
+{
+	static struct DisallowedCharacter
+	{
+		const TCHAR Character;
+		bool First;
+	}
+	DisallowedCharacters[] =
+	{
+		// breaks simple csv files
+		{ TEXT('\n'), true },
+		{ TEXT('\r'), true },
+		{ TEXT(','), true },
+	};
+
+	// sanitize strings for a bog-simple csv file
+	bool bDisallowed = false;
+	int32 Index = 0;
+	for (struct DisallowedCharacter& DisallowedCharacter : DisallowedCharacters)
+	{
+		if (String.FindChar(DisallowedCharacter.Character, Index))
+		{
+			if (DisallowedCharacter.First)
+			{
+				UE_LOG(LogSummarizeTrace, Display, TEXT("A string contains disallowed character '%c'. See log for full list."), DisallowedCharacter.Character);
+				DisallowedCharacter.First = false;
+			}
+
+			UE_LOG(LogSummarizeTrace, Verbose, TEXT("String '%s' contains disallowed character '%c', skipping..."), *String, DisallowedCharacter.Character);				bDisallowed = true;
+		}
+
+		if (bDisallowed)
+		{
+			break;
+		}
+	}
+
+	return !bDisallowed;
+}
+
+void CsvUtils::WriteAsUTF8String(IFileHandle* Handle, const FString& String)
+{
+	const auto& UTF8String = StringCast<ANSICHAR>(*String);
+	Handle->Write(reinterpret_cast<const uint8*>(UTF8String.Get()), UTF8String.Length());
+}
\ No newline at end of file
diff --git a/Engine/Source/Editor/UnrealEd/Private/Commandlets/SummarizeTraceUtils.h b/Engine/Source/Editor/UnrealEd/Private/Commandlets/SummarizeTraceUtils.h
new file mode 100644
--- /dev/null
+++ b/Engine/Source/Editor/UnrealEd/Private/Commandlets/SummarizeTraceUtils.h
@@ -0,0 +1,352 @@
+// Copyright Epic Games, Inc. All Rights Reserved.
+
+#pragma once
+
+#include "TraceServices/Model/Threads.h"
+#include "TraceServices/Model/TimingProfiler.h"
+
+namespace TraceServices { class IAnalysisSession; }
+namespace UE::Trace { class IAnalyzer; }
+
+// Defined here to prevent adding logs in the code above, which will likely be moved elsewhere.
+DEFINE_LOG_CATEGORY_STATIC(LogSummarizeTrace, Log, All);
+
+/*
+*
+* Helper classes for the SummarizeTrace commandlet. Aggregates statistics about a trace.
+* Code preceding this comment should eventually be relocated to the Trace/Insights modules.
+* Everything succeeding this comment should stay in this commandlet.
+*
+*/
+
+class FIncrementalVariance
+{
+public:
+	FIncrementalVariance()
+		: Count(0)
+		, Mean(0.0)
+		, VarianceAccumulator(0.0)
+	{
+
+	}
+
+	uint64 GetCount() const;
+
+	double GetMean() const;
+
+	/**
+	* Compute the variance given Welford's accumulator and the overall count
+	*
+	* @return The variance in sample units squared
+	*/
+	double Variance() const;
+
+	/**
+	* Compute the standard deviation given Welford's accumulator and the overall count
+	*
+	* @return The standard deviation in sample units
+	*/
+	double Deviation() const;
+
+	/**
+	* Perform an increment of work for Welford's variance, from which we can compute variation and standard deviation
+	*
+	* @param InSample	The new sample value to operate on
+	*/
+	void Increment(const double InSample);
+
+	/**
+	* Merge with another IncrementalVariance series in progress
+	*
+	* @param Other	The other variance incremented from another mutually exclusive population of analogous data.
+	*/
+	void Merge(const FIncrementalVariance& Other);
+
+	/**
+	* Reset state back to initialized.
+	*/
+	void Reset();
+
+private:
+	uint64 Count;
+	double Mean;
+	double VarianceAccumulator;
+};
+
+/**
+ * Base class to extend for CPU scope analysis. Derived class instances are meant to be registered as
+ * delegates to the FSummarizeCpuProfilerProvider to handle scope events and perform meaningful analysis.
+ */
+class FSummarizeCpuScopeAnalyzer
+{
+public:
+	enum class EScopeEventType : uint32
+	{
+		Enter,
+		Exit
+	};
+
+	struct FScopeEvent
+	{
+		EScopeEventType ScopeEventType;
+		uint32 ScopeId;
+		uint32 ThreadId;
+		double Timestamp; // As Seconds
+	};
+
+	struct FScope
+	{
+		uint32 ScopeId;
+		uint32 ThreadId;
+		double EnterTimestamp; // As Seconds
+		double ExitTimestamp;  // As Seconds
+	};
+
+public:
+	virtual ~FSummarizeCpuScopeAnalyzer() = default;
+
+	/** Invoked when a CPU scope is discovered. This function is always invoked first when a CPU scope is encountered for the first time.*/
+	virtual void OnCpuScopeDiscovered(uint32 ScopeId) {}
+
+	/** Invoked when CPU scope specification is encountered in the trace stream. */
+	virtual void OnCpuScopeName(uint32 ScopeId, const FStringView& ScopeName) {};
+
+	/** Invoked when a scope is entered. The scope name might not be known yet. */
+	virtual void OnCpuScopeEnter(const FScopeEvent& ScopeEnter, const FString* ScopeName) {};
+
+	/** Invoked when a scope is exited. The scope name might not be known yet. */
+	virtual void OnCpuScopeExit(const FScope& Scope, const FString* ScopeName) {};
+
+	/** Invoked when a root event on the specified thread along with all child events down to the leaves are known. */
+	virtual void OnCpuScopeTree(uint32 ThreadId, const TArray64<FSummarizeCpuScopeAnalyzer::FScopeEvent>& ScopeEvents, const TFunction<const FString* (uint32)>& ScopeLookupNameFn) {};
+
+	/** Invoked when the trace stream has been fully consumed/processed. */
+	virtual void OnCpuScopeAnalysisEnd() {};
+
+	static constexpr uint32 CoroutineSpecId = (1u << 31u) - 1u;
+	static constexpr uint32 CoroutineUnknownSpecId = (1u << 31u) - 2u;
+};
+
+/**
+ * Decodes, format and routes CPU scope events embedded into a trace stream to specialized CPU scope analyzers. This processor
+ * decodes the low level events, keeps a small state and publishes higher level events to registered analyzers. The purpose
+ * is to decode the stream only once and let several registered analyzers reuse the small state built up by this processor. This
+ * matters when processing very large traces with possibly billions of scope events.
+ */
+class FSummarizeCpuProfilerProvider
+	: public TraceServices::IEditableThreadProvider
+	, public TraceServices::IEditableTimingProfilerProvider
+{
+public:
+	FSummarizeCpuProfilerProvider();
+
+	/** Register an analyzer with this processor. The processor decodes the trace stream and invokes the registered analyzers when a CPU scope event occurs.*/
+	void AddCpuScopeAnalyzer(TSharedPtr<FSummarizeCpuScopeAnalyzer> Analyzer);
+
+	void AnalysisComplete();
+
+private:
+	struct FScopeEnter
+	{
+		uint32 ScopeId;
+		double Timestamp; // As Seconds
+	};
+
+	// Contains scope events for a root scope and its children along with extra info to analyze that tree at once.
+	struct FScopeTreeInfo
+	{
+		// Records the current root scope and its children to run analysis that needs to know the parent/child relationship.
+		TArray64<FSummarizeCpuScopeAnalyzer::FScopeEvent> ScopeEvents;
+
+		// Indicates if one of the scope in the current hierarchy is nameless. (Its names specs hasn't been received yet).
+		bool bHasNamelessScopes = false;
+
+		void Reset()
+		{
+			ScopeEvents.Reset();
+			bHasNamelessScopes = false;
+		}
+	};
+
+	// For each thread we track what the stack of scopes are, for matching end-to-start
+	struct FThread
+		: public TraceServices::IEditableTimeline<TraceServices::FTimingProfilerEvent>
+	{
+		FThread(uint32 InThreadId, FSummarizeCpuProfilerProvider* InProvider)
+			: ThreadId(InThreadId)
+			, Provider(InProvider)
+		{
+
+		}
+
+		virtual void AppendBeginEvent(double StartTime, const TraceServices::FTimingProfilerEvent& Event) override;
+		virtual void AppendEndEvent(double EndTime) override;
+
+		// The ThreadId of this thread
+		uint32 ThreadId;
+
+		// The provider to forward calls to
+		FSummarizeCpuProfilerProvider* Provider;
+
+		// The current stack state
+		TArray<FScopeEnter> ScopeStack;
+
+		// The events recorded for the current root scope and its children to run analysis that needs to know the parent/child relationship, for example to compute time including childreen and time
+		// excluding childreen.
+		FScopeTreeInfo ScopeTreeInfo;
+
+		// Scope trees for which as least one scope name was unknown. Some analysis need scope names, but scope names/event can be emitted out of order by the engine depending on thread scheduling.
+		// Some scope tree cannot be analyzed right away and need to be delayed until all scope names are discovered.
+		TArray<FScopeTreeInfo> DelayedScopeTreeInfo;
+	};
+
+private:
+	// TraceServices::IEditableThreadProvider
+	virtual void					AddThread(uint32 Id, const TCHAR* Name, EThreadPriority Priority) override;
+
+	// TraceServices::IEditableTimingProfilerProvider
+	virtual uint32					AddCpuTimer(FStringView Name, const TCHAR* File, uint32 Line) override;
+	virtual void					SetTimerNameAndLocation(uint32 TimerId, FStringView Name, const TCHAR* File, uint32 Line) override;
+	virtual uint32					AddMetadata(uint32 MasterTimerId, TArray<uint8>&& Metadata) override;
+	virtual TArrayView<const uint8> GetMetadata(uint32 TimerId) const override;
+	virtual TraceServices::IEditableTimeline<TraceServices::FTimingProfilerEvent>& GetCpuThreadEditableTimeline(uint32 ThreadId) override;
+
+	// Callbacks from FThread to forward calls to Analyzers
+	virtual void AppendBeginEvent(const FSummarizeCpuScopeAnalyzer::FScopeEvent& ScopeEvent);
+	virtual void AppendEndEvent(const FSummarizeCpuScopeAnalyzer::FScope& ScopeEvent, const FString* ScopeName);
+
+	// Processing the ScopeTree once it's complete
+	void OnCpuScopeTree(uint32 ThreadId, const FScopeTreeInfo& ScopeTreeInfo);
+
+	// Resolve a ScopeId to string
+	const FString* LookupScopeName(uint32 ScopeId);
+
+private:
+	// The state at any moment of the threads
+	TMap<uint32, TUniquePtr<FThread>> Threads;
+
+	// The scope names, the array index correspond to the scope Id. If the optional is not set, the scope hasn't been encountered yet.
+	TArray<TOptional<FString>> ScopeNames;
+
+	// List of analyzers to invoke when a scope event is decoded.
+	TArray<TSharedPtr<FSummarizeCpuScopeAnalyzer>> ScopeAnalyzers;
+
+	// Scope name lookup function, cached for efficiency.
+	TFunction<const FString* (uint32 ScopeId)> LookupScopeNameFn;
+
+	struct FMetadata
+	{
+		TArray<uint8> Payload;
+		uint32 TimerId;
+	};
+
+	TArray<FMetadata> Metadatas;
+};
+
+struct FSummarizeScope
+{
+	FString Name;
+	FIncrementalVariance DurationVariance;
+	double TotalDurationSeconds = 0.0;
+
+	double FirstStartSeconds = 0.0;
+	double FirstFinishSeconds = 0.0;
+	double FirstDurationSeconds = 0.0;
+
+	double LastStartSeconds = 0.0;
+	double LastFinishSeconds = 0.0;
+	double LastDurationSeconds = 0.0;
+
+	double MinDurationSeconds = TNumericLimits<double>::Max();
+	double MaxDurationSeconds = TNumericLimits<double>::Min();
+
+	TArray<double> EndTimeArray;
+	TArray<double> BeginTimeArray;
+
+	void AddDuration(double StartSeconds, double FinishSeconds);
+
+	uint64 GetCount() const;
+
+	double GetMeanDurationSeconds() const;
+
+	double GetDeviationDurationSeconds() const;
+
+	double GetCountPerSecond() const;
+
+	void Merge(const FSummarizeScope& Other);
+
+	FString GetValue(const FStringView& Statistic) const;
+
+	// for deduplication
+	bool operator==(const FSummarizeScope& Scope) const;
+
+	// for sorting descending
+	bool operator<(const FSummarizeScope& Scope) const;
+};
+
+static uint32 GetTypeHash(const FSummarizeScope& Scope)
+{
+	return FCrc::StrCrc32(*Scope.Name);
+}
+
+/**
+ * This analyzer aggregates CPU scopes having the same name together, counting the number of occurrences, total duration,
+ * average occurrence duration, for each scope name.
+ */
+class FSummarizeCpuScopeDurationAnalyzer
+	: public FSummarizeCpuScopeAnalyzer
+{
+public:
+	FSummarizeCpuScopeDurationAnalyzer(TFunction<void(const TArray<FSummarizeScope>&)> InPublishFn);
+
+protected:
+	virtual void OnCpuScopeDiscovered(uint32 ScopeId) override;
+	virtual void OnCpuScopeName(uint32 ScopeId, const FStringView& ScopeName) override;
+	virtual void OnCpuScopeExit(const FScope& Scope, const FString* ScopeName) override;
+	virtual void OnCpuScopeAnalysisEnd() override;
+
+private:
+	TMap<uint32, FSummarizeScope> Scopes;
+
+	// Function invoked to publish the list of scopes.
+	TFunction<void(const TArray<FSummarizeScope>&)> PublishFn;
+};
+
+//
+// FSummarizeBookmarksProvider - Tally Bookmarks from bookmark events
+//
+struct FSummarizeBookmark
+{
+	FString Name;
+	uint64 Count = 0;
+
+	double FirstSeconds = 0.0;
+	double LastSeconds = 0.0;
+
+	void AddTimestamp(double Seconds);
+
+	FString GetValue(const FStringView& Statistic) const;
+
+	// for deduplication
+	bool operator==(const FSummarizeBookmark& Bookmark) const
+	{
+		return Name == Bookmark.Name;
+	}
+};
+
+static uint32 GetTypeHash(const FSummarizeBookmark& Bookmark)
+{
+	return FCrc::StrCrc32(*Bookmark.Name);
+}
+
+namespace CsvUtils
+{
+
+	/*
+	* Helpers for the csv files
+	*/
+
+	bool IsCsvSafeString(const FString& String);
+
+	void WriteAsUTF8String(IFileHandle* Handle, const FString& String);
+}
\ No newline at end of file
