[Splines] Fix an issue where the component visualizer would stop drawing on undo/redo operations

#rnx

diff --git a/Engine/Source/Editor/UnrealEd/Private/UnrealEdEngine.cpp b/Engine/Source/Editor/UnrealEd/Private/UnrealEdEngine.cpp
--- a/Engine/Source/Editor/UnrealEd/Private/UnrealEdEngine.cpp
+++ b/Engine/Source/Editor/UnrealEd/Private/UnrealEdEngine.cpp
@@ -1509,7 +1509,8 @@ void UUnrealEdEngine::OnEditorSelectionChanged(UObject* SelectionThatChanged)
 		}
 	};
 
-	if (SelectionThatChanged == GetSelectedActors())
+	const int32 SelectedComponentCount = GetSelectedComponents()->Num();
+	if (SelectionThatChanged == GetSelectedActors() && SelectedComponentCount == 0)
 	{
 		// actor selection changed.  Update the list of component visualizers
 		// This is expensive so we do not search for visualizers each time they want to draw
@@ -1525,31 +1526,29 @@ void UUnrealEdEngine::OnEditorSelectionChanged(UObject* SelectionThatChanged)
 			}
 		}
 	}
-	else if (SelectionThatChanged == GetSelectedComponents())
+
+	// Do not proceed if the selection contains no components. This occurs when a component is
+	// deselected while selecting its owner actor. But a corresponding actor selection is not invoked
+	// so if the visualizers are cleared here, they will not be properly reset for the selected actor. 
+	else if (SelectionThatChanged == GetSelectedComponents() && SelectedComponentCount > 0)
 	{
 		if (USelection* Selection = Cast<USelection>(SelectionThatChanged))
 		{
-			// Do not proceed if the selection contains no components. This occurs when a component is
-			// deselected while selecting its owner actor. But a corresponding actor selection is not invoked
-			// so if the visualizers are cleared here, they will not be properly reset for the selected actor. 
-			if (Selection->Num() > 0)
-			{
-				VisualizersForSelection.Empty();
+			VisualizersForSelection.Empty();
 
-				TArray<AActor*> ActorsProcessed;
+			TArray<AActor*> ActorsProcessed;
 
-				// Iterate over all selected components
-				for (FSelectionIterator It(GetSelectedComponentIterator()); It; ++It)
+			// Iterate over all selected components
+			for (FSelectionIterator It(GetSelectedComponentIterator()); It; ++It)
+			{
+				if (UActorComponent* Comp = Cast<UActorComponent>(*It))
 				{
-					if (UActorComponent* Comp = Cast<UActorComponent>(*It))
+					if (AActor* Actor = Comp->GetOwner())
 					{
-						if (AActor* Actor = Comp->GetOwner())
+						if (!ActorsProcessed.Contains(Actor))
 						{
-							if (!ActorsProcessed.Contains(Actor))
-							{
-								GetVisualizersForSelection(Actor, Comp);
-								ActorsProcessed.Emplace(Actor);
-							}
+							GetVisualizersForSelection(Actor, Comp);
+							ActorsProcessed.Emplace(Actor);
 						}
 					}
 				}
