Without it, ForceDeleteObjects can fail to delete assets when those assets are referenced by FBlueprintActionDatabase. This affected utility blueprint callable functions like UEditorAssetSubsystem::DeleteAsset which defer to ObjectTools::ForceDeleteObjects.

diff --git a/Engine/Source/Editor/UnrealEd/Private/ObjectTools.cpp b/Engine/Source/Editor/UnrealEd/Private/ObjectTools.cpp
--- a/Engine/Source/Editor/UnrealEd/Private/ObjectTools.cpp
+++ b/Engine/Source/Editor/UnrealEd/Private/ObjectTools.cpp
@@ -3616,6 +3616,9 @@ namespace ObjectTools
 				CollectGarbage(GARBAGE_COLLECTION_KEEPFLAGS);
 				bNeedsGarbageCollection = false;
 			}
+			
+			// Give systems opportunity to clean up references to the objects being deleted
+			FEditorDelegates::OnAssetsPreDelete.Broadcast(ShownObjectsToDelete);
 
 			// Load the asset tools module to get access to the browser type maps
 			FAssetToolsModule& AssetToolsModule = FModuleManager::LoadModuleChecked<FAssetToolsModule>(TEXT("AssetTools"));
