[MPCook] minor fixes:
- ensure tickable cook objects are ticked with bCookComplete=true on worker processes
- flush cook worker logs an additional time during shutdown to ensure log messages printed after final collector tick get sent back to the director

diff --git a/Engine/Source/Editor/UnrealEd/Private/CookOnTheFlyServer.cpp b/Engine/Source/Editor/UnrealEd/Private/CookOnTheFlyServer.cpp
--- a/Engine/Source/Editor/UnrealEd/Private/CookOnTheFlyServer.cpp
+++ b/Engine/Source/Editor/UnrealEd/Private/CookOnTheFlyServer.cpp
@@ -10908,6 +10908,13 @@ void UCookOnTheFlyServer::CookAsCookWorkerFinished()
 	}
 	CookWorkerClient->SetHasRunFinished(true);
 
+	{
+		UE_SCOPED_COOKTIMER(TickCookableObjects);
+		const double CurrentTime = FPlatformTime::Seconds();
+		FTickableCookObject::TickObjects(static_cast<float>(CurrentTime - LastCookableObjectTickTime), true /* bCookComplete */);
+		LastCookableObjectTickTime = CurrentTime;
+	}
+
 	FString LibraryName = GetProjectShaderLibraryName();
 	FString ActualLibraryName = GenerateShaderCodeLibraryName(LibraryName, IsCookFlagSet(ECookInitializationFlags::IterateSharedBuild));
 	FShaderLibraryCooker::EndCookingLibrary(ActualLibraryName);
@@ -10920,6 +10927,7 @@ void UCookOnTheFlyServer::CookAsCookWorkerFinished()
 		GShaderCompilingManager->SkipShaderCompilation(false);
 	}
 	LogCookWorkerStats();
+	CookWorkerClient->FlushLogs();
 }
 
 void UCookOnTheFlyServer::GetPackagesToRetract(int32 NumToRetract, TArray<FName>& OutRetractionPackages)
diff --git a/Engine/Source/Editor/UnrealEd/Private/Cooker/CookWorkerClient.cpp b/Engine/Source/Editor/UnrealEd/Private/Cooker/CookWorkerClient.cpp
--- a/Engine/Source/Editor/UnrealEd/Private/Cooker/CookWorkerClient.cpp
+++ b/Engine/Source/Editor/UnrealEd/Private/Cooker/CookWorkerClient.cpp
@@ -31,9 +31,9 @@ constexpr float WaitForConnectReplyTimeout = 60.f;
 FCookWorkerClient::FCookWorkerClient(UCookOnTheFlyServer& InCOTFS)
 	: COTFS(InCOTFS)
 {
-	FLogMessagesMessageHandler* Handler = new FLogMessagesMessageHandler();
-	Handler->InitializeClient();
-	Register(Handler);
+	LogMessageHandler = new FLogMessagesMessageHandler();
+	LogMessageHandler->InitializeClient();
+	Register(LogMessageHandler);
 	Register(new IMPCollectorCbClientMessage<FRetractionRequestMessage>([this]
 	(FMPCollectorClientMessageContext& Context, bool bReadSuccessful, FRetractionRequestMessage&& Message)
 		{
@@ -815,7 +815,13 @@ void FCookWorkerClient::Unregister(IMPCollector* Collector)
 	}
 }
 
-void FCookWorkerClient::TickCollectors(FTickStackData& StackData, bool bFlush)
+void FCookWorkerClient::FlushLogs()
+{
+	FTickStackData TickData(MAX_flt, ECookTickFlags::None);
+	TickCollectors(TickData, true, LogMessageHandler);
+}
+
+void FCookWorkerClient::TickCollectors(FTickStackData& StackData, bool bFlush, IMPCollector* SingleCollector)
 {
 	if (StackData.LoopStartTime < NextTickCollectorsTimeSeconds && !bFlush)
 	{
@@ -829,9 +835,8 @@ void FCookWorkerClient::TickCollectors(FTickStackData& StackData, bool bFlush)
 		Context.bFlush = bFlush;
 		TArray<UE::CompactBinaryTCP::FMarshalledMessage> MarshalledMessages;
 
-		for (const TPair<FGuid, TRefCountPtr<IMPCollector>>& Pair: Collectors)
+		auto TickCollector = [&MarshalledMessages, &Context](IMPCollector* Collector)
 		{
-			IMPCollector* Collector = Pair.Value.GetReference();
 			Collector->ClientTick(Context);
 			if (!Context.Messages.IsEmpty())
 			{
@@ -842,6 +847,18 @@ void FCookWorkerClient::TickCollectors(FTickStackData& StackData, bool bFlush)
 				}
 				Context.Messages.Reset();
 			}
+		};
+
+		if (SingleCollector)
+		{
+			TickCollector(SingleCollector);
+		}
+		else
+		{
+			for (const TPair<FGuid, TRefCountPtr<IMPCollector>>& Pair : Collectors)
+			{
+				TickCollector(Pair.Value.GetReference());
+			}
 		}
 
 		if (!MarshalledMessages.IsEmpty())
diff --git a/Engine/Source/Editor/UnrealEd/Private/Cooker/CookWorkerClient.h b/Engine/Source/Editor/UnrealEd/Private/Cooker/CookWorkerClient.h
--- a/Engine/Source/Editor/UnrealEd/Private/Cooker/CookWorkerClient.h
+++ b/Engine/Source/Editor/UnrealEd/Private/Cooker/CookWorkerClient.h
@@ -11,6 +11,7 @@
 #include "Misc/Guid.h"
 #include "Templates/UniquePtr.h"
 
+namespace UE::Cook { class FLogMessagesMessageHandler; }
 namespace UE::Cook { class FMPCollectorClientMessageContext; }
 namespace UE::Cook { class IMPCollector; }
 namespace UE::Cook { struct FAbortPackagesMessage; }
@@ -71,6 +72,9 @@ public:
 	/** Unegister a Collector that was registered. */
 	void Unregister(IMPCollector* Collector);
 
+	/** Called on worker cook process shutdown to flush any remaining log messages. */
+	void FlushLogs();
+
 private:
 	enum class EConnectStatus
 	{
@@ -111,8 +115,8 @@ private:
 	void LogInvalidMessage(const TCHAR* MessageTypeName);
 	/** Send packages assigned from the server into the request state. */
 	void AssignPackages(FAssignPackagesMessage& Message);
-	/** Tick the registered collectors. */
-	void TickCollectors(FTickStackData& StackData, bool bFlush);
+	/** Tick the registered collectors, or the single given collector if non-null. */
+	void TickCollectors(FTickStackData& StackData, bool bFlush, IMPCollector* SingleCollector = nullptr);
 	/** Helper for ReportDemote/ReportPromote: Collect IMPCollectors and asynchronously add the message to pending. */
 	void ReportPackageMessage(FName PackageName, TUniquePtr<FPackageRemoteResult>&& ResultOwner);
 	
@@ -141,6 +145,7 @@ private:
 	// Variables Read/Write only from the Scheduler thread
 	TSharedPtr<FInternetAddr> DirectorAddr;
 	TUniquePtr<FInitialConfigMessage> InitialConfigMessage;
+	TRefCountPtr<FLogMessagesMessageHandler> LogMessageHandler;
 	TArray<ITargetPlatform*> OrderedSessionPlatforms;
 	TArray<ITargetPlatform*> OrderedSessionAndSpecialPlatforms;
 	TArray<FDiscoveredPackageReplication> PendingDiscoveredPackages;
