[Selection] Fix an issue with the legacy sync event on USelection which resulted double notification to the element list on undo/redo, causing out of order updates when updating component visualizers

Upgraded this function to use the typed element selection set, which has combined actor and component selection, which removes the duplicate notify on the legacy bridge.

diff --git a/Engine/Source/Editor/UnrealEd/Classes/Editor/EditorEngine.h b/Engine/Source/Editor/UnrealEd/Classes/Editor/EditorEngine.h
--- a/Engine/Source/Editor/UnrealEd/Classes/Editor/EditorEngine.h
+++ b/Engine/Source/Editor/UnrealEd/Classes/Editor/EditorEngine.h
@@ -2800,12 +2800,6 @@ private:
 	/** A temporary copy used to add resilience during the undo/redo broadcast. Used to allow unregistering of clients while the broadcast is happening.
 	 * DO NOT add to this set while a broadcast is happening. */
 	TSet<class FEditorUndoClient*> InflightUndoClients;
-	
-	/** List of actors that were selected before Undo/redo */
-	TArray<AActor*> OldSelectedActors;
-
-	/** List of components that were selected before Undo/redo */
-	TArray<UActorComponent*> OldSelectedComponents;
 
 	/** The notification item to use for undo/redo */
 	TSharedPtr<class SNotificationItem> UndoRedoNotificationItem;
diff --git a/Engine/Source/Editor/UnrealEd/Classes/Editor/UnrealEdEngine.h b/Engine/Source/Editor/UnrealEd/Classes/Editor/UnrealEdEngine.h
--- a/Engine/Source/Editor/UnrealEd/Classes/Editor/UnrealEdEngine.h
+++ b/Engine/Source/Editor/UnrealEd/Classes/Editor/UnrealEdEngine.h
@@ -870,8 +870,8 @@ public:
 
 protected:
 
-	/** Called when global editor selection changes */
-	UNREALED_API void OnEditorSelectionChanged(UObject* SelectionThatChanged);
+	/** Called when global editor selection is about to change */
+	void OnEditorSelectionPreChange(const UTypedElementSelectionSet* SelectionSet);
 
 	/** Called when the element selection set pointer set on the global editor selection changes */
 	UNREALED_API void OnEditorElementSelectionPtrChanged(USelection* Selection, UTypedElementSelectionSet* OldSelectionSet, UTypedElementSelectionSet* NewSelectionSet);
diff --git a/Engine/Source/Editor/UnrealEd/Private/EditorSelectUtils.cpp b/Engine/Source/Editor/UnrealEd/Private/EditorSelectUtils.cpp
--- a/Engine/Source/Editor/UnrealEd/Private/EditorSelectUtils.cpp
+++ b/Engine/Source/Editor/UnrealEd/Private/EditorSelectUtils.cpp
@@ -273,6 +273,12 @@ void UUnrealEdEngine::ResetPivot()
 /*-----------------------------------------------------------------------------
 	Selection.
 -----------------------------------------------------------------------------*/
+void UUnrealEdEngine::OnEditorSelectionPreChange(const UTypedElementSelectionSet* SelectionSet)
+{
+	// Clear active editing visualizer on selection change
+	ComponentVisManager.ClearActiveComponentVis();
+	VisualizersForSelection.Empty();
+}
 
 void UUnrealEdEngine::OnEditorElementSelectionPtrChanged(USelection* Selection, UTypedElementSelectionSet* OldSelectionSet, UTypedElementSelectionSet* NewSelectionSet)
 {
@@ -281,11 +287,13 @@ void UUnrealEdEngine::OnEditorElementSelectionPtrChanged(USelection* Selection,
 		if (OldSelectionSet)
 		{
 			OldSelectionSet->OnChanged().RemoveAll(this);
+			OldSelectionSet->OnPreChange().RemoveAll(this);
 		}
 
 		if (NewSelectionSet)
 		{
 			NewSelectionSet->OnChanged().AddUObject(this, &UUnrealEdEngine::OnEditorElementSelectionChanged);
+			OldSelectionSet->OnPreChange().AddUObject(this, &UUnrealEdEngine::OnEditorSelectionPreChange);
 		}
 	}
 }
@@ -293,6 +301,76 @@ void UUnrealEdEngine::OnEditorElementSelectionPtrChanged(USelection* Selection,
 
 void UUnrealEdEngine::OnEditorElementSelectionChanged(const UTypedElementSelectionSet* SelectionSet)
 {
+	auto GetVisualizersForSelection = [this](AActor* Actor, const UActorComponent* SelectedComponent)
+	{
+		// Iterate over components of that actor (and recurse through child components)
+		TInlineComponentArray<UActorComponent*> Components;
+		Actor->GetComponents(Components, true);
+
+		for (int32 CompIdx = 0; CompIdx < Components.Num(); CompIdx++)
+		{
+			UActorComponent* Comp = Components[CompIdx];
+			if (Comp->IsRegistered())
+			{
+				// Try and find a visualizer
+				TSharedPtr<FComponentVisualizer> Visualizer = FindComponentVisualizer(Comp->GetClass());
+				if (Visualizer.IsValid() && (Comp == SelectedComponent || Visualizer->ShouldShowForSelectedSubcomponents(Comp)))
+				{
+					FCachedComponentVisualizer CachedComponentVisualizer(Comp, Visualizer);
+					FComponentVisualizerForSelection Temp{ CachedComponentVisualizer };
+
+					FComponentVisualizerForSelection& ComponentVisualizerForSelection = VisualizersForSelection.Add_GetRef(MoveTemp(Temp));
+
+					if (Comp != SelectedComponent)
+					{
+						ComponentVisualizerForSelection.IsEnabledDelegate.Emplace([]() { return GetDefault<UEditorPerProjectUserSettings>()->bShowSelectionSubcomponents == true; });
+					}
+				}
+			}
+		}
+	};
+
+	UTypedElementSelectionSet* LevelEditorSelectionSet = GetSelectedActors()->GetElementSelectionSet();
+	TSet<AActor*> ActorsProcessed;
+	LevelEditorSelectionSet->ForEachSelectedObject<UActorComponent>([&ActorsProcessed, &GetVisualizersForSelection](UActorComponent* InActorComponent)
+		{
+			if (AActor* Actor = InActorComponent->GetOwner())
+			{
+				if (!ActorsProcessed.Contains(Actor))
+				{
+					GetVisualizersForSelection(Actor, InActorComponent);
+					ActorsProcessed.Emplace(Actor);
+				}
+			}
+			return true;
+		});
+
+	if (ActorsProcessed.Num() == 0)
+	{
+		LevelEditorSelectionSet->ForEachSelectedObject<AActor>([&GetVisualizersForSelection](AActor* InActor)
+			{
+				GetVisualizersForSelection(InActor, InActor->GetRootComponent());
+				return true;
+			});
+	}
+
+	// Restore the active component visualizer, since an undo/redo may have changed the selection
+	if (!ComponentVisManager.IsActive() && VisualizersForSelection.Num() > 0)
+	{
+		for (FComponentVisualizerForSelection& VisualizerForSelection : VisualizersForSelection)
+		{
+			if (VisualizerForSelection.ComponentVisualizer.Visualizer->GetEditedComponent() != nullptr)
+			{
+				ComponentVisManager.SetActiveComponentVis(GCurrentLevelEditingViewportClient, VisualizerForSelection.ComponentVisualizer.Visualizer);
+				break;
+			}
+		}
+	}
+
+#if PLATFORM_MAC
+	FPlatformApplicationMisc::bChachedMacMenuStateNeedsUpdate = true;
+#endif
+
 	NoteSelectionChange();
 }
 
@@ -394,9 +472,6 @@ void UUnrealEdEngine::NoteSelectionChange(bool bNotify)
 	// The selection changed, so make sure the pivot (widget) is located in the right place
 	UpdatePivotLocationForSelection( true );
 
-	// Clear active editing visualizer on selection change
-	ComponentVisManager.ClearActiveComponentVis();
-
 	const bool bComponentSelectionChanged = GetSelectedComponentCount() > 0;
 	if (bNotify)
 	{
diff --git a/Engine/Source/Editor/UnrealEd/Private/EditorServer.cpp b/Engine/Source/Editor/UnrealEd/Private/EditorServer.cpp
--- a/Engine/Source/Editor/UnrealEd/Private/EditorServer.cpp
+++ b/Engine/Source/Editor/UnrealEd/Private/EditorServer.cpp
@@ -1149,22 +1149,6 @@ void UEditorEngine::ShowUndoRedoNotification(const FText& NotificationText, bool
 
 void UEditorEngine::HandleTransactorBeforeRedoUndo(const FTransactionContext& TransactionContext)
 {
-	//Get the list of all selected actors before the undo/redo is performed
-	OldSelectedActors.Empty();
-	for ( FSelectionIterator It( GetSelectedActorIterator() ) ; It ; ++It )
-	{
-		AActor* Actor = CastChecked<AActor>( *It );
-		OldSelectedActors.Add( Actor);
-	}
-
-	// Get the list of selected components as well
-	OldSelectedComponents.Empty();
-	for (FSelectionIterator It(GetSelectedComponentIterator()); It; ++It)
-	{
-		auto Component = CastChecked<UActorComponent>(*It);
-		OldSelectedComponents.Add(Component);
-	}
-
 	// Before an undo, store the current operation and hook on object transaction, if we do not have an outer operation already
 	if (CurrentUndoRedoContext->OperationDepth++ == 0)
 	{
@@ -1274,28 +1258,14 @@ UTransactor* UEditorEngine::CreateTrans()
 
 void UEditorEngine::PostUndo(bool)
 {
+	UTypedElementSelectionSet* LevelEditorSelection = GetSelectedActors()->GetElementSelectionSet();
+
 	// Cache any Actor that needs to be re-instanced because it still points to a REINST_ class
 	TMap< UClass*, UClass* > OldToNewClassMapToReinstance;
-
-	//Update the actor selection followed by the component selection if needed (note: order is important)
-		
-	//Get the list of all selected actors after the operation
-	TArray<AActor*> SelectedActors;
-	for (FSelectionIterator It(GetSelectedActorIterator()); It; ++It)
+	LevelEditorSelection->ForEachSelectedObject<AActor>([&OldToNewClassMapToReinstance](AActor* InActor)
 	{
-		AActor* Actor = CastChecked<AActor>(*It);
-		//if this actor is NOT in a hidden level add it to the list - otherwise de-select it
-		if (FLevelUtils::IsLevelLocked(Actor) == false)
-		{
-			SelectedActors.Add(Actor);
-		}
-		else
-		{
-			GetSelectedActors()->Select(Actor, false);
-		}
-
 		// If the Actor's Class is not the AuthoritativeClass, then it needs to be re-instanced
-		UClass* OldClass = Actor->GetClass();
+		UClass* OldClass = InActor->GetClass();
 		if (OldClass->HasAnyClassFlags(CLASS_NewerVersionExists))
 		{
 			UClass* NewClass = OldClass->GetAuthoritativeClass();
@@ -1306,101 +1276,8 @@ void UEditorEngine::PostUndo(bool)
 
 			OldToNewClassMapToReinstance.Add(OldClass, NewClass);
 		}
-	}
-
-	USelection* ActorSelection = GetSelectedActors();
-	ActorSelection->BeginBatchSelectOperation();
-
-	//Deselect all of the actors that were selected prior to the operation
-	for (int32 OldSelectedActorIndex = OldSelectedActors.Num() - 1; OldSelectedActorIndex >= 0; --OldSelectedActorIndex)
-	{
-		AActor* Actor = OldSelectedActors[OldSelectedActorIndex];
-
-		//To stop us from unselecting and then reselecting again (causing two force update components, we will remove (from both lists) any object that was selected and should continue to be selected
-		int32 FoundIndex;
-		if (SelectedActors.Find(Actor, FoundIndex))
-		{
-			OldSelectedActors.RemoveAt(OldSelectedActorIndex);
-			SelectedActors.RemoveAt(FoundIndex);
-		}
-		else
-		{
-			SelectActor(Actor, false, false);//First false is to deselect, 2nd is to notify
-			Actor->UpdateComponentTransforms();
-		}
-	}
-
-	//Select all of the actors in SelectedActors
-	for (int32 SelectedActorIndex = 0; SelectedActorIndex < SelectedActors.Num(); ++SelectedActorIndex)
-	{
-		AActor* Actor = SelectedActors[SelectedActorIndex];
-		SelectActor(Actor, true, false);	//false is to stop notify which is done below if bOpWasSuccessful
-		Actor->UpdateComponentTransforms();
-	}
-
-	OldSelectedActors.Empty();
-	ActorSelection->EndBatchSelectOperation();
-	
-	if (GetSelectedComponentCount() > 0)
-	{
-		//@todo Check to see if component owner is in a hidden level
-		
-		// Get a list of all selected components after the operation
-		TArray<UActorComponent*> SelectedComponents;
-		for (FSelectionIterator It(GetSelectedComponentIterator()); It; ++It)
-		{
-			SelectedComponents.Add(CastChecked<UActorComponent>(*It));
-		}
-		
-		USelection* ComponentSelection = GetSelectedComponents();
-		ComponentSelection->BeginBatchSelectOperation();
-
-		//Deselect all of the actors that were selected prior to the operation
-		for (int32 OldSelectedComponentIndex = OldSelectedComponents.Num() - 1; OldSelectedComponentIndex >= 0; --OldSelectedComponentIndex)
-		{
-			UActorComponent* Component = OldSelectedComponents[OldSelectedComponentIndex];
-
-			//To stop us from unselecting and then reselecting again (causing two force update components, we will remove (from both lists) any object that was selected and should continue to be selected
-			int32 FoundIndex;
-			if (SelectedComponents.Find(Component, FoundIndex))
-			{
-				OldSelectedComponents.RemoveAt(OldSelectedComponentIndex);
-				SelectedComponents.RemoveAt(FoundIndex);
-			}
-			else
-			{
-				// Deselect without any notification
-				SelectComponent(Component, false, false);
-
-				AActor* Owner = Component->GetOwner();
-				if (Owner && Owner->IsSelected())
-				{
-					// Synchronize selection with owner actors
-					SelectActor(Owner, false, false, true);
-				}
-			}
-		}
-
-		//Select all of the components left in SelectedComponents
-		for (int32 SelectedComponentIndex = 0; SelectedComponentIndex < SelectedComponents.Num(); ++SelectedComponentIndex)
-		{
-			UActorComponent* Component = SelectedComponents[SelectedComponentIndex];
-			SelectComponent(Component, true, false);	//false is to stop notify which is done below if bOpWasSuccessful
-
-			AActor* Owner = Component->GetOwner();
-			if (Owner && !Owner->IsSelected())
-			{
-				// Synchronize selection with owner actors
-				SelectActor(Owner, true, false, true);
-			}
-		}
-
-		OldSelectedComponents.Empty();
-
-		// We want to broadcast the component SelectionChangedEvent even if the selection didn't actually change
-		ComponentSelection->ForceBatchDirty();
-		ComponentSelection->EndBatchSelectOperation();
-	}
+		return true;
+	});
 
 	// Re-instance any actors that need it
 	FBlueprintCompileReinstancer::BatchReplaceInstancesOfClass(OldToNewClassMapToReinstance, FReplaceInstancesOfClassParameters());
diff --git a/Engine/Source/Editor/UnrealEd/Private/UnrealEdEngine.cpp b/Engine/Source/Editor/UnrealEd/Private/UnrealEdEngine.cpp
--- a/Engine/Source/Editor/UnrealEd/Private/UnrealEdEngine.cpp
+++ b/Engine/Source/Editor/UnrealEd/Private/UnrealEdEngine.cpp
@@ -137,8 +137,6 @@ void UUnrealEdEngine::Init(IEngineLoop* InEngineLoop)
 	FEditorSupportDelegates::PostWindowsMessage.AddUObject(this, &UUnrealEdEngine::OnPostWindowsMessage);
 
 	FHierarchicalInstancedStaticMeshDelegates::OnTreeBuilt.AddUObject(this, &UUnrealEdEngine::OnHISMTreeBuilt);
-
-	USelection::SelectionChangedEvent.AddUObject(this, &UUnrealEdEngine::OnEditorSelectionChanged);
 	USelection::SelectionElementSelectionPtrChanged.AddUObject(this, &UUnrealEdEngine::OnEditorElementSelectionPtrChanged);
 
 	// Initialize the snap manager
@@ -1478,101 +1476,6 @@ void UUnrealEdEngine::DrawComponentVisualizersHUD(const FViewport* Viewport, con
 	}
 }
 
-void UUnrealEdEngine::OnEditorSelectionChanged(UObject* SelectionThatChanged)
-{
-	auto GetVisualizersForSelection = [this](AActor* Actor, const UActorComponent* SelectedComponent)
-	{
-		// Iterate over components of that actor (and recurse through child components)
-		TInlineComponentArray<UActorComponent*> Components;
-		Actor->GetComponents(Components, true);
-
-		for (int32 CompIdx = 0; CompIdx < Components.Num(); CompIdx++)
-		{
-			UActorComponent* Comp = Components[CompIdx];
-			if (Comp->IsRegistered())
-			{
-				// Try and find a visualizer
-				TSharedPtr<FComponentVisualizer> Visualizer = FindComponentVisualizer(Comp->GetClass());
-				if (Visualizer.IsValid() && (Comp == SelectedComponent || Visualizer->ShouldShowForSelectedSubcomponents(Comp)))
-				{
-					FCachedComponentVisualizer CachedComponentVisualizer(Comp, Visualizer);
-					FComponentVisualizerForSelection Temp{CachedComponentVisualizer};
-
-					FComponentVisualizerForSelection& ComponentVisualizerForSelection = VisualizersForSelection.Add_GetRef(MoveTemp(Temp));
-
-					if (Comp != SelectedComponent)
-					{
-						ComponentVisualizerForSelection.IsEnabledDelegate.Emplace([](){ return GetDefault<UEditorPerProjectUserSettings>()->bShowSelectionSubcomponents == true; });
-					}
-				}
-			}
-		}
-	};
-
-	const int32 SelectedComponentCount = GetSelectedComponents()->Num();
-	if (SelectionThatChanged == GetSelectedActors() && SelectedComponentCount == 0)
-	{
-		// actor selection changed.  Update the list of component visualizers
-		// This is expensive so we do not search for visualizers each time they want to draw
-		VisualizersForSelection.Empty();
-
-		// Iterate over all selected actors
-		for (FSelectionIterator It(GetSelectedActorIterator()); It; ++It)
-		{
-			AActor* Actor = Cast<AActor>(*It);
-			if (Actor != nullptr)
-			{
-				GetVisualizersForSelection(Actor, Actor->GetRootComponent());
-			}
-		}
-	}
-
-	// Do not proceed if the selection contains no components. This occurs when a component is
-	// deselected while selecting its owner actor. But a corresponding actor selection is not invoked
-	// so if the visualizers are cleared here, they will not be properly reset for the selected actor. 
-	else if (SelectionThatChanged == GetSelectedComponents() && SelectedComponentCount > 0)
-	{
-		if (USelection* Selection = Cast<USelection>(SelectionThatChanged))
-		{
-			VisualizersForSelection.Empty();
-
-			TArray<AActor*> ActorsProcessed;
-
-			// Iterate over all selected components
-			for (FSelectionIterator It(GetSelectedComponentIterator()); It; ++It)
-			{
-				if (UActorComponent* Comp = Cast<UActorComponent>(*It))
-				{
-					if (AActor* Actor = Comp->GetOwner())
-					{
-						if (!ActorsProcessed.Contains(Actor))
-						{
-							GetVisualizersForSelection(Actor, Comp);
-							ActorsProcessed.Emplace(Actor);
-						}
-					}
-				}
-			}
-		}
-	}
-
-	// If there is an undo/redo operation in progress, restore the active component visualizer.
-	if (GIsTransacting)
-	{
-		for (FComponentVisualizerForSelection& VisualizerForSelection : VisualizersForSelection)
-		{
-			if (VisualizerForSelection.ComponentVisualizer.Visualizer->GetEditedComponent() != nullptr)
-			{
-				ComponentVisManager.SetActiveComponentVis(GCurrentLevelEditingViewportClient, VisualizerForSelection.ComponentVisualizer.Visualizer);
-				break;
-			}
-		}
-	}
-#if PLATFORM_MAC
-	FPlatformApplicationMisc::bChachedMacMenuStateNeedsUpdate = true;
-#endif
-}
-
 bool UUnrealEdEngine::HasMountWritePermissionForPackage(const FString& PackageName)
 {
 	FName MountPoint = FPackageName::GetPackageMountPoint(PackageName, false);
