CIS Fix for missing include

#rnx

diff --git a/Engine/Source/Editor/UnrealEd/Private/EditorSelectUtils.cpp b/Engine/Source/Editor/UnrealEd/Private/EditorSelectUtils.cpp
--- a/Engine/Source/Editor/UnrealEd/Private/EditorSelectUtils.cpp
+++ b/Engine/Source/Editor/UnrealEd/Private/EditorSelectUtils.cpp
@@ -9,6 +9,7 @@
 #include "Components/SceneComponent.h"
 #include "GameFramework/Actor.h"
 #include "Components/PrimitiveComponent.h"
+#include "Editor/EditorPerProjectUserSettings.h"
 #include "Editor/UnrealEdEngine.h"
 #include "Editor/GroupActor.h"
 #include "Components/ChildActorComponent.h"
