Cooker: Logging: Downgrade cook stats and profile data to Log to make it easier to scan the display output of the cook.
#rnx

diff --git a/Engine/Source/Editor/UnrealEd/Private/Cooker/CookProfiling.cpp b/Engine/Source/Editor/UnrealEd/Private/Cooker/CookProfiling.cpp
--- a/Engine/Source/Editor/UnrealEd/Private/Cooker/CookProfiling.cpp
+++ b/Engine/Source/Editor/UnrealEd/Private/Cooker/CookProfiling.cpp
@@ -176,7 +176,7 @@ void OutputHierarchyTimers(const FHierarchicalTimerInfo* TimerInfo, int32 Depth)
 	static const TCHAR LeftPad[] = TEXT("                                ");
 	const SIZE_T PadOffset = FMath::Max<int>(UE_ARRAY_COUNT(LeftPad) - 1 - Depth * 2, 0);
 
-	UE_LOG(LogCook, Display, TEXT("  %s%s: %.3fs (%u)"), &LeftPad[PadOffset], *TimerName, TimerInfo->Length, TimerInfo->HitCount);
+	UE_LOG(LogCook, Log, TEXT("  %s%s: %.3fs (%u)"), &LeftPad[PadOffset], *TimerName, TimerInfo->Length, TimerInfo->HitCount);
 
 	// We need to print in reverse order since the child list begins with the most recently added child
 
@@ -197,7 +197,7 @@ void OutputHierarchyTimers(const FHierarchicalTimerInfo* TimerInfo, int32 Depth)
 
 void OutputHierarchyTimers()
 {
-	UE_LOG(LogCook, Display, TEXT("Hierarchy Timer Information:"));
+	UE_LOG(LogCook, Log, TEXT("Hierarchy Timer Information:"));
 
 	OutputHierarchyTimers(&RootTimerInfo, 0);
 }
@@ -1041,29 +1041,29 @@ void LogCookStats(ECookMode::Type CookMode)
 
 	FCookStatsManager::LogCookStats(LogStatsFunc);
 
-	UE_LOG(LogCook, Display, TEXT("Misc Cook Stats"));
-	UE_LOG(LogCook, Display, TEXT("==============="));
+	UE_LOG(LogCook, Log, TEXT("Misc Cook Stats"));
+	UE_LOG(LogCook, Log, TEXT("==============="));
 	for (FString& StatCategory : StatCategories)
 	{
-		UE_LOG(LogCook, Display, TEXT("%s"), *StatCategory);
+		UE_LOG(LogCook, Log, TEXT("%s"), *StatCategory);
 		TArray<FCookStatsManager::StringKeyValue>& StatsInCategory = StatsInCategories.FindOrAdd(StatCategory);
 
 		// log each key/value pair, with the equal signs lined up.
 		for (const FCookStatsManager::StringKeyValue& StatKeyValue : StatsInCategory)
 		{
-			UE_LOG(LogCook, Display, TEXT("    %s=%s"), *StatKeyValue.Key, *StatKeyValue.Value);
+			UE_LOG(LogCook, Log, TEXT("    %s=%s"), *StatKeyValue.Key, *StatKeyValue.Value);
 		}
 	}
 
 	// DDC Usage stats are custom formatted, and the above code just accumulated them into a TSet. Now log it with our special formatting for readability.
 	if (CookProfileData.Num() > 0)
 	{
-		UE_LOG(LogCook, Display, TEXT(""));
-		UE_LOG(LogCook, Display, TEXT("Cook Profile"));
-		UE_LOG(LogCook, Display, TEXT("============"));
+		UE_LOG(LogCook, Log, TEXT(""));
+		UE_LOG(LogCook, Log, TEXT("Cook Profile"));
+		UE_LOG(LogCook, Log, TEXT("============"));
 		for (const auto& ProfileEntry : CookProfileData)
 		{
-			UE_LOG(LogCook, Display, TEXT("%s.%s=%s"), *ProfileEntry.Path, *ProfileEntry.Key, *ProfileEntry.Value);
+			UE_LOG(LogCook, Log, TEXT("%s.%s=%s"), *ProfileEntry.Path, *ProfileEntry.Key, *ProfileEntry.Value);
 		}
 
 		FString CookStatsFileName;
@@ -1108,12 +1108,12 @@ void LogCookStats(ECookMode::Type CookMode)
 	}
 	if (DDCSummaryStats.Num() > 0)
 	{
-		UE_LOG(LogCook, Display, TEXT(""));
-		UE_LOG(LogCook, Display, TEXT("DDC Summary Stats"));
-		UE_LOG(LogCook, Display, TEXT("================="));
+		UE_LOG(LogCook, Log, TEXT(""));
+		UE_LOG(LogCook, Log, TEXT("DDC Summary Stats"));
+		UE_LOG(LogCook, Log, TEXT("================="));
 		for (const auto& Attr : DDCSummaryStats)
 		{
-			UE_LOG(LogCook, Display, TEXT("%-16s=%10s"), *Attr.Key, *Attr.Value);
+			UE_LOG(LogCook, Log, TEXT("%-16s=%10s"), *Attr.Key, *Attr.Value);
 		}
 	}
 
@@ -1133,14 +1133,14 @@ void LogCookStats(ECookMode::Type CookMode)
 				return LHS.TotalTimeSec > RHS.TotalTimeSec;
 			});
 
-		UE_LOG(LogCook, Display, TEXT(""));
-		UE_LOG(LogCook, Display, TEXT("DDC Resource Stats"));
-		UE_LOG(LogCook, Display, TEXT("======================================================================================================="));
-		UE_LOG(LogCook, Display, TEXT("Asset Type                          Total Time (Sec)  GameThread Time (Sec)  Assets Built  MB Processed"));
-		UE_LOG(LogCook, Display, TEXT("----------------------------------  ----------------  ---------------------  ------------  ------------"));
+		UE_LOG(LogCook, Log, TEXT(""));
+		UE_LOG(LogCook, Log, TEXT("DDC Resource Stats"));
+		UE_LOG(LogCook, Log, TEXT("======================================================================================================="));
+		UE_LOG(LogCook, Log, TEXT("Asset Type                          Total Time (Sec)  GameThread Time (Sec)  Assets Built  MB Processed"));
+		UE_LOG(LogCook, Log, TEXT("----------------------------------  ----------------  ---------------------  ------------  ------------"));
 		for (const FDDCResourceUsageStat& Stat : SortedDDCResourceUsageStats)
 		{
-			UE_LOG(LogCook, Display, TEXT("%-34s  %16.2f  %21.2f  %12d  %12.2f"), *Stat.AssetType, Stat.TotalTimeSec, Stat.GameThreadTimeSec, Stat.AssetsBuilt, Stat.SizeMB);
+			UE_LOG(LogCook, Log, TEXT("%-34s  %16.2f  %21.2f  %12d  %12.2f"), *Stat.AssetType, Stat.TotalTimeSec, Stat.GameThreadTimeSec, Stat.AssetsBuilt, Stat.SizeMB);
 		}
 	}
 
diff --git a/Engine/Source/Editor/UnrealEd/Private/Cooker/PackageBuildDependencyTracker.cpp b/Engine/Source/Editor/UnrealEd/Private/Cooker/PackageBuildDependencyTracker.cpp
--- a/Engine/Source/Editor/UnrealEd/Private/Cooker/PackageBuildDependencyTracker.cpp
+++ b/Engine/Source/Editor/UnrealEd/Private/Cooker/PackageBuildDependencyTracker.cpp
@@ -47,18 +47,18 @@ void FPackageBuildDependencyTracker::DumpStats() const
 			++ReferenceCount;
 		}
 	}
-	UE_LOG(LogPackageBuildDependencyTracker, Display, TEXT("Package Accesses (%u referencing packages with a total of %u unique accesses)"), ReferencingPackageCount, ReferenceCount);
+	UE_LOG(LogPackageBuildDependencyTracker, Log, TEXT("Package Accesses (%u referencing packages with a total of %u unique accesses)"), ReferencingPackageCount, ReferenceCount);
 
 	constexpr bool bDetailedDump = false;
 	if (bDetailedDump)
 	{
-		UE_LOG(LogPackageBuildDependencyTracker, Display, TEXT("========================================================================="));
+		UE_LOG(LogPackageBuildDependencyTracker, Log, TEXT("========================================================================="));
 		for (const TPair<FName, TSet<FBuildDependencyAccessData>>& PackageAccessRecord : Records)
 		{
-			UE_LOG(LogPackageBuildDependencyTracker, Display, TEXT("%s:"), *PackageAccessRecord.Key.ToString());
+			UE_LOG(LogPackageBuildDependencyTracker, Log, TEXT("%s:"), *PackageAccessRecord.Key.ToString());
 			for (const FBuildDependencyAccessData& AccessedData : PackageAccessRecord.Value)
 			{
-				UE_LOG(LogPackageBuildDependencyTracker, Display, TEXT("    %s"), *AccessedData.ReferencedPackage.ToString());
+				UE_LOG(LogPackageBuildDependencyTracker, Log, TEXT("    %s"), *AccessedData.ReferencedPackage.ToString());
 			}
 		}
 	}
diff --git a/Engine/Source/Editor/UnrealEd/Private/DerivedDataBuildRemoteExecutor.cpp b/Engine/Source/Editor/UnrealEd/Private/DerivedDataBuildRemoteExecutor.cpp
--- a/Engine/Source/Editor/UnrealEd/Private/DerivedDataBuildRemoteExecutor.cpp
+++ b/Engine/Source/Editor/UnrealEd/Private/DerivedDataBuildRemoteExecutor.cpp
@@ -324,25 +324,25 @@ private:
 
 		void Dump()
 		{
-			UE_LOG(LogDerivedDataBuildRemoteExecutor, Display, TEXT(""));
-			UE_LOG(LogDerivedDataBuildRemoteExecutor, Display, TEXT("DDC Remote Execution Stats"));
-			UE_LOG(LogDerivedDataBuildRemoteExecutor, Display, TEXT("=========================="));
-			UE_LOG(LogDerivedDataBuildRemoteExecutor, Display, TEXT("%-36s=%10") UINT64_FMT, TEXT("Total remote builds"), TotalRemoteBuilds.load());
-			UE_LOG(LogDerivedDataBuildRemoteExecutor, Display, TEXT("%-36s=%10") UINT64_FMT, TEXT("Successful remote builds"), TotalSuccessfulRemoteBuilds.load());
-			UE_LOG(LogDerivedDataBuildRemoteExecutor, Display, TEXT("%-36s=%10") UINT64_FMT, TEXT("Timed out remote builds"), TotalTimedOutRemoteBuilds.load());
-
-			UE_LOG(LogDerivedDataBuildRemoteExecutor, Display, TEXT("%-36s=%10") UINT64_FMT, TEXT("Uploaded worker objects (quantity)"), TotalWorkerObjectsUploaded.Quantity.load());
-			UE_LOG(LogDerivedDataBuildRemoteExecutor, Display, TEXT("%-36s=%10") UINT64_FMT, TEXT("Uploaded worker objects (KB)"), TotalWorkerObjectsUploaded.Bytes.load()/1024);
-			UE_LOG(LogDerivedDataBuildRemoteExecutor, Display, TEXT("%-36s=%10") UINT64_FMT, TEXT("Uploaded worker packages (quantity)"), TotalWorkerPackagesUploaded.Quantity.load());
-			UE_LOG(LogDerivedDataBuildRemoteExecutor, Display, TEXT("%-36s=%10") UINT64_FMT, TEXT("Uploaded worker packages (KB)"), TotalWorkerPackagesUploaded.Bytes.load() / 1024);
-
-			UE_LOG(LogDerivedDataBuildRemoteExecutor, Display, TEXT("%-36s=%10") UINT64_FMT, TEXT("Uploaded action objects (quantity)"), TotalActionObjectsUploaded.Quantity.load());
-			UE_LOG(LogDerivedDataBuildRemoteExecutor, Display, TEXT("%-36s=%10") UINT64_FMT, TEXT("Uploaded action objects (KB)"), TotalActionObjectsUploaded.Bytes.load() / 1024);
-			UE_LOG(LogDerivedDataBuildRemoteExecutor, Display, TEXT("%-36s=%10") UINT64_FMT, TEXT("Uploaded action packages (quantity)"), TotalActionPackagesUploaded.Quantity.load());
-			UE_LOG(LogDerivedDataBuildRemoteExecutor, Display, TEXT("%-36s=%10") UINT64_FMT, TEXT("Uploaded action packages (KB)"), TotalActionPackagesUploaded.Bytes.load() / 1024);
-
-			UE_LOG(LogDerivedDataBuildRemoteExecutor, Display, TEXT("%-36s=%10") UINT64_FMT, TEXT("Downloaded packages (quantity)"), TotalPackagesDownloaded.Quantity.load());
-			UE_LOG(LogDerivedDataBuildRemoteExecutor, Display, TEXT("%-36s=%10") UINT64_FMT, TEXT("Downloaded packages (KB)"), TotalPackagesDownloaded.Bytes.load() / 1024);
+			UE_LOG(LogDerivedDataBuildRemoteExecutor, Log, TEXT(""));
+			UE_LOG(LogDerivedDataBuildRemoteExecutor, Log, TEXT("DDC Remote Execution Stats"));
+			UE_LOG(LogDerivedDataBuildRemoteExecutor, Log, TEXT("=========================="));
+			UE_LOG(LogDerivedDataBuildRemoteExecutor, Log, TEXT("%-36s=%10") UINT64_FMT, TEXT("Total remote builds"), TotalRemoteBuilds.load());
+			UE_LOG(LogDerivedDataBuildRemoteExecutor, Log, TEXT("%-36s=%10") UINT64_FMT, TEXT("Successful remote builds"), TotalSuccessfulRemoteBuilds.load());
+			UE_LOG(LogDerivedDataBuildRemoteExecutor, Log, TEXT("%-36s=%10") UINT64_FMT, TEXT("Timed out remote builds"), TotalTimedOutRemoteBuilds.load());
+
+			UE_LOG(LogDerivedDataBuildRemoteExecutor, Log, TEXT("%-36s=%10") UINT64_FMT, TEXT("Uploaded worker objects (quantity)"), TotalWorkerObjectsUploaded.Quantity.load());
+			UE_LOG(LogDerivedDataBuildRemoteExecutor, Log, TEXT("%-36s=%10") UINT64_FMT, TEXT("Uploaded worker objects (KB)"), TotalWorkerObjectsUploaded.Bytes.load()/1024);
+			UE_LOG(LogDerivedDataBuildRemoteExecutor, Log, TEXT("%-36s=%10") UINT64_FMT, TEXT("Uploaded worker packages (quantity)"), TotalWorkerPackagesUploaded.Quantity.load());
+			UE_LOG(LogDerivedDataBuildRemoteExecutor, Log, TEXT("%-36s=%10") UINT64_FMT, TEXT("Uploaded worker packages (KB)"), TotalWorkerPackagesUploaded.Bytes.load() / 1024);
+
+			UE_LOG(LogDerivedDataBuildRemoteExecutor, Log, TEXT("%-36s=%10") UINT64_FMT, TEXT("Uploaded action objects (quantity)"), TotalActionObjectsUploaded.Quantity.load());
+			UE_LOG(LogDerivedDataBuildRemoteExecutor, Log, TEXT("%-36s=%10") UINT64_FMT, TEXT("Uploaded action objects (KB)"), TotalActionObjectsUploaded.Bytes.load() / 1024);
+			UE_LOG(LogDerivedDataBuildRemoteExecutor, Log, TEXT("%-36s=%10") UINT64_FMT, TEXT("Uploaded action packages (quantity)"), TotalActionPackagesUploaded.Quantity.load());
+			UE_LOG(LogDerivedDataBuildRemoteExecutor, Log, TEXT("%-36s=%10") UINT64_FMT, TEXT("Uploaded action packages (KB)"), TotalActionPackagesUploaded.Bytes.load() / 1024);
+
+			UE_LOG(LogDerivedDataBuildRemoteExecutor, Log, TEXT("%-36s=%10") UINT64_FMT, TEXT("Downloaded packages (quantity)"), TotalPackagesDownloaded.Quantity.load());
+			UE_LOG(LogDerivedDataBuildRemoteExecutor, Log, TEXT("%-36s=%10") UINT64_FMT, TEXT("Downloaded packages (KB)"), TotalPackagesDownloaded.Bytes.load() / 1024);
 		}
 	};
 
