Fixing a typo, which was causing the wrong variable to be used and some users to crash on editor startup

#rnx

diff --git a/Engine/Source/Editor/UnrealEd/Private/EditorSelectUtils.cpp b/Engine/Source/Editor/UnrealEd/Private/EditorSelectUtils.cpp
--- a/Engine/Source/Editor/UnrealEd/Private/EditorSelectUtils.cpp
+++ b/Engine/Source/Editor/UnrealEd/Private/EditorSelectUtils.cpp
@@ -294,7 +294,7 @@ void UUnrealEdEngine::OnEditorElementSelectionPtrChanged(USelection* Selection,
 		if (NewSelectionSet)
 		{
 			NewSelectionSet->OnChanged().AddUObject(this, &UUnrealEdEngine::OnEditorElementSelectionChanged);
-			OldSelectionSet->OnPreChange().AddUObject(this, &UUnrealEdEngine::OnEditorSelectionPreChange);
+			NewSelectionSet->OnPreChange().AddUObject(this, &UUnrealEdEngine::OnEditorSelectionPreChange);
 		}
 	}
 }
