CIS - Fix a mac compile issue

#rnx

diff --git a/Engine/Source/Editor/UnrealEd/Private/EditorSelectUtils.cpp b/Engine/Source/Editor/UnrealEd/Private/EditorSelectUtils.cpp
--- a/Engine/Source/Editor/UnrealEd/Private/EditorSelectUtils.cpp
+++ b/Engine/Source/Editor/UnrealEd/Private/EditorSelectUtils.cpp
@@ -36,6 +36,10 @@
 #include "Elements/Interfaces/TypedElementWorldInterface.h"
 #include "Elements/Interfaces/TypedElementObjectInterface.h"
 
+#if PLATFORM_MAC
+#include "HAL/PlatformApplicationMisc.h"
+#endif // PLATFORM_MAC
+
 #define LOCTEXT_NAMESPACE "EditorSelectUtils"
 
 DEFINE_LOG_CATEGORY_STATIC(LogEditorSelectUtils, Log, All);
