Broadcast CookByTheBookStartedEvent and CookByTheBookFinishedEvent when starting/finishing cooking as a cook worker, providing that the director is cooking by the book.

diff --git a/Engine/Source/Editor/UnrealEd/Private/CookOnTheFlyServer.cpp b/Engine/Source/Editor/UnrealEd/Private/CookOnTheFlyServer.cpp
--- a/Engine/Source/Editor/UnrealEd/Private/CookOnTheFlyServer.cpp
+++ b/Engine/Source/Editor/UnrealEd/Private/CookOnTheFlyServer.cpp
@@ -10896,10 +10896,8 @@ void UCookOnTheFlyServer::StartCookAsCookWorker()
 	if (IsDirectorCookByTheBook())
 	{
 		RegisterCookByTheBookDelegates();
-	}
-	if (IsDirectorCookByTheBook())
-	{
 		BeginCookFinishShaderCodeLibrary(BeginContext);
+		CookByTheBookStartedEvent.Broadcast();
 	}
 }
 
@@ -10939,6 +10937,10 @@ void UCookOnTheFlyServer::CookAsCookWorkerFinished()
 	{
 		GShaderCompilingManager->SkipShaderCompilation(false);
 	}
+	if (IsDirectorCookByTheBook())
+	{
+		CookByTheBookFinishedEvent.Broadcast();
+	}
 	LogCookWorkerStats();
 	CookWorkerClient->FlushLogs();
 }
