Fixed package load warning from UAssetEditorSubsystem::OpenEditorForAsset

#rnx

diff --git a/Engine/Source/Editor/UnrealEd/Private/Subsystems/AssetEditorSubsystem.cpp b/Engine/Source/Editor/UnrealEd/Private/Subsystems/AssetEditorSubsystem.cpp
--- a/Engine/Source/Editor/UnrealEd/Private/Subsystems/AssetEditorSubsystem.cpp
+++ b/Engine/Source/Editor/UnrealEd/Private/Subsystems/AssetEditorSubsystem.cpp
@@ -765,45 +765,24 @@ void UAssetEditorSubsystem::HandleRequestOpenAssetMessage(const FAssetEditorRequ
 	OpenEditorForAsset(Message.AssetName);
 }
 
-void UAssetEditorSubsystem::OpenEditorForAsset(const FSoftObjectPath& AssetPath)
+void UAssetEditorSubsystem::OpenEditorForAsset(const FSoftObjectPath& AssetPath, const EAssetTypeActivationOpenedMethod OpenedMethod)
 {
-	OpenEditorForAsset(AssetPath.ToString());
-}
-
-void UAssetEditorSubsystem::OpenEditorForAsset(const FString& AssetPathName, const EAssetTypeActivationOpenedMethod OpenedMethod)
-{
-	UPackage* Package = FindPackage(NULL, *AssetPathName);
-
-	if(!Package)
+	UObject* Object = FindObject<UObject>(AssetPath.GetAssetPath());
+	if (!Object)
 	{
-		// An asset needs loading
-		Package =  LoadPackage(NULL, *AssetPathName, LOAD_NoRedirects);
-	}
-	
-	if (Package)
-	{
-		Package->FullyLoad();
-
-		FString AssetName = FPaths::GetBaseFilename(AssetPathName);
-		UObject* Object = FindObject<UObject>(Package, *AssetName);
-		
-		if (Object != nullptr)
-		{
-			OpenEditorForAsset(Object, EToolkitMode::Standalone, TSharedPtr<IToolkitHost>(), true, OpenedMethod);
-		}
+		Object = LoadObject<UObject>(nullptr, *AssetPath.GetAssetPathString(), nullptr, LOAD_NoRedirects);
 	}
-	else
+	if (Object)
 	{
-		// fallback for unsaved assets
-		UObject* Object = FindObject<UObject>(nullptr, *AssetPathName);
-		
-		if (Object != nullptr)
-		{
-			OpenEditorForAsset(Object, EToolkitMode::Standalone, TSharedPtr<IToolkitHost>(), true, OpenedMethod);
-		}
+		OpenEditorForAsset(Object, EToolkitMode::Standalone, TSharedPtr<IToolkitHost>(), true, OpenedMethod);
 	}
 }
 
+void UAssetEditorSubsystem::OpenEditorForAsset(const FString& AssetPathName, const EAssetTypeActivationOpenedMethod OpenedMethod)
+{
+	OpenEditorForAsset(FSoftObjectPath(AssetPathName), OpenedMethod);
+}
+
 bool UAssetEditorSubsystem::HandleTicker(float DeltaTime)
 {
 	QUICK_SCOPE_CYCLE_COUNTER(STAT_UAssetEditorSubsystem_HandleTicker);
diff --git a/Engine/Source/Editor/UnrealEd/Public/Subsystems/AssetEditorSubsystem.h b/Engine/Source/Editor/UnrealEd/Public/Subsystems/AssetEditorSubsystem.h
--- a/Engine/Source/Editor/UnrealEd/Public/Subsystems/AssetEditorSubsystem.h
+++ b/Engine/Source/Editor/UnrealEd/Public/Subsystems/AssetEditorSubsystem.h
@@ -104,7 +104,7 @@ public:
 
 	/** Opens an asset by path */
 	UNREALED_API void OpenEditorForAsset(const FString& AssetPathName, const EAssetTypeActivationOpenedMethod OpenedMethod = EAssetTypeActivationOpenedMethod::Edit);
-	UNREALED_API void OpenEditorForAsset(const FSoftObjectPath& AssetPath);
+	UNREALED_API void OpenEditorForAsset(const FSoftObjectPath& AssetPath, const EAssetTypeActivationOpenedMethod OpenedMethod = EAssetTypeActivationOpenedMethod::Edit);
 
 	/**
 	 * Tries to open an editor for the specified asset.  Returns true if the asset is open in an editor.
