Build graph and performance metrics

A new commandlet was added to record the performance data of workflows tests.

Modified the build graph script for the EditorPerf test to include the task nodes and agents for the workflows tests.

diff --git a/Engine/Build/EditorPerf.xml b/Engine/Build/EditorPerf.xml
--- a/Engine/Build/EditorPerf.xml
+++ b/Engine/Build/EditorPerf.xml
@@ -142,7 +142,7 @@
 	<Property Name="GauntletArgs" Value="$(GauntletArgs) -maxduration=28800"/>
 
 	<!-- Gauntlet arguments for profile runs -->
-	<Property Name="ProfileArgs" Value="-test=EditorTests -testfilter=Project.Iteration.PIE -tracefiletrunc"/>
+	<Property Name="ProfileArgs" Value="-test=EditorTests -tracefiletrunc"/>
 	<Property Name="ProfileArgs" Value="$(ProfileArgs) -EnablePlugins=$(ExtraPluginsCommaDelims)" If="'$(ExtraPluginsCommaDelims)' != ''"/>
 	<Property Name="ProfileArgs" Value="$(ProfileArgs) -SuppressLogWarnings -SuppressLogErrors"/>
 	<Property Name="ProfileArgs" Value="$(ProfileArgs) -GameInstanceLostTimerSeconds=3600"/>
@@ -165,6 +165,8 @@
 	<!-- The main test running macro logic that is expanded with platform and dependency node names -->
 	<Macro Name="RunTest" Arguments="TestNameSuffix;Platform;ProjectName;ProjectTarget;NodeName;Phase;Iteration;Args;EditorArgs">
 		<Property Name="TestName" Value="EditorPerf$(Phase)$(TestNameSuffix)"/>
+		<Property Name="TestFilter" Value="Project.Iteration.PIE"/>
+		<Property Name="TestFilterArgs" Value="-testfilter=$(TestFilter)"/>
 		<Property Name="IterationSuffix" Value=""/>
 		<Property Name="IterationSuffix" Value="$(Iteration)" If="'$(Iteration)' != '1'"/>
 
@@ -185,7 +187,7 @@
 		<!-- <Commandlet Project="$(ProjectTarget)" Name="DDCCleanup" Arguments="$(Args) $(CommandletArgs)"/> -->
 
 		<!-- Run the Gauntlet EditorPerf Tests -->
-		<Command Name="RunUnreal" Arguments="$(AllGauntletArgs) $(ProfileArgs) $(RunArgs) -sessionlabel=&quot;$(TestName)&quot; -tracefile=&quot;$(TraceFilePath)&quot;"/>
+		<Command Name="RunUnreal" Arguments="$(AllGauntletArgs) $(ProfileArgs) $(TestFilterArgs) $(RunArgs) -sessionlabel=&quot;$(TestName)&quot; -tracefile=&quot;$(TraceFilePath)&quot;"/>
 
 		<!-- Translate the utrace file into csv data -->
 		<Do If="$(RunTraceSummary)">
@@ -198,14 +200,41 @@
 		</Do>		
 	</Macro>
 	
+	<!-- The main editor workflows test running macro logic that is expanded with platform and dependency node names -->
+	<Macro Name="RunEditorWorkflowsTest" Arguments="Platform;ProjectName;NodeName">
+		<Property Name="TestName" Value="Editor.Workflows"/>
+		<Property Name="TestFilter" Value="Editor.Workflows"/>
+		<Property Name="TestFilterArgs" Value="-testfilter=$(TestFilter)"/>
+
+		<Property Name="TestOutputDir" Value="$(AutomationOutputDir)/$(ProjectName)$(Platform)EditorWorkflows"/>
+
+		<Property Name="LogOutputParams" Value="-logdir=&quot;$(TestOutputDir)&quot;"/>
+		<Property Name="ReportParams" Value="-WriteTestResultsForHorde -HordeTestDataKey=&quot;$(NodeName)$(PreflightSuffix)&quot;"/>
+		<Property Name="AllGauntletArgs" Value="-Project=$(ProjectName) -Platform=$(Platform) $(GauntletArgs) $(ReportParams) $(LogOutputParams)"/>
+
+		<Property Name="TraceFilePath" Value="$(TraceDir)/Editor.Workflows/$(ProjectName)$(Platform)EditorWorkflows.utrace" />
+
+		<Command Name="RunUnreal" Arguments="$(AllGauntletArgs) $(ProfileArgs) $(TestFilterArgs) -sessionlabel=&quot;$(TestName)&quot; -tracefile=&quot;$(TraceFilePath)&quot;" />
+
+		<!-- Translate the utrace file into csv data -->
+		<Commandlet Project="$(ProjectName)" Name="SummarizeTraceEditorWorkflows" Arguments="-inputfile=&quot;$(TraceFilePath)&quot; -testname=&quot;$(TestName)&quot; $(CommandletArgs)" />
+	</Macro>
+	
 	<Macro Name="UploadTestTraces" Arguments="Platform;ProjectName">		
 		<!-- Save all files named for this project in the trace directory -->
 		<Property Name="TestOutputDir" Value="$(AutomationOutputDir)/$(ProjectName)$(Platform)EditorPerf"/>
 		<Copy From="$(TraceDir)/$(ProjectName)$(Platform)EditorPerf*.*" To="$(TestOutputDir)"/>
 	</Macro>
 
-	<!-- This will be set to the names of all nodes we want to execute based on specified platforms/options -->
+	<Macro Name="UploadEditorWorkflowsPerformanceDataAndTraces" Arguments="Platform;ProjectName">
+		<!-- Copy the utrace file and the generated csv performance data to the output directory -->
+		<Property Name="TestOutputDir" Value="$(AutomationOutputDir)/$(ProjectName)$(Platform)EditorWorkflows"/>
+		<Copy From="$(TraceDir)/Editor.Workflows/*" To="$(TestOutputDir)/Editor.Workflows"/>
+	</Macro>
+
+	<!-- These will be set to the names of all nodes we want to execute based on specified platforms/options -->
 	<Property Name="RunGraphNodes" Value=""/>
+	<Property Name="RunEditorWorkflowsGraphNodes" Value=""/>
 
 	<!-- Mimic of the Runner that retrieves the stashed telemetry file and copies it to the depot location -->
 	<Macro Name="BaselineResult" Arguments="Platform;Project;Phase">
@@ -342,6 +371,57 @@
 
 	</ForEach> <!-- Platforms -->
 
+	<!-- Generate graph nodes for each platform to run -->
+	<ForEach Name="Platform" Values="$(Platforms)">
+
+		<!-- Define output name for the EditorWorkflows test -->
+		<Property Name="EditorWorkflowsOutputName" Value="EngineTest$(Platform)EditorWorkflows"/>
+		<Property Name="CompileNodeName" Value="Compile EngineTest $(Platform) EditorWorkflows"/>
+		<Property Name="AgentModifier" Value=""/>
+		<Property Name="AgentModifier" Value=" EngineTest" If="!$(SingleAgentCompile)"/>
+
+		<!-- Compile Agent configuration -->
+		<Agent Name="Compile Editor Agent $(Platform)$(AgentModifier)" Type="$($(Platform)CompileAgentType)">
+			<Node Name="$(CompileNodeName)">
+				<SetVersion Change="$(Change)" Branch="$(EscapedBranch)" If="$(IsBuildMachine)"/>
+				<Compile Target="ShaderCompileWorker" Platform="$(Platform)" Configuration="Development"/>
+				<Compile Target="CrashReportClientEditor" Platform="$(Platform)" Project="EngineTest" Configuration="Shipping"/>
+				<Compile Target="EngineTestEditor" Platform="$(Platform)" Project="EngineTest" Configuration="Development" Arguments="$(EditorCompileArgs)"/>
+			</Node>
+		</Agent>
+		<Label Category="$(CompileCategory)" Name="EngineTest Editor Compile" Requires="$(CompileNodeName)"/>
+
+		<!-- Set up properties for run and upload nodes -->
+		<Property Name="RunEditorWorkflowsNodeRequires" Value="$(CompileNodeName)"/>
+		<Property Name="UploadEditorWorkflowsNodeList" Value=""/>
+		<Property Name="RunEditorWorkflowsNodeList" Value=""/>
+
+		<!-- Configure EditorWorkflows Agent -->
+		<Agent Name="EditorWorkflows Agent $(Platform) EngineTest" Type="EditorPerf$(Platform)">
+			<Property Name="RunEditorWorkflowsNodeName" Value="Run $(EditorWorkflowsOutputName)" />
+			<Node Name="$(RunEditorWorkflowsNodeName)" Requires="$(RunEditorWorkflowsNodeRequires);$(RunEditorWorkflowsNodeList)">
+				<Expand Name="RunEditorWorkflowsTest" Platform="$(Platform)" ProjectName="EngineTest" NodeName="$(RunEditorWorkflowsNodeName)"/>
+			</Node>
+
+			<Property Name="RunEditorWorkflowsNodeList" Value="$(RunEditorWorkflowsNodeList);$(RunEditorWorkflowsNodeName)" />
+			<Property Name="RunNodeUploadTraces" Value="Upload $(EditorWorkflowsOutputName) Traces"/>
+
+			<!-- Set up Node for uploading traces -->
+			<Node Name="$(RunNodeUploadTraces)">
+				<Expand Name="UploadEditorWorkflowsPerformanceDataAndTraces" Platform="$(Platform)" ProjectName="EngineTest"/>
+			</Node>
+
+			<!-- Update UploadEditorWorkflowsNodeList -->
+			<Property Name="UploadEditorWorkflowsNodeList" Value="$(UploadEditorWorkflowsNodeList);$(RunNodeUploadTraces)"/>
+		</Agent>
+		<Label Category="$(RunCategory)" Name="EngineTest Editor Run" Requires="$(RunEditorWorkflowsNodeRequires);$(RunEditorWorkflowsNodeList);$(UploadEditorWorkflowsNodeList)" UgsBadge="EngineTest Editor Run"/>
+		<Property Name="RunEditorWorkflowsGraphNodes" Value="$(RunEditorWorkflowsGraphNodes);$(RunEditorWorkflowsNodeList);$(UploadEditorWorkflowsNodeList)"/>
+	</ForEach>
+	<!-- End ForEach Platforms -->
+
+	<!-- These Nodes run the Editor Workflows test and gather profile data -->
+	<Aggregate Name="Run EditorWorkflows Tests" Requires="$(RunEditorWorkflowsGraphNodes)"/>
+
 	<!-- These Nodes run the Editor and gather profile data -->
 	<Aggregate Name="Run EditorPerf Tests" Requires="$(RunGraphNodes)"/>
 
diff --git a/Engine/Source/Editor/UnrealEd/Classes/Commandlets/SummarizeTraceEditorWorkflowsCommandlet.h b/Engine/Source/Editor/UnrealEd/Classes/Commandlets/SummarizeTraceEditorWorkflowsCommandlet.h
new file mode 100644
--- /dev/null
+++ b/Engine/Source/Editor/UnrealEd/Classes/Commandlets/SummarizeTraceEditorWorkflowsCommandlet.h
@@ -0,0 +1,25 @@
+// Copyright Epic Games, Inc. All Rights Reserved.
+
+#pragma once
+#include "Commandlets/Commandlet.h"
+#include "Commandlets/SummarizeTraceUtils.h"
+#include "TraceServices/Model/Counters.h"
+#include "SummarizeTraceEditorWorkflowsCommandlet.generated.h"
+
+UCLASS(config=Editor)
+class USummarizeTraceEditorWorkflowsCommandlet : public UCommandlet
+{
+	GENERATED_UCLASS_BODY()
+
+public:
+	//~ Begin UCommandlet Interface
+
+	virtual int32 Main(const FString& CmdLineParams) override;
+
+	//~ End UCommandlet Interface
+
+private:
+	TPair<int32, TArray<int32>> ClipScopes(const TArray<const FSummarizeScope*>& Scopes, double BeginTime, double EndTime);
+	TUniquePtr<IFileHandle> OpenCSVFile(const FString& Path);
+	bool GenerateTimeCSV(const FString& Path, const TArray<FSummarizeScope>& Scopes, double BeginTime, double EndTime);
+};
diff --git a/Engine/Source/Editor/UnrealEd/Private/Commandlets/SummarizeTraceEditorWorkflowsCommandlet.cpp b/Engine/Source/Editor/UnrealEd/Private/Commandlets/SummarizeTraceEditorWorkflowsCommandlet.cpp
new file mode 100644
--- /dev/null
+++ b/Engine/Source/Editor/UnrealEd/Private/Commandlets/SummarizeTraceEditorWorkflowsCommandlet.cpp
@@ -0,0 +1,339 @@
+// Copyright Epic Games, Inc. All Rights Reserved.
+
+#include "Commandlets/SummarizeTraceEditorWorkflowsCommandlet.h"
+
+#include "HAL/PlatformFileManager.h"
+#include "Misc/Paths.h"
+#include "Misc/FileHelper.h"
+#include "ProfilingDebugging/CountersTrace.h"
+#include "String/ParseTokens.h"
+#include "Trace/Detail/Channel.h"
+#include "Trace/Analysis.h"
+#include "Trace/DataStream.h"
+#include "TraceServices/AnalyzerFactories.h"
+#include "TraceServices/Model/Bookmarks.h"
+#include "TraceServices/Model/Counters.h"
+#include "TraceServices/Utils.h"
+
+class FSummarizeBookmarksProviderEditorWorkflows
+	: public TraceServices::IEditableBookmarkProvider
+{
+public:
+	virtual void UpdateBookmarkSpec(uint64 BookmarkPoint, const TCHAR* FormatString, const TCHAR* File, int32 Line) override;
+	virtual void AppendBookmark(uint64 BookmarkPoint, double Time, const uint8* FormatString) override;
+	virtual void AppendBookmark(uint64 BookmarkPoint, double Time, const TCHAR* Text) override;
+	struct FBookmarkSpec
+	{
+		const TCHAR* File = nullptr;
+		const TCHAR* FormatString = nullptr;
+		int32 Line = 0;
+	};
+	FBookmarkSpec& GetSpec(uint64 BookmarkPoint);
+	FSummarizeBookmark* FindStartBookmarkForEndBookmark(const FString& Name);
+
+	// Keyed by a unique memory address
+	TMap<uint64, FBookmarkSpec> BookmarkSpecs;
+	// Keyed by name
+	TMap<FString, FSummarizeBookmark> Bookmarks;
+
+private:
+	enum
+	{
+		FormatBufferSize = 65536
+	};
+	TCHAR FormatBuffer[FormatBufferSize];
+	TCHAR TempBuffer[FormatBufferSize];
+};
+
+FSummarizeBookmarksProviderEditorWorkflows::FBookmarkSpec& FSummarizeBookmarksProviderEditorWorkflows::GetSpec(uint64 BookmarkPoint)
+{
+	FBookmarkSpec* Found = BookmarkSpecs.Find(BookmarkPoint);
+	if (Found)
+	{
+		return *Found;
+	}
+
+	FBookmarkSpec& Spec = BookmarkSpecs.Add(BookmarkPoint, FBookmarkSpec());
+	Spec.File = TEXT("<unknown>");
+	Spec.FormatString = TEXT("<unknown>");
+	return Spec;
+}
+
+void FSummarizeBookmarksProviderEditorWorkflows::UpdateBookmarkSpec(uint64 BookmarkPoint, const TCHAR* FormatString, const TCHAR* File, int32 Line)
+{
+	FBookmarkSpec& BookmarkSpec = GetSpec(BookmarkPoint);
+	BookmarkSpec.FormatString = FormatString;
+	BookmarkSpec.File = File;
+	BookmarkSpec.Line = Line;
+}
+
+void FSummarizeBookmarksProviderEditorWorkflows::AppendBookmark(uint64 BookmarkPoint, double Time, const uint8* FormatArgs)
+{
+	FBookmarkSpec& BookmarkSpec = GetSpec(BookmarkPoint);
+	TraceServices::StringFormat(FormatBuffer, FormatBufferSize - 1, TempBuffer, FormatBufferSize - 1, BookmarkSpec.FormatString, FormatArgs);
+	FSummarizeBookmarksProviderEditorWorkflows::AppendBookmark(BookmarkPoint, Time, FormatBuffer);
+}
+
+void FSummarizeBookmarksProviderEditorWorkflows::AppendBookmark(uint64 BookmarkPoint, double Time, const TCHAR* Text)
+{
+	FString Name(Text);
+
+	FSummarizeBookmark* FoundBookmark = Bookmarks.Find(Name);
+	if (!FoundBookmark)
+	{
+		FoundBookmark = &Bookmarks.Add(Name, FSummarizeBookmark());
+		FoundBookmark->Name = Name;
+	}
+
+	FoundBookmark->AddTimestamp(Time);
+}
+
+FSummarizeBookmark* FSummarizeBookmarksProviderEditorWorkflows::FindStartBookmarkForEndBookmark(const FString& Name)
+{
+	int32 Index = Name.Find(TEXT("Complete"));
+	if (Index != -1)
+	{
+		FString StartName = Name;
+		StartName.RemoveAt(Index, TCString<TCHAR>::Strlen(TEXT("Complete")));
+		return Bookmarks.Find(StartName);
+	}
+
+	return nullptr;
+}
+
+USummarizeTraceEditorWorkflowsCommandlet::USummarizeTraceEditorWorkflowsCommandlet(const FObjectInitializer& ObjectInitializer)
+	: Super(ObjectInitializer)
+{
+}
+
+int32 USummarizeTraceEditorWorkflowsCommandlet::Main(const FString& CmdLineParams)
+{
+	TArray<FString> Tokens;
+	TArray<FString> Switches;
+	TMap<FString, FString> ParamVals;
+	UCommandlet::ParseCommandLine(*CmdLineParams, Tokens, Switches, ParamVals);
+
+	// Display help
+	if (Switches.Contains("help"))
+	{
+		UE_LOG(LogSummarizeTrace, Log, TEXT("SummarizeTraceEditorWorkflows"));
+		UE_LOG(LogSummarizeTrace, Log, TEXT("This commandlet will summarize a utrace into a csv to ease of parsing by reporting tool."));
+		UE_LOG(LogSummarizeTrace, Log, TEXT("Options:"));
+		UE_LOG(LogSummarizeTrace, Log, TEXT(" Required: -inputfile=<utrace path>   (The utrace you wish to process)"));
+		return 0;
+	}
+
+	FString TraceFileName;
+	if (FParse::Value(*CmdLineParams, TEXT("inputfile="), TraceFileName, true))
+	{
+		UE_LOG(LogSummarizeTrace, Display, TEXT("Loading trace from %s"), *TraceFileName);
+	}
+	else
+	{
+		UE_LOG(LogSummarizeTrace, Error, TEXT("You must specify a utrace file using -inputfile=<path>"));
+		return 1;
+	}
+
+	if (!FPaths::FileExists(TraceFileName))
+	{
+		UE_LOG(LogSummarizeTrace, Error, TEXT("Trace file '%s' was not found"), *TraceFileName);
+		return 1;
+	}
+
+	UE::Trace::FFileDataStream* DataStream = new UE::Trace::FFileDataStream();
+	if (!DataStream->Open(*TraceFileName))
+	{
+		UE_LOG(LogSummarizeTrace, Error, TEXT("Unable to open trace file '%s' for read"), *TraceFileName);
+		delete DataStream;
+		return 1;
+	}
+
+	// setup analysis context with analyzers
+	UE::Trace::FAnalysisContext AnalysisContext;
+
+	// List of summarized scopes.
+	TArray<FSummarizeScope> CollectedScopeSummaries;
+
+	// Analyze CPU scope timer individually.
+	TSharedPtr<FSummarizeCpuScopeDurationAnalyzer> IndividualScopeAnalyzer = MakeShared<FSummarizeCpuScopeDurationAnalyzer>(
+		[&CollectedScopeSummaries](const TArray<FSummarizeScope>& ScopeSummaries)
+		{
+			// Collect all individual scopes summary from this analyzer.
+			CollectedScopeSummaries.Append(ScopeSummaries);
+		});
+
+	TSharedPtr<TraceServices::IAnalysisSession> Session = TraceServices::CreateAnalysisSession(0, nullptr, TUniquePtr<UE::Trace::IInDataStream>(DataStream));
+
+	FSummarizeBookmarksProviderEditorWorkflows BookmarksProvider;
+	TSharedPtr<UE::Trace::IAnalyzer> BookmarksAnalyzer = TraceServices::CreateBookmarksAnalyzer(*Session, BookmarksProvider);
+	AnalysisContext.AddAnalyzer(*BookmarksAnalyzer);
+
+	FSummarizeCpuProfilerProvider CpuProfilerProvider;
+	CpuProfilerProvider.AddCpuScopeAnalyzer(IndividualScopeAnalyzer);
+	TSharedPtr<UE::Trace::IAnalyzer> CpuProfilerAnalyzer = TraceServices::CreateCpuProfilerAnalyzer(*Session, CpuProfilerProvider, CpuProfilerProvider);
+	AnalysisContext.AddAnalyzer(*CpuProfilerAnalyzer);
+
+	// kick processing on a thread
+	UE::Trace::FAnalysisProcessor AnalysisProcessor = AnalysisContext.Process(*DataStream);
+
+	// sync on completion
+	AnalysisProcessor.Wait();
+
+	CpuProfilerProvider.AnalysisComplete();
+
+	// some locals to help with all the derived files we are about to generate
+	FString TracePath = FPaths::GetPath(TraceFileName);
+
+	// Define a map of test names to a pair of their begin and end times.
+	TMap<FString, TPair<double, double>> TestNameAndTimes;
+	for (const TMap<FString, FSummarizeBookmark>::ElementType& Bookmark : BookmarksProvider.Bookmarks)
+	{
+		if (!CsvUtils::IsCsvSafeString(Bookmark.Value.Name))
+		{
+			continue;
+		}
+
+		FString TestName;
+		FString BookmarkAnnotation;
+		Bookmark.Value.Name.Split(" ", &TestName, &BookmarkAnnotation);
+		const TCHAR* ProfileBeginBookmark = TEXT("ProfileBegin");
+		const TCHAR* ProfileEndBookmark = TEXT("ProfileEnd");
+
+		if (BookmarkAnnotation != ProfileBeginBookmark && BookmarkAnnotation != ProfileEndBookmark)
+		{
+			continue;
+		}
+
+		TPair<double, double>& TestTime = TestNameAndTimes.FindOrAdd(TestName);
+
+		if (BookmarkAnnotation == ProfileBeginBookmark)
+		{
+			TestTime.Key = Bookmark.Value.FirstSeconds;
+		}
+		else if (BookmarkAnnotation == ProfileEndBookmark)
+		{
+			TestTime.Value = Bookmark.Value.FirstSeconds;
+		}
+		else 
+		{
+			UE_LOG(LogSummarizeTrace, Error, TEXT("Unrecognized bookmark: '%s', bookmark should end with either 'ProfileBegin' or 'ProfileEnd'"), *BookmarkAnnotation);
+		}
+	}
+
+	// Loop through our map of test times and generate the CSV files.
+	IPlatformFile& PlatformFile = FPlatformFileManager::Get().GetPlatformFile();
+	for (const TMap<FString, TPair<double, double>>::ElementType& TestNameAndTime : TestNameAndTimes)
+	{
+		FString DirectoryPath = TracePath;
+		DirectoryPath.Append(TEXT("/"));
+		DirectoryPath.Append(TestNameAndTime.Key);
+
+		// Create directory for each individual test
+		if (!PlatformFile.DirectoryExists(*DirectoryPath))
+		{
+			PlatformFile.CreateDirectory(*DirectoryPath);
+
+			if (!PlatformFile.DirectoryExists(*DirectoryPath))
+			{
+				UE_LOG(LogSummarizeTrace, Error, TEXT("Failed to create directory: '%s'"), *DirectoryPath);
+				return 1;
+			}
+		}
+
+		if (!GenerateTimeCSV(DirectoryPath, CollectedScopeSummaries, TestNameAndTime.Value.Key, TestNameAndTime.Value.Value))
+		{
+			return 1;
+		}
+	}
+
+	return 0;
+}
+
+TPair<int32, TArray<int32>> USummarizeTraceEditorWorkflowsCommandlet::ClipScopes(const TArray<const FSummarizeScope*>& Scopes, double BeginTime, double EndTime)
+{
+	TArray<int32> BeginIndices;
+	int32 Size = TNumericLimits<int32>::Max();
+
+	for (const FSummarizeScope* Scope : Scopes)
+	{
+		int32 BeginIdx = Algo::LowerBound(Scope->BeginTimeArray, BeginTime);
+		int32 EndIdx = Algo::UpperBound(Scope->EndTimeArray, EndTime);
+
+		// Ensure the EndIndex is not less than BeginIndex
+		if (EndIdx < BeginIdx)
+		{
+			UE_LOG(LogSummarizeTrace, Error, TEXT("End index is smaller than begin index detected during clipping for ScopeName: '%s'"), *(Scope->Name));
+			EndIdx = BeginIdx;
+		}
+
+		BeginIndices.Add(BeginIdx);
+		Size = FMath::Min(Size, EndIdx - BeginIdx);
+	}
+
+	// Return the size and the array of Begin Indices
+	return { Size, BeginIndices };
+}
+
+TUniquePtr<IFileHandle> USummarizeTraceEditorWorkflowsCommandlet::OpenCSVFile(const FString& Path)
+{
+	FString CsvFileName = FPaths::Combine(Path, FPaths::SetExtension("Result", "csv"));
+	UE_LOG(LogSummarizeTrace, Display, TEXT("Writing %s..."), *CsvFileName);
+	TUniquePtr<IFileHandle> CsvHandle(FPlatformFileManager::Get().GetPlatformFile().OpenWrite(*CsvFileName));
+
+	if (!CsvHandle)
+	{
+		UE_LOG(LogSummarizeTrace, Error, TEXT("Unable to open csv '%s' for write"), *CsvFileName);
+	}
+
+	return CsvHandle;
+}
+
+bool USummarizeTraceEditorWorkflowsCommandlet::GenerateTimeCSV(const FString& Path, const TArray<FSummarizeScope>& Scopes, double BeginTime, double EndTime)
+{
+	TUniquePtr<IFileHandle> CsvHandle = OpenCSVFile(Path);
+	if (!CsvHandle)
+	{
+		return false;
+	}
+
+	const TSet<FString> ScopeNamesToTrack = { TEXT("FEngineLoop::Tick"), TEXT("Slate::Tick (Time and Widgets)") };
+	TArray<const FSummarizeScope*> TrackedScopes;
+	TrackedScopes.Reserve(ScopeNamesToTrack.Num());
+	for (const FSummarizeScope& Scope : Scopes)
+	{
+		if (ScopeNamesToTrack.Contains(Scope.Name))
+		{
+			TrackedScopes.Add(&Scope);
+		}
+	}
+
+	// Write the header
+	FString header = "FrameNumber";
+	for (const FString& name : ScopeNamesToTrack)
+	{
+		header.Append(",");
+		header.Append(name);
+	}
+	CsvUtils::WriteAsUTF8String(CsvHandle.Get(), header);
+
+	// Clip
+	TPair<int32, TArray<int32>> TimeRange = ClipScopes(TrackedScopes, BeginTime, EndTime);
+
+	// Write the data
+	for (int32 FrameNum = 0; FrameNum < TimeRange.Key; ++FrameNum)
+	{
+		FString DataLine = FString::Printf(TEXT("\n%d"), FrameNum);
+		for (int32 i = 0; i < TrackedScopes.Num(); ++i)
+		{
+			const FSummarizeScope* Scope = TrackedScopes[i];
+			int32 Idx = TimeRange.Value[i] + FrameNum;
+
+			DataLine.Append(FString::Printf(TEXT(",%f"), Scope->EndTimeArray[Idx] - Scope->BeginTimeArray[Idx]));
+		}
+
+		CsvUtils::WriteAsUTF8String(CsvHandle.Get(), DataLine);
+	}
+
+	CsvHandle->Flush();
+	return true;
+}
\ No newline at end of file
