CookShadersCommandlet for generating a subset of shader symbol files based upon input params
Can filter from a ShaderSymbols.info file, or by material name directly

#tests cooking Lyra and generating shader symbol files

diff --git a/Engine/Source/Editor/UnrealEd/Classes/Commandlets/CookShadersCommandlet.h b/Engine/Source/Editor/UnrealEd/Classes/Commandlets/CookShadersCommandlet.h
new file mode 100644
--- /dev/null
+++ b/Engine/Source/Editor/UnrealEd/Classes/Commandlets/CookShadersCommandlet.h
@@ -0,0 +1,18 @@
+// Copyright Epic Games, Inc. All Rights Reserved.
+
+#pragma once
+
+#include "CoreMinimal.h"
+#include "UObject/ObjectMacros.h"
+#include "Commandlets/Commandlet.h"
+#include "CookShadersCommandlet.generated.h"
+
+UCLASS(config = Editor)
+class UCookShadersCommandlet : public UCommandlet
+{
+	GENERATED_UCLASS_BODY()
+
+	//~ Begin UCommandlet Interface
+	virtual int32 Main(const FString& Params) override;
+	//~ End UCommandlet Interface
+};
\ No newline at end of file
diff --git a/Engine/Source/Editor/UnrealEd/Private/Commandlets/CookShadersCommandlet.cpp b/Engine/Source/Editor/UnrealEd/Private/Commandlets/CookShadersCommandlet.cpp
new file mode 100644
--- /dev/null
+++ b/Engine/Source/Editor/UnrealEd/Private/Commandlets/CookShadersCommandlet.cpp
@@ -0,0 +1,353 @@
+// Copyright Epic Games, Inc. All Rights Reserved.
+
+#include "Commandlets/CookShadersCommandlet.h"
+
+#include "AssetRegistry/AssetData.h"
+#include "AssetRegistry/AssetRegistryModule.h"
+#include "GlobalShader.h"
+#include "HAL/FileManager.h"
+#include "Interfaces/ITargetPlatform.h"
+#include "Interfaces/ITargetPlatformManagerModule.h"
+#include "MaterialShared.h"
+#include "Materials/Material.h"
+#include "Materials/MaterialInstance.h"
+#include "ShaderCompiler.h"
+#include "UObject/UObjectIterator.h"
+
+DEFINE_LOG_CATEGORY_STATIC(LogCookShadersCommandlet, Log, All);
+
+static const TCHAR* GlobalName = TEXT("Global");
+
+// Examples
+// UnrealEditor-Cmd.exe <proj> -run=CookShaders -targetPlatform=<platform> -infoFile=D:\ShaderSymbols\ShaderSymbols.info -ShaderSymbolsExport=D:\ShaderSymbols\Out -filter=Mannequin
+// UnrealEditor-Cmd.exe <proj> -run=CookShaders -targetPlatform=<platform> -infoFile=D:\ShaderSymbols\ShaderSymbols.info -ShaderSymbolsExport=D:\ShaderSymbols\Out -filter=00FB89F127D2DC10 -noglobals
+// UnrealEditor-Cmd.exe <proj> -run=CookShaders -targetPlatform=<platform> -ShaderSymbolsExport=D:\ShaderSymbols\Out -material=M_UI_Base_BordersAndButtons
+
+namespace CookShadersCommandlet {
+
+	// ShaderSymbols.info files will have a series of lines like the following, where the specifics of hash and extension are platform specific
+	// hash0.extension Global/FTonemapCS/2233
+	// hash1.extension M_Default_ad9c64900150ee77/Default/FLocalVertexFactory/TBasePassPSFNoLightMapPolicy/0
+	//
+	// FInfoRecord will contain a deconstructed version of a single line from this file
+	struct FInfoRecord
+	{
+		FString Hash;
+		FString Material;
+		FString VertexFactory;
+		FString Shader;
+		FString Permutation;
+		FString ID;
+		bool bGlobal = false;
+		bool bDefault = false;
+	};
+
+	bool LoadAndParse(const FString& Path, const FString& Filter, TArray<FInfoRecord>& OutInfo)
+	{
+		IFileManager& FileManager = IFileManager::Get();
+		TUniquePtr<FArchive> Reader = TUniquePtr<FArchive>(FileManager.CreateFileReader(*Path));
+		if (Reader.IsValid())
+		{
+			int64 Size = Reader->TotalSize();
+			TArray<uint8> RawData;
+			RawData.AddUninitialized(Size);
+			Reader->Serialize(RawData.GetData(), Size);
+			Reader->Close();
+
+			// figure out how long the hash is on this platform, each platform is different, but consistent and delimited by a space
+			uint32 HashSize = 0;
+			for (uint32 i = 0; i < Size; i++)
+			{
+				if (RawData[i] == ' ')
+				{
+					HashSize = i;
+					break;
+				}
+			}
+
+			if (HashSize == 0)
+			{
+				return false;
+			}
+
+			// go through each line
+			for (int i = 0; i < Size;)
+			{
+				if (i + HashSize < Size)
+				{
+					RawData[i + HashSize] = 0;
+					uint8* Hash = RawData.GetData() + i;
+					uint32 j = i + HashSize + 1;
+					uint8* Data = RawData.GetData() + j;
+					for (; j < Size; j++)
+					{
+						if (RawData[j] == '\n')
+						{
+							RawData[j] = 0;
+							break;
+						}
+					}
+
+					if (j >= Size)
+					{
+						return false;
+					}
+
+					FString HashString = StringCast<TCHAR>(reinterpret_cast<ANSICHAR*>(Hash)).Get();
+					FString DataString = StringCast<TCHAR>(reinterpret_cast<ANSICHAR*>(Data)).Get();
+
+					// add to our list if it passes the filter
+					if (Filter.IsEmpty() || HashString.Contains(Filter) || DataString.Contains(Filter))
+					{
+						FInfoRecord Record;
+						Record.Hash = HashString;
+
+						TArray<FString> Substrings;
+						DataString.ParseIntoArray(Substrings, TEXT("/"));
+
+						for (auto& S : Substrings)
+						{
+							if (Record.Material.IsEmpty())
+							{
+								TArray<FString> NameParts;
+								S.ParseIntoArray(NameParts, TEXT("_"));
+								if (NameParts.Num() == 1)
+								{
+									Record.Material = S;
+								}
+								else
+								{
+									Record.ID = NameParts[NameParts.Num() - 1];
+									Record.Material = S.Left(S.Len() - Record.ID.Len() - 1);
+								}
+
+								if (S == GlobalName)
+								{
+									Record.bGlobal = true;
+								}
+							}
+							else if (S == TEXT("Default"))
+							{
+								Record.bDefault = true;
+							}
+							else if (S.Contains(TEXT("VertexFactory")))
+							{
+								Record.VertexFactory = S;
+							}
+							else if (Record.Shader.IsEmpty())
+							{
+								Record.Shader = S;
+							}
+							else if (Record.Permutation.IsEmpty())
+							{
+								Record.Permutation = S;
+							}
+						}
+
+						OutInfo.Emplace(Record);
+					}
+
+					i = j + 1;
+				}
+			}
+
+			return true;
+		}
+
+		return false;
+	}
+
+};
+
+using namespace CookShadersCommandlet;
+
+UCookShadersCommandlet::UCookShadersCommandlet(const FObjectInitializer& ObjectInitializer)
+	: Super(ObjectInitializer)
+{
+}
+
+int32 UCookShadersCommandlet::Main(const FString& Params)
+{
+	TArray<FString> Tokens;
+	TArray<FString> Switches;
+	TMap<FString, FString> ParamVals;
+	UCommandlet::ParseCommandLine(*Params, Tokens, Switches, ParamVals);
+
+	// Display help
+	if (Switches.Contains("help"))
+	{
+		UE_LOG(LogCookShadersCommandlet, Log, TEXT("CookShadersCommandlet"));
+		UE_LOG(LogCookShadersCommandlet, Log, TEXT("Cook shaders based upon the options, ideal for generating pdbs for shaders you need"));
+		UE_LOG(LogCookShadersCommandlet, Log, TEXT("Options:"));
+		UE_LOG(LogCookShadersCommandlet, Log, TEXT(" Required: -targetPlatform=<platform>     (Which target platform do you want results, e.g. WindowsClient, etc."));
+		UE_LOG(LogCookShadersCommandlet, Log, TEXT(" Required: -ShaderSymbolsExport=<path>    (Force shader symbols on and where to write them."));
+		UE_LOG(LogCookShadersCommandlet, Log, TEXT(" Optional: -infoFile=<path>               (Path to ShaderSymbols.info file you want to find shaders from."));
+		UE_LOG(LogCookShadersCommandlet, Log, TEXT(" Optional: -filter=<string>               (Recommended! Filter to shaders with <string> in their hash or info data, requires -infoFile)."));
+		UE_LOG(LogCookShadersCommandlet, Log, TEXT(" Optional: -material=<string>             (Cook this material if you don't have a .info file, can be Global for global shaders)."));
+		UE_LOG(LogCookShadersCommandlet, Log, TEXT(" Optional: -noglobals                     (Don't do global shaders, even if they match the filter.)"));
+		return 0;
+	}
+
+	// Setup
+	FString Filter;
+	FParse::Value(*Params, TEXT("filter="), Filter, true);
+	FString MaterialString;
+	FParse::Value(*Params, TEXT("material="), MaterialString, true);
+	FString InfoFilePath;
+	FParse::Value(*Params, TEXT("infoFile="), InfoFilePath, true);
+	FString ExportPath;
+	FParse::Value(*Params, TEXT("ShaderSymbolsExport="), ExportPath, true);
+	const bool bNoGlobals = FParse::Param(FCommandLine::Get(), TEXT("noglobals"));
+
+	if (ExportPath.IsEmpty())
+	{
+		UE_LOG(LogCookShadersCommandlet, Log, TEXT("Note: Providing -ShaderSymbolsExport=<path> is highly recommended to automate activation of shader symbol file creation and set the location to save them!"));
+	}
+
+	TArray<FInfoRecord> Info;
+
+	// Mock up a material if requested
+	if (!MaterialString.IsEmpty())
+	{
+		FInfoRecord Record;
+		Record.bGlobal = (MaterialString == GlobalName);
+		Record.Material = MaterialString;
+		Info.Emplace(Record);
+	}
+
+	// Load info file if requested
+	if (!InfoFilePath.IsEmpty())
+	{
+		if (!LoadAndParse(InfoFilePath, Filter, Info))
+		{
+			UE_LOG(LogCookShadersCommandlet, Log, TEXT("Unabled to read / parse info file '%s'"), *InfoFilePath);
+			return 0;
+		}
+	}
+
+	// Load material assets
+	UE_LOG(LogCookShadersCommandlet, Display, TEXT("Loading Asset Registry..."));
+	IAssetRegistry& AssetRegistry = FModuleManager::Get().LoadModuleChecked<FAssetRegistryModule>("AssetRegistry").Get();
+	AssetRegistry.SearchAllAssets(true);
+
+	TArray<FAssetData> MaterialList;
+	TArray<FAssetData> MaterialInstanceList;
+	if (!AssetRegistry.IsLoadingAssets())
+	{
+		AssetRegistry.GetAssetsByClass(UMaterial::StaticClass()->GetClassPathName(), MaterialList, true);
+		AssetRegistry.GetAssetsByClass(UMaterialInstance::StaticClass()->GetClassPathName(), MaterialInstanceList, true);
+	}
+
+	// Pre-process the info we have, finding materials of interest
+	bool bCookGlobals = false;
+	TArray<FString> MaterialsRequested;
+	TSet<FString> MaterialsToFind;
+	for (const auto& I : Info)
+	{
+		if (I.bGlobal)
+		{
+			bCookGlobals = true;
+		}
+		else if (!I.Material.IsEmpty())
+		{
+			MaterialsToFind.Add(I.Material);
+		}
+	}
+
+	TSet<UMaterial*> MaterialsProcessed;
+	TSet<UMaterialInstance*> MaterialInstancesProcessed;
+	for (const auto& I : MaterialsToFind)
+	{
+		// Find matching materials
+		UMaterialInterface* MatchingMaterial = nullptr;
+		for (const FAssetData& It : MaterialList)
+		{
+			UMaterial* Material = Cast<UMaterial>(It.GetAsset());
+
+			if (Material && Material->GetName() == *I)
+			{
+				bool bAlreadyInSet = false;
+				MaterialsProcessed.Add(Material, &bAlreadyInSet);
+				if (!bAlreadyInSet)
+				{
+					MaterialsRequested.Add(Material->GetPathName());
+					MatchingMaterial = Material;
+				}
+				break;
+			}
+		}
+
+		// Locate material instances from the matched material
+		if (MatchingMaterial)
+		{
+			for (const FAssetData& It : MaterialInstanceList)
+			{
+				UMaterialInstance* MaterialInstance = Cast<UMaterialInstance>(It.GetAsset());
+
+				if (MaterialInstance && MaterialInstance->IsDependent(MatchingMaterial))
+				{
+					bool bAlreadyInSet = false;
+					MaterialInstancesProcessed.Add(MaterialInstance, &bAlreadyInSet);
+					if (!bAlreadyInSet)
+					{
+						MaterialsRequested.Add(MaterialInstance->GetPathName());
+					}
+				}
+			}
+		}
+	}
+
+	// Did we find anything to do? 
+	if (!bCookGlobals && MaterialsRequested.IsEmpty())
+	{
+		UE_LOG(LogCookShadersCommandlet, Log, TEXT("Couldn't find any globals or materials to process!"));
+		return 0;
+	}
+
+	// Iterate over the active platforms
+	ITargetPlatformManagerModule* TPM = GetTargetPlatformManager();
+	const TArray<ITargetPlatform*>& Platforms = TPM->GetActiveTargetPlatforms();
+	for (int32 Index = 0; Index < Platforms.Num(); Index++)
+	{
+		TArray<FName> DesiredShaderFormats;
+		Platforms[Index]->GetAllTargetedShaderFormats(DesiredShaderFormats);
+
+		for (int32 FormatIndex = 0; FormatIndex < DesiredShaderFormats.Num(); FormatIndex++)
+		{
+			const auto* Platform = Platforms[Index];
+			const EShaderPlatform ShaderPlatform = ShaderFormatToLegacyShaderPlatform(DesiredShaderFormats[FormatIndex]);
+			FString ShaderPlatformName = LexToString(ShaderPlatform);
+			FString PlatformName = Platform->PlatformName();
+			UE_LOG(LogCookShadersCommandlet, Log, TEXT("Working on %s %s"), *Platform->PlatformName(), *ShaderPlatformName);
+
+			// Setup
+			TArray<uint8> OutGlobalShaderMap;
+			TArray<uint8> OutMeshMaterialMaps;
+			TArray<FString> OutModifiedFiles;
+			FString OutputDir;
+			FShaderRecompileData Arguments(PlatformName, ShaderPlatform, ODSCRecompileCommand::None, &OutModifiedFiles, &OutMeshMaterialMaps, &OutGlobalShaderMap);
+
+			// Cook global shaders unless disabled
+			if (bCookGlobals && !bNoGlobals)
+			{
+				UE_LOG(LogCookShadersCommandlet, Log, TEXT("Cooking Global Shaders..."));
+				Arguments.CommandType = ODSCRecompileCommand::Global;
+				RecompileShadersForRemote(Arguments, OutputDir);
+			}
+
+			// Cook materials
+			if (!MaterialsRequested.IsEmpty())
+			{
+				UE_LOG(LogCookShadersCommandlet, Log, TEXT("Cooking Materials..."));
+				Arguments.CommandType = ODSCRecompileCommand::Material;
+				Arguments.MaterialsToLoad = MaterialsRequested;
+				RecompileShadersForRemote(Arguments, OutputDir);
+			}
+		}
+	}
+
+	UE_LOG(LogCookShadersCommandlet, Log, TEXT("Done CookShadersCommandlet"));
+	return 0;
+}
+
+
diff --git a/Engine/Source/Runtime/Engine/Private/ShaderCompiler/ShaderCompiler.cpp b/Engine/Source/Runtime/Engine/Private/ShaderCompiler/ShaderCompiler.cpp
--- a/Engine/Source/Runtime/Engine/Private/ShaderCompiler/ShaderCompiler.cpp
+++ b/Engine/Source/Runtime/Engine/Private/ShaderCompiler/ShaderCompiler.cpp
@@ -9330,11 +9330,15 @@ void RecompileShadersForRemote(
 					SaveShaderMapsForRemote(TargetPlatform, CompiledShaderMaps, Args.MeshMaterialMaps);
 				}
 
-				// save it out so the client can get it (and it's up to date next time)
-				FString GlobalShaderFilename = SaveGlobalShaderFile(ShaderPlatform, OutputDirectory, TargetPlatform);
+				// save it out so the client can get it (and it's up to date next time), if we were sent a OutputDirectory to put it in
+				FString GlobalShaderFilename;
+				if (!OutputDirectory.IsEmpty())
+				{
+					GlobalShaderFilename = SaveGlobalShaderFile(ShaderPlatform, OutputDirectory, TargetPlatform);
+				}
 
 				// add this to the list of files to tell the other end about
-				if (Args.ModifiedFiles)
+				if (Args.ModifiedFiles && !GlobalShaderFilename.IsEmpty())
 				{
 					// need to put it in non-sandbox terms
 					FString SandboxPath(GlobalShaderFilename);
