-Integration of PR #10625: [TTFontImport] Fix infinite loop when importing TTF glyph ranges.
-Bonus: fixed issue with the FontEditor not refreshing the details view when changing the font cache type (it was requiring to close and reopen the font editor to get updated content).

#tests changed the font cache type => the details view is refreshed as soon as we validate the change. Imported TTF with a glyph range specified => no infinite loop.

diff --git a/Engine/Source/Editor/FontEditor/Private/FontEditor.cpp b/Engine/Source/Editor/FontEditor/Private/FontEditor.cpp
--- a/Engine/Source/Editor/FontEditor/Private/FontEditor.cpp
+++ b/Engine/Source/Editor/FontEditor/Private/FontEditor.cpp
@@ -491,6 +491,7 @@ void FFontEditor::NotifyPostChange( const FPropertyChangedEvent& PropertyChanged
 
 			// If we changed the font cache type, then we need to update the UI to hide the invalid tabs and spawn the new ones
 			UpdateLayout();
+			FontProperties->ForceRefresh();
 		}
 		else
 		{
diff --git a/Engine/Source/Editor/UnrealEd/Private/Factories/TTFontImport.cpp b/Engine/Source/Editor/UnrealEd/Private/Factories/TTFontImport.cpp
--- a/Engine/Source/Editor/UnrealEd/Private/Factories/TTFontImport.cpp
+++ b/Engine/Source/Editor/UnrealEd/Private/Factories/TTFontImport.cpp
@@ -2227,7 +2227,7 @@ bool UTrueTypeFontFactory::ImportTrueTypeFont(
 		uint16 j = MinRangeCharacter;
 		uint16 Min = std::numeric_limits<uint16>::max();
 		uint16 Max = 0;
-		for( uint16 i = MinRangeCharacter; i <= std::numeric_limits<uint16>::max(); ++i )
+		for( uint16 i = MinRangeCharacter; i < std::numeric_limits<uint16>::max(); ++i )
 		{
 			if( Chars[i] )
 			{
