Fixed ObjectTools::AddExtraObjectsToDelete not to include newly created packages as it will try to delete packages that don't exist and fail which can prevent SaveAs to succeed.

diff --git a/Engine/Source/Editor/UnrealEd/Private/ObjectTools.cpp b/Engine/Source/Editor/UnrealEd/Private/ObjectTools.cpp
--- a/Engine/Source/Editor/UnrealEd/Private/ObjectTools.cpp
+++ b/Engine/Source/Editor/UnrealEd/Private/ObjectTools.cpp
@@ -2718,7 +2718,11 @@ namespace ObjectTools
 
 				for (UPackage* Package : World->GetOutermost()->GetExternalPackages())
 				{
-					ObjectsToDelete.AddUnique(Package);
+					// Don't include newly created packages
+					if (!Package->HasAnyPackageFlags(PKG_NewlyCreated))
+					{
+						ObjectsToDelete.AddUnique(Package);
+					}
 				}
 			}
 		}
