Fix crash attempting to initialize mainframe when slate not initialized

diff --git a/Engine/Source/Editor/UnrealEd/Private/UnrealEdGlobals.cpp b/Engine/Source/Editor/UnrealEd/Private/UnrealEdGlobals.cpp
--- a/Engine/Source/Editor/UnrealEd/Private/UnrealEdGlobals.cpp
+++ b/Engine/Source/Editor/UnrealEd/Private/UnrealEdGlobals.cpp
@@ -21,6 +21,7 @@
 #include "UnrealEngine.h"
 #include "UnrealEdMisc.h"
 #include "EditorModes.h"
+#include "Framework/Application/SlateApplication.h"
 #include "DesktopPlatformModule.h"
 
 #include "DebugToolExec.h"
@@ -162,13 +163,24 @@ int32 EditorInit( IEngineLoop& EngineLoop )
 	{
 		TRACE_CPUPROFILER_EVENT_SCOPE(EditorInit::MainFrame);
 
-		// Startup Slate main frame and other editor windows
+		// Startup Slate main frame and other editor windows if possible
 		{
 			const bool bStartImmersive = bIsImmersive;
 			const bool bStartPIE = bIsImmersive || bIsPlayInEditorRequested;
 
 			IMainFrameModule& MainFrameModule = FModuleManager::LoadModuleChecked<IMainFrameModule>(TEXT("MainFrame"));
-			MainFrameModule.CreateDefaultMainFrame( bStartImmersive, bStartPIE );
+			if (!MainFrameModule.IsWindowInitialized())
+			{
+				if (FSlateApplication::IsInitialized())
+				{
+					MainFrameModule.CreateDefaultMainFrame(bStartImmersive, bStartPIE);
+				}
+				else
+				{
+					RequestEngineExit(TEXT("Slate Application terminated or not initialized for MainFrame"));
+					return 1;
+				}
+			}
 		}
 	}
 
