Skeletal Mesh: Enable importing of vertex maps from FBX into skeletal mesh assets.

diff --git a/Engine/Plugins/Tests/EditorTests/Source/EditorTests/Private/UnrealEd/FbxAutomationTests.cpp b/Engine/Plugins/Tests/EditorTests/Source/EditorTests/Private/UnrealEd/FbxAutomationTests.cpp
--- a/Engine/Plugins/Tests/EditorTests/Source/EditorTests/Private/UnrealEd/FbxAutomationTests.cpp
+++ b/Engine/Plugins/Tests/EditorTests/Source/EditorTests/Private/UnrealEd/FbxAutomationTests.cpp
@@ -414,6 +414,7 @@ bool FFbxImportAssetsAutomationTest::RunTest(const FString& Parameters)
 					ImportData->ImportContentType = TestPlan->ImportUI->SkeletalMeshImportData->ImportContentType;
 					ImportData->bImportMeshesInBoneHierarchy = TestPlan->ImportUI->SkeletalMeshImportData->bImportMeshesInBoneHierarchy;
 					ImportData->bImportMorphTargets = TestPlan->ImportUI->SkeletalMeshImportData->bImportMorphTargets;
+					ImportData->bImportVertexAttributes = TestPlan->ImportUI->SkeletalMeshImportData->bImportVertexAttributes;
 					ImportData->ThresholdPosition = TestPlan->ImportUI->SkeletalMeshImportData->ThresholdPosition;
 					ImportData->ThresholdTangentNormal = TestPlan->ImportUI->SkeletalMeshImportData->ThresholdTangentNormal;
 					ImportData->MorphThresholdPosition = TestPlan->ImportUI->SkeletalMeshImportData->MorphThresholdPosition;
@@ -510,6 +511,7 @@ bool FFbxImportAssetsAutomationTest::RunTest(const FString& Parameters)
 					ImportData->ImportContentType = TestPlan->ImportUI->SkeletalMeshImportData->ImportContentType;
 					ImportData->bImportMeshesInBoneHierarchy = TestPlan->ImportUI->SkeletalMeshImportData->bImportMeshesInBoneHierarchy;
 					ImportData->bImportMorphTargets = TestPlan->ImportUI->SkeletalMeshImportData->bImportMorphTargets;
+					ImportData->bImportVertexAttributes = TestPlan->ImportUI->SkeletalMeshImportData->bImportVertexAttributes;
 					ImportData->ThresholdPosition = TestPlan->ImportUI->SkeletalMeshImportData->ThresholdPosition;
 					ImportData->ThresholdTangentNormal = TestPlan->ImportUI->SkeletalMeshImportData->ThresholdTangentNormal;
 					ImportData->ThresholdUV = TestPlan->ImportUI->SkeletalMeshImportData->ThresholdUV;
diff --git a/Engine/Source/Editor/UnrealEd/Classes/Factories/FbxSceneImportOptionsSkeletalMesh.h b/Engine/Source/Editor/UnrealEd/Classes/Factories/FbxSceneImportOptionsSkeletalMesh.h
--- a/Engine/Source/Editor/UnrealEd/Classes/Factories/FbxSceneImportOptionsSkeletalMesh.h
+++ b/Engine/Source/Editor/UnrealEd/Classes/Factories/FbxSceneImportOptionsSkeletalMesh.h
@@ -42,6 +42,10 @@ class UFbxSceneImportOptionsSkeletalMesh : public UObject
 	UPROPERTY(EditAnywhere, config, Category = SkeletalMesh, meta = (ToolTip = "If enabled, creates Unreal morph objects for the imported meshes"))
 	uint32 bImportMorphTargets : 1;
 
+	/** True to import morph target meshes from the FBX file */
+	UPROPERTY(EditAnywhere, config, Category = SkeletalMesh, meta = (ToolTip = "If enabled, import single-channel/weight/alpha vertex attributes"))
+	uint32 bImportVertexAttributes : 1;
+
 	/** Threshold to compare vertex position equality. */
 	UPROPERTY(EditAnywhere, config, Category = "SkeletalMesh|Thresholds", meta = (NoSpinbox = "true", ClampMin = "0.0"))
 	float ThresholdPosition;
diff --git a/Engine/Source/Editor/UnrealEd/Classes/Factories/FbxSkeletalMeshImportData.h b/Engine/Source/Editor/UnrealEd/Classes/Factories/FbxSkeletalMeshImportData.h
--- a/Engine/Source/Editor/UnrealEd/Classes/Factories/FbxSkeletalMeshImportData.h
+++ b/Engine/Source/Editor/UnrealEd/Classes/Factories/FbxSkeletalMeshImportData.h
@@ -83,6 +83,11 @@ public:
 	/** True to import morph target meshes from the FBX file */
 	UPROPERTY(EditAnywhere, AdvancedDisplay, config, Category=Mesh, meta=(ImportType="SkeletalMesh|GeoOnly", ToolTip="If enabled, creates Unreal morph objects for the imported meshes"))
 	uint32 bImportMorphTargets:1;
+	
+	/** True to import per-vertex attributes from the FBX file */
+	UPROPERTY(EditAnywhere, AdvancedDisplay, config, Category=Mesh, meta=(ImportType="SkeletalMesh|GeoOnly", ToolTip="If enabled, creates named vertex attributes for secondary vertex color data."))
+	uint32 bImportVertexAttributes:1;
+	
 
 	/** Threshold to compare vertex position equality. */
 	UPROPERTY(EditAnywhere, config, Category="Mesh", meta = (ImportType = "SkeletalMesh|GeoOnly", SubCategory = "Thresholds", NoSpinbox = "true", ClampMin = "0.0"))
diff --git a/Engine/Source/Editor/UnrealEd/Private/Fbx/FbxMainImport.cpp b/Engine/Source/Editor/UnrealEd/Private/Fbx/FbxMainImport.cpp
--- a/Engine/Source/Editor/UnrealEd/Private/Fbx/FbxMainImport.cpp
+++ b/Engine/Source/Editor/UnrealEd/Private/Fbx/FbxMainImport.cpp
@@ -429,6 +429,7 @@ void ApplyImportUIToImportOptions(UFbxImportUI* ImportUI, FBXImportOptions& InOu
 		InOutImportOptions.VertexColorImportOption		= SkeletalMeshData->VertexColorImportOption;
 		InOutImportOptions.VertexOverrideColor			= SkeletalMeshData->VertexOverrideColor;
 		InOutImportOptions.bReorderMaterialToFbxOrder	= SkeletalMeshData->bReorderMaterialToFbxOrder;
+		InOutImportOptions.bImportVertexAttributes		= SkeletalMeshData->bImportVertexAttributes;
 
 		if(ImportUI->bImportAnimations)
 		{
diff --git a/Engine/Source/Editor/UnrealEd/Private/Fbx/FbxSceneImportData.cpp b/Engine/Source/Editor/UnrealEd/Private/Fbx/FbxSceneImportData.cpp
--- a/Engine/Source/Editor/UnrealEd/Private/Fbx/FbxSceneImportData.cpp
+++ b/Engine/Source/Editor/UnrealEd/Private/Fbx/FbxSceneImportData.cpp
@@ -121,6 +121,7 @@ UnFbx::FBXImportOptions *JSONToFbxOption(TSharedPtr<FJsonValue> OptionJsonValue,
 	OptionObj->TryGetBoolField("bPreserveSmoothingGroups", Option->bPreserveSmoothingGroups);
 	OptionObj->TryGetBoolField("bImportMeshesInBoneHierarchy", Option->bImportMeshesInBoneHierarchy);
 	OptionObj->TryGetBoolField("bImportMorphTargets", Option->bImportMorph);
+	OptionObj->TryGetBoolField("bImportVertexAttributes", Option->bImportVertexAttributes);
 	
 	if (OptionObj->TryGetObjectField("OverlappingThresholds", DataObj))
 	{
@@ -204,12 +205,13 @@ FString FbxOptionToJSON(FString OptionName, UnFbx::FBXImportOptions *Option)
 		);
 
 	//Skeletal mesh options
-	JsonString += FString::Printf(TEXT("\"bUpdateSkeletonReferencePose\" : \"%d\", \"bUseT0AsRefPose\" : \"%d\", \"bPreserveSmoothingGroups\" : \"%d\", \"bImportMeshesInBoneHierarchy\" : \"%d\", \"bImportMorphTargets\" : \"%d\", \"OverlappingThresholds\" : {\"ThresholdPosition\" : \"%f\", \"ThresholdTangentNormal\" : \"%f\", \"ThresholdUV\" : \"%f\", \"MorphThresholdPosition\" : \"%f\"},"),
+	JsonString += FString::Printf(TEXT("\"bUpdateSkeletonReferencePose\" : \"%d\", \"bUseT0AsRefPose\" : \"%d\", \"bPreserveSmoothingGroups\" : \"%d\", \"bImportMeshesInBoneHierarchy\" : \"%d\", \"bImportMorphTargets\" : \"%d\", \"bImportVertexAttributes\" : \"%d\", \"OverlappingThresholds\" : {\"ThresholdPosition\" : \"%f\", \"ThresholdTangentNormal\" : \"%f\", \"ThresholdUV\" : \"%f\", \"MorphThresholdPosition\" : \"%f\"},"),
 		Option->bUpdateSkeletonReferencePose ? 1 : 0,
 		Option->bUseT0AsRefPose ? 1 : 0,
 		Option->bPreserveSmoothingGroups ? 1 : 0,
 		Option->bImportMeshesInBoneHierarchy ? 1 : 0,
 		Option->bImportMorph ? 1 : 0,
+		Option->bImportVertexAttributes ? 1 : 0,
 		Option->OverlappingThresholds.ThresholdPosition,
 		Option->OverlappingThresholds.ThresholdTangentNormal,
 		Option->OverlappingThresholds.ThresholdUV,
diff --git a/Engine/Source/Editor/UnrealEd/Private/Fbx/FbxSceneImportOptionsSkeletalMesh.cpp b/Engine/Source/Editor/UnrealEd/Private/Fbx/FbxSceneImportOptionsSkeletalMesh.cpp
--- a/Engine/Source/Editor/UnrealEd/Private/Fbx/FbxSceneImportOptionsSkeletalMesh.cpp
+++ b/Engine/Source/Editor/UnrealEd/Private/Fbx/FbxSceneImportOptionsSkeletalMesh.cpp
@@ -18,6 +18,7 @@ UFbxSceneImportOptionsSkeletalMesh::UFbxSceneImportOptionsSkeletalMesh(const FOb
 	, bPreserveSmoothingGroups(false)
 	, bImportMeshesInBoneHierarchy(true)
 	, bImportMorphTargets(false)
+	, bImportVertexAttributes(false)
 	, ThresholdPosition(THRESH_POINTS_ARE_SAME)
 	, ThresholdTangentNormal(THRESH_NORMALS_ARE_SAME)
 	, ThresholdUV(THRESH_UVS_ARE_SAME)
@@ -40,6 +41,8 @@ void UFbxSceneImportOptionsSkeletalMesh::FillSkeletalMeshInmportData(UFbxSkeleta
 	check(SkeletalMeshImportData != nullptr);
 	SkeletalMeshImportData->bImportMeshesInBoneHierarchy = bImportMeshesInBoneHierarchy;
 	SkeletalMeshImportData->bImportMorphTargets = bImportMorphTargets;
+	SkeletalMeshImportData->bImportVertexAttributes = bImportVertexAttributes;
+
 	SkeletalMeshImportData->ThresholdPosition = ThresholdPosition;
 	SkeletalMeshImportData->ThresholdTangentNormal = ThresholdTangentNormal;
 	SkeletalMeshImportData->ThresholdUV = ThresholdUV;
diff --git a/Engine/Source/Editor/UnrealEd/Private/Fbx/FbxSkeletalMeshImport.cpp b/Engine/Source/Editor/UnrealEd/Private/Fbx/FbxSkeletalMeshImport.cpp
--- a/Engine/Source/Editor/UnrealEd/Private/Fbx/FbxSkeletalMeshImport.cpp
+++ b/Engine/Source/Editor/UnrealEd/Private/Fbx/FbxSkeletalMeshImport.cpp
@@ -73,7 +73,6 @@
 
 using namespace UnFbx;
 
-
 // Get the geometry deformation local to a node. It is never inherited by the
 // children.
 FbxAMatrix GetGeometry(FbxNode* pNode) 
@@ -3799,6 +3798,199 @@ bool UnFbx::FFbxImporter::FillSkelMeshImporterFromFbx( FSkeletalMeshImportData&
 		}
 	}
 
+	//
+	// Get the vertex attribute layers from all layers, even the first layer which may be used as a vertex color layer.
+	// Currently we're only interested in alpha-only layers, since those are the only layer types the engine
+	// currently exposes for vertex attributes. Internally we can store 1-4 components, but there's no tooling for that
+	// 2-4 channels as of yet.
+	//
+	struct FNamedVertexAttribute : SkeletalMeshImportData::FVertexAttribute
+	{
+		FNamedVertexAttribute(FString&& InAttributeName, TArray<float>&& InAttributeValues, const int32 InComponentCount)
+			: FVertexAttribute(MoveTemp(InAttributeValues), InComponentCount)
+			, AttributeName(InAttributeName)
+		{}
+			
+		FString AttributeName;
+	};
+
+	if (ImportOptions->bImportVertexAttributes)
+	{
+		TArray<FNamedVertexAttribute> NamedVertexAttributes;
+
+		for (int32 LayerIndex = 0; LayerIndex < LayerCount; LayerIndex++)
+		{
+			FbxLayerElementVertexColor* LayerElementVertexAttribute = Mesh->GetLayer(LayerIndex)->GetVertexColors();
+			if (!LayerElementVertexAttribute)
+			{
+				continue;
+			}
+
+			// Check if this is an alpha-only attribute, by ensuring the RGB values are all zero, otherwise skip.
+			bool bIsValidAttribute = true;
+			const FbxLayerElementArrayTemplate<FbxColor>& AttributeValues = LayerElementVertexAttribute->GetDirectArray();
+			for (int32 Index = 0; Index < AttributeValues.GetCount(); Index++)
+			{
+				// We do an exact comparison, since that's how empty channels would be represented in the FBX file.
+				const FbxColor& Value = AttributeValues.GetAt(Index); 
+				if (Value.mRed != 0.0 || Value.mGreen != 0.0 || Value.mBlue != 0.0)
+				{
+					bIsValidAttribute = false;
+					break;
+				}
+			}
+
+			// We can only do attributes that are mapped per-vertex.
+			if (!bIsValidAttribute)
+			{
+				continue;
+			}
+
+			const int32 AttributeComponentCount = 1;	// Number of component values per vertex. See comment above. 
+			TArray<float> AttributeComponentValues;
+			
+			switch(LayerElementVertexAttribute->GetMappingMode())
+			{
+			case FbxLayerElement::eByControlPoint:
+				{
+					AttributeComponentValues.AddZeroed(ControlPointsCount);
+
+					if (LayerElementVertexAttribute->GetReferenceMode() == FbxLayerElement::eDirect)
+					{
+						for (int32 Index = 0; Index < AttributeValues.GetCount(); Index++)
+						{
+							AttributeComponentValues[Index] = AttributeValues.GetAt(Index).mAlpha;
+						}
+					}
+					else // LayerElementVertexAttribute->GetReferenceMode() == FbxLayerElement::eIndexToDirect
+					{
+						const FbxLayerElementArrayTemplate<int>& IndexArray = LayerElementVertexAttribute->GetIndexArray();
+						for (int32 Index = 0; Index < IndexArray.GetCount(); Index++)
+						{
+							AttributeComponentValues[Index] = AttributeValues.GetAt(IndexArray[Index]).mAlpha;
+						}
+					}
+				}
+				break;
+			case FbxLayerElement::eByPolygonVertex:
+				{
+					// Vertex attributes are stored per-vertex, not per-vertex instance. To work around this we average
+					// together values that share a vertex.
+					TArray<int32> SharedVertexCount;
+					SharedVertexCount.AddZeroed(ControlPointsCount);
+					AttributeComponentValues.AddZeroed(ControlPointsCount);
+
+					const FbxLayerElementArrayTemplate<int>* IndexArray = nullptr;
+					if (LayerElementVertexAttribute->GetReferenceMode() == FbxLayerElement::eIndexToDirect)
+					{
+						IndexArray = &LayerElementVertexAttribute->GetIndexArray();
+					}
+
+					const int* PolygonControlPointIndexes = Mesh->GetPolygonVertices();
+					
+					for(int32 TriangleIndex = 0; TriangleIndex < TriangleCount; TriangleIndex++)
+					{
+						for (int32 InnerIndex = 0; InnerIndex < 3; InnerIndex++)
+						{
+							const int32 PolygonVertexIndex = TriangleIndex * 3 + InnerIndex;;
+							const int32 PointIndex = PolygonControlPointIndexes[PolygonVertexIndex];
+
+							AttributeComponentValues[PointIndex] +=
+									AttributeValues.GetAt(IndexArray ? IndexArray->GetAt(PolygonVertexIndex) : PolygonVertexIndex).mAlpha;
+							SharedVertexCount[PointIndex]++;
+						}					
+					}
+
+					for (int32 PointIndex = 0; PointIndex < ControlPointsCount; PointIndex++)
+					{
+						if (SharedVertexCount[PointIndex] > 1)
+						{
+							AttributeComponentValues[PointIndex] /= static_cast<float>(SharedVertexCount[PointIndex]);
+						}
+					}
+				}
+				break;
+			default:
+				break;
+			}
+
+			if (!AttributeComponentValues.IsEmpty())
+			{
+				FString AttributeName(UTF8_TO_TCHAR(LayerElementVertexAttribute->GetName()));
+				NamedVertexAttributes.Emplace(MoveTemp(AttributeName), MoveTemp(AttributeComponentValues), AttributeComponentCount);
+			}
+		}
+
+		//
+		// Add in the attributes that we received, and pad any non-matching existing attributes with zero values to
+		// match the new point count.
+		//
+		if (ExistFaceNum)
+		{
+			// TODO: In the future, once we support attributes with non-unity component counts, we need to decide how to deal
+			// with attributes that share the same name but differing component counts. Two options:
+			// 1) We create a new attribute with a new name that differs from the existing name. Preferably with the component
+			//    count somehow embedded in the name so that we can deal with appending again.
+			// 2) Widen the existing attribute's component count to match the new one (or vice versa, widening the new attribute
+			//    so that we can append).
+			
+			const int32 VertexCount = Mesh->GetControlPointsCount();
+			
+			for (int32 ExistingAttributeIndex = 0; ExistingAttributeIndex < ImportData.VertexAttributes.Num(); ExistingAttributeIndex++)
+			{
+				const FString& ExistingAttributeName = ImportData.VertexAttributeNames[ExistingAttributeIndex];
+				const int32 ExistingAttributeComponentCount = ImportData.VertexAttributes[ExistingAttributeIndex].ComponentCount;
+				
+				int32 NewAttributeIndex = NamedVertexAttributes.IndexOfByPredicate(
+					[ExistingAttributeName](const FNamedVertexAttribute& InAttribute)
+					{
+						return InAttribute.AttributeName == ExistingAttributeName && ensure(InAttribute.ComponentCount == 1);
+					});
+
+				if (ExistingAttributeComponentCount != 1)
+				{
+					// For now, if an existing attribute has a non-unity component count, we don't append the new attribute values
+					// on top, just pad with zeroes. 
+					const FText ErrorMsg = FText::Format(LOCTEXT("FbxSkeletaLMeshimport_MaterialIndexInconsistency", "Existing attribute '{0}' has more than one component. Ignoring imported attribute of same name."),
+							FText::FromString(ExistingAttributeName)); 
+					AddTokenizedErrorMessage(FTokenizedMessage::Create(EMessageSeverity::Warning, ErrorMsg), FFbxErrors::SkeletalMesh_AttributeComponentCountMismatch);
+					
+					NewAttributeIndex = INDEX_NONE;
+				}
+
+				SkeletalMeshImportData::FVertexAttribute& ExistingAttribute = ImportData.VertexAttributes[ExistingAttributeIndex];
+				if (NewAttributeIndex != INDEX_NONE)
+				{
+					ExistingAttribute.AttributeValues.Append(NamedVertexAttributes[NewAttributeIndex].AttributeValues);
+					NamedVertexAttributes.RemoveAt(NewAttributeIndex);
+				}
+				else
+				{
+					ExistingAttribute.AttributeValues.AddZeroed(VertexCount * ExistingAttribute.ComponentCount);
+				}
+			}
+
+			// Any remaining attributes we add + padding for existing points.
+			for (FNamedVertexAttribute& NamedVertexAttribute: NamedVertexAttributes)
+			{
+				ImportData.VertexAttributeNames.Emplace(MoveTemp(NamedVertexAttribute.AttributeName));
+				
+				SkeletalMeshImportData::FVertexAttribute& NewAttribute = ImportData.VertexAttributes.AddDefaulted_GetRef();
+				NewAttribute.ComponentCount = NamedVertexAttribute.ComponentCount;
+				NewAttribute.AttributeValues.AddZeroed(ExistPointNum * NamedVertexAttribute.ComponentCount);
+				NewAttribute.AttributeValues.Append(NamedVertexAttribute.AttributeValues);
+			}
+		}
+		else
+		{
+			for (FNamedVertexAttribute& NamedVertexAttribute: NamedVertexAttributes)
+			{
+				ImportData.VertexAttributes.Emplace(MoveTemp(NamedVertexAttribute.AttributeValues), NamedVertexAttribute.ComponentCount);
+				ImportData.VertexAttributeNames.Emplace(MoveTemp(NamedVertexAttribute.AttributeName));
+			}
+		}
+	}
+
 	//
 	// clean up
 	//
diff --git a/Engine/Source/Editor/UnrealEd/Private/Fbx/SFbxSceneOptionWindow.cpp b/Engine/Source/Editor/UnrealEd/Private/Fbx/SFbxSceneOptionWindow.cpp
--- a/Engine/Source/Editor/UnrealEd/Private/Fbx/SFbxSceneOptionWindow.cpp
+++ b/Engine/Source/Editor/UnrealEd/Private/Fbx/SFbxSceneOptionWindow.cpp
@@ -1496,6 +1496,7 @@ void SFbxSceneOptionWindow::CopySkeletalMeshOptionsToFbxOptions(UnFbx::FBXImport
 	ImportSettings->bCreatePhysicsAsset = SkeletalMeshOptions->bCreatePhysicsAsset;
 	ImportSettings->bImportMeshesInBoneHierarchy = SkeletalMeshOptions->bImportMeshesInBoneHierarchy;
 	ImportSettings->bImportMorph = SkeletalMeshOptions->bImportMorphTargets;
+	ImportSettings->bImportVertexAttributes = SkeletalMeshOptions->bImportVertexAttributes;
 	ImportSettings->OverlappingThresholds.ThresholdPosition = SkeletalMeshOptions->ThresholdPosition;
 	ImportSettings->OverlappingThresholds.ThresholdTangentNormal = SkeletalMeshOptions->ThresholdTangentNormal;
 	ImportSettings->OverlappingThresholds.ThresholdUV = SkeletalMeshOptions->ThresholdUV;
@@ -1523,6 +1524,7 @@ void SFbxSceneOptionWindow::CopyFbxOptionsToSkeletalMeshOptions(UnFbx::FBXImport
 	SkeletalMeshOptions->bCreatePhysicsAsset = ImportSettings->bCreatePhysicsAsset;
 	SkeletalMeshOptions->bImportMeshesInBoneHierarchy = ImportSettings->bImportMeshesInBoneHierarchy;
 	SkeletalMeshOptions->bImportMorphTargets = ImportSettings->bImportMorph;
+	SkeletalMeshOptions->bImportVertexAttributes = ImportSettings->bImportVertexAttributes;
 	SkeletalMeshOptions->ThresholdPosition = ImportSettings->OverlappingThresholds.ThresholdPosition;
 	SkeletalMeshOptions->ThresholdTangentNormal = ImportSettings->OverlappingThresholds.ThresholdTangentNormal;
 	SkeletalMeshOptions->ThresholdUV = ImportSettings->OverlappingThresholds.ThresholdUV;
diff --git a/Engine/Source/Editor/UnrealEd/Public/FbxImporter.h b/Engine/Source/Editor/UnrealEd/Public/FbxImporter.h
--- a/Engine/Source/Editor/UnrealEd/Public/FbxImporter.h
+++ b/Engine/Source/Editor/UnrealEd/Public/FbxImporter.h
@@ -175,6 +175,7 @@ struct FBXImportOptions
 	bool bReorderMaterialToFbxOrder;
 	// Skeletal Mesh options
 	bool bImportMorph;
+	bool bImportVertexAttributes;
 	bool bImportAnimations;
 	bool bUpdateSkeletonReferencePose;
 	bool bResample;
diff --git a/Engine/Source/Runtime/Core/Private/Misc/FbxErrors.cpp b/Engine/Source/Runtime/Core/Private/Misc/FbxErrors.cpp
--- a/Engine/Source/Runtime/Core/Private/Misc/FbxErrors.cpp
+++ b/Engine/Source/Runtime/Core/Private/Misc/FbxErrors.cpp
@@ -143,6 +143,8 @@ FLazyName FFbxErrors::SkeletalMesh_FillImportDataFailed(TEXT("SkeletalMesh_FillI
 
 FLazyName FFbxErrors::SkeletalMesh_InvalidPosition(TEXT("SkeletalMesh_InvalidPosition"));
 
+FLazyName FFbxErrors::SkeletalMesh_AttributeComponentCountMismatch(TEXT("SkeletalMesh_AttributeComponentCountMismatch"));
+
 /** Animation **/
 FLazyName FFbxErrors::Animation_CouldNotFindRootTrack(TEXT("Animation_CouldNotFindRootTrack"));
 
diff --git a/Engine/Source/Runtime/Core/Public/Misc/FbxErrors.h b/Engine/Source/Runtime/Core/Public/Misc/FbxErrors.h
--- a/Engine/Source/Runtime/Core/Public/Misc/FbxErrors.h
+++ b/Engine/Source/Runtime/Core/Public/Misc/FbxErrors.h
@@ -153,6 +153,8 @@ struct FFbxErrors
 
 	static CORE_API FLazyName SkeletalMesh_InvalidPosition;
 
+	static CORE_API FLazyName SkeletalMesh_AttributeComponentCountMismatch;
+
 	/** Animation **/
 	static CORE_API FLazyName Animation_CouldNotFindRootTrack;
 
diff --git a/Engine/Source/Runtime/Engine/Public/Rendering/SkeletalMeshLODImporterData.h b/Engine/Source/Runtime/Engine/Public/Rendering/SkeletalMeshLODImporterData.h
--- a/Engine/Source/Runtime/Engine/Public/Rendering/SkeletalMeshLODImporterData.h
+++ b/Engine/Source/Runtime/Engine/Public/Rendering/SkeletalMeshLODImporterData.h
@@ -293,6 +293,13 @@ namespace SkeletalMeshImportData
 
 	struct FVertexAttribute
 	{
+		FVertexAttribute() = default;
+		FVertexAttribute(TArray<float>&& InAttributeValues, int32 InComponentCount) :
+			AttributeValues(InAttributeValues), ComponentCount(InComponentCount)
+		{}
+		FVertexAttribute(const FVertexAttribute&) = default;
+		FVertexAttribute(FVertexAttribute&&) = default;
+		
 		TArray<float> AttributeValues;
 		int32 ComponentCount;
 		
