Fbx: Fix localization tag.

#rnx

diff --git a/Engine/Source/Editor/UnrealEd/Private/Fbx/FbxSkeletalMeshImport.cpp b/Engine/Source/Editor/UnrealEd/Private/Fbx/FbxSkeletalMeshImport.cpp
--- a/Engine/Source/Editor/UnrealEd/Private/Fbx/FbxSkeletalMeshImport.cpp
+++ b/Engine/Source/Editor/UnrealEd/Private/Fbx/FbxSkeletalMeshImport.cpp
@@ -3951,7 +3951,7 @@ bool UnFbx::FFbxImporter::FillSkelMeshImporterFromFbx( FSkeletalMeshImportData&
 				{
 					// For now, if an existing attribute has a non-unity component count, we don't append the new attribute values
 					// on top, just pad with zeroes. 
-					const FText ErrorMsg = FText::Format(LOCTEXT("FbxSkeletaLMeshimport_MaterialIndexInconsistency", "Existing attribute '{0}' has more than one component. Ignoring imported attribute of same name."),
+					const FText ErrorMsg = FText::Format(LOCTEXT("FbxSkeletalMeshimport_AttributeComponentCountMismatch", "Existing attribute '{0}' has more than one component. Ignoring imported attribute of same name."),
 							FText::FromString(ExistingAttributeName)); 
 					AddTokenizedErrorMessage(FTokenizedMessage::Create(EMessageSeverity::Warning, ErrorMsg), FFbxErrors::SkeletalMesh_AttributeComponentCountMismatch);
 					
