Prevent all actor types from displaying in the blueprint reparent dropdown if the blueprint explicitly sets any allowed actor classes.

This fixes nDisplay root actors being able to be reparented to any actor type.

diff --git a/Engine/Source/Editor/UnrealEd/Private/Kismet2/BlueprintEditorUtils.cpp b/Engine/Source/Editor/UnrealEd/Private/Kismet2/BlueprintEditorUtils.cpp
--- a/Engine/Source/Editor/UnrealEd/Private/Kismet2/BlueprintEditorUtils.cpp
+++ b/Engine/Source/Editor/UnrealEd/Private/Kismet2/BlueprintEditorUtils.cpp
@@ -8525,7 +8525,23 @@ TSharedRef<SWidget> FBlueprintEditorUtils::ConstructBlueprintParentClassPicker(
 		else
 		{
 			// Don't allow conversion outside of the Actor hierarchy
-			Filter->AllowedChildrenOfClasses.Add( AActor::StaticClass() );
+			{
+				bool bHasAllowedActorClasses = false;
+				for (const UClass* Class : Filter->AllowedChildrenOfClasses)
+				{
+					if (Class && Class->IsChildOf(AActor::StaticClass()))
+					{
+						// This blueprint has already defined explicit actor classes to use
+						bHasAllowedActorClasses = true;
+						break;
+					}
+				}
+			
+				if (!bHasAllowedActorClasses)
+				{
+					Filter->AllowedChildrenOfClasses.Add( AActor::StaticClass() );
+				}
+			}
 
 			// Don't allow non-LevelScriptActor->LevelScriptActor conversion
 			Filter->DisallowedChildrenOfClasses.Add( ALevelScriptActor::StaticClass() );
