[BugFix] Variable change in BP diff defaults panel doesn't scroll to or highlight property

diff --git a/Engine/Source/Editor/UnrealEd/Private/DiffUtils.cpp b/Engine/Source/Editor/UnrealEd/Private/DiffUtils.cpp
--- a/Engine/Source/Editor/UnrealEd/Private/DiffUtils.cpp
+++ b/Engine/Source/Editor/UnrealEd/Private/DiffUtils.cpp
@@ -467,6 +467,19 @@ FPropertyPath FPropertySoftPath::ResolvePath(const UObject* Object) const
 		FName PropertyIdentifier = PropertyChain[I].PropertyName;
 		FProperty* ResolvedProperty = UEDiffUtils_Private::Resolve(ContainerStruct, PropertyIdentifier);
 
+		// if property wasn't found in ContainerStruct, check for it in SparseClassData
+		if (!ResolvedProperty && RootTypeHint && RootTypeHint != ContainerStruct)
+		{
+			if (const UClass* AsClass = Cast<UClass>(ContainerStruct))
+			{
+				const UScriptStruct* SparseClassDataStruct = AsClass->GetSparseClassDataStruct();
+				if (SparseClassDataStruct && SparseClassDataStruct->IsChildOf(RootTypeHint))
+				{
+					ResolvedProperty = UEDiffUtils_Private::Resolve(RootTypeHint, PropertyIdentifier);
+				}
+			}
+		}
+
 		// If the property didn't exist inside the container, return an invalid property
 		if (!ResolvedProperty)
 		{
