Cooker Logs: Downgrade FixupPackageDependenciesForChunks log to verbose and cull 0-value entries even in verbose, to improve cook log readability.
#rnx

diff --git a/Engine/Source/Editor/UnrealEd/Private/Commandlets/AssetRegistryGenerator.cpp b/Engine/Source/Editor/UnrealEd/Private/Commandlets/AssetRegistryGenerator.cpp
--- a/Engine/Source/Editor/UnrealEd/Private/Commandlets/AssetRegistryGenerator.cpp
+++ b/Engine/Source/Editor/UnrealEd/Private/Commandlets/AssetRegistryGenerator.cpp
@@ -2416,8 +2416,11 @@ void FAssetRegistryGenerator::FixupPackageDependenciesForChunks(UE::Cook::FCookS
 		const int32 ChunkManifestNum = ChunkManifests[PakchunkIndex] ? ChunkManifests[PakchunkIndex]->Num() : 0;
 		check(PakchunkIndex < FinalChunkManifests.Num() && FinalChunkManifests[PakchunkIndex]);
 		const int32 FinalChunkManifestNum = FinalChunkManifests[PakchunkIndex]->Num();
-		UE_LOG(LogAssetRegistryGenerator, Log, TEXT("Chunk: %i, Started with %i packages, Final after dependency resolve: %i"),
-			PakchunkIndex, ChunkManifestNum, FinalChunkManifestNum);
+		if (ChunkManifestNum != 0 || FinalChunkManifestNum != 0)
+		{
+			UE_LOG(LogAssetRegistryGenerator, Verbose, TEXT("Chunk: %i, Started with %i packages, Final after dependency resolve: %i"),
+				PakchunkIndex, ChunkManifestNum, FinalChunkManifestNum);
+		}
 	}
 	
 	// Fix up the data in the FAssetRegistryState to reflect this chunk layout
