MPCook: Add debugging parameter -cookforceremote
#rnx

diff --git a/Engine/Source/Editor/UnrealEd/Private/Cooker/CookDirector.cpp b/Engine/Source/Editor/UnrealEd/Private/Cooker/CookDirector.cpp
--- a/Engine/Source/Editor/UnrealEd/Private/Cooker/CookDirector.cpp
+++ b/Engine/Source/Editor/UnrealEd/Private/Cooker/CookDirector.cpp
@@ -190,6 +190,8 @@ void FCookDirector::ParseConfig(int32 CookProcessCount, bool& bOutValid)
 		}
 	}
 
+	bAllowLocalCooks = !FParse::Param(CommandLine, TEXT("CookForceRemote"));
+
 	int32 MultiprocessId;
 	if (FParse::Value(CommandLine, TEXT("-MultiprocessId="), MultiprocessId))
 	{
@@ -292,7 +294,10 @@ void FCookDirector::AssignRequests(TArrayView<FPackageData*> Requests, TArray<FW
 	LocalRemoteWorkers = CopyRemoteWorkers();;
 
 	WorkerIds.Reserve(LocalRemoteWorkers.Num() + 1);
-	WorkerIds.Add(FWorkerId::Local());
+	if (bAllowLocalCooks)
+	{
+		WorkerIds.Add(FWorkerId::Local());
+	}
 	for (const TRefCountPtr<FCookWorkerServer>& RemoteWorker : LocalRemoteWorkers)
 	{
 		WorkerIds.Add(RemoteWorker->GetWorkerId());
@@ -467,12 +472,18 @@ void FCookDirector::TickFromSchedulerThread()
 		TickCommunication(ECookDirectorThread::SchedulerThread);
 	}
 
-	int32 BusiestNumAssignments = COTFS.NumMultiprocessLocalWorkerAssignments();
-	bool bLocalWorkerIdle = BusiestNumAssignments == 0;
+	int32 BusiestNumAssignments = 0;
+	bool bLocalWorkerIdle = true;
+	bool bAnyIdle = false;
 	float DeltaTime = static_cast<float>(CurrentTime - LastTickTimeSeconds);
-	bool bAnyIdle = bLocalWorkerIdle;
 	LastTickTimeSeconds = CurrentTime;
-	LocalWorkerProfileData->UpdateIdle(bLocalWorkerIdle, DeltaTime);
+	if (bAllowLocalCooks)
+	{
+		BusiestNumAssignments = COTFS.NumMultiprocessLocalWorkerAssignments();
+		bLocalWorkerIdle = BusiestNumAssignments == 0;
+		bAnyIdle = bLocalWorkerIdle;
+		LocalWorkerProfileData->UpdateIdle(bLocalWorkerIdle, DeltaTime);
+	}
 	
 	bool bSendHeartbeat;
 	int32 LocalHeartbeatNumber;
@@ -1464,6 +1475,7 @@ void FCookDirector::FRetractionHandler::Initialize()
 	TArray<FWorkerId> IdleWorkers;
 	int32 BusiestNumAssignments = 0;
 
+	if (Director.bAllowLocalCooks)
 	{
 		int32 NumAssignments = Director.COTFS.NumMultiprocessLocalWorkerAssignments();
 		BusiestWorker = FWorkerId::Local();
@@ -1478,7 +1490,7 @@ void FCookDirector::FRetractionHandler::Initialize()
 	for (const TRefCountPtr<FCookWorkerServer>& RemoteWorker : LocalRemoteWorkers)
 	{
 		int32 NumAssignments = RemoteWorker->NumAssignments();
-		if (NumAssignments > BusiestNumAssignments)
+		if (BusiestWorker.IsInvalid() || NumAssignments > BusiestNumAssignments)
 		{
 			BusiestWorker = RemoteWorker->GetWorkerId();
 			BusiestNumAssignments = NumAssignments;
@@ -1726,7 +1738,7 @@ TArray<FWorkerId> FCookDirector::FRetractionHandler::CalculateWorkersToSplitOver
 	TConstArrayView<TRefCountPtr<FCookWorkerServer>> LocalRemoteWorkers)
 {
 	TArray<TPair<FWorkerId, int32>> WorkerNumPackages;
-	if (FromWorker != FWorkerId::Local())
+	if (FromWorker != FWorkerId::Local() && Director.bAllowLocalCooks)
 	{
 		WorkerNumPackages.Emplace(FWorkerId::Local(), Director.COTFS.NumMultiprocessLocalWorkerAssignments());
 	}
diff --git a/Engine/Source/Editor/UnrealEd/Private/Cooker/CookDirector.h b/Engine/Source/Editor/UnrealEd/Private/Cooker/CookDirector.h
--- a/Engine/Source/Editor/UnrealEd/Private/Cooker/CookDirector.h
+++ b/Engine/Source/Editor/UnrealEd/Private/Cooker/CookDirector.h
@@ -248,6 +248,7 @@ private:
 	int32 CoreLimit = 0;
 	EShowWorker ShowWorkerOption = EShowWorker::CombinedLogs;
 	ELoadBalanceAlgorithm LoadBalanceAlgorithm = ELoadBalanceAlgorithm::CookBurden;
+	bool bAllowLocalCooks = true;
 
 	// Data only accessible from the CommunicationThread (or if the CommunicationThread is inactive)
 	FSocket* WorkerConnectSocket = nullptr;
