MPCook: Do not send per-platform package result messages from CookWorkers for the session platforms that were not requested for a given package in a multiprocess cook. And do not send platform result messages that have already been sent previously.
#rnx

diff --git a/Engine/Source/Editor/UnrealEd/Private/Commandlets/AssetRegistryGenerator.cpp b/Engine/Source/Editor/UnrealEd/Private/Commandlets/AssetRegistryGenerator.cpp
--- a/Engine/Source/Editor/UnrealEd/Private/Commandlets/AssetRegistryGenerator.cpp
+++ b/Engine/Source/Editor/UnrealEd/Private/Commandlets/AssetRegistryGenerator.cpp
@@ -2931,6 +2931,11 @@ void FAssetRegistryMPCollector::ClientTickPackage(FMPCollectorClientTickPackageC
 	bool bLoggedWarning = false;
 	for (const FMPCollectorClientTickPackageContext::FPlatformData& ContextPlatformData : Context.GetPlatformDatas())
 	{
+		if (ContextPlatformData.CookResults == ECookResult::Invalid)
+		{
+			continue;
+		}
+
 		const ITargetPlatform* TargetPlatform = ContextPlatformData.TargetPlatform;
 		FPlatformData* PlatformData = COTFS.PlatformManager->GetPlatformData(TargetPlatform);
 		FAssetRegistryReporterRemote& RegistryReporter = static_cast<FAssetRegistryReporterRemote&>(*PlatformData->RegistryReporter);
diff --git a/Engine/Source/Editor/UnrealEd/Private/Cooker/CookPackageData.h b/Engine/Source/Editor/UnrealEd/Private/Cooker/CookPackageData.h
--- a/Engine/Source/Editor/UnrealEd/Private/Cooker/CookPackageData.h
+++ b/Engine/Source/Editor/UnrealEd/Private/Cooker/CookPackageData.h
@@ -181,7 +181,13 @@ struct FPackagePlatformData
 	ECookResult GetCookResults() const { return (ECookResult)CookResults; }
 	bool IsCookAttempted() const { return CookResults != (uint32)ECookResult::NotAttempted; }
 	bool IsCookSucceeded() const { return CookResults == (uint32)ECookResult::Succeeded; }
-	void SetCookResults(ECookResult Value) { CookResults = (uint32)Value; }
+	void SetCookResults(ECookResult Value)
+	{
+		// ECookResult::Invalid is only used in replication and is not allowed in FPackagePlatformData
+		check(Value != ECookResult::Invalid);
+		CookResults = (uint32)Value;
+		SetReportedToDirector(false);
+	}
 
 	/** Return if we need to call SavePackage on it (reachable, cookable, not yet cooked) */
 	bool NeedsCooking(const ITargetPlatform* PlatformItBelongsTo) const;
@@ -189,6 +195,10 @@ struct FPackagePlatformData
 	bool IsRegisteredForCachedObjectsInOuter() const { return bRegisteredForCachedObjectsInOuter != 0; }
 	void SetRegisteredForCachedObjectsInOuter(bool bValue) { bRegisteredForCachedObjectsInOuter = bValue; }
 
+	/** Only read/written on CookWorkers */
+	bool IsReportedToDirector() { return bReportedToDirector != 0; }
+	void SetReportedToDirector(bool bValue) { bReportedToDirector = (uint32)bValue; }
+
 private:
 	uint32 bReachable : 1;
 	uint32 bVisitedByCluster : 1;
@@ -198,6 +208,7 @@ private:
 	uint32 bExplorableOverride : 1;
 	uint32 bIterativelyUnmodified : 1;
 	uint32 bRegisteredForCachedObjectsInOuter : 1;
+	uint32 bReportedToDirector : 1;
 	uint32 CookResults : (int)ECookResult::NumBits;
 };
 
diff --git a/Engine/Source/Editor/UnrealEd/Private/Cooker/CookTypes.h b/Engine/Source/Editor/UnrealEd/Private/Cooker/CookTypes.h
--- a/Engine/Source/Editor/UnrealEd/Private/Cooker/CookTypes.h
+++ b/Engine/Source/Editor/UnrealEd/Private/Cooker/CookTypes.h
@@ -117,6 +117,8 @@ namespace UE::Cook
 		Failed,
 		/** The package is a NeverCook package that needs to be added to cookresults for dependency tracking. */
 		NeverCookPlaceholder,
+		/** No information for this platform (used in CookWorker replication) */
+		Invalid,
 		Count,
 		NumBits= FPlatformMath::ConstExprCeilLogTwo(ECookResult::Count),
 	};
diff --git a/Engine/Source/Editor/UnrealEd/Private/Cooker/CookWorkerClient.cpp b/Engine/Source/Editor/UnrealEd/Private/Cooker/CookWorkerClient.cpp
--- a/Engine/Source/Editor/UnrealEd/Private/Cooker/CookWorkerClient.cpp
+++ b/Engine/Source/Editor/UnrealEd/Private/Cooker/CookWorkerClient.cpp
@@ -205,7 +205,17 @@ void FCookWorkerClient::ReportPromoteToSaveComplete(FPackageData& PackageData)
 		ITargetPlatform* TargetPlatform = OrderedSessionPlatforms[PlatformIndex];
 		FPackageRemoteResult::FPlatformResult& PlatformResults = Result->GetPlatforms()[PlatformIndex];
 		FPackagePlatformData& PackagePlatformData = PackageData.FindOrAddPlatformData(TargetPlatform);
-		PlatformResults.SetCookResults(PackagePlatformData.GetCookResults());
+		if (!PackagePlatformData.IsCookAttempted() || PackagePlatformData.IsReportedToDirector())
+		{
+			// We didn't attempt to cook this platform for this package, or we cooked it previously and already sent the
+			// information about it
+			PlatformResults.SetCookResults(ECookResult::Invalid);
+		}
+		else
+		{
+			PlatformResults.SetCookResults(PackagePlatformData.GetCookResults());
+			PackagePlatformData.SetReportedToDirector(true);
+		}
 	}
 
 	ReportPackageMessage(PackageName, MoveTemp(ResultOwner));
diff --git a/Engine/Source/Editor/UnrealEd/Private/Cooker/CookWorkerServer.cpp b/Engine/Source/Editor/UnrealEd/Private/Cooker/CookWorkerServer.cpp
--- a/Engine/Source/Editor/UnrealEd/Private/Cooker/CookWorkerServer.cpp
+++ b/Engine/Source/Editor/UnrealEd/Private/Cooker/CookWorkerServer.cpp
@@ -836,7 +836,10 @@ void FCookWorkerServer::RecordResults(FPackageResultsMessage& Message)
 			{
 				ITargetPlatform* TargetPlatform = OrderedSessionPlatforms[PlatformIndex];
 				FPackageRemoteResult::FPlatformResult& PlatformResult = Result.GetPlatforms()[PlatformIndex];
-				PackageData->SetPlatformCooked(TargetPlatform, PlatformResult.GetCookResults());
+				if (PlatformResult.GetCookResults() != ECookResult::Invalid)
+				{
+					PackageData->SetPlatformCooked(TargetPlatform, PlatformResult.GetCookResults());
+				}
 				HandleReceivedPackagePlatformMessages(*PackageData, TargetPlatform, PlatformResult.ReleaseMessages());
 			}
 			COTFS.RecordExternalActorDependencies(Result.GetExternalActorDependencies());
@@ -1315,6 +1318,10 @@ void FPackageWriterMPCollector::ClientTickPackage(FMPCollectorClientTickPackageC
 {
 	for (const FMPCollectorClientTickPackageContext::FPlatformData& PlatformData : Context.GetPlatformDatas())
 	{
+		if (PlatformData.CookResults == ECookResult::Invalid)
+		{
+			continue;
+		}
 		ICookedPackageWriter& PackageWriter = COTFS.FindOrCreatePackageWriter(PlatformData.TargetPlatform);
 		TFuture<FCbObject> ObjectFuture = PackageWriter.WriteMPCookMessageForPackage(Context.GetPackageName());
 		Context.AddAsyncPlatformMessage(PlatformData.TargetPlatform, MoveTemp(ObjectFuture));
