[Localization] Display number of found assets for localization during cook.

diff --git a/Engine/Source/Editor/UnrealEd/Private/CookOnTheFlyServer.cpp b/Engine/Source/Editor/UnrealEd/Private/CookOnTheFlyServer.cpp
--- a/Engine/Source/Editor/UnrealEd/Private/CookOnTheFlyServer.cpp
+++ b/Engine/Source/Editor/UnrealEd/Private/CookOnTheFlyServer.cpp
@@ -11991,6 +11991,8 @@ void UCookOnTheFlyServer::GenerateLocalizationReferences()
 	TArray<FAssetData> AssetDataForCultures;
 	AssetRegistry->GetAssets(Filter, AssetDataForCultures);
 
+	UE_LOG(LogCook, Display, TEXT("Found %d localized assets"), AssetDataForCultures.Num());
+
 	for (const FAssetData& AssetData : AssetDataForCultures)
 	{
 		const FName LocalizedPackageName = AssetData.PackageName;
