Cooker Stats: Add hierarchical cook timer for CookByTheBookFinished. Convert the timing log data in AssetRegistryGenerator to a hierarchical cook timer. Remove some no-longer-necessary marker logs.
#rnx

diff --git a/Engine/Source/Editor/UnrealEd/Classes/CookOnTheSide/CookOnTheFlyServer.h b/Engine/Source/Editor/UnrealEd/Classes/CookOnTheSide/CookOnTheFlyServer.h
--- a/Engine/Source/Editor/UnrealEd/Classes/CookOnTheSide/CookOnTheFlyServer.h
+++ b/Engine/Source/Editor/UnrealEd/Classes/CookOnTheSide/CookOnTheFlyServer.h
@@ -982,6 +982,7 @@ private:
 
 	/** Called at the end of CookByTheBook to write aggregated data such as AssetRegistry and shaders. */
 	UNREALED_API void CookByTheBookFinished();
+	void CookByTheBookFinishedInternal();
 
 	/** Clears session-lifetime data from COTFS and CookPackageDatas. */
 	UNREALED_API void ShutdownCookSession();
diff --git a/Engine/Source/Editor/UnrealEd/Private/Commandlets/AssetRegistryGenerator.cpp b/Engine/Source/Editor/UnrealEd/Private/Commandlets/AssetRegistryGenerator.cpp
--- a/Engine/Source/Editor/UnrealEd/Private/Commandlets/AssetRegistryGenerator.cpp
+++ b/Engine/Source/Editor/UnrealEd/Private/Commandlets/AssetRegistryGenerator.cpp
@@ -12,6 +12,7 @@
 #include "Cooker/CookMPCollector.h"
 #include "Cooker/CookPackageData.h"
 #include "Cooker/CookPlatformManager.h"
+#include "Cooker/CookProfiling.h"
 #include "Cooker/CookSandbox.h"
 #include "Cooker/CookWorkerClient.h"
 #include "CookMetadata.h"
@@ -35,7 +36,6 @@
 #include "Misc/PathViews.h"
 #include "PakFileUtilities.h"
 #include "Policies/PrettyJsonPrintPolicy.h"
-#include "ProfilingDebugging/ScopedTimers.h"
 #include "Serialization/ArrayReader.h"
 #include "Serialization/ArrayWriter.h"
 #include "Serialization/CompactBinarySerialization.h"
@@ -1701,8 +1701,6 @@ bool FAssetRegistryGenerator::SaveAssetRegistry(const FString& SandboxPath, bool
 {
 	LLM_SCOPE_BYTAG(Cooker_GeneratedAssetRegistry);
 	UE_LOG(LogAssetRegistryGenerator, Display, TEXT("Saving asset registry v%d."), FAssetRegistryVersion::Type::LatestVersion);
-	FAutoScopedDurationTimer Timer;
-
 	
 	// Write development first, this will always write
 	FAssetRegistrySerializationOptions DevelopmentSaveOptions;
@@ -1871,8 +1869,6 @@ bool FAssetRegistryGenerator::SaveAssetRegistry(const FString& SandboxPath, bool
 		}
 	}
 	
-	UE_LOG(LogAssetRegistryGenerator, Display, TEXT("Done saving asset registry. Took %.1f seconds."), Timer.GetTime());
-
 	return true;
 }
 
@@ -2357,8 +2353,7 @@ FString FAssetRegistryGenerator::GetTempPackagingDirectoryForPlatform(const FStr
 
 void FAssetRegistryGenerator::FixupPackageDependenciesForChunks(UE::Cook::FCookSandbox& InSandboxFile)
 {
-	UE_LOG(LogAssetRegistryGenerator, Log, TEXT("Starting FixupPackageDependenciesForChunks..."));
-	SCOPE_LOG_TIME_IN_SECONDS(TEXT("... FixupPackageDependenciesForChunks complete."), nullptr);
+	UE_SCOPED_HIERARCHICAL_COOKTIMER(FixupPackageDependenciesForChunks);
 
 	// Clear any existing manifests from the final array
 	FinalChunkManifests.Empty(ChunkManifests.Num());
diff --git a/Engine/Source/Editor/UnrealEd/Private/CookOnTheFlyServer.cpp b/Engine/Source/Editor/UnrealEd/Private/CookOnTheFlyServer.cpp
--- a/Engine/Source/Editor/UnrealEd/Private/CookOnTheFlyServer.cpp
+++ b/Engine/Source/Editor/UnrealEd/Private/CookOnTheFlyServer.cpp
@@ -1296,7 +1296,7 @@ uint32 UCookOnTheFlyServer::TickCookByTheBook(const float TimeSlice, ECookTickFl
 	TickMainCookLoop(StackData);
 
 	CookByTheBookOptions->CookTime += StackData.Timer.GetTickTimeTillNow();
-	// Make sure no UE_SCOPED_HIERARCHICAL_COOKTIMERs are around CookByTheBookFinishes or CancelCookByTheBook, as those functions delete memory for them
+	// Make sure no UE_SCOPED_HIERARCHICAL_COOKTIMERs are around CookByTheBookFinished or CancelCookByTheBook, as those functions delete memory for them
 	if (StackData.bCookCancelled)
 	{
 		CancelCookByTheBook();
@@ -9528,6 +9528,20 @@ void UCookOnTheFlyServer::WriteCookMetadata(const ITargetPlatform* InTargetPlatf
 }
 
 void UCookOnTheFlyServer::CookByTheBookFinished()
+{
+	{
+		// Add a timer around most of CookByTheBookFinished, but the timer can not exist during or after
+		// ShutdownCookSession because it deletes memory for the timers
+		UE_SCOPED_HIERARCHICAL_COOKTIMER(CookByTheBookFinished);
+		CookByTheBookFinishedInternal();
+	}
+
+	ShutdownCookSession();
+	CookByTheBookFinishedEvent.Broadcast();
+	UE_LOG(LogCook, Display, TEXT("Done!"));
+}
+
+void UCookOnTheFlyServer::CookByTheBookFinishedInternal()
 {
 	using namespace UE::Cook;
 
@@ -9781,11 +9795,6 @@ void UCookOnTheFlyServer::CookByTheBookFinished()
 	GenerateCachedEditorThumbnails();
 
 	FinalizePackageStore();
-	ShutdownCookSession();
-
-	CookByTheBookFinishedEvent.Broadcast();
-
-	UE_LOG(LogCook, Display, TEXT("Done!"));
 }
 
 void UCookOnTheFlyServer::ShutdownCookSession()
@@ -10596,7 +10605,6 @@ void UCookOnTheFlyServer::FinalizePackageStore()
 
 		FindOrCreatePackageWriter(TargetPlatform).EndCook(CookInfo);
 	}
-	UE_LOG(LogCook, Display, TEXT("Done finalizing package store(s)"));
 }
 
 void UCookOnTheFlyServer::ClearPackageStoreContexts()
