Cooker Stats: Remove duplicate call to PrintDetailedStats. Remove CookByTheBookTimeSec since its calculation is currently broken and it is unnecessary. Rearrange Cook.Profile stats to separate cook paramaters from timing data.
#rnx

diff --git a/Engine/Source/Editor/UnrealEd/Private/Commandlets/CookCommandlet.cpp b/Engine/Source/Editor/UnrealEd/Private/Commandlets/CookCommandlet.cpp
--- a/Engine/Source/Editor/UnrealEd/Private/Commandlets/CookCommandlet.cpp
+++ b/Engine/Source/Editor/UnrealEd/Private/Commandlets/CookCommandlet.cpp
@@ -248,19 +248,10 @@ bool UCookCommandlet::CookByTheBook(const TArray<ITargetPlatform*>& Platforms)
 	TRACE_CPUPROFILER_EVENT_SCOPE_ON_CHANNEL(CookByTheBook, CookChannel);
 #endif // OUTPUT_COOKTIMING
 
-	COOK_STAT(TOptional<FScopedDurationTimer> CookByTheBookTimer);
-	COOK_STAT(CookByTheBookTimer.Emplace(DetailedCookStats::CookByTheBookTimeSec));
-
 	UCookOnTheFlyServer* CookOnTheFlyServer = NewObject<UCookOnTheFlyServer>();
 	// make sure that the cookonthefly server doesn't get cleaned up while we are garbage collecting below :)
 	UE::Cook::FScopeRootObject S(CookOnTheFlyServer);
 
-	ON_SCOPE_EXIT
-	{
-		COOK_STAT(CookByTheBookTimer.Reset()); // Destruct the timer and write its duration
-		COOK_STAT(CookOnTheFlyServer->PrintDetailedCookStats());
-	};
-
 	UCookerSettings const* CookerSettings = GetDefault<UCookerSettings>();
 	const UProjectPackagingSettings* const PackagingSettings = GetDefault<UProjectPackagingSettings>();
 	ECookInitializationFlags IterateFlags = ECookInitializationFlags::Iterative;
diff --git a/Engine/Source/Editor/UnrealEd/Private/Cooker/CookProfiling.cpp b/Engine/Source/Editor/UnrealEd/Private/Cooker/CookProfiling.cpp
--- a/Engine/Source/Editor/UnrealEd/Private/Cooker/CookProfiling.cpp
+++ b/Engine/Source/Editor/UnrealEd/Private/Cooker/CookProfiling.cpp
@@ -220,7 +220,6 @@ namespace DetailedCookStats
 	double CookStartTime = 0.0;
 	double CookWallTimeSec = 0.0;
 	double StartupWallTimeSec = 0.0;
-	double CookByTheBookTimeSec = 0.0;
 	double StartCookByTheBookTimeSec = 0.0;
 	double TickCookOnTheSideTimeSec = 0.0;
 	double TickCookOnTheSideLoadPackagesTimeSec = 0.0;
@@ -871,30 +870,31 @@ FCookStatsManager::FAutoRegisterCallback RegisterCookStats([](FCookStatsManager:
 	#define ADD_COOK_STAT_FLT(Path, Name) AddStat(StatName, FCookStatsManager::CreateKeyValueArray(TEXT("Path"), TEXT(Path), TEXT(#Name), Name))
 	ADD_COOK_STAT_FLT(" 0", CookWallTimeSec);
 	ADD_COOK_STAT_FLT(" 0. 0", StartupWallTimeSec);
-	ADD_COOK_STAT_FLT(" 0. 1", CookByTheBookTimeSec);
-	ADD_COOK_STAT_FLT(" 0. 1. 0", StartCookByTheBookTimeSec);
-	ADD_COOK_STAT_FLT(" 0. 1. 0. 0", BlockOnAssetRegistryTimeSec);
-	ADD_COOK_STAT_FLT(" 0. 1. 0. 1", GameCookModificationDelegateTimeSec);
-	ADD_COOK_STAT_FLT(" 0. 1. 1", TickCookOnTheSideTimeSec);
-	ADD_COOK_STAT_FLT(" 0. 1. 1. 0", TickCookOnTheSideLoadPackagesTimeSec);
-	ADD_COOK_STAT_FLT(" 0. 1. 1. 1", TickCookOnTheSideSaveCookedPackageTimeSec);
-	ADD_COOK_STAT_FLT(" 0. 1. 1. 1. 0", TickCookOnTheSideResolveRedirectorsTimeSec);
-	ADD_COOK_STAT_FLT(" 0. 1. 1. 2", TickCookOnTheSidePrepareSaveTimeSec);
-	ADD_COOK_STAT_FLT(" 0. 1. 2", TickLoopGCTimeSec);
-	ADD_COOK_STAT_FLT(" 0. 1. 3", TickLoopRecompileShaderRequestsTimeSec);
-	ADD_COOK_STAT_FLT(" 0. 1. 4", TickLoopShaderProcessAsyncResultsTimeSec);
-	ADD_COOK_STAT_FLT(" 0. 1. 5", TickLoopProcessDeferredCommandsTimeSec);
-	ADD_COOK_STAT_FLT(" 0. 1. 6", TickLoopTickCommandletStatsTimeSec);
-	ADD_COOK_STAT_FLT(" 0. 1. 7", TickLoopFlushRenderingCommandsTimeSec);
-	ADD_COOK_STAT_FLT(" 0. 1. 8", TargetPlatforms);
-	ADD_COOK_STAT_FLT(" 0. 1. 9", CookProject);
-	ADD_COOK_STAT_FLT(" 0. 1. 10", CookCultures);
-	ADD_COOK_STAT_FLT(" 0. 1. 11", IsCookAll);
-	ADD_COOK_STAT_FLT(" 0. 1. 12", IsCookOnTheFly);
-	ADD_COOK_STAT_FLT(" 0. 1. 13", IsIterativeCook);
-	ADD_COOK_STAT_FLT(" 0. 1. 14", IsUnversioned);
-	ADD_COOK_STAT_FLT(" 0. 1. 15", CookLabel);
-	ADD_COOK_STAT_FLT(" 0. 1. 16", IsFastCook);
+	ADD_COOK_STAT_FLT(" 0. 1", StartCookByTheBookTimeSec);
+	ADD_COOK_STAT_FLT(" 0. 1. 0", BlockOnAssetRegistryTimeSec);
+	ADD_COOK_STAT_FLT(" 0. 1. 1", GameCookModificationDelegateTimeSec);
+	ADD_COOK_STAT_FLT(" 0. 2", TickCookOnTheSideTimeSec);
+	ADD_COOK_STAT_FLT(" 0. 2. 0", TickCookOnTheSideLoadPackagesTimeSec);
+	ADD_COOK_STAT_FLT(" 0. 2. 1", TickCookOnTheSideSaveCookedPackageTimeSec);
+	ADD_COOK_STAT_FLT(" 0. 2. 1. 0", TickCookOnTheSideResolveRedirectorsTimeSec);
+	ADD_COOK_STAT_FLT(" 0. 2. 2", TickCookOnTheSidePrepareSaveTimeSec);
+	ADD_COOK_STAT_FLT(" 0. 3", TickLoopGCTimeSec);
+	ADD_COOK_STAT_FLT(" 0. 4", TickLoopRecompileShaderRequestsTimeSec);
+	ADD_COOK_STAT_FLT(" 0. 5", TickLoopShaderProcessAsyncResultsTimeSec);
+	ADD_COOK_STAT_FLT(" 0. 6", TickLoopProcessDeferredCommandsTimeSec);
+	ADD_COOK_STAT_FLT(" 0. 7", TickLoopTickCommandletStatsTimeSec);
+	ADD_COOK_STAT_FLT(" 0. 8", TickLoopFlushRenderingCommandsTimeSec);
+	FString CookParameters; // Empty value to write a header with name "CookParameters"
+	ADD_COOK_STAT_FLT(" 1", CookParameters);
+	ADD_COOK_STAT_FLT(" 1. 0", TargetPlatforms);
+	ADD_COOK_STAT_FLT(" 1. 1", CookProject);
+	ADD_COOK_STAT_FLT(" 1. 2", CookCultures);
+	ADD_COOK_STAT_FLT(" 1. 3", IsCookAll);
+	ADD_COOK_STAT_FLT(" 1. 4", IsCookOnTheFly);
+	ADD_COOK_STAT_FLT(" 1. 5", IsIterativeCook);
+	ADD_COOK_STAT_FLT(" 1. 6", IsUnversioned);
+	ADD_COOK_STAT_FLT(" 1. 7", CookLabel);
+	ADD_COOK_STAT_FLT(" 1. 8", IsFastCook);
 		
 	#undef ADD_COOK_STAT_FLT
 });
diff --git a/Engine/Source/Editor/UnrealEd/Private/Cooker/CookProfiling.h b/Engine/Source/Editor/UnrealEd/Private/Cooker/CookProfiling.h
--- a/Engine/Source/Editor/UnrealEd/Private/Cooker/CookProfiling.h
+++ b/Engine/Source/Editor/UnrealEd/Private/Cooker/CookProfiling.h
@@ -134,7 +134,6 @@ extern FString TargetPlatforms;
 extern double CookStartTime;
 extern double CookWallTimeSec;
 extern double StartupWallTimeSec;
-extern double CookByTheBookTimeSec;
 extern double StartCookByTheBookTimeSec;
 extern double TickCookOnTheSideTimeSec;
 extern double TickCookOnTheSideLoadPackagesTimeSec;
