Cooker Logging: Move stat logs into a separate log category so they can be toggled separately.
#rnx

diff --git a/Engine/Source/Editor/UnrealEd/Classes/CookOnTheSide/CookLog.h b/Engine/Source/Editor/UnrealEd/Classes/CookOnTheSide/CookLog.h
--- a/Engine/Source/Editor/UnrealEd/Classes/CookOnTheSide/CookLog.h
+++ b/Engine/Source/Editor/UnrealEd/Classes/CookOnTheSide/CookLog.h
@@ -5,4 +5,5 @@
 #include "Logging/LogMacros.h"
 
 DECLARE_LOG_CATEGORY_EXTERN(LogCook, Log, All);
+DECLARE_LOG_CATEGORY_EXTERN(LogCookStats, Log, All);
 extern FName LogCookName;
diff --git a/Engine/Source/Editor/UnrealEd/Private/CookOnTheFlyServer.cpp b/Engine/Source/Editor/UnrealEd/Private/CookOnTheFlyServer.cpp
--- a/Engine/Source/Editor/UnrealEd/Private/CookOnTheFlyServer.cpp
+++ b/Engine/Source/Editor/UnrealEd/Private/CookOnTheFlyServer.cpp
@@ -154,6 +154,7 @@
 #define LOCTEXT_NAMESPACE "Cooker"
 
 DEFINE_LOG_CATEGORY(LogCook);
+DEFINE_LOG_CATEGORY(LogCookStats);
 DEFINE_LOG_CATEGORY_STATIC(LogCookList, Log, All);
 
 FName LogCookName(TEXT("LogCook"));
@@ -9869,7 +9870,7 @@ void UCookOnTheFlyServer::PrintDetailedCookStats()
 			}
 			DetailedCookStats::CookWallTimeSec = Now - GStartTime;
 			DetailedCookStats::StartupWallTimeSec = DetailedCookStats::CookStartTime - GStartTime;
-			DetailedCookStats::LogCookStats(CurrentCookMode);
+			DetailedCookStats::SendLogCookStats(CurrentCookMode);
 		});
 }
 
diff --git a/Engine/Source/Editor/UnrealEd/Private/Cooker/CookProfiling.cpp b/Engine/Source/Editor/UnrealEd/Private/Cooker/CookProfiling.cpp
--- a/Engine/Source/Editor/UnrealEd/Private/Cooker/CookProfiling.cpp
+++ b/Engine/Source/Editor/UnrealEd/Private/Cooker/CookProfiling.cpp
@@ -177,7 +177,7 @@ void OutputHierarchyTimers(const FHierarchicalTimerInfo* TimerInfo, int32 Depth)
 	static const TCHAR LeftPad[] = TEXT("                                ");
 	const SIZE_T PadOffset = FMath::Max<int>(UE_ARRAY_COUNT(LeftPad) - 1 - Depth * 2, 0);
 
-	UE_LOG(LogCook, Display, TEXT("  %s%s: %.3fs (%u)"), &LeftPad[PadOffset], *TimerName, TimerInfo->Length, TimerInfo->HitCount);
+	UE_LOG(LogCookStats, Display, TEXT("  %s%s: %.3fs (%u)"), &LeftPad[PadOffset], *TimerName, TimerInfo->Length, TimerInfo->HitCount);
 
 	// We need to print in reverse order since the child list begins with the most recently added child
 
@@ -198,7 +198,7 @@ void OutputHierarchyTimers(const FHierarchicalTimerInfo* TimerInfo, int32 Depth)
 
 void OutputHierarchyTimers()
 {
-	UE_LOG(LogCook, Display, TEXT("Hierarchy Timer Information:"));
+	UE_LOG(LogCookStats, Display, TEXT("Hierarchy Timer Information:"));
 
 	OutputHierarchyTimers(&RootTimerInfo, 0);
 }
@@ -899,7 +899,7 @@ FCookStatsManager::FAutoRegisterCallback RegisterCookStats([](FCookStatsManager:
 	#undef ADD_COOK_STAT_FLT
 });
 
-void LogCookStats(ECookMode::Type CookMode)
+void SendLogCookStats(ECookMode::Type CookMode)
 {
 	if (IsCookingInEditor(CookMode))
 	{
@@ -970,29 +970,29 @@ void LogCookStats(ECookMode::Type CookMode)
 
 	FCookStatsManager::LogCookStats(LogStatsFunc);
 
-	UE_LOG(LogCook, Display, TEXT("Misc Cook Stats"));
-	UE_LOG(LogCook, Display, TEXT("==============="));
+	UE_LOG(LogCookStats, Display, TEXT("Misc Cook Stats"));
+	UE_LOG(LogCookStats, Display, TEXT("==============="));
 	for (FString& StatCategory : StatCategories)
 	{
-		UE_LOG(LogCook, Display, TEXT("%s"), *StatCategory);
+		UE_LOG(LogCookStats, Display, TEXT("%s"), *StatCategory);
 		TArray<FCookStatsManager::StringKeyValue>& StatsInCategory = StatsInCategories.FindOrAdd(StatCategory);
 
 		// log each key/value pair, with the equal signs lined up.
 		for (const FCookStatsManager::StringKeyValue& StatKeyValue : StatsInCategory)
 		{
-			UE_LOG(LogCook, Display, TEXT("    %s=%s"), *StatKeyValue.Key, *StatKeyValue.Value);
+			UE_LOG(LogCookStats, Display, TEXT("    %s=%s"), *StatKeyValue.Key, *StatKeyValue.Value);
 		}
 	}
 
 	// DDC Usage stats are custom formatted, and the above code just accumulated them into a TSet. Now log it with our special formatting for readability.
 	if (CookProfileData.Num() > 0)
 	{
-		UE_LOG(LogCook, Display, TEXT(""));
-		UE_LOG(LogCook, Display, TEXT("Cook Profile"));
-		UE_LOG(LogCook, Display, TEXT("============"));
+		UE_LOG(LogCookStats, Display, TEXT(""));
+		UE_LOG(LogCookStats, Display, TEXT("Cook Profile"));
+		UE_LOG(LogCookStats, Display, TEXT("============"));
 		for (const auto& ProfileEntry : CookProfileData)
 		{
-			UE_LOG(LogCook, Display, TEXT("%s.%s=%s"), *ProfileEntry.Path, *ProfileEntry.Key, *ProfileEntry.Value);
+			UE_LOG(LogCookStats, Display, TEXT("%s.%s=%s"), *ProfileEntry.Path, *ProfileEntry.Key, *ProfileEntry.Value);
 		}
 
 		FString CookStatsFileName;
@@ -1025,7 +1025,7 @@ void LogCookStats(ECookMode::Type CookMode)
 			TUniquePtr<FArchive> JsonFile(IFileManager::Get().CreateFileWriter(*CookStatsFileName));
 			if (!JsonFile)
 			{
-				UE_LOG(LogCook, Warning, TEXT("Could not write to CookStatsFile %s."), *CookStatsFileName);
+				UE_LOG(LogCookStats, Warning, TEXT("Could not write to CookStatsFile %s."), *CookStatsFileName);
 			}
 			else
 			{
@@ -1037,12 +1037,12 @@ void LogCookStats(ECookMode::Type CookMode)
 	}
 	if (DDCSummaryStats.Num() > 0)
 	{
-		UE_LOG(LogCook, Display, TEXT(""));
-		UE_LOG(LogCook, Display, TEXT("DDC Summary Stats"));
-		UE_LOG(LogCook, Display, TEXT("================="));
+		UE_LOG(LogCookStats, Display, TEXT(""));
+		UE_LOG(LogCookStats, Display, TEXT("DDC Summary Stats"));
+		UE_LOG(LogCookStats, Display, TEXT("================="));
 		for (const auto& Attr : DDCSummaryStats)
 		{
-			UE_LOG(LogCook, Display, TEXT("%-16s=%10s"), *Attr.Key, *Attr.Value);
+			UE_LOG(LogCookStats, Display, TEXT("%-16s=%10s"), *Attr.Key, *Attr.Value);
 		}
 	}
 
@@ -1052,14 +1052,14 @@ void LogCookStats(ECookMode::Type CookMode)
 	{
 		Algo::SortBy(DDCResourceUsageStats, [](const FDerivedDataCacheResourceStat& Stat) { return Stat.BuildTimeSec + Stat.LoadTimeSec; }, TGreater());
 
-		UE_LOG(LogCook, Display, TEXT(""));
-		UE_LOG(LogCook, Display, TEXT("DDC Resource Stats"));
-		UE_LOG(LogCook, Display, TEXT("======================================================================================================="));
-		UE_LOG(LogCook, Display, TEXT("Asset Type                          Total Time (Sec)  GameThread Time (Sec)  Assets Built  MB Processed"));
-		UE_LOG(LogCook, Display, TEXT("----------------------------------  ----------------  ---------------------  ------------  ------------"));
+		UE_LOG(LogCookStats, Display, TEXT(""));
+		UE_LOG(LogCookStats, Display, TEXT("DDC Resource Stats"));
+		UE_LOG(LogCookStats, Display, TEXT("======================================================================================================="));
+		UE_LOG(LogCookStats, Display, TEXT("Asset Type                          Total Time (Sec)  GameThread Time (Sec)  Assets Built  MB Processed"));
+		UE_LOG(LogCookStats, Display, TEXT("----------------------------------  ----------------  ---------------------  ------------  ------------"));
 		for (const FDerivedDataCacheResourceStat& Stat : DDCResourceUsageStats)
 		{
-			UE_LOG(LogCook, Display, TEXT("%-34s  %16.2f  %21.2f  %12d  %12.2f"),
+			UE_LOG(LogCookStats, Display, TEXT("%-34s  %16.2f  %21.2f  %12d  %12.2f"),
 				*Stat.AssetType, Stat.LoadTimeSec + Stat.BuildTimeSec, Stat.GameThreadTimeSec,
 				Stat.BuildCount, Stat.LoadSizeMB + Stat.BuildSizeMB);
 		}
diff --git a/Engine/Source/Editor/UnrealEd/Private/Cooker/CookProfiling.h b/Engine/Source/Editor/UnrealEd/Private/Cooker/CookProfiling.h
--- a/Engine/Source/Editor/UnrealEd/Private/Cooker/CookProfiling.h
+++ b/Engine/Source/Editor/UnrealEd/Private/Cooker/CookProfiling.h
@@ -162,7 +162,7 @@ extern int32 PeakRequestQueueSize;
 extern int32 PeakLoadQueueSize;
 extern int32 PeakSaveQueueSize;
 
-void LogCookStats(ECookMode::Type CookMode);
+void SendLogCookStats(ECookMode::Type CookMode);
 
 }
 #endif
