Disabled FMergeWithSelfTests since FMergeWithSelfTests::GetTests hitches the editor

diff --git a/Engine/Source/Editor/UnrealEd/Private/MergeUtils.cpp b/Engine/Source/Editor/UnrealEd/Private/MergeUtils.cpp
--- a/Engine/Source/Editor/UnrealEd/Private/MergeUtils.cpp
+++ b/Engine/Source/Editor/UnrealEd/Private/MergeUtils.cpp
@@ -1190,7 +1190,7 @@ EAssetCommandResult MergeUtils::Merge(const FAssetManualMergeArgs& MergeArgs)
 	return NotifyResolution(EAssetMergeResult::Cancelled);
 }
 
-#if WITH_DEV_AUTOMATION_TESTS
+#if false && WITH_DEV_AUTOMATION_TESTS
 #include "Misc/AutomationTest.h"
 namespace UE::MergeUtilsTests
 {
