FBX: Fix tooltip text on vertex attribute import toggle.

#rnx

diff --git a/Engine/Source/Editor/UnrealEd/Classes/Factories/FbxSkeletalMeshImportData.h b/Engine/Source/Editor/UnrealEd/Classes/Factories/FbxSkeletalMeshImportData.h
--- a/Engine/Source/Editor/UnrealEd/Classes/Factories/FbxSkeletalMeshImportData.h
+++ b/Engine/Source/Editor/UnrealEd/Classes/Factories/FbxSkeletalMeshImportData.h
@@ -85,9 +85,8 @@ public:
 	uint32 bImportMorphTargets:1;
 	
 	/** True to import per-vertex attributes from the FBX file */
-	UPROPERTY(EditAnywhere, AdvancedDisplay, config, Category=Mesh, meta=(ImportType="SkeletalMesh|GeoOnly", ToolTip="If enabled, creates named vertex attributes for secondary vertex color data."))
+	UPROPERTY(EditAnywhere, AdvancedDisplay, config, Category=Mesh, meta=(ImportType="SkeletalMesh|GeoOnly", ToolTip="If enabled, creates a named vertex attribute for each single-channel weight map of the imported mesh."))
 	uint32 bImportVertexAttributes:1;
-	
 
 	/** Threshold to compare vertex position equality. */
 	UPROPERTY(EditAnywhere, config, Category="Mesh", meta = (ImportType = "SkeletalMesh|GeoOnly", SubCategory = "Thresholds", NoSpinbox = "true", ClampMin = "0.0"))
