Fix crash when thumbnails are requested for skeletal meshes that are compiling asynchronously

diff --git a/Engine/Source/Editor/UnrealEd/Classes/ThumbnailRendering/AnimSequenceThumbnailRenderer.h b/Engine/Source/Editor/UnrealEd/Classes/ThumbnailRendering/AnimSequenceThumbnailRenderer.h
--- a/Engine/Source/Editor/UnrealEd/Classes/ThumbnailRendering/AnimSequenceThumbnailRenderer.h
+++ b/Engine/Source/Editor/UnrealEd/Classes/ThumbnailRendering/AnimSequenceThumbnailRenderer.h
@@ -22,6 +22,7 @@ class UAnimSequenceThumbnailRenderer : public UDefaultSizedThumbnailRenderer
 
 
 	// Begin UThumbnailRenderer Object
+	UNREALED_API virtual bool CanVisualizeAsset(UObject* Object) override;
 	UNREALED_API virtual void Draw(UObject* Object, int32 X, int32 Y, uint32 Width, uint32 Height, FRenderTarget*, FCanvas* Canvas, bool bAdditionalViewFamily) override;
 	// End UThumbnailRenderer Object
 
diff --git a/Engine/Source/Editor/UnrealEd/Classes/ThumbnailRendering/BlendSpaceThumbnailRenderer.h b/Engine/Source/Editor/UnrealEd/Classes/ThumbnailRendering/BlendSpaceThumbnailRenderer.h
--- a/Engine/Source/Editor/UnrealEd/Classes/ThumbnailRendering/BlendSpaceThumbnailRenderer.h
+++ b/Engine/Source/Editor/UnrealEd/Classes/ThumbnailRendering/BlendSpaceThumbnailRenderer.h
@@ -22,6 +22,7 @@ class UBlendSpaceThumbnailRenderer : public UDefaultSizedThumbnailRenderer
 
 
 	// Begin UThumbnailRenderer Object
+	UNREALED_API virtual bool CanVisualizeAsset(UObject* Object) override;
 	UNREALED_API virtual void Draw(UObject* Object, int32 X, int32 Y, uint32 Width, uint32 Height, FRenderTarget*, FCanvas* Canvas, bool bAdditionalViewFamily) override;
 	// End UThumbnailRenderer Object
 
diff --git a/Engine/Source/Editor/UnrealEd/Classes/ThumbnailRendering/SkeletonThumbnailRenderer.h b/Engine/Source/Editor/UnrealEd/Classes/ThumbnailRendering/SkeletonThumbnailRenderer.h
--- a/Engine/Source/Editor/UnrealEd/Classes/ThumbnailRendering/SkeletonThumbnailRenderer.h
+++ b/Engine/Source/Editor/UnrealEd/Classes/ThumbnailRendering/SkeletonThumbnailRenderer.h
@@ -23,6 +23,7 @@ class USkeletonThumbnailRenderer : public UDefaultSizedThumbnailRenderer
 
 
 	// Begin UThumbnailRenderer Object
+	UNREALED_API virtual bool CanVisualizeAsset(UObject* Object) override;
 	UNREALED_API virtual void Draw(UObject* Object, int32 X, int32 Y, uint32 Width, uint32 Height, FRenderTarget*, FCanvas* Canvas, bool bAdditionalViewFamily) override;
 	UNREALED_API virtual EThumbnailRenderFrequency GetThumbnailRenderFrequency(UObject* Object) const override;
 	// End UThumbnailRenderer Object
diff --git a/Engine/Source/Editor/UnrealEd/Private/ThumbnailRendering/AnimBlueprintThumbnailRenderer.cpp b/Engine/Source/Editor/UnrealEd/Private/ThumbnailRendering/AnimBlueprintThumbnailRenderer.cpp
--- a/Engine/Source/Editor/UnrealEd/Private/ThumbnailRendering/AnimBlueprintThumbnailRenderer.cpp
+++ b/Engine/Source/Editor/UnrealEd/Private/ThumbnailRendering/AnimBlueprintThumbnailRenderer.cpp
@@ -5,6 +5,8 @@
 #include "SceneView.h"
 #include "Misc/App.h"
 #include "Animation/AnimBlueprint.h"
+#include "Animation/Skeleton.h"
+#include "Engine/SkeletalMesh.h"
 
 UAnimBlueprintThumbnailRenderer::UAnimBlueprintThumbnailRenderer(const FObjectInitializer& ObjectInitializer)
 	: Super(ObjectInitializer)
@@ -14,7 +16,25 @@ UAnimBlueprintThumbnailRenderer::UAnimBlueprintThumbnailRenderer(const FObjectIn
 bool UAnimBlueprintThumbnailRenderer::CanVisualizeAsset(UObject* Object)
 {
 	UAnimBlueprint* AnimBlueprint = Cast<UAnimBlueprint>(Object);
-	return AnimBlueprint->BlueprintType != BPTYPE_Interface && !AnimBlueprint->bIsTemplate;
+
+	if (AnimBlueprint)
+	{
+		if (USkeleton* Skeleton = AnimBlueprint->TargetSkeleton)
+		{
+			USkeletalMesh* PreviewSkeletalMesh = Skeleton->GetAssetPreviewMesh(AnimBlueprint);
+			if (PreviewSkeletalMesh == nullptr)
+			{
+				PreviewSkeletalMesh = Skeleton->FindCompatibleMesh();
+			}
+
+			if (PreviewSkeletalMesh && (PreviewSkeletalMesh->IsCompiling() || PreviewSkeletalMesh->GetResourceForRendering() == nullptr))
+			{
+				return false;
+			}
+		}
+	}
+
+	return AnimBlueprint && AnimBlueprint->BlueprintType != BPTYPE_Interface && !AnimBlueprint->bIsTemplate;
 }
 
 void UAnimBlueprintThumbnailRenderer::Draw(UObject* Object, int32 X, int32 Y, uint32 Width, uint32 Height, FRenderTarget* RenderTarget, FCanvas* Canvas, bool bAdditionalViewFamily)
diff --git a/Engine/Source/Editor/UnrealEd/Private/ThumbnailRendering/AnimSequenceThumbnailRenderer.cpp b/Engine/Source/Editor/UnrealEd/Private/ThumbnailRendering/AnimSequenceThumbnailRenderer.cpp
--- a/Engine/Source/Editor/UnrealEd/Private/ThumbnailRendering/AnimSequenceThumbnailRenderer.cpp
+++ b/Engine/Source/Editor/UnrealEd/Private/ThumbnailRendering/AnimSequenceThumbnailRenderer.cpp
@@ -6,6 +6,8 @@
 #include "SceneView.h"
 #include "Animation/AnimSequenceBase.h"
 #include "ThumbnailHelpers.h"
+#include "Animation/Skeleton.h"
+#include "Engine/SkeletalMesh.h"
 
 UAnimSequenceThumbnailRenderer::UAnimSequenceThumbnailRenderer(const FObjectInitializer& ObjectInitializer)
 	: Super(ObjectInitializer)
@@ -13,6 +15,29 @@ UAnimSequenceThumbnailRenderer::UAnimSequenceThumbnailRenderer(const FObjectInit
 	ThumbnailScene = nullptr;
 }
 
+bool UAnimSequenceThumbnailRenderer::CanVisualizeAsset(UObject* Object)
+{
+	UAnimSequenceBase* AnimSequence = Cast<UAnimSequenceBase>(Object);
+	if (AnimSequence != nullptr)
+	{
+		if (USkeleton* Skeleton = AnimSequence->GetSkeleton())
+		{
+			USkeletalMesh* PreviewSkeletalMesh = Skeleton->GetAssetPreviewMesh(AnimSequence);
+			if (PreviewSkeletalMesh == nullptr)
+			{
+				PreviewSkeletalMesh = Skeleton->FindCompatibleMesh();
+			}
+
+			if (PreviewSkeletalMesh && (PreviewSkeletalMesh->IsCompiling() || PreviewSkeletalMesh->GetResourceForRendering() == nullptr))
+			{
+				return false;
+			}
+		}
+	}
+
+	return true;
+}
+
 void UAnimSequenceThumbnailRenderer::Draw(UObject* Object, int32 X, int32 Y, uint32 Width, uint32 Height, FRenderTarget* RenderTarget, FCanvas* Canvas, bool bAdditionalViewFamily)
 {
 	UAnimSequenceBase* Anim = Cast<UAnimSequenceBase>(Object);
diff --git a/Engine/Source/Editor/UnrealEd/Private/ThumbnailRendering/BlendSpaceThumbnailRenderer.cpp b/Engine/Source/Editor/UnrealEd/Private/ThumbnailRendering/BlendSpaceThumbnailRenderer.cpp
--- a/Engine/Source/Editor/UnrealEd/Private/ThumbnailRendering/BlendSpaceThumbnailRenderer.cpp
+++ b/Engine/Source/Editor/UnrealEd/Private/ThumbnailRendering/BlendSpaceThumbnailRenderer.cpp
@@ -5,8 +5,9 @@
 #include "ShowFlags.h"
 #include "SceneView.h"
 #include "ThumbnailHelpers.h"
-
 #include "Animation/BlendSpace.h"
+#include "Animation/Skeleton.h"
+#include "Engine/SkeletalMesh.h"
 
 UBlendSpaceThumbnailRenderer::UBlendSpaceThumbnailRenderer(const FObjectInitializer& ObjectInitializer)
 	: Super(ObjectInitializer)
@@ -14,6 +15,29 @@ UBlendSpaceThumbnailRenderer::UBlendSpaceThumbnailRenderer(const FObjectInitiali
 	ThumbnailScene = nullptr;
 }
 
+bool UBlendSpaceThumbnailRenderer::CanVisualizeAsset(UObject* Object)
+{
+	UBlendSpace* BlendSpace = Cast<UBlendSpace>(Object);
+	if (BlendSpace != nullptr)
+	{
+		if (USkeleton* Skeleton = BlendSpace->GetSkeleton())
+		{
+			USkeletalMesh* PreviewSkeletalMesh = Skeleton->GetAssetPreviewMesh(BlendSpace);
+			if (PreviewSkeletalMesh == nullptr)
+			{
+				PreviewSkeletalMesh = Skeleton->FindCompatibleMesh();
+			}
+
+			if (PreviewSkeletalMesh && (PreviewSkeletalMesh->IsCompiling() || PreviewSkeletalMesh->GetResourceForRendering() == nullptr))
+			{
+				return false;
+			}
+		}
+	}
+
+	return true;
+}
+
 void UBlendSpaceThumbnailRenderer::Draw(UObject* Object, int32 X, int32 Y, uint32 Width, uint32 Height, FRenderTarget* RenderTarget, FCanvas* Canvas, bool bAdditionalViewFamily)
 {
 	UBlendSpace* BlendSpace = Cast<UBlendSpace>(Object);
diff --git a/Engine/Source/Editor/UnrealEd/Private/ThumbnailRendering/SkeletonThumbnailRenderer.cpp b/Engine/Source/Editor/UnrealEd/Private/ThumbnailRendering/SkeletonThumbnailRenderer.cpp
--- a/Engine/Source/Editor/UnrealEd/Private/ThumbnailRendering/SkeletonThumbnailRenderer.cpp
+++ b/Engine/Source/Editor/UnrealEd/Private/ThumbnailRendering/SkeletonThumbnailRenderer.cpp
@@ -15,6 +15,22 @@ USkeletonThumbnailRenderer::USkeletonThumbnailRenderer(const FObjectInitializer&
 {
 }
 
+bool USkeletonThumbnailRenderer::CanVisualizeAsset(UObject* Object)
+{
+	USkeleton* Skeleton = Cast<USkeleton>(Object);
+
+	constexpr bool bFindIfNotSet = true;
+	if (USkeletalMesh* SkeletalMesh = Skeleton->GetPreviewMesh(bFindIfNotSet))
+	{
+		if (SkeletalMesh && (SkeletalMesh->IsCompiling() || SkeletalMesh->GetResourceForRendering() == nullptr))
+		{
+			return false;
+		}
+	}
+
+	return true;
+}
+
 void USkeletonThumbnailRenderer::Draw(UObject* Object, int32 X, int32 Y, uint32 Width, uint32 Height, FRenderTarget* RenderTarget, FCanvas* Canvas, bool bAdditionalViewFamily)
 {
 	USkeleton* Skeleton = Cast<USkeleton>(Object);
