Cooker: IncrementalCook + CookSinglePackageNoRefs: Fix RequestCluster to run IncrementalCook's tests for iteratively unmodified even when cooking transitive references is turned off for the cook.
#rnx

diff --git a/Engine/Source/Editor/UnrealEd/Private/Cooker/CookRequestCluster.cpp b/Engine/Source/Editor/UnrealEd/Private/Cooker/CookRequestCluster.cpp
--- a/Engine/Source/Editor/UnrealEd/Private/Cooker/CookRequestCluster.cpp
+++ b/Engine/Source/Editor/UnrealEd/Private/Cooker/CookRequestCluster.cpp
@@ -475,9 +475,18 @@ void FRequestCluster::PumpExploration(const FCookerTimer& CookerTimer, bool& bOu
 
 	if (!GraphSearch)
 	{
-		GraphSearch.Reset(new FGraphSearch(*this));
+		ETraversalTier TraversalTier = ETraversalTier::None;
+		if (COTFS.IsCookWorkerMode())
+		{
+			TraversalTier = ETraversalTier::None;
+		}
+		else
+		{
+			TraversalTier = bAllowHardDependencies ? ETraversalTier::FollowDependencies : ETraversalTier::FetchEdgeData;
+		}
+		GraphSearch.Reset(new FGraphSearch(*this, TraversalTier));
 
-		if (!bAllowHardDependencies || COTFS.IsCookWorkerMode())
+		if (TraversalTier == ETraversalTier::None)
 		{
 			GraphSearch->VisitWithoutDependencies();
 			GraphSearch.Reset();
@@ -557,12 +566,12 @@ void FRequestCluster::PumpExploration(const FCookerTimer& CookerTimer, bool& bOu
 	bDependenciesComplete = true;
 }
 
-FRequestCluster::FGraphSearch::FGraphSearch(FRequestCluster& InCluster)
+FRequestCluster::FGraphSearch::FGraphSearch(FRequestCluster& InCluster, ETraversalTier InTraversalTier)
 	: Cluster(InCluster)
+	, TraversalTier(InTraversalTier)
 	, AsyncResultsReadyEvent(EEventMode::ManualReset)
 {
 	AsyncResultsReadyEvent->Trigger();
-	bCookAttachmentsEnabled = !Cluster.bFullBuild && Cluster.COTFS.bHybridIterativeEnabled;
 	LastActivityTime = FPlatformTime::Seconds();
 	VertexAllocator.SetMaxBlockSize(1024);
 	VertexAllocator.SetMaxBlockSize(65536);
@@ -599,7 +608,7 @@ void FRequestCluster::FGraphSearch::VisitWithoutDependencies()
 	{
 		FVertexData Vertex;
 		Vertex.PackageData = Pair.Key;
-		VisitVertex(Vertex, true /* bSkipDependencies */);
+		VisitVertex(Vertex);
 	}
 }
 
@@ -814,7 +823,7 @@ void FRequestCluster::FGraphSearch::TickExploration(bool& bOutDone)
 void FRequestCluster::FGraphSearch::UpdateDisplay()
 {
 	constexpr double WarningTimeout = 10.0;
-	if (FPlatformTime::Seconds() > LastActivityTime + WarningTimeout && bCookAttachmentsEnabled)
+	if (FPlatformTime::Seconds() > LastActivityTime + WarningTimeout && Cluster.IsIncrementalCook())
 	{
 		FScopeLock ScopeLock(&Lock);
 		int32 NumVertices = 0;
@@ -831,7 +840,7 @@ void FRequestCluster::FGraphSearch::UpdateDisplay()
 	}
 }
 
-void FRequestCluster::FGraphSearch::VisitVertex(FVertexData& Vertex, bool bSkipDependencies)
+void FRequestCluster::FGraphSearch::VisitVertex(FVertexData& Vertex)
 {
 	// Only called from PumpExploration thread
 
@@ -864,7 +873,9 @@ void FRequestCluster::FGraphSearch::VisitVertex(FVertexData& Vertex, bool bSkipD
 			if (!PlatformData.IsVisitedByCluster())
 			{
 				VisitVertexForPlatform(Vertex, Pair.Key, PlatformData, SuppressCookReason);
-				if (!bSkipDependencies && PlatformData.IsExplorable())
+				if ((TraversalTier >= ETraversalTier::FetchEdgeData) && 
+					(((TraversalTier >= ETraversalTier::FollowDependencies) && PlatformData.IsExplorable())
+						|| Cluster.IsIncrementalCook()))
 				{
 					ExplorePlatforms.Add(Pair.Key);
 				}
@@ -920,15 +931,14 @@ void FRequestCluster::FGraphSearch::VisitVertex(FVertexData& Vertex, bool bSkipD
 		CookerLoadingPlatform->SetCookable(true);
 		CookerLoadingPlatform->SetExplorable(true);
 		CookerLoadingPlatform->SetVisitedByCluster(true);
-		if (!bSkipDependencies)
+		if (TraversalTier >= ETraversalTier::FollowDependencies)
 		{
 			ExplorePlatforms.Add(CookerLoadingPlatformKey);
 		}
 	}
 
-	if (!ExplorePlatforms.IsEmpty())
+	if (!ExplorePlatforms.IsEmpty() && TraversalTier >= ETraversalTier::FetchEdgeData)
 	{
-		check(!bSkipDependencies);
 		QueueEdgesFetch(Vertex, ExplorePlatforms);
 	}
 }
@@ -1054,6 +1064,69 @@ void FRequestCluster::FGraphSearch::ExploreVertexEdges(FVertexData& Vertex)
 	};
 
 	FQueryPlatformData& PlatformAgnosticQueryPlatformData = Vertex.QueryData->Platforms[PlatformAgnosticPlatformIndex];
+
+	auto ProcessPlatformAttachments = [this, PackageName, &PackageData, &PlatformAgnosticQueryPlatformData, &AddPlatformDependencyRange]
+		(int32 PlatformIndex, const ITargetPlatform* TargetPlatform, FFetchPlatformData& FetchPlatformData,
+			FPackagePlatformData& PackagePlatformData, const FCookAttachments& PlatformAttachments, bool bExploreDependencies)
+	{
+		bool bFoundBuildDefinitions = false;
+		ICookedPackageWriter* PackageWriter = FetchPlatformData.Writer;
+
+		if (Cluster.IsIncrementalCook())
+		{
+			bool bIterativelyUnmodified = false;
+			if (IsCookAttachmentsValid(PackageName, PlatformAttachments))
+			{
+				if (IsIterativeEnabled(PackageName, Cluster.COTFS.bHybridIterativeAllowAllClasses))
+				{
+					bIterativelyUnmodified = true;
+					PackagePlatformData.SetIterativelyUnmodified(true);
+				}
+				if (bExploreDependencies)
+				{
+					AddPlatformDependencyRange(PlatformAttachments.BuildDependencies, PlatformIndex,
+						true /* bHardDependency */);
+					if (Cluster.bAllowSoftDependencies)
+					{
+						AddPlatformDependencyRange(PlatformAttachments.RuntimeOnlyDependencies, PlatformIndex,
+							true /* bHardDependency */);
+					}
+				}
+
+				if (Cluster.bPreQueueBuildDefinitions)
+				{
+					bFoundBuildDefinitions = true;
+					Cluster.BuildDefinitions.AddBuildDefinitionList(PackageName, TargetPlatform,
+						PlatformAttachments.BuildDefinitionList);
+				}
+			}
+			bool bShouldIterativelySkip = bIterativelyUnmodified;
+			PackageWriter->UpdatePackageModificationStatus(PackageName, bIterativelyUnmodified, bShouldIterativelySkip);
+			if (bShouldIterativelySkip)
+			{
+				// Call SetPlatformCooked instead of just PackagePlatformData.SetCookResults because we might also need
+				// to set OnFirstCookedPlatformAdded
+				PackageData.SetPlatformCooked(TargetPlatform, ECookResult::Succeeded);
+				if (PlatformIndex == FirstSessionPlatformIndex)
+				{
+					COOK_STAT(++DetailedCookStats::NumPackagesIterativelySkipped);
+				}
+				// Declare the package to the EDLCookInfo verification so we don't warn about missing exports from it
+				UE::SavePackageUtilities::EDLCookInfoAddIterativelySkippedPackage(PackageName);
+			}
+		}
+
+		if (Cluster.bPreQueueBuildDefinitions && !bFoundBuildDefinitions)
+		{
+			if (PlatformAgnosticQueryPlatformData.bActive &&
+				IsCookAttachmentsValid(PackageName, PlatformAgnosticQueryPlatformData.CookAttachments))
+			{
+				Cluster.BuildDefinitions.AddBuildDefinitionList(PackageName, TargetPlatform,
+					PlatformAgnosticQueryPlatformData.CookAttachments.BuildDefinitionList);
+			}
+		}
+	};
+
 	for (int32 PlatformIndex = 0; PlatformIndex < LocalNumFetchPlatforms; ++PlatformIndex)
 	{
 		FQueryPlatformData& QueryPlatformData = Vertex.QueryData->Platforms[PlatformIndex];
@@ -1062,6 +1135,23 @@ void FRequestCluster::FGraphSearch::ExploreVertexEdges(FVertexData& Vertex)
 			continue;
 		}
 
+		FFetchPlatformData& FetchPlatformData = FetchPlatforms[PlatformIndex];
+		const ITargetPlatform* TargetPlatform = FetchPlatformData.Platform;
+		FPackagePlatformData& PackagePlatformData = PackageData.FindOrAddPlatformData(TargetPlatform);
+		if ((TraversalTier < ETraversalTier::FollowDependencies) || !PackagePlatformData.IsExplorable())
+		{
+			// ExploreVertexEdges is responsible for updating package modification status so we might
+			// have been called for this platform even if not explorable. If not explorable, just update
+			// package modification status for the given platform, except for CookerLoadingPlatformIndex which has
+			// no status to update.
+			if (PlatformIndex != CookerLoadingPlatformIndex)
+			{
+				ProcessPlatformAttachments(PlatformIndex, TargetPlatform, FetchPlatformData, PackagePlatformData,
+					QueryPlatformData.CookAttachments, false /* bExploreDependencies */);
+			}
+			continue;
+		}
+
 		if (PlatformIndex == CookerLoadingPlatformIndex)
 		{
 			TArray<FName>& CookerLoadingDependencies(Scratch.CookerLoadingDependencies);
@@ -1087,63 +1177,10 @@ void FRequestCluster::FGraphSearch::ExploreVertexEdges(FVertexData& Vertex)
 		}
 		else
 		{
-			FFetchPlatformData& FetchPlatformData = FetchPlatforms[PlatformIndex];
-			const ITargetPlatform* TargetPlatform = FetchPlatformData.Platform;
-
 			AddPlatformDependencyRange(HardGameDependencies, PlatformIndex, true /* bHardDependency */);
 			AddPlatformDependencyRange(SoftGameDependencies, PlatformIndex, false /* bHardDependency */);
-
-			const FCookAttachments& PlatformAttachments = QueryPlatformData.CookAttachments;
-			bool bFoundBuildDefinitions = false;
-			ICookedPackageWriter* PackageWriter = FetchPlatformData.Writer;
-			if (!Cluster.bFullBuild && Cluster.COTFS.bHybridIterativeEnabled)
-			{
-				bool bIterativelyUnmodified = false;
-				if (IsCookAttachmentsValid(PackageName, PlatformAttachments))
-				{
-					if (IsIterativeEnabled(PackageName, Cluster.COTFS.bHybridIterativeAllowAllClasses))
-					{
-						bIterativelyUnmodified = true;
-						PackageData.FindOrAddPlatformData(TargetPlatform).SetIterativelyUnmodified(true);
-					}
-					AddPlatformDependencyRange(PlatformAttachments.BuildDependencies, PlatformIndex,
-						true /* bHardDependency */);
-					if (Cluster.bAllowSoftDependencies)
-					{
-						AddPlatformDependencyRange(PlatformAttachments.RuntimeOnlyDependencies, PlatformIndex,
-							true /* bHardDependency */);
-					}
-
-					if (Cluster.bPreQueueBuildDefinitions)
-					{
-						bFoundBuildDefinitions = true;
-						Cluster.BuildDefinitions.AddBuildDefinitionList(PackageName, TargetPlatform,
-							PlatformAttachments.BuildDefinitionList);
-					}
-				}
-				bool bShouldIterativelySkip = bIterativelyUnmodified;
-				PackageWriter->UpdatePackageModificationStatus(PackageName, bIterativelyUnmodified, bShouldIterativelySkip);
-				if (bShouldIterativelySkip)
-				{
-					PackageData.SetPlatformCooked(TargetPlatform, ECookResult::Succeeded);
-					if (PlatformIndex == FirstSessionPlatformIndex)
-					{
-						COOK_STAT(++DetailedCookStats::NumPackagesIterativelySkipped);
-					}
-					// Declare the package to the EDLCookInfo verification so we don't warn about missing exports from it
-					UE::SavePackageUtilities::EDLCookInfoAddIterativelySkippedPackage(PackageName);
-				}
-			}
-
-			if (Cluster.bPreQueueBuildDefinitions && !bFoundBuildDefinitions)
-			{
-				if (PlatformAgnosticQueryPlatformData.bActive &&
-					IsCookAttachmentsValid(PackageName, PlatformAgnosticQueryPlatformData.CookAttachments))
-				{
-					Cluster.BuildDefinitions.AddBuildDefinitionList(PackageName, TargetPlatform,
-						PlatformAgnosticQueryPlatformData.CookAttachments.BuildDefinitionList);
-				}
-			}
+			ProcessPlatformAttachments(PlatformIndex, TargetPlatform, FetchPlatformData, PackagePlatformData,
+				QueryPlatformData.CookAttachments, true /* bExploreDependencies  */);
 		}
 		if (DiscoveredDependencies)
 		{
@@ -1411,7 +1448,7 @@ void FRequestCluster::FQueryVertexBatch::Send()
 		}
 		FFetchPlatformData& FetchPlatformData = ThreadSafeOnlyVars.FetchPlatforms[PlatformIndex];
 
-		if (ThreadSafeOnlyVars.bCookAttachmentsEnabled // Only FetchCookAttachments if our cookmode supports it. Otherwise keep them all empty
+		if (ThreadSafeOnlyVars.Cluster.IsIncrementalCook() // Only FetchCookAttachments if our cookmode supports it. Otherwise keep them all empty
 			&& !FetchPlatformData.bIsCookerLoadingPlatform // The CookerLoadingPlatform has no stored CookAttachments; always use empty
 			)
 		{
@@ -1474,6 +1511,11 @@ TMap<FPackageData*, TArray<FPackageData*>>& FRequestCluster::FGraphSearch::GetGr
 	return GraphEdges;
 }
 
+bool FRequestCluster::IsIncrementalCook() const
+{
+	return !bFullBuild && COTFS.bHybridIterativeEnabled;
+}
+
 void FRequestCluster::IsRequestCookable(const ITargetPlatform* Platform, FPackageData& PackageData,
 	UCookOnTheFlyServer& COTFS, ESuppressCookReason& OutReason, bool& bOutCookable, bool& bOutExplorable)
 {
diff --git a/Engine/Source/Editor/UnrealEd/Private/Cooker/CookRequestCluster.h b/Engine/Source/Editor/UnrealEd/Private/Cooker/CookRequestCluster.h
--- a/Engine/Source/Editor/UnrealEd/Private/Cooker/CookRequestCluster.h
+++ b/Engine/Source/Editor/UnrealEd/Private/Cooker/CookRequestCluster.h
@@ -178,6 +178,21 @@ private:
 	static constexpr int32 CookerLoadingPlatformIndex = 1;
 	static constexpr int32 FirstSessionPlatformIndex = 2;
 
+	/** How much traversal the GraphSearch should do based on settings for the entire cook. */
+	enum class ETraversalTier
+	{
+		/** Do not fetch any edgedata. Used on CookWorkers; the director already did the fetch. */
+		None,
+		/**
+		 * Fetch the edgedata and use it for ancillary calculation like updating whether a package is iteratively
+		 * unmodified. Do not explore the discovered dependencies.
+		 */
+		FetchEdgeData,
+		/** Fetch the edgedata, update ancillary calculations, and explore the discovered dependencies. */
+		FollowDependencies,
+		All=FollowDependencies,
+	};
+
 	/**
 	 * Variables and functions that are only used during PumpExploration. PumpExploration executes a graph search
 	 * over the graph of packages (vertices) and their hard/soft dependencies upon other packages (edges). 
@@ -188,8 +203,7 @@ private:
 	struct FGraphSearch
 	{
 	public:
-		/** Data used during the reentrant dependencies and topological sort operation of FRequestCluster. */
-		FGraphSearch(FRequestCluster& InCluster);
+		FGraphSearch(FRequestCluster& InCluster, ETraversalTier TraversalTier);
 		~FGraphSearch();
 
 		// All public functions are callable only from the process thread
@@ -240,7 +254,7 @@ private:
 		/** Asynchronously fetch the dependencies and previous incremental results for a vertex */
 		void QueueEdgesFetch(FVertexData& Vertex, TConstArrayView<const ITargetPlatform*> Platforms);
 		/** Calculate and store the vertex's PackageData's cookability for each reachable platform. Kick off edges fetch. */
-		void VisitVertex(FVertexData& VertexData, bool bSkipDependencies=false);
+		void VisitVertex(FVertexData& VertexData);
 		/** Calculate and store the vertex's PackageData's cookability for the platform. */
 		void VisitVertexForPlatform(FVertexData& VertexData, const ITargetPlatform* Platform,
 			FPackagePlatformData& PlatformData, ESuppressCookReason& AccumulatedSuppressCookReason);
@@ -283,7 +297,7 @@ private:
 		// Variables that are read-only during multithreading
 		TArray<FFetchPlatformData> FetchPlatforms;
 		FRequestCluster& Cluster;
-		bool bCookAttachmentsEnabled = false;
+		ETraversalTier TraversalTier = ETraversalTier::All;
 
 		// Variables that are accessible only from the Process thread
 		/** Scratch-space variables that are reused to avoid allocations. */
@@ -315,6 +329,7 @@ private:
 	void FetchPackageNames(const FCookerTimer& CookerTimer, bool& bOutComplete);
 	void PumpExploration(const FCookerTimer& CookerTimer, bool& bOutComplete);
 	void StartAsync(const FCookerTimer& CookerTimer, bool& bOutComplete);
+	bool IsIncrementalCook() const;
 	void IsRequestCookable(const ITargetPlatform* TargetPlatform, FName PackageName, FPackageData& PackageData,
 		ESuppressCookReason& OutReason, bool& bOutCookable, bool& bOutExplorable);
 	static void IsRequestCookable(const ITargetPlatform* TargetPlatform, FName PackageName, FPackageData& PackageData,
