Avoid putting book keeping renames into the transaction buffer, this prevents capturing of outdated CDOs in the transaction buffer, now that CDOs are always marked as transient by the blueprint compiler

diff --git a/Engine/Source/Editor/UnrealEd/Private/Kismet2/KismetReinstanceUtilities.cpp b/Engine/Source/Editor/UnrealEd/Private/Kismet2/KismetReinstanceUtilities.cpp
--- a/Engine/Source/Editor/UnrealEd/Private/Kismet2/KismetReinstanceUtilities.cpp
+++ b/Engine/Source/Editor/UnrealEd/Private/Kismet2/KismetReinstanceUtilities.cpp
@@ -1813,7 +1813,7 @@ UClass* FBlueprintCompileReinstancer::MoveCDOToNewClass(UClass* OwnerClass, cons
 		OriginalNames.Add(OwnedObject->GetFName());
 		if(OwnedObject->HasAnyFlags(RF_ArchetypeObject))
 		{
-			OwnedObject->Rename(nullptr, GetTransientPackage(), REN_DoNotDirty | REN_DontCreateRedirectors | REN_ForceNoResetLoaders);
+			OwnedObject->Rename(nullptr, GetTransientPackage(), REN_DoNotDirty | REN_DontCreateRedirectors | REN_ForceNoResetLoaders | REN_NonTransactional);
 		}
 	}
 
@@ -1859,7 +1859,7 @@ UClass* FBlueprintCompileReinstancer::MoveCDOToNewClass(UClass* OwnerClass, cons
 		UObject* OwnedArchetype = OwnedObjects[I];
 		if(OwnedArchetype->HasAnyFlags(RF_ArchetypeObject))
 		{
-			OwnedArchetype->Rename(*OriginalNames[I].ToString(), OwnerClass, REN_DoNotDirty | REN_DontCreateRedirectors | REN_ForceNoResetLoaders);
+			OwnedArchetype->Rename(*OriginalNames[I].ToString(), OwnerClass, REN_DoNotDirty | REN_DontCreateRedirectors | REN_ForceNoResetLoaders | REN_NonTransactional);
 		}
 	}
 
@@ -1882,7 +1882,7 @@ UClass* FBlueprintCompileReinstancer::MoveCDOToNewClass(UClass* OwnerClass, cons
 		if(bAvoidCDODuplication)
 		{
 			OwnerClass->ClassDefaultObject = nullptr;
-			OldCDO->Rename(nullptr, CopyOfOwnerClass->GetOuter(), REN_DoNotDirty | REN_DontCreateRedirectors | REN_ForceNoResetLoaders);
+			OldCDO->Rename(nullptr, CopyOfOwnerClass->GetOuter(), REN_DoNotDirty | REN_DontCreateRedirectors | REN_ForceNoResetLoaders | REN_NonTransactional);
 			CopyOfOwnerClass->ClassDefaultObject = OldCDO;
 		}
 		OldCDO->SetClass(CopyOfOwnerClass);
