Changing list command line parameter to CookList to be consistent with existing CookList parameter which will terminate and not run a cook after listing out all packages to be cooked.

diff --git a/Engine/Source/Editor/UnrealEd/Private/CookOnTheFlyServer.cpp b/Engine/Source/Editor/UnrealEd/Private/CookOnTheFlyServer.cpp
--- a/Engine/Source/Editor/UnrealEd/Private/CookOnTheFlyServer.cpp
+++ b/Engine/Source/Editor/UnrealEd/Private/CookOnTheFlyServer.cpp
@@ -11232,7 +11232,7 @@ void UCookOnTheFlyServer::GenerateInitialRequests(FBeginCookContext& BeginContex
 		}
 	}
 
-	if (FParse::Param(FCommandLine::Get(), TEXT("List")))
+	if (FParse::Param(FCommandLine::Get(), TEXT("CookList")))
 	{
 		WorkerRequests->LogAllRequestedFiles();
 	}
