Adding a simple pointer check for additional security against a crash when pasting nodes to older engine versions.

diff --git a/Engine/Source/Editor/UnrealEd/Private/MaterialGraphNode.cpp b/Engine/Source/Editor/UnrealEd/Private/MaterialGraphNode.cpp
--- a/Engine/Source/Editor/UnrealEd/Private/MaterialGraphNode.cpp
+++ b/Engine/Source/Editor/UnrealEd/Private/MaterialGraphNode.cpp
@@ -142,7 +142,7 @@ bool UMaterialGraphNode::CanPasteHere(const UEdGraph* TargetGraph) const
 {
 	if (Super::CanPasteHere(TargetGraph) && MaterialExpression)
 	{
-		if (MaterialExpression->IsA(UMaterialExpressionStrataLegacyConversion::StaticClass()))
+		if (MaterialExpression && MaterialExpression->IsA(UMaterialExpressionStrataLegacyConversion::StaticClass()))
 		{
 			// We could have used CanDuplicateNode() returning false to prevent the copy but it is nicer to have a notification about why the copy is not happening.
 			FText NotificationText = NSLOCTEXT("UMaterialGraphNode", "SkippedSubstrateLegacyConversion", "SubstrateLegacyConversion node cannot be copied! It is only used to convert legacy material to Substrate.");
