Import: Keep sections separate
-	Adds an option to not merge skeletal mesh sections sharing the same material during FBX import. This allows you to control visibility at a per-section level without having to give each section a unique material. [Github Pull Request: 9587]

diff --git a/Engine/Config/BaseEditorPerProjectUserSettings.ini b/Engine/Config/BaseEditorPerProjectUserSettings.ini
--- a/Engine/Config/BaseEditorPerProjectUserSettings.ini
+++ b/Engine/Config/BaseEditorPerProjectUserSettings.ini
@@ -670,6 +670,7 @@ VertexOverrideColor=(R=0,G=0,B=0,A=0)
 bUpdateSkeletonReferencePose=False
 bUseT0AsRefPose=False
 bPreserveSmoothingGroups=True
+bKeepSectionsSeparate=False
 bImportMeshesInBoneHierarchy=True
 bImportMorphTargets=False
 ThresholdPosition=0.00002
diff --git a/Engine/Plugins/Tests/EditorTests/Source/EditorTests/Private/UnrealEd/FbxAutomationTests.cpp b/Engine/Plugins/Tests/EditorTests/Source/EditorTests/Private/UnrealEd/FbxAutomationTests.cpp
--- a/Engine/Plugins/Tests/EditorTests/Source/EditorTests/Private/UnrealEd/FbxAutomationTests.cpp
+++ b/Engine/Plugins/Tests/EditorTests/Source/EditorTests/Private/UnrealEd/FbxAutomationTests.cpp
@@ -420,6 +420,7 @@ bool FFbxImportAssetsAutomationTest::RunTest(const FString& Parameters)
 					ImportData->MorphThresholdPosition = TestPlan->ImportUI->SkeletalMeshImportData->MorphThresholdPosition;
 					ImportData->ThresholdUV = TestPlan->ImportUI->SkeletalMeshImportData->ThresholdUV;
 					ImportData->bPreserveSmoothingGroups = TestPlan->ImportUI->SkeletalMeshImportData->bPreserveSmoothingGroups;
+					ImportData->bKeepSectionsSeparate = TestPlan->ImportUI->SkeletalMeshImportData->bKeepSectionsSeparate;
 					ImportData->bUpdateSkeletonReferencePose = TestPlan->ImportUI->SkeletalMeshImportData->bUpdateSkeletonReferencePose;
 					ImportData->bUseT0AsRefPose = TestPlan->ImportUI->SkeletalMeshImportData->bUseT0AsRefPose;
 					//Copy UFbxMeshImportData
@@ -517,6 +518,7 @@ bool FFbxImportAssetsAutomationTest::RunTest(const FString& Parameters)
 					ImportData->ThresholdUV = TestPlan->ImportUI->SkeletalMeshImportData->ThresholdUV;
 					ImportData->MorphThresholdPosition = TestPlan->ImportUI->SkeletalMeshImportData->MorphThresholdPosition;
 					ImportData->bPreserveSmoothingGroups = TestPlan->ImportUI->SkeletalMeshImportData->bPreserveSmoothingGroups;
+					ImportData->bKeepSectionsSeparate = TestPlan->ImportUI->SkeletalMeshImportData->bKeepSectionsSeparate;
 					ImportData->bUpdateSkeletonReferencePose = TestPlan->ImportUI->SkeletalMeshImportData->bUpdateSkeletonReferencePose;
 					ImportData->bUseT0AsRefPose = TestPlan->ImportUI->SkeletalMeshImportData->bUseT0AsRefPose;
 					//Copy UFbxMeshImportData
diff --git a/Engine/Source/Developer/MeshUtilities/Private/MeshUtilities.cpp b/Engine/Source/Developer/MeshUtilities/Private/MeshUtilities.cpp
--- a/Engine/Source/Developer/MeshUtilities/Private/MeshUtilities.cpp
+++ b/Engine/Source/Developer/MeshUtilities/Private/MeshUtilities.cpp
@@ -6020,6 +6020,7 @@ void FMeshUtilities::CreateImportDataFromLODModel(USkeletalMesh* SkeletalMesh) c
 		ImportData.bHasTangents = true;
 		ImportData.bUseT0AsRefPose = false;
 		ImportData.bHasVertexColors = SkeletalMesh->GetHasVertexColors();
+		ImportData.bKeepSectionsSeparate = false;
 
 		TArray<FSkeletalMaterial>& SKMaterials = SkeletalMesh->GetMaterials();
 		ImportData.Materials.Reserve(SkeletalMesh->GetMaterials().Num());
diff --git a/Engine/Source/Editor/UnrealEd/Classes/Factories/FbxSceneImportOptionsSkeletalMesh.h b/Engine/Source/Editor/UnrealEd/Classes/Factories/FbxSceneImportOptionsSkeletalMesh.h
--- a/Engine/Source/Editor/UnrealEd/Classes/Factories/FbxSceneImportOptionsSkeletalMesh.h
+++ b/Engine/Source/Editor/UnrealEd/Classes/Factories/FbxSceneImportOptionsSkeletalMesh.h
@@ -34,6 +34,10 @@ class UFbxSceneImportOptionsSkeletalMesh : public UObject
 	UPROPERTY(EditAnywhere, config, Category = SkeletalMesh)
 	uint32 bPreserveSmoothingGroups : 1;
 
+	/** If checked, sections with matching materials are kept separate and will not get combined. */
+	UPROPERTY(EditAnywhere, config, Category = SkeletalMesh)
+	uint32 bKeepSectionsSeparate : 1;
+
 	/** If checked, meshes nested in bone hierarchies will be imported instead of being converted to bones. */
 	UPROPERTY(EditAnywhere, config, Category = SkeletalMesh)
 	uint32 bImportMeshesInBoneHierarchy : 1;
diff --git a/Engine/Source/Editor/UnrealEd/Classes/Factories/FbxSkeletalMeshImportData.h b/Engine/Source/Editor/UnrealEd/Classes/Factories/FbxSkeletalMeshImportData.h
--- a/Engine/Source/Editor/UnrealEd/Classes/Factories/FbxSkeletalMeshImportData.h
+++ b/Engine/Source/Editor/UnrealEd/Classes/Factories/FbxSkeletalMeshImportData.h
@@ -76,6 +76,10 @@ public:
 	UPROPERTY(EditAnywhere, AdvancedDisplay, config, Category=Mesh, meta=(ImportType="SkeletalMesh|GeoOnly"))
 	uint32 bPreserveSmoothingGroups:1;
 
+	/** If checked, sections with matching materials are kept separate and will not get combined. */
+	UPROPERTY(EditAnywhere, AdvancedDisplay, config, Category=Mesh, meta=(ImportType="SkeletalMesh|GeoOnly"))
+	uint32 bKeepSectionsSeparate:1;
+
 	/** If checked, meshes nested in bone hierarchies will be imported instead of being converted to bones. */
 	UPROPERTY(EditAnywhere, AdvancedDisplay, config, Category=Mesh, meta=(ImportType="SkeletalMesh"))
 	uint32 bImportMeshesInBoneHierarchy:1;
diff --git a/Engine/Source/Editor/UnrealEd/Private/Fbx/FbxMainImport.cpp b/Engine/Source/Editor/UnrealEd/Private/Fbx/FbxMainImport.cpp
--- a/Engine/Source/Editor/UnrealEd/Private/Fbx/FbxMainImport.cpp
+++ b/Engine/Source/Editor/UnrealEd/Private/Fbx/FbxMainImport.cpp
@@ -465,6 +465,7 @@ void ApplyImportUIToImportOptions(UFbxImportUI* ImportUI, FBXImportOptions& InOu
 		InOutImportOptions.bImportRigidMesh				= ImportUI->OriginalImportType == FBXIT_StaticMesh && ImportUI->MeshTypeToImport == FBXIT_SkeletalMesh;
 		InOutImportOptions.bUseT0AsRefPose				= ImportUI->SkeletalMeshImportData->bUseT0AsRefPose;
 		InOutImportOptions.bPreserveSmoothingGroups		= ImportUI->SkeletalMeshImportData->bPreserveSmoothingGroups;
+		InOutImportOptions.bKeepSectionsSeparate		= ImportUI->SkeletalMeshImportData->bKeepSectionsSeparate;
 		InOutImportOptions.OverlappingThresholds.ThresholdPosition = ImportUI->SkeletalMeshImportData->ThresholdPosition;
 		InOutImportOptions.OverlappingThresholds.ThresholdTangentNormal = ImportUI->SkeletalMeshImportData->ThresholdTangentNormal;
 		InOutImportOptions.OverlappingThresholds.ThresholdUV = ImportUI->SkeletalMeshImportData->ThresholdUV;
@@ -1742,6 +1743,7 @@ bool FFbxImporter::ImportFromFile(const FString& Filename, const FString& Type,
 							Attribs.Add(FAnalyticsEventAttribute(TEXT("SkeletalMeshOpt ImportMorph"), CaptureImportOptions->bImportMorph));
 							Attribs.Add(FAnalyticsEventAttribute(TEXT("SkeletalMeshOpt ImportSkeletalMeshLODs"), CaptureImportOptions->bImportSkeletalMeshLODs));
 							Attribs.Add(FAnalyticsEventAttribute(TEXT("SkeletalMeshOpt PreserveSmoothingGroups"), CaptureImportOptions->bPreserveSmoothingGroups));
+							Attribs.Add(FAnalyticsEventAttribute(TEXT("SkeletalMeshOpt KeepSectionsSeparate"), CaptureImportOptions->bKeepSectionsSeparate));
 							Attribs.Add(FAnalyticsEventAttribute(TEXT("SkeletalMeshOpt UpdateSkeletonReferencePose"), CaptureImportOptions->bUpdateSkeletonReferencePose));
 							Attribs.Add(FAnalyticsEventAttribute(TEXT("SkeletalMeshOpt UseT0AsRefPose"), CaptureImportOptions->bUseT0AsRefPose));
 							Attribs.Add(FAnalyticsEventAttribute(TEXT("SkeletalMeshOpt OverlappingThresholds.ThresholdPosition"), CaptureImportOptions->OverlappingThresholds.ThresholdPosition));
diff --git a/Engine/Source/Editor/UnrealEd/Private/Fbx/FbxSceneImportData.cpp b/Engine/Source/Editor/UnrealEd/Private/Fbx/FbxSceneImportData.cpp
--- a/Engine/Source/Editor/UnrealEd/Private/Fbx/FbxSceneImportData.cpp
+++ b/Engine/Source/Editor/UnrealEd/Private/Fbx/FbxSceneImportData.cpp
@@ -119,6 +119,7 @@ UnFbx::FBXImportOptions *JSONToFbxOption(TSharedPtr<FJsonValue> OptionJsonValue,
 	//OptionObj->TryGetBoolField("bUseT0AsRefPose", Option->bUseT0AsRefPose);
 	Option->bUseT0AsRefPose = false;
 	OptionObj->TryGetBoolField("bPreserveSmoothingGroups", Option->bPreserveSmoothingGroups);
+	OptionObj->TryGetBoolField("bKeepSectionsSeparate", Option->bKeepSectionsSeparate);
 	OptionObj->TryGetBoolField("bImportMeshesInBoneHierarchy", Option->bImportMeshesInBoneHierarchy);
 	OptionObj->TryGetBoolField("bImportMorphTargets", Option->bImportMorph);
 	OptionObj->TryGetBoolField("bImportVertexAttributes", Option->bImportVertexAttributes);
@@ -204,11 +205,11 @@ FString FbxOptionToJSON(FString OptionName, UnFbx::FBXImportOptions *Option)
 		Option->bImportStaticMeshLODs ? 1 : 0
 		);
 
-	//Skeletal mesh options
-	JsonString += FString::Printf(TEXT("\"bUpdateSkeletonReferencePose\" : \"%d\", \"bUseT0AsRefPose\" : \"%d\", \"bPreserveSmoothingGroups\" : \"%d\", \"bImportMeshesInBoneHierarchy\" : \"%d\", \"bImportMorphTargets\" : \"%d\", \"bImportVertexAttributes\" : \"%d\", \"OverlappingThresholds\" : {\"ThresholdPosition\" : \"%f\", \"ThresholdTangentNormal\" : \"%f\", \"ThresholdUV\" : \"%f\", \"MorphThresholdPosition\" : \"%f\"},"),
+	JsonString += FString::Printf(TEXT("\"bUpdateSkeletonReferencePose\" : \"%d\", \"bUseT0AsRefPose\" : \"%d\", \"bPreserveSmoothingGroups\" : \"%d\", \"bKeepSectionsSeparate\" : \"%d\", \"bImportMeshesInBoneHierarchy\" : \"%d\", \"bImportMorphTargets\" : \"%d\", , \"bImportVertexAttributes\" : \"%d\", \"OverlappingThresholds\" : {\"ThresholdPosition\" : \"%f\", \"ThresholdTangentNormal\" : \"%f\", \"ThresholdUV\" : \"%f\", \"MorphThresholdPosition\" : \"%f\"},"),
 		Option->bUpdateSkeletonReferencePose ? 1 : 0,
 		Option->bUseT0AsRefPose ? 1 : 0,
 		Option->bPreserveSmoothingGroups ? 1 : 0,
+		Option->bKeepSectionsSeparate ? 1 : 0,
 		Option->bImportMeshesInBoneHierarchy ? 1 : 0,
 		Option->bImportMorph ? 1 : 0,
 		Option->bImportVertexAttributes ? 1 : 0,
diff --git a/Engine/Source/Editor/UnrealEd/Private/Fbx/FbxSceneImportFactory.cpp b/Engine/Source/Editor/UnrealEd/Private/Fbx/FbxSceneImportFactory.cpp
--- a/Engine/Source/Editor/UnrealEd/Private/Fbx/FbxSceneImportFactory.cpp
+++ b/Engine/Source/Editor/UnrealEd/Private/Fbx/FbxSceneImportFactory.cpp
@@ -123,6 +123,8 @@ bool GetFbxSceneImportOptions(UnFbx::FFbxImporter* FbxImporter
 	GlobalImportSettings->bImportTextures = true;
 	//Make sure Material get imported
 	GlobalImportSettings->bImportMaterials = true;
+	//Make sure skeletal mesh sections with the same material are combined (and not kept separate)
+	GlobalImportSettings->bKeepSectionsSeparate = false;
 	//TODO support T0AsRefPose
 	GlobalImportSettings->bUseT0AsRefPose = false;
 
diff --git a/Engine/Source/Editor/UnrealEd/Private/Fbx/FbxSceneImportOptionsSkeletalMesh.cpp b/Engine/Source/Editor/UnrealEd/Private/Fbx/FbxSceneImportOptionsSkeletalMesh.cpp
--- a/Engine/Source/Editor/UnrealEd/Private/Fbx/FbxSceneImportOptionsSkeletalMesh.cpp
+++ b/Engine/Source/Editor/UnrealEd/Private/Fbx/FbxSceneImportOptionsSkeletalMesh.cpp
@@ -16,6 +16,7 @@ UFbxSceneImportOptionsSkeletalMesh::UFbxSceneImportOptionsSkeletalMesh(const FOb
 	, bCreatePhysicsAsset(false)
 	, bUseT0AsRefPose(false)
 	, bPreserveSmoothingGroups(false)
+	, bKeepSectionsSeparate(false)
 	, bImportMeshesInBoneHierarchy(true)
 	, bImportMorphTargets(false)
 	, bImportVertexAttributes(false)
@@ -48,6 +49,7 @@ void UFbxSceneImportOptionsSkeletalMesh::FillSkeletalMeshInmportData(UFbxSkeleta
 	SkeletalMeshImportData->ThresholdUV = ThresholdUV;
 	SkeletalMeshImportData->MorphThresholdPosition = MorphThresholdPosition;
 	SkeletalMeshImportData->bPreserveSmoothingGroups = bPreserveSmoothingGroups;
+	SkeletalMeshImportData->bKeepSectionsSeparate = bKeepSectionsSeparate;
 	SkeletalMeshImportData->bUpdateSkeletonReferencePose = bUpdateSkeletonReferencePose;
 	SkeletalMeshImportData->bUseT0AsRefPose = bUseT0AsRefPose;
 
diff --git a/Engine/Source/Editor/UnrealEd/Private/Fbx/FbxSkeletalMeshImport.cpp b/Engine/Source/Editor/UnrealEd/Private/Fbx/FbxSkeletalMeshImport.cpp
--- a/Engine/Source/Editor/UnrealEd/Private/Fbx/FbxSkeletalMeshImport.cpp
+++ b/Engine/Source/Editor/UnrealEd/Private/Fbx/FbxSkeletalMeshImport.cpp
@@ -271,6 +271,16 @@ void RemapSkeletalMeshVertexColorToImportData(const USkeletalMesh* SkeletalMesh,
 	}
 }
 
+// Utility function for creating a unique material import name based on the node's initial name and its material/section index
+FString MakeNodeMaterialIndexName(FbxNode* Node, int32 MaterialIndex)
+{
+	FString NodeName = FFbxImporter::MakeName(Node->GetInitialName());
+	NodeName.Append("_");
+	NodeName.AppendInt(MaterialIndex);
+
+	return NodeName;
+}
+
 void FFbxImporter::SkinControlPointsToPose(FSkeletalMeshImportData& ImportData, FbxMesh* FbxMesh, FbxShape* FbxShape, bool bUseT0 )
 {
 	FbxTime poseTime = FBXSDK_TIME_INFINITE;
@@ -1371,13 +1381,29 @@ bool UnFbx::FFbxImporter::FillSkeletalMeshImportData(TArray<FbxNode*>& NodeArray
 		for (int32 MaterialIndex = 0; MaterialIndex < MaterialCount; ++MaterialIndex)
 		{
 			FbxSurfaceMaterial* FbxMaterial = Node->GetMaterial(MaterialIndex);
-			if (!FbxMaterials.Contains(FbxMaterial))
+
+			// Store unique FbxNode materials resulting in sections being combined.
+			if (!ImportOptions->bKeepSectionsSeparate)
 			{
+				if (!FbxMaterials.Contains(FbxMaterial))
+				{
+					FbxMaterials.Add(FbxMaterial);
+
+					SkeletalMeshImportData::FMaterial NewMaterial;
+
+					NewMaterial.MaterialImportName = MakeName(FbxMaterial->GetName());
+					// Add an entry for each unique material
+					SkelMeshImportDataPtr->Materials.Add(NewMaterial);
+				}
+			}
+			else
+			{
+				// Store each FbxNode material regardless of duplicates to avoid sections being combined.
 				FbxMaterials.Add(FbxMaterial);
 
 				SkeletalMeshImportData::FMaterial NewMaterial;
 
-				NewMaterial.MaterialImportName = MakeName(FbxMaterial->GetName());
+				NewMaterial.MaterialImportName = MakeNodeMaterialIndexName(Node, MaterialIndex);
 				// Add an entry for each unique material
 				SkelMeshImportDataPtr->Materials.Add(NewMaterial);
 			}
@@ -3205,6 +3231,8 @@ bool UnFbx::FFbxImporter::FillSkelMeshImporterFromFbx( FSkeletalMeshImportData&
 		ImportTexturesFromNode(Node);
 	}
 
+	ImportData.bKeepSectionsSeparate = ImportOptions->bKeepSectionsSeparate;
+
 	// Maps local mesh material index to global material index
 	const int32 MaterialCount = Node->GetMaterialCount();
 	TArray<int32> MaterialMapping;
@@ -3215,15 +3243,49 @@ bool UnFbx::FFbxImporter::FillSkelMeshImporterFromFbx( FSkeletalMeshImportData&
 		FbxSurfaceMaterial* FbxMaterial = Node->GetMaterial( MaterialIndex );
 
 		int32 ExistingMatIndex = INDEX_NONE;
-		FbxMaterials.Find( FbxMaterial, ExistingMatIndex ); 
-		if( ExistingMatIndex != INDEX_NONE )
+
+		// Default behavior finds the first material in FbxMaterials that matches the node's materials.
+		// This results in combining nodes/sections that use the same material.
+		if (!ImportOptions->bKeepSectionsSeparate)
+		{
+			FbxMaterials.Find(FbxMaterial, ExistingMatIndex);
+		}
+		else
+		{
+			// Otherwise find the material in FbxMaterials specific to that node's section/surface
+			for (int32 FbxMatIndex = 0; FbxMatIndex < FbxMaterials.Num(); FbxMatIndex++)
+			{
+				FString NodeMaterialIndexName = MakeNodeMaterialIndexName(Node, MaterialIndex);
+				bool MaterialImportNameMatches = ImportData.Materials[FbxMatIndex].MaterialImportName == NodeMaterialIndexName;
+				bool MaterialInitialNameMatches = FbxMaterials[FbxMatIndex]->GetInitialName() == FbxMaterial->GetInitialName();
+
+				if (MaterialImportNameMatches && MaterialInitialNameMatches)
+				{
+					ExistingMatIndex = FbxMatIndex;
+					break;
+				}
+			}
+		}
+
+		if (ExistingMatIndex != INDEX_NONE)
 		{
 			// Reuse existing material
 			MaterialMapping[MaterialIndex] = ExistingMatIndex;
 
-			if (Materials.IsValidIndex(MaterialIndex) )
+			if (!ImportOptions->bKeepSectionsSeparate)
+			{
+				if (Materials.IsValidIndex(MaterialIndex))
+				{
+					ImportData.Materials[ExistingMatIndex].Material = Materials[MaterialIndex];
+				}
+			}
+			else
 			{
-				ImportData.Materials[ExistingMatIndex].Material = Materials[MaterialIndex];
+				if (Materials.IsValidIndex(MaterialIndex) && ImportData.Materials.IsValidIndex(ExistingMatIndex))
+				{
+					ImportData.Materials[ExistingMatIndex].MaterialImportName = MakeNodeMaterialIndexName(Node, MaterialIndex);
+					ImportData.Materials[ExistingMatIndex].Material = Materials[MaterialIndex];
+				}
 			}
 		}
 		else
diff --git a/Engine/Source/Editor/UnrealEd/Private/Fbx/FbxSkeletalMeshImportData.cpp b/Engine/Source/Editor/UnrealEd/Private/Fbx/FbxSkeletalMeshImportData.cpp
--- a/Engine/Source/Editor/UnrealEd/Private/Fbx/FbxSkeletalMeshImportData.cpp
+++ b/Engine/Source/Editor/UnrealEd/Private/Fbx/FbxSkeletalMeshImportData.cpp
@@ -13,6 +13,7 @@ UFbxSkeletalMeshImportData::UFbxSkeletalMeshImportData(const FObjectInitializer&
 	bBakePivotInVertex = false;
 	VertexColorImportOption = EVertexColorImportOption::Replace;
 	LastImportContentType = EFBXImportContentType::FBXICT_All;
+	bKeepSectionsSeparate = false;
 }
 
 void UFbxSkeletalMeshImportData::Serialize(FArchive& Ar)
diff --git a/Engine/Source/Editor/UnrealEd/Private/Fbx/ReimportFbxSceneFactory.cpp b/Engine/Source/Editor/UnrealEd/Private/Fbx/ReimportFbxSceneFactory.cpp
--- a/Engine/Source/Editor/UnrealEd/Private/Fbx/ReimportFbxSceneFactory.cpp
+++ b/Engine/Source/Editor/UnrealEd/Private/Fbx/ReimportFbxSceneFactory.cpp
@@ -239,6 +239,8 @@ bool GetFbxSceneReImportOptions(UnFbx::FFbxImporter* FbxImporter
 	GlobalImportSettings->bImportMaterials = true;
 	//TODO support T0AsRefPose
 	GlobalImportSettings->bUseT0AsRefPose = false;
+	//Make sure skeletal mesh sections with the same material are combined (and not kept separate)
+	GlobalImportSettings->bKeepSectionsSeparate = false;
 	//Make sure we do not mess with AutoComputeLodDistances when re-importing
 	GlobalImportSettings->bAutoComputeLodDistances = true;
 	GlobalImportSettings->LodNumber = 0;
diff --git a/Engine/Source/Editor/UnrealEd/Private/Fbx/SFbxSceneOptionWindow.cpp b/Engine/Source/Editor/UnrealEd/Private/Fbx/SFbxSceneOptionWindow.cpp
--- a/Engine/Source/Editor/UnrealEd/Private/Fbx/SFbxSceneOptionWindow.cpp
+++ b/Engine/Source/Editor/UnrealEd/Private/Fbx/SFbxSceneOptionWindow.cpp
@@ -1502,6 +1502,7 @@ void SFbxSceneOptionWindow::CopySkeletalMeshOptionsToFbxOptions(UnFbx::FBXImport
 	ImportSettings->OverlappingThresholds.ThresholdUV = SkeletalMeshOptions->ThresholdUV;
 	ImportSettings->OverlappingThresholds.MorphThresholdPosition = SkeletalMeshOptions->MorphThresholdPosition;
 	ImportSettings->bPreserveSmoothingGroups = SkeletalMeshOptions->bPreserveSmoothingGroups;
+	ImportSettings->bKeepSectionsSeparate = SkeletalMeshOptions->bKeepSectionsSeparate;
 	ImportSettings->bUpdateSkeletonReferencePose = SkeletalMeshOptions->bUpdateSkeletonReferencePose;
 	ImportSettings->bUseT0AsRefPose = SkeletalMeshOptions->bUseT0AsRefPose;
 
@@ -1530,6 +1531,7 @@ void SFbxSceneOptionWindow::CopyFbxOptionsToSkeletalMeshOptions(UnFbx::FBXImport
 	SkeletalMeshOptions->ThresholdUV = ImportSettings->OverlappingThresholds.ThresholdUV;
 	SkeletalMeshOptions->MorphThresholdPosition = ImportSettings->OverlappingThresholds.MorphThresholdPosition;
 	SkeletalMeshOptions->bPreserveSmoothingGroups = ImportSettings->bPreserveSmoothingGroups;
+	SkeletalMeshOptions->bKeepSectionsSeparate = ImportSettings->bKeepSectionsSeparate;
 	SkeletalMeshOptions->bUpdateSkeletonReferencePose = ImportSettings->bUpdateSkeletonReferencePose;
 	SkeletalMeshOptions->bUseT0AsRefPose = ImportSettings->bUseT0AsRefPose;
 
diff --git a/Engine/Source/Editor/UnrealEd/Public/FbxImporter.h b/Engine/Source/Editor/UnrealEd/Public/FbxImporter.h
--- a/Engine/Source/Editor/UnrealEd/Public/FbxImporter.h
+++ b/Engine/Source/Editor/UnrealEd/Public/FbxImporter.h
@@ -184,6 +184,7 @@ struct FBXImportOptions
 	bool bImportRigidMesh;
 	bool bUseT0AsRefPose;
 	bool bPreserveSmoothingGroups;
+	bool bKeepSectionsSeparate;
 	FOverlappingThresholds OverlappingThresholds;
 	bool bImportMeshesInBoneHierarchy;
 	bool bCreatePhysicsAsset;
diff --git a/Engine/Source/Runtime/Engine/Private/Rendering/SkeletalMeshLODImporterData.cpp b/Engine/Source/Runtime/Engine/Private/Rendering/SkeletalMeshLODImporterData.cpp
--- a/Engine/Source/Runtime/Engine/Private/Rendering/SkeletalMeshLODImporterData.cpp
+++ b/Engine/Source/Runtime/Engine/Private/Rendering/SkeletalMeshLODImporterData.cpp
@@ -30,6 +30,7 @@ void FSkeletalMeshImportData::CopyDataNeedByMorphTargetImport(FSkeletalMeshImpor
 	Other.PointToRawMap = PointToRawMap;
 	Other.bDiffPose = bDiffPose;
 	Other.bUseT0AsRefPose = bUseT0AsRefPose;
+	Other.bKeepSectionsSeparate = bKeepSectionsSeparate;
 }
 
 void FSkeletalMeshImportData::KeepAlternateSkinningBuildDataOnly()
@@ -158,6 +159,7 @@ bool FSkeletalMeshImportData::ReplaceSkeletalMeshGeometryImportData(const USkele
 	ImportData->bHasTangents = OriginalSkeletalMeshImportData.bHasTangents;
 	ImportData->bHasVertexColors = OriginalSkeletalMeshImportData.bHasVertexColors;
 	ImportData->NumTexCoords = OriginalSkeletalMeshImportData.NumTexCoords;
+	ImportData->bKeepSectionsSeparate = OriginalSkeletalMeshImportData.bKeepSectionsSeparate;
 
 	ImportData->Materials.Reset();
 	ImportData->Points.Reset();
@@ -215,6 +217,7 @@ bool FSkeletalMeshImportData::ReplaceSkeletalMeshRigImportData(const USkeletalMe
 
 	ImportData->bDiffPose = OriginalSkeletalMeshImportData.bDiffPose;
 	ImportData->bUseT0AsRefPose = OriginalSkeletalMeshImportData.bUseT0AsRefPose;
+	ImportData->bKeepSectionsSeparate = OriginalSkeletalMeshImportData.bKeepSectionsSeparate;
 
 	ImportData->RefBonesBinary.Reset();
 	ImportData->RefBonesBinary += OriginalSkeletalMeshImportData.RefBonesBinary;
@@ -659,6 +662,7 @@ enum
 	RAW_SKELETAL_MESH_BULKDATA_VER_RebuildSystem = 2,
 	RAW_SKELETAL_MESH_BULKDATA_VER_CompressMorphTargetData = 3,
 	RAW_SKELETAL_MESH_BULKDATA_VER_VertexAttributes = 4,
+	RAW_SKELETAL_MESH_BULKDATA_VER_Keep_Sections_Separate = 5,
 	// Add new raw mesh versions here.
 
 	RAW_SKELETAL_MESH_BULKDATA_VER_PLUS_ONE,
@@ -760,6 +764,15 @@ FArchive& operator<<(FArchive& Ar, FSkeletalMeshImportData& RawMesh)
 		}
 	}
 
+	if (Version >= RAW_SKELETAL_MESH_BULKDATA_VER_Keep_Sections_Separate)
+	{
+		Ar << RawMesh.bKeepSectionsSeparate;
+	}
+	else if (Ar.IsLoading())
+	{
+		RawMesh.bKeepSectionsSeparate = false;
+	}
+
 	return Ar;
 }
 
diff --git a/Engine/Source/Runtime/Engine/Public/Rendering/SkeletalMeshLODImporterData.h b/Engine/Source/Runtime/Engine/Public/Rendering/SkeletalMeshLODImporterData.h
--- a/Engine/Source/Runtime/Engine/Public/Rendering/SkeletalMeshLODImporterData.h
+++ b/Engine/Source/Runtime/Engine/Public/Rendering/SkeletalMeshLODImporterData.h
@@ -340,6 +340,7 @@ public:
 	bool bHasTangents; // If true there are tangents in the imported file
 	bool bUseT0AsRefPose; // If true, then the pose at time=0 will be used instead of the ref pose
 	bool bDiffPose; // If true, one of the bones has a different pose at time=0 vs the ref pose
+	bool bKeepSectionsSeparate; // If true, sections with matching materials are kept separate and will not get combined
 
 	// Morph targets imported(i.e. FBX) data. The name is the morph target name
 	TArray<FSkeletalMeshImportData> MorphTargets;
@@ -363,6 +364,7 @@ public:
 		, bHasTangents(false)
 		, bUseT0AsRefPose(false)
 		, bDiffPose(false)
+		, bKeepSectionsSeparate(false)
 	{
 
 	}
