MPCook: Speculative fix for error "FZenStoreWriter encountered the same package twice in a cook!". One way this can happen is Platform1 -> CW1, CW2 marks it reachable, director retracts it and sends it to CW2, CW2 cooks it, Platform2 discovered, Platform2 -> CW1, CW1 cooks both Platform1 and Platform2 because Platform1 was still reachable from the pre-retraction.

Fix that case by clearing reachability upon retraction, and also ignoring cookresults for platforms that the director already has a cookresult for.
#rnx

diff --git a/Engine/Source/Editor/UnrealEd/Private/Cooker/CookWorkerClient.cpp b/Engine/Source/Editor/UnrealEd/Private/Cooker/CookWorkerClient.cpp
--- a/Engine/Source/Editor/UnrealEd/Private/Cooker/CookWorkerClient.cpp
+++ b/Engine/Source/Editor/UnrealEd/Private/Cooker/CookWorkerClient.cpp
@@ -913,6 +913,7 @@ void FCookWorkerClient::HandleRetractionMessage(FMPCollectorClientMessageContext
 	{
 		FPackageData* PackageData = COTFS.PackageDatas->FindPackageDataByPackageName(PackageName);
 		check(PackageData);
+		PackageData->ResetReachable();
 		COTFS.DemoteToIdle(*PackageData, ESendFlags::QueueAddAndRemove, ESuppressCookReason::RetractedByCookDirector);
 	}
 
diff --git a/Engine/Source/Editor/UnrealEd/Private/Cooker/CookWorkerServer.cpp b/Engine/Source/Editor/UnrealEd/Private/Cooker/CookWorkerServer.cpp
--- a/Engine/Source/Editor/UnrealEd/Private/Cooker/CookWorkerServer.cpp
+++ b/Engine/Source/Editor/UnrealEd/Private/Cooker/CookWorkerServer.cpp
@@ -836,11 +836,25 @@ void FCookWorkerServer::RecordResults(FPackageResultsMessage& Message)
 			{
 				ITargetPlatform* TargetPlatform = OrderedSessionPlatforms[PlatformIndex];
 				FPackageRemoteResult::FPlatformResult& PlatformResult = Result.GetPlatforms()[PlatformIndex];
-				if (PlatformResult.GetCookResults() != ECookResult::Invalid)
+				FPackagePlatformData& ExistingData = PackageData->FindOrAddPlatformData(TargetPlatform);
+				if (!ExistingData.NeedsCooking(TargetPlatform))
 				{
-					PackageData->SetPlatformCooked(TargetPlatform, PlatformResult.GetCookResults());
+					if (PlatformResult.GetCookResults() != ECookResult::Invalid)
+					{
+						UE_LOG(LogCook, Display,
+							TEXT("CookWorkerServer %d received FPackageResultsMessage for package %s, platform %s, but that platform has already been cooked. Ignoring the results for that platform."),
+							ProfileId, *Result.GetPackageName().ToString(), *TargetPlatform->PlatformName());
+					}
+					continue;
+				}
+				else
+				{
+					if (PlatformResult.GetCookResults() != ECookResult::Invalid)
+					{
+						PackageData->SetPlatformCooked(TargetPlatform, PlatformResult.GetCookResults());
+					}
+					HandleReceivedPackagePlatformMessages(*PackageData, TargetPlatform, PlatformResult.ReleaseMessages());
 				}
-				HandleReceivedPackagePlatformMessages(*PackageData, TargetPlatform, PlatformResult.ReleaseMessages());
 			}
 			COTFS.RecordExternalActorDependencies(Result.GetExternalActorDependencies());
 			if (Result.IsReferencedOnlyByEditorOnlyData())
