Various fixes for Orthographic cameras in Editor:
- Switched the camera preview window to Lit mode by default
- Disabled faux-ortho mode to fix viewport lighting artefacts (faux-ortho to be removed entirely at a later date).
- Fixed DPI scaling for measurements in wireframe viewports.

diff --git a/Engine/Source/Editor/UnrealEd/Private/EditorViewportClient.cpp b/Engine/Source/Editor/UnrealEd/Private/EditorViewportClient.cpp
--- a/Engine/Source/Editor/UnrealEd/Private/EditorViewportClient.cpp
+++ b/Engine/Source/Editor/UnrealEd/Private/EditorViewportClient.cpp
@@ -74,7 +74,7 @@
 #define LOCTEXT_NAMESPACE "EditorViewportClient"
 
 const EViewModeIndex FEditorViewportClient::DefaultPerspectiveViewMode = VMI_Lit;
-const EViewModeIndex FEditorViewportClient::DefaultOrthoViewMode = VMI_BrushWireframe;
+const EViewModeIndex FEditorViewportClient::DefaultOrthoViewMode = VMI_Lit;
 
 static TAutoConsoleVariable<int32> CVarAlignedOrthoZoom(
 	TEXT("r.Editor.AlignedOrthoZoom"),
@@ -1233,7 +1233,6 @@ FSceneView* FEditorViewportClient::CalcSceneView(FSceneViewFamily* ViewFamily, c
 
 	// for ortho views to steal perspective view origin
 	ViewInitOptions.OverrideLODViewOrigin = FVector::ZeroVector;
-	ViewInitOptions.bUseFauxOrthoViewPos = true;
 
 	ViewInitOptions.FOV = ModifiedViewFOV;
 	if (bUseControllingActorViewInfo)
@@ -3544,7 +3543,7 @@ FString FEditorViewportClient::UnrealUnitsToSiUnits(float UnrealUnits)
 
 void FEditorViewportClient::DrawScaleUnits(FViewport* InViewport, FCanvas* Canvas, const FSceneView& InView)
 {
-	const float UnitsPerPixel = GetOrthoUnitsPerPixel(InViewport);
+	const float UnitsPerPixel = GetOrthoUnitsPerPixel(InViewport) * Canvas->GetDPIScale();
 
 	// Find the closest power of ten to our target width
 	static const int32 ApproxTargetMarkerWidthPx = 100;
