Cooker Incremental Cooking: Fix -diffonly to work (and recook everything) when Incremental cooking is enabled (incremental cooking cooks iteratively all the time, so -diffonly needs to be a special exception).
#rnx

diff --git a/Engine/Source/Editor/UnrealEd/Private/CookOnTheFlyServer.cpp b/Engine/Source/Editor/UnrealEd/Private/CookOnTheFlyServer.cpp
--- a/Engine/Source/Editor/UnrealEd/Private/CookOnTheFlyServer.cpp
+++ b/Engine/Source/Editor/UnrealEd/Private/CookOnTheFlyServer.cpp
@@ -10369,6 +10369,7 @@ void UCookOnTheFlyServer::LoadBeginCookIterativeFlagsLocal(FBeginCookContext& Be
 			UE_LOG(LogCook, Display, TEXT("Keeping cooked content for platform %s for DiffOnly"), *TargetPlatform->PlatformName());
 			// When looking for deterministic cooking differences in cooked packages, don't delete the packages on disk
 			PlatformContext.bFullBuild = false;
+			PlatformContext.bAllowIterativeResults = false;
 			PlatformContext.bClearMemoryResults = true;
 			PlatformContext.bPopulateMemoryResultsFromDiskResults = false;
 			PlatformContext.bIterateSharedBuild = false;
@@ -10389,6 +10390,7 @@ void UCookOnTheFlyServer::LoadBeginCookIterativeFlagsLocal(FBeginCookContext& Be
 			if (bIterativeAllowed)
 			{
 				PlatformContext.bFullBuild = false;
+				PlatformContext.bAllowIterativeResults = true;
 				PlatformContext.bClearMemoryResults = false;
 				PlatformContext.bPopulateMemoryResultsFromDiskResults = !PlatformContext.bHasMemoryResults;
 				PlatformContext.bIterateSharedBuild = bIterateSharedBuild;
@@ -10396,12 +10398,14 @@ void UCookOnTheFlyServer::LoadBeginCookIterativeFlagsLocal(FBeginCookContext& Be
 			else
 			{
 				PlatformContext.bFullBuild = true;
+				PlatformContext.bAllowIterativeResults = false;
 				PlatformContext.bClearMemoryResults = true;
 				PlatformContext.bPopulateMemoryResultsFromDiskResults = false;
 				PlatformContext.bIterateSharedBuild = false;
 			}
 		}
 		PlatformData->bFullBuild = PlatformContext.bFullBuild;
+		PlatformData->bAllowIterativeResults = PlatformContext.bAllowIterativeResults;
 		PlatformData->bIterateSharedBuild = PlatformContext.bIterateSharedBuild;
 		PlatformData->bWorkerOnSharedSandbox = PlatformContext.bWorkerOnSharedSandbox;
 	}
diff --git a/Engine/Source/Editor/UnrealEd/Private/Cooker/CookPlatformManager.h b/Engine/Source/Editor/UnrealEd/Private/Cooker/CookPlatformManager.h
--- a/Engine/Source/Editor/UnrealEd/Private/Cooker/CookPlatformManager.h
+++ b/Engine/Source/Editor/UnrealEd/Private/Cooker/CookPlatformManager.h
@@ -59,8 +59,13 @@ namespace Cook
 		*/
 		bool bIsSandboxInitialized = false;
 
-		/* Whether BeginCookSandbox specified a full, non-iterative build for this platform. */
+		/** Whether BeginCookSandbox specified a full, non-iterative build for this platform. */
 		bool bFullBuild = false;
+		/**
+		 * Whether BeginCookSandBox specied non-full, and is in an iterative mode that allows using previous results.
+		 * -diffonly is the expected case where bFullBuild=false but bAllowIterativeResults=false.
+		 */
+		bool bAllowIterativeResults = true;
 
 		/** If true we are cooking iteratively, from results in a shared build (e.g. from buildfarm) rather than from our previous cook. */
 		bool bIterateSharedBuild = false;
diff --git a/Engine/Source/Editor/UnrealEd/Private/Cooker/CookRequestCluster.cpp b/Engine/Source/Editor/UnrealEd/Private/Cooker/CookRequestCluster.cpp
--- a/Engine/Source/Editor/UnrealEd/Private/Cooker/CookRequestCluster.cpp
+++ b/Engine/Source/Editor/UnrealEd/Private/Cooker/CookRequestCluster.cpp
@@ -61,22 +61,22 @@ FRequestCluster::FRequestCluster(UCookOnTheFlyServer& InCOTFS)
 	}
 	GConfig->GetBool(TEXT("CookSettings"), TEXT("PreQueueBuildDefinitions"), bPreQueueBuildDefinitions, GEditorIni);
 
-	bFullBuild = false;
+	bAllowIterativeResults = true;
 	bool bFirst = true;
 	for (const ITargetPlatform* TargetPlatform : COTFS.PlatformManager->GetSessionPlatforms())
 	{
 		FPlatformData* PlatformData = COTFS.PlatformManager->GetPlatformData(TargetPlatform);
 		if (bFirst)
 		{
-			bFullBuild = PlatformData->bFullBuild;
+			bAllowIterativeResults = PlatformData->bAllowIterativeResults;
 			bFirst = false;
 		}
 		else
 		{
-			if (PlatformData->bFullBuild != bFullBuild)
+			if (PlatformData->bAllowIterativeResults != bAllowIterativeResults)
 			{
 				UE_LOG(LogCook, Warning, TEXT("Full build is requested for some platforms but not others, but this is not supported. All platforms will be built full."));
-				bFullBuild = true;
+				bAllowIterativeResults = false;
 			}
 		}
 	}
@@ -1513,7 +1513,7 @@ TMap<FPackageData*, TArray<FPackageData*>>& FRequestCluster::FGraphSearch::GetGr
 
 bool FRequestCluster::IsIncrementalCook() const
 {
-	return !bFullBuild && COTFS.bHybridIterativeEnabled;
+	return bAllowIterativeResults && COTFS.bHybridIterativeEnabled;
 }
 
 void FRequestCluster::IsRequestCookable(const ITargetPlatform* Platform, FPackageData& PackageData,
diff --git a/Engine/Source/Editor/UnrealEd/Private/Cooker/CookRequestCluster.h b/Engine/Source/Editor/UnrealEd/Private/Cooker/CookRequestCluster.h
--- a/Engine/Source/Editor/UnrealEd/Private/Cooker/CookRequestCluster.h
+++ b/Engine/Source/Editor/UnrealEd/Private/Cooker/CookRequestCluster.h
@@ -354,7 +354,7 @@ private:
 	bool bPackageNamesComplete = false;
 	bool bDependenciesComplete = false;
 	bool bStartAsyncComplete = false;
-	bool bFullBuild = false;
+	bool bAllowIterativeResults = false;
 	bool bPreQueueBuildDefinitions = true;
 };
 
diff --git a/Engine/Source/Editor/UnrealEd/Private/Cooker/CookTypes.h b/Engine/Source/Editor/UnrealEd/Private/Cooker/CookTypes.h
--- a/Engine/Source/Editor/UnrealEd/Private/Cooker/CookTypes.h
+++ b/Engine/Source/Editor/UnrealEd/Private/Cooker/CookTypes.h
@@ -536,8 +536,13 @@ struct FBeginCookContextPlatform
 	UE::Cook::FPlatformData* PlatformData = nullptr;
 	TMap<FName, FString> CurrentCookSettings;
 
-	/** If true, we are deleting all old results from disk and rebuilding every package. If false, we are building iteratively. */
+	/** If true, we are deleting all old results from the previous cook. If false, we are keeping the old results. */
 	bool bFullBuild = false;
+	/**
+	 * If true, we will use results from the previous cook for the new cook, if present. If false, we will recook.
+	 * -diffonly is the expected case where bFullBuild=false but bAllowIterativeResults=false.
+	 */
+	bool bAllowIterativeResults = true;
 	/** If true, a cook has already been run in the current process and we still have results from it. */
 	bool bHasMemoryResults = false;
 	/** If true, we should delete the in-memory results from an earlier cook in the same process, if we have any. */
