Cooker: IncrementalCook: Fix bug in FRequestCluster that causes spurious load requests on the editordomainoplog for the fake PlatformAgnosticPlatform.
#rnx

diff --git a/Engine/Source/Editor/UnrealEd/Private/Cooker/CookRequestCluster.cpp b/Engine/Source/Editor/UnrealEd/Private/Cooker/CookRequestCluster.cpp
--- a/Engine/Source/Editor/UnrealEd/Private/Cooker/CookRequestCluster.cpp
+++ b/Engine/Source/Editor/UnrealEd/Private/Cooker/CookRequestCluster.cpp
@@ -1449,6 +1449,7 @@ void FRequestCluster::FQueryVertexBatch::Send()
 		FFetchPlatformData& FetchPlatformData = ThreadSafeOnlyVars.FetchPlatforms[PlatformIndex];
 
 		if (ThreadSafeOnlyVars.Cluster.IsIncrementalCook() // Only FetchCookAttachments if our cookmode supports it. Otherwise keep them all empty
+			&& !FetchPlatformData.bIsPlatformAgnosticPlatform // The PlatformAgnosticPlatform has no stored CookAttachments; always use empty
 			&& !FetchPlatformData.bIsCookerLoadingPlatform // The CookerLoadingPlatform has no stored CookAttachments; always use empty
 			)
 		{
