Added ExportPath to exported T3D actors and fixed incorrect masking of PPF_ExportsNotFullyQualified

#rnx

diff --git a/Engine/Source/Editor/UnrealEd/Private/EditorExporters.cpp b/Engine/Source/Editor/UnrealEd/Private/EditorExporters.cpp
--- a/Engine/Source/Editor/UnrealEd/Private/EditorExporters.cpp
+++ b/Engine/Source/Editor/UnrealEd/Private/EditorExporters.cpp
@@ -563,9 +563,20 @@ bool ULevelExporterT3D::ExportText( const FExportObjectInnerContext* Context, UO
 				FString ActorFolderPath = (Level->IsUsingActorFolders() ? FString::Printf(TEXT(" ActorFolderPath=%s"), *Actor->GetFolderPath().ToString()) : TEXT(""));
 				FString CopyPasteId = (Actor->CopyPasteId != INDEX_NONE) ? FString::Printf(TEXT(" CopyPasteId=%d"), Actor->CopyPasteId ) : TEXT("");
 				FString ContentBundleGuid = (Actor->GetContentBundleGuid().IsValid() ? FString::Printf(TEXT(" ActorContentBundleGuid=%s"), *Actor->GetContentBundleGuid().ToString()) : TEXT(""));
-				Ar.Logf( TEXT("%sBegin Actor Class=%s Name=%s Archetype=%s'%s'%s%s%s%s%s%s%s") LINE_TERMINATOR, 
+				Ar.Logf( TEXT("%sBegin Actor Class=%s Name=%s Archetype=%s%s%s%s%s%s%s%s"), 
 					FCString::Spc(TextIndent), *Actor->GetClass()->GetPathName(), *Actor->GetName(),
-					*Actor->GetArchetype()->GetClass()->GetPathName(), *Actor->GetArchetype()->GetPathName(), *ParentActorString, *SocketNameString, *GroupActor, *GroupFolder, *ActorFolderPath, *CopyPasteId, *ContentBundleGuid );
+					*FObjectPropertyBase::GetExportPath(Actor->GetArchetype(), nullptr, nullptr, (PortFlags | PPF_Delimited) & ~PPF_ExportsNotFullyQualified), 
+					*ParentActorString, *SocketNameString, *GroupActor, *GroupFolder, *ActorFolderPath, *CopyPasteId, *ContentBundleGuid );
+
+				// When exporting for diffs, export paths can cause false positives. since diff files don't get imported, we can
+				// skip adding this info the file.
+				if (!(PortFlags & PPF_ForDiff))
+				{
+					// Emit the actor path
+					Ar.Logf(TEXT(" ExportPath=%s"), *FObjectPropertyBase::GetExportPath(Actor, nullptr, nullptr, (PortFlags | PPF_Delimited) & ~PPF_ExportsNotFullyQualified));
+				}
+
+				Ar.Logf(LINE_TERMINATOR);
 
 				ExportRootScope = Actor;
 				ExportObjectInner( Context, Actor, Ar, PortFlags | PPF_ExportsNotFullyQualified );
diff --git a/Engine/Source/Editor/UnrealEd/Private/Elements/Actor/ActorElementEditorCopyAndPaste.cpp b/Engine/Source/Editor/UnrealEd/Private/Elements/Actor/ActorElementEditorCopyAndPaste.cpp
--- a/Engine/Source/Editor/UnrealEd/Private/Elements/Actor/ActorElementEditorCopyAndPaste.cpp
+++ b/Engine/Source/Editor/UnrealEd/Private/Elements/Actor/ActorElementEditorCopyAndPaste.cpp
@@ -70,13 +70,20 @@ bool UActorElementsExporterT3D::ExportText(const FExportObjectInnerContext* Cont
 			                           ? FString::Printf(
 				                           TEXT(" ActorFolderPath=%s"), *Actor->GetFolderPath().ToString())
 			                           : TEXT(""));
-		FString ActorPath = FString::Printf(
-			TEXT(" ExportPath=%s"),
-			*FObjectPropertyBase::GetExportPath(Actor, nullptr, nullptr, PortFlags | PPF_Delimited));
-		Out.Logf(TEXT("%sBegin Actor Class=%s Name=%s Archetype=%s'%s'%s%s%s%s%s%s") LINE_TERMINATOR,
+		Out.Logf(TEXT("%sBegin Actor Class=%s Name=%s Archetype=%s%s%s%s%s%s"),
 		        FCString::Spc(TextIndent), *Actor->GetClass()->GetPathName(), *Actor->GetName(),
-		        *Actor->GetArchetype()->GetClass()->GetPathName(), *Actor->GetArchetype()->GetPathName(),
-		        *ParentActorString, *SocketNameString, *GroupActor, *GroupFolder, *ActorFolderPath, *ActorPath);
+				*FObjectPropertyBase::GetExportPath(Actor->GetArchetype(), nullptr, nullptr, (PortFlags | PPF_Delimited) & ~PPF_ExportsNotFullyQualified),
+		        *ParentActorString, *SocketNameString, *GroupActor, *GroupFolder, *ActorFolderPath);
+
+		// When exporting for diffs, export paths can cause false positives. since diff files don't get imported, we can
+		// skip adding this info the file.
+		if (!(PortFlags & PPF_ForDiff))
+		{
+			// Emit the actor path
+			Out.Logf(TEXT(" ExportPath=%s"), *FObjectPropertyBase::GetExportPath(Actor, nullptr, nullptr, (PortFlags | PPF_Delimited) & ~PPF_ExportsNotFullyQualified));
+		}
+
+		Out.Logf(LINE_TERMINATOR);
 
 		ExportRootScope = Actor;
 		ExportObjectInner(Context, Actor, Out, PortFlags);
diff --git a/Engine/Source/Runtime/Engine/Private/UnrealExporter.cpp b/Engine/Source/Runtime/Engine/Private/UnrealExporter.cpp
--- a/Engine/Source/Runtime/Engine/Private/UnrealExporter.cpp
+++ b/Engine/Source/Runtime/Engine/Private/UnrealExporter.cpp
@@ -469,7 +469,7 @@ void UExporter::EmitBeginObject( FOutputDevice& Ar, UObject* Obj, uint32 PortFla
 			UObject* Archetype = Obj->GetArchetype();
 			// since we could have two object owners with the same name (like named Blueprints in different folders),
 			// we need the fully qualified path for the archetype (so we don't get confused when unpacking this)
-			Ar.Logf(TEXT(" Archetype=%s"), *FObjectPropertyBase::GetExportPath(Archetype, Archetype->GetOutermost(), /*ExportRootScope =*/nullptr, PortFlags & ~PPF_ExportsNotFullyQualified));
+			Ar.Logf(TEXT(" Archetype=%s"), *FObjectPropertyBase::GetExportPath(Archetype, nullptr, /*ExportRootScope =*/nullptr, (PortFlags | PPF_Delimited) & ~PPF_ExportsNotFullyQualified));
 		}
 	}
 
@@ -478,10 +478,10 @@ void UExporter::EmitBeginObject( FOutputDevice& Ar, UObject* Obj, uint32 PortFla
 	if (!(PortFlags & PPF_ForDiff))
 	{
 		// Emit the object path
-		Ar.Logf(TEXT(" ExportPath=%s"), *FObjectPropertyBase::GetExportPath(Obj, nullptr, nullptr, PortFlags | (PPF_Delimited  & ~PPF_ExportsNotFullyQualified)));
+		Ar.Logf(TEXT(" ExportPath=%s"), *FObjectPropertyBase::GetExportPath(Obj, nullptr, nullptr, (PortFlags | PPF_Delimited) & ~PPF_ExportsNotFullyQualified));
 	}
 	// end in a return
-	Ar.Logf(TEXT("\r\n"));
+	Ar.Logf(LINE_TERMINATOR);
 
 	if ( bEnableDebugBrackets )
 	{
