Fixed soft-object actor references failing to fix-up when moving (via copy/paste) actors to a new world

Breaking Change:
 * ImportObjectProperties ActorRemapper argument is now called ObjectRemapper and maps from FSoftObjectPath to UObject*

diff --git a/Engine/Source/Editor/UnrealEd/Private/EditorObject.cpp b/Engine/Source/Editor/UnrealEd/Private/EditorObject.cpp
--- a/Engine/Source/Editor/UnrealEd/Private/EditorObject.cpp
+++ b/Engine/Source/Editor/UnrealEd/Private/EditorObject.cpp
@@ -116,7 +116,7 @@ void UEditorEngine::RenameObject(UObject* Object,UObject* NewOuter,const TCHAR*
 }
 
 
-static void RemapProperty(FProperty* Property, int32 Index, const TMap<AActor*, AActor*>& ActorRemapper, uint8* DestData)
+static void RemapProperty(FProperty* Property, int32 Index, const TMap<FSoftObjectPath, UObject*>& ObjectRemapper, uint8* DestData)
 {
 	if (FObjectPropertyBase* ObjectProperty = CastField<FObjectPropertyBase>(Property))
 	{
@@ -126,16 +126,24 @@ static void RemapProperty(FProperty* Property, int32 Index, const TMap<AActor*,
 		for (int32 Count = 0; Count < Num; Count++)
 		{
 			uint8* PropertyAddr = ObjectProperty->ContainerPtrToValuePtr<uint8>(DestData, StartIndex + Count);
-			AActor* Actor = Cast<AActor>(ObjectProperty->GetObjectPropertyValue(PropertyAddr));
-			if (Actor)
+			UObject* Object = ObjectProperty->GetObjectPropertyValue(PropertyAddr);
+			if (Object)
 			{
-				AActor* const* RemappedObject = ActorRemapper.Find(Actor);
+				UObject* const* RemappedObject = ObjectRemapper.Find(Object);
 				if (RemappedObject && (*RemappedObject)->GetClass()->IsChildOf(ObjectProperty->PropertyClass))
 				{
 					ObjectProperty->SetObjectPropertyValue(PropertyAddr, *RemappedObject);
 				}
 			}
-
+			else if (FSoftObjectProperty* SoftObjectProperty = CastField<FSoftObjectProperty>(Property))
+			{
+				const FSoftObjectPtr& SoftObjectPtr = SoftObjectProperty->GetPropertyValue(PropertyAddr);
+				UObject* const* RemappedObject = ObjectRemapper.Find(SoftObjectPtr.ToSoftObjectPath());
+				if (RemappedObject && (*RemappedObject)->GetClass()->IsChildOf(SoftObjectProperty->PropertyClass))
+				{
+					SoftObjectProperty->SetObjectPropertyValue(PropertyAddr, *RemappedObject);
+				}
+			}
 		}
 	}
 	else if (FArrayProperty* ArrayProperty = CastField<FArrayProperty>(Property))
@@ -143,24 +151,38 @@ static void RemapProperty(FProperty* Property, int32 Index, const TMap<AActor*,
 		FScriptArrayHelper ArrayHelper(ArrayProperty, ArrayProperty->ContainerPtrToValuePtr<void>(DestData));
 		if (Index != INDEX_NONE)
 		{
-			RemapProperty(ArrayProperty->Inner, INDEX_NONE, ActorRemapper, ArrayHelper.GetRawPtr(Index));
+			RemapProperty(ArrayProperty->Inner, INDEX_NONE, ObjectRemapper, ArrayHelper.GetRawPtr(Index));
 		}
 		else
 		{
 			for (int32 ArrayIndex = 0; ArrayIndex < ArrayHelper.Num(); ArrayIndex++)
 			{
-				RemapProperty(ArrayProperty->Inner, INDEX_NONE, ActorRemapper, ArrayHelper.GetRawPtr(ArrayIndex));
+				RemapProperty(ArrayProperty->Inner, INDEX_NONE, ObjectRemapper, ArrayHelper.GetRawPtr(ArrayIndex));
 			}
 		}
 	}
 	else if (FStructProperty* StructProperty = CastField<FStructProperty>(Property))
 	{
-		if (Index != INDEX_NONE)
+		if (StructProperty->Struct == TBaseStructure<FSoftObjectPath>::Get())
+		{
+			// If there's a concrete index, use that, otherwise iterate all array members (for the case that this property is inside a struct, or there is exactly one element)
+			const int32 Num = (Index == INDEX_NONE) ? ObjectProperty->ArrayDim : 1;
+			const int32 StartIndex = (Index == INDEX_NONE) ? 0 : Index;
+			for (int32 Count = 0; Count < Num; Count++)
+			{
+				FSoftObjectPath* SoftObjectPathPtr = ObjectProperty->ContainerPtrToValuePtr<FSoftObjectPath>(DestData, StartIndex + Count);
+				if (UObject* const* RemappedObject = ObjectRemapper.Find(*SoftObjectPathPtr))
+				{
+					*SoftObjectPathPtr = *RemappedObject;
+				}
+			}
+		}
+		else if (Index != INDEX_NONE)
 		{
 			// If a concrete index was given, remap just that
 			for (TFieldIterator<FProperty> It(StructProperty->Struct); It; ++It)
 			{
-				RemapProperty(*It, INDEX_NONE, ActorRemapper, StructProperty->ContainerPtrToValuePtr<uint8>(DestData, Index));
+				RemapProperty(*It, INDEX_NONE, ObjectRemapper, StructProperty->ContainerPtrToValuePtr<uint8>(DestData, Index));
 			}
 		}
 		else
@@ -171,7 +193,7 @@ static void RemapProperty(FProperty* Property, int32 Index, const TMap<AActor*,
 			{
 				for (TFieldIterator<FProperty> It(StructProperty->Struct); It; ++It)
 				{
-					RemapProperty(*It, INDEX_NONE, ActorRemapper, StructProperty->ContainerPtrToValuePtr<uint8>(DestData, Count));
+					RemapProperty(*It, INDEX_NONE, ObjectRemapper, StructProperty->ContainerPtrToValuePtr<uint8>(DestData, Count));
 				}
 			}
 		}
@@ -216,7 +238,7 @@ static bool IsEndOfProperties(const TCHAR* Str, int32 Depth)
  * @param	Warn						output device to use for log messages
  * @param	Depth						current nesting level
  * @param	InstanceGraph				contains the mappings of instanced objects and components to their templates
- * @param	ActorRemapper				a map of existing actors to new instances, used to replace internal references when a number of actors are copy+pasted
+ * @param	ObjectRemapper				a map of existing objects, typically actors, to new instances, used to replace internal references when a number of objects are copy+pasted
  *
  * @return	NULL if the default values couldn't be imported
  */
@@ -229,7 +251,7 @@ static const TCHAR* ImportProperties(
 	FFeedbackContext*			Warn,
 	int32						Depth,
 	FObjectInstancingGraph&		InstanceGraph,
-	const TMap<AActor*, AActor*>* ActorRemapper
+	const TMap<FSoftObjectPath, UObject*>* ObjectRemapper
 	)
 {
 	check(ObjectStruct!=NULL);
@@ -482,7 +504,7 @@ static const TCHAR* ImportProperties(
 			{
 				// since we're redefining an object in the same text block, only need to import properties again
 				SourceText = ImportObjectProperties( (uint8*)BaseTemplate, SourceText, TemplateClass, SubobjectRoot, BaseTemplate,
-													Warn, Depth + 1, ContextSupplier ? ContextSupplier->CurrentLine : 0, &InstanceGraph, ActorRemapper );
+													Warn, Depth + 1, ContextSupplier ? ContextSupplier->CurrentLine : 0, &InstanceGraph, ObjectRemapper );
 			}
 			else 
 			{
@@ -663,7 +685,7 @@ static const TCHAR* ImportProperties(
 					Depth+1,
 					ContextSupplier ? ContextSupplier->CurrentLine : 0,
 					&InstanceGraph,
-					ActorRemapper
+					ObjectRemapper
 					);
 			}
 		}
@@ -689,11 +711,11 @@ static const TCHAR* ImportProperties(
 		}
 	}
 
-	if (ActorRemapper)
+	if (ObjectRemapper)
 	{
 		for (const auto& DefinedProperty : DefinedProperties)
 		{
-			RemapProperty(DefinedProperty.Property, DefinedProperty.Index, *ActorRemapper, DestData);
+			RemapProperty(DefinedProperty.Property, DefinedProperty.Index, *ObjectRemapper, DestData);
 		}
 	}
 
@@ -775,7 +797,7 @@ const TCHAR* ImportObjectProperties( FImportObjectParams& InParams )
 			InParams.Warn,
 			InParams.Depth,
 			InstanceGraph,
-			InParams.ActorRemapper
+			InParams.ObjectRemapper
 			);
 
 	if ( InParams.SubobjectOuter != NULL )
@@ -844,7 +866,7 @@ const TCHAR* ImportObjectProperties(
 	int32					Depth,
 	int32					LineNumber,
 	FObjectInstancingGraph* InInstanceGraph,
-	const TMap<AActor*, AActor*>* ActorRemapper
+	const TMap<FSoftObjectPath, UObject*>* ObjectRemapper
 	)
 {
 	FImportObjectParams Params;
@@ -858,7 +880,7 @@ const TCHAR* ImportObjectProperties(
 		Params.Depth = Depth;
 		Params.LineNumber = LineNumber;
 		Params.InInstanceGraph = InInstanceGraph;
-		Params.ActorRemapper = ActorRemapper;
+		Params.ObjectRemapper = ObjectRemapper;
 
 		// This implementation always calls PreEditChange/PostEditChange
 		Params.bShouldCallEditChange = true;
@@ -1503,15 +1525,6 @@ static const TCHAR* ImportPropertiesStep(
 		}
 	}
 
-	/**
-	if (ActorRemapper)
-	{
-		for (const auto& DefinedProperty : DefinedProperties)
-		{
-			RemapProperty(DefinedProperty.Property, DefinedProperty.Index, *ActorRemapper, DestData);
-		}
-	}*/
-
 	return CurrentSourceText;
 }
 
diff --git a/Engine/Source/Editor/UnrealEd/Private/Factories/EditorFactories.cpp b/Engine/Source/Editor/UnrealEd/Private/Factories/EditorFactories.cpp
--- a/Engine/Source/Editor/UnrealEd/Private/Factories/EditorFactories.cpp
+++ b/Engine/Source/Editor/UnrealEd/Private/Factories/EditorFactories.cpp
@@ -720,7 +720,6 @@ UObject* ULevelFactory::FactoryCreateText
 
 	UPackage* RootMapPackage = Cast<UPackage>(InParent);
 	TMap<FString, UPackage*> MapPackages;
-	TMap<AActor*, AActor*> MapActors;
 	// Assumes data is being imported over top of a new, valid map.
 	FParse::Next( &Buffer );
 	if (GetBEGIN(&Buffer, TEXT("MAP")))
@@ -797,8 +796,8 @@ UObject* ULevelFactory::FactoryCreateText
 	// Maintain a lookup for the new actors, keyed by their source FName.
 	TMap<FName, AActor*> NewActorsFNames;
 
-	// Maintain a lookup from existing to new actors, used when replacing internal references when copy+pasting / duplicating
-	TMap<AActor*, AActor*> ExistingToNewMap;
+	// Maintain a lookup from existing to new objects, typically actors, used when replacing internal references when copy+pasting / duplicating
+	TMap<FSoftObjectPath, UObject*> ExistingToNewMap;
 
 	// Maintain a lookup of the new actors to their parent and socket attachment if provided.
 	struct FAttachmentDetail
@@ -859,7 +858,7 @@ UObject* ULevelFactory::FactoryCreateText
 		else if( GetBEGIN(&Str,TEXT("ACTOR")) )
 		{
 			UClass* TempClass;
-			if( ParseObject<UClass>( Str, TEXT("CLASS="), TempClass, nullptr ) )
+			if( ParseObject<UClass>( Str, TEXT("CLASS="), TempClass, nullptr, EParseObjectLoadingPolicy::FindOrLoad ) )
 			{
 				// Get actor name.
 				FName ActorUniqueName(NAME_None);
@@ -872,6 +871,10 @@ UObject* ULevelFactory::FactoryCreateText
 				bool bNeedGloballyUniqueName = World->GetCurrentLevel()->IsUsingExternalActors() && CastChecked<AActor>(TempClass->GetDefaultObject())->SupportsExternalPackaging();
 				ActorUniqueName = FActorSpawnUtils::MakeUniqueActorName(World->GetCurrentLevel(), TempClass, FActorSpawnUtils::GetBaseName(ActorUniqueName), bNeedGloballyUniqueName);
 
+				// Get the full ExportPath for this actor (used to redirection if the existing actor cannot be found)
+				FString ExportedActorFullName;
+				FParse::Value(Str, TEXT("ExportPath="), ExportedActorFullName);
+
 				// Get parent name for attachment.
 				FName ActorParentName(NAME_None);
 				FParse::Value( Str, TEXT("ParentActor="), ActorParentName );
@@ -1067,9 +1070,6 @@ UObject* ULevelFactory::FactoryCreateText
 							// Store the new actor and the text it should be initialized with.
 							NewActorMap.Add( NewActor, *PropText );
 
-							// Store the copy to original actor mapping
-							MapActors.Add(NewActor, Found);
-
 							// Store the new actor against its source actor name (not the one that may have been made unique)
 							if( ActorSourceName!=NAME_None )
 							{
@@ -1078,6 +1078,10 @@ UObject* ULevelFactory::FactoryCreateText
 								{
 									ExistingToNewMap.Add(Found, NewActor);
 								}
+								else if (!ExportedActorFullName.IsEmpty())
+								{
+									ExistingToNewMap.Add(ExportedActorFullName, NewActor);
+								}
 							}
 
 							// Store the new actor with its parent's FName, and socket FName if applicable
@@ -1341,7 +1345,7 @@ UObject* ULevelFactory::FactoryCreateText
 	}
 
 	// Pass 2: Sort out any attachment parenting on the new actors now that all actors have the correct properties set
-	for( auto It = MapActors.CreateIterator(); It; ++It )
+	for( auto It = NewActorMap.CreateIterator(); It; ++It )
 	{
 		AActor* const Actor = It.Key();
 
diff --git a/Engine/Source/Editor/UnrealEd/Public/Editor.h b/Engine/Source/Editor/UnrealEd/Public/Editor.h
--- a/Engine/Source/Editor/UnrealEd/Public/Editor.h
+++ b/Engine/Source/Editor/UnrealEd/Public/Editor.h
@@ -464,58 +464,41 @@ private:
 struct FImportObjectParams
 {
 	/** the location to import the property values to */
-	uint8*				DestData;
+	uint8*				DestData = nullptr;
 
 	/** pointer to a buffer containing the values that should be parsed and imported */
-	const TCHAR*		SourceText;
+	const TCHAR*		SourceText = nullptr;
 
 	/** the struct for the data we're importing */
-	UStruct*			ObjectStruct;
+	UStruct*			ObjectStruct = nullptr;
 
 	/** the original object that ImportObjectProperties was called for.
 		if SubobjectOuter is a subobject, corresponds to the first object in SubobjectOuter's Outer chain that is not a subobject itself.
 		if SubobjectOuter is not a subobject, should normally be the same value as SubobjectOuter */
-	UObject*			SubobjectRoot;
+	UObject*			SubobjectRoot = nullptr;
 
 	/** the object corresponding to DestData; this is the object that will used as the outer when creating subobjects from definitions contained in SourceText */
-	UObject*			SubobjectOuter;
+	UObject*			SubobjectOuter = nullptr;
 
 	/** output device to use for log messages */
-	FFeedbackContext*	Warn;
+	FFeedbackContext*	Warn = nullptr;
 
 	/** current nesting level */
-	int32					Depth;
+	int32					Depth = 0;
 
 	/** used when importing defaults during script compilation for tracking which line we're currently for the purposes of printing compile errors */
-	int32					LineNumber;
+	int32					LineNumber = INDEX_NONE;
 
 	/** contains the mappings of instanced objects and components to their templates; used when recursively calling ImportObjectProperties; generally
 		not necessary to specify a value when calling this function from other code */
-	FObjectInstancingGraph* InInstanceGraph;
+	FObjectInstancingGraph* InInstanceGraph = nullptr;
 
-	/** provides a mapping from an existing actor to a new instance to which it should be remapped */
-	const TMap<AActor*, AActor*>* ActorRemapper;
+	/** provides a mapping from an existing object, typically an actor, (which may no longer be loaded) to a new instance to which it should be remapped */
+	const TMap<FSoftObjectPath, UObject*>* ObjectRemapper = nullptr;
 
 	/** True if we should call PreEditChange/PostEditChange on the object as it's imported.  Pass false here
 		if you're going to do that on your own. */
-	bool				bShouldCallEditChange;
-
-
-	/** Constructor */
-	FImportObjectParams()
-		: DestData( NULL ),
-		  SourceText( NULL ),
-		  ObjectStruct( NULL ),
-		  SubobjectRoot( NULL ),
-		  SubobjectOuter( NULL ),
-		  Warn( NULL ),
-		  Depth( 0 ),
-		  LineNumber( INDEX_NONE ),
-		  InInstanceGraph( NULL ),
-		  ActorRemapper( NULL ),
-		  bShouldCallEditChange( true )
-	{
-	}
+	bool				bShouldCallEditChange = true;
 };
 
 
@@ -543,7 +526,7 @@ UNREALED_API const TCHAR* ImportObjectProperties( FImportObjectParams& InParams
  * @param	LineNumber			used when importing defaults during script compilation for tracking which line the defaultproperties block begins on
  * @param	InstanceGraph		contains the mappings of instanced objects and components to their templates; used when recursively calling ImportObjectProperties; generally
  *								not necessary to specify a value when calling this function from other code
- * @param	ActorRemapper		used when duplicating actors to remap references from a source actor to the duplicated actor
+ * @param	ObjectRemapper		used when duplicating objects, typically actors, to remap references from a source object to the duplicated object
  *
  * @return	NULL if the default values couldn't be imported
  */
@@ -556,8 +539,8 @@ const TCHAR* ImportObjectProperties(
 	FFeedbackContext*	Warn,
 	int32					Depth,
 	int32					LineNumber = INDEX_NONE,
-	FObjectInstancingGraph* InstanceGraph = NULL,
-	const TMap<AActor*, AActor*>* ActorRemapper = NULL
+	FObjectInstancingGraph* InstanceGraph = nullptr,
+	const TMap<FSoftObjectPath, UObject*>* ObjectRemapper = nullptr
 	);
 
 //
