Add a command line option to the cooking server to skip cooking unsolicted packages.

- This adds `-odsc` which for now just skips unsolicted packages.
- This would only be used when you are launching the cooking server to just process ODSC recompileshader requests and not cooked packages.

diff --git a/Engine/Source/Editor/UnrealEd/Classes/CookOnTheSide/CookOnTheFlyServer.h b/Engine/Source/Editor/UnrealEd/Classes/CookOnTheSide/CookOnTheFlyServer.h
--- a/Engine/Source/Editor/UnrealEd/Classes/CookOnTheSide/CookOnTheFlyServer.h
+++ b/Engine/Source/Editor/UnrealEd/Classes/CookOnTheSide/CookOnTheFlyServer.h
@@ -1480,6 +1480,8 @@ private:
 	bool bIterativeIgnoreExe = false;
 	/** Whether we should calculate the exe's hash; might be true even if bIterativeIgnoreExe is true. */
 	bool bIterativeCalculateExe = true;
+	/** If true this will ignore cooking unsolicited packages.  This is only useful if launching the COTF server to only process shader recompile requests. */
+	bool bIgnoreUnsolicitedPackages = false;
 	/** Timers for tracking how long we have been busy, to manage retries and warnings of deadlock */
 	double SaveBusyStartTimeSeconds = MAX_flt;
 	double SaveBusyRetryTimeSeconds = MAX_flt;
diff --git a/Engine/Source/Editor/UnrealEd/Private/CookOnTheFlyServer.cpp b/Engine/Source/Editor/UnrealEd/Private/CookOnTheFlyServer.cpp
--- a/Engine/Source/Editor/UnrealEd/Private/CookOnTheFlyServer.cpp
+++ b/Engine/Source/Editor/UnrealEd/Private/CookOnTheFlyServer.cpp
@@ -4097,6 +4097,11 @@ void UCookOnTheFlyServer::DumpCrashContext(FCrashContextExtendedWriter& Writer)
 void UCookOnTheFlyServer::ProcessUnsolicitedPackages(TArray<FName>* OutDiscoveredPackageNames,
 	TMap<FName, UE::Cook::FInstigator>* OutInstigators)
 {
+	if (bIgnoreUnsolicitedPackages)
+	{
+		return;
+	}
+
 	using namespace UE::Cook;
 
 	TMap<UPackage*, UE::Cook::FInstigator> NewPackages = PackageTracker->GetNewPackages();
@@ -6773,6 +6778,8 @@ void UCookOnTheFlyServer::SetInitializeConfigSettings(UE::Cook::FInitializeConfi
 	// Calculate the exe hash if IterativeExeInvalidation is required by ini OR required by commandline
 	// It would be better to always calculate it, but we want to avoid the performance cost until it becomes more widely used
 	bIterativeCalculateExe = !bIterativeIgnoreExe || !bConfigSettingSetIterativeIgnoreExe;
+
+	bIgnoreUnsolicitedPackages = FParse::Param(FCommandLine::Get(), TEXT("odsc"));
 }
 
 void UCookOnTheFlyServer::ParseCookFilters()
