Updated multi-step ObjectRemapper to key against FSoftObjectPath for consistency with ImportObjectProperties

diff --git a/Engine/Source/Editor/UnrealEd/Private/EditorObject.cpp b/Engine/Source/Editor/UnrealEd/Private/EditorObject.cpp
--- a/Engine/Source/Editor/UnrealEd/Private/EditorObject.cpp
+++ b/Engine/Source/Editor/UnrealEd/Private/EditorObject.cpp
@@ -917,7 +917,7 @@ static const TCHAR* ImportCreateSubObjectsStep(
 	FFeedbackContext* Warn,
 	int32 Depth,
 	FObjectInstancingGraph& InstanceGraph,
-	TMap<FString, UObject*>* ObjectRemapper
+	TMap<FSoftObjectPath, UObject*>* ObjectRemapper
 	)
 {
 	check(ObjectStruct!=nullptr);
@@ -1307,7 +1307,7 @@ static const TCHAR* ImportPropertiesStep(
 	FFeedbackContext* Warn,
 	int32 Depth,
 	FObjectInstancingGraph& InstanceGraph,
-	TMap<FString, UObject*>* ObjectRemapper,
+	TMap<FSoftObjectPath, UObject*>* ObjectRemapper,
 	TSet<FProperty*>* PropertiesToSkip
 	)
 {
@@ -1502,10 +1502,9 @@ static const TCHAR* ImportPropertiesStep(
 							{ 
 								// Select the next character
 								PtrToEqual++;
-								FString PropertyValue(PtrToEqual, FullPropertyText.Len() - (PtrToEqual - FullPropertyText.GetData()));
-								if (UObject* const* PointerToObject = ObjectRemapper->Find(PropertyValue))
+								FStringView PropertyValue(PtrToEqual, FullPropertyText.Len() - (PtrToEqual - FullPropertyText.GetData()));
+								if (UObject* const* PointerToObject = ObjectRemapper->Find(FSoftObjectPath(PropertyValue)))
 								{
-
 									FString RedirectedFullName = FObjectPropertyBase::GetExportPath(*PointerToObject, nullptr, nullptr, PortFlags | PPF_Delimited);
 
 									FStringView TextBeforeValue(FullPropertyText.GetData(), PtrToEqual - FullPropertyText.GetData());
diff --git a/Engine/Source/Editor/UnrealEd/Private/Elements/Actor/ActorElementEditorCopyAndPaste.cpp b/Engine/Source/Editor/UnrealEd/Private/Elements/Actor/ActorElementEditorCopyAndPaste.cpp
--- a/Engine/Source/Editor/UnrealEd/Private/Elements/Actor/ActorElementEditorCopyAndPaste.cpp
+++ b/Engine/Source/Editor/UnrealEd/Private/Elements/Actor/ActorElementEditorCopyAndPaste.cpp
@@ -205,7 +205,7 @@ TArray<FTypedElementHandle> FActorElementEditorPasteImporter::GetImportedElement
 
 void FActorElementEditorPasteImporter::PostImportProcess()
 {
-	TMap<FString, UObject*> InExportedPathToObjectCreated;
+	TMap<FSoftObjectPath, UObject*> InExportedPathToObjectCreated;
 	// CreateObjects
 	{
 		UWorld* World = Context.World;
diff --git a/Engine/Source/Editor/UnrealEd/Public/Editor.h b/Engine/Source/Editor/UnrealEd/Public/Editor.h
--- a/Engine/Source/Editor/UnrealEd/Public/Editor.h
+++ b/Engine/Source/Editor/UnrealEd/Public/Editor.h
@@ -830,7 +830,7 @@ namespace EditorUtilities
 		 * Provides a mapping from an exported path to a new instance to which it should be remapped 
 		 * Imported object will be added to this map when possible during the create objects step.
 		 */
-		TMap<FString, UObject*>* ObjectRemapper = nullptr;
+		TMap<FSoftObjectPath, UObject*>* ObjectRemapper = nullptr;
 
 		/** 
 		 * Tell what properties shouldn't be imported when importing the properties
