Cooker: Improve feedback of warning message "ClientTickPackage was called for package" by adding the platform.
#rnx

diff --git a/Engine/Source/Editor/UnrealEd/Private/Commandlets/AssetRegistryGenerator.cpp b/Engine/Source/Editor/UnrealEd/Private/Commandlets/AssetRegistryGenerator.cpp
--- a/Engine/Source/Editor/UnrealEd/Private/Commandlets/AssetRegistryGenerator.cpp
+++ b/Engine/Source/Editor/UnrealEd/Private/Commandlets/AssetRegistryGenerator.cpp
@@ -2943,8 +2943,8 @@ void FAssetRegistryMPCollector::ClientTickPackage(FMPCollectorClientTickPackageC
 			{
 				// For a successful package, UpdateAssetRegistryData should have been called
 				UE_CLOG(!bLoggedWarning, LogAssetRegistryGenerator, Warning,
-					TEXT("ClientTickPackage was called for package %s, but UpdateAssetRegistryData was not called for that package. We will not have up to date AssetDataTags in the generated AssetRegistry."),
-					*Context.GetPackageName().ToString());
+					TEXT("ClientTickPackage was called for package %s, platform %s, but UpdateAssetRegistryData was not called for that package. We will not have up to date AssetDataTags in the generated AssetRegistry."),
+					*Context.GetPackageName().ToString(), *ContextPlatformData.TargetPlatform->PlatformName());
 				bLoggedWarning = true;
 			}
 		}
