EditorDomain: Turn on in all editor modes the EditorDomain for its role of providing EditorDomainDigest. Separate out its role of injecting itself into the PackageResourceManager to save and load packages from the EditorDomain DDC bucket.
#rnx

diff --git a/Engine/Build/Graph/Tests/IterativeValidate.xml b/Engine/Build/Graph/Tests/IterativeValidate.xml
--- a/Engine/Build/Graph/Tests/IterativeValidate.xml
+++ b/Engine/Build/Graph/Tests/IterativeValidate.xml
@@ -25,8 +25,6 @@
 	<!-- Agent=IterativeValidate variables -->
 	<Property Name="EditorPlatform" Value="Win64"/>
 	<Property Name="CommonCookArgs" Value="-targetplatform=$(TargetPlatform) $(IterativeArgs) -skiponlyeditoronly=true"/>
-	<!-- TODO: Remove these editordomain arguments once we enable editordomain for all cooks -->
-	<Property Name="CommonCookArgs" Value="$(CommonCookArgs) -editordomain -noeditordomainread -noeditordomainwrite"/>
 	<!-- Force singleprocess cook because IterativeValidatePackageWriter includes data that is not replicated -->
 	<Property Name="CommonCookArgs" Value="$(CommonCookArgs) -cookprocesscount=1"/>
 
diff --git a/Engine/Source/Editor/UnrealEd/Private/CookOnTheFlyServer.cpp b/Engine/Source/Editor/UnrealEd/Private/CookOnTheFlyServer.cpp
--- a/Engine/Source/Editor/UnrealEd/Private/CookOnTheFlyServer.cpp
+++ b/Engine/Source/Editor/UnrealEd/Private/CookOnTheFlyServer.cpp
@@ -6068,6 +6068,7 @@ void FSaveCookedPackageContext::FinishPlatform()
 		TArray<UE::DerivedData::FBuildDefinition> BuildDefinitions;
 		FCbObject BuildDefinitionList = UE::TargetDomain::BuildDefinitionListToObject(BuildDefinitions);
 		FCbObject TargetDomainDependencies;
+		if (COTFS.bHybridIterativeEnabled)
 		{
 			UE_SCOPED_HIERARCHICAL_COOKTIMER(TargetDomainDependencies);
 			TargetDomainDependencies = UE::TargetDomain::CollectDependenciesObject(Package, TargetPlatform, nullptr /* ErrorMessage */);
diff --git a/Engine/Source/Editor/UnrealEd/Private/Cooker/CookPackageData.cpp b/Engine/Source/Editor/UnrealEd/Private/Cooker/CookPackageData.cpp
--- a/Engine/Source/Editor/UnrealEd/Private/Cooker/CookPackageData.cpp
+++ b/Engine/Source/Editor/UnrealEd/Private/Cooker/CookPackageData.cpp
@@ -993,7 +993,8 @@ bool FPackageData::TryPreload()
 	}
 	if (!PreloadableFile.Get())
 	{
-		if (FEditorDomain* EditorDomain = FEditorDomain::Get())
+		if (FEditorDomain* EditorDomain(FEditorDomain::Get());
+			EditorDomain && EditorDomain->IsReadingPackages())
 		{
 			EditorDomain->PrecachePackageDigest(GetPackageName());
 		}
diff --git a/Engine/Source/Editor/UnrealEd/Private/Cooker/CookRequestCluster.cpp b/Engine/Source/Editor/UnrealEd/Private/Cooker/CookRequestCluster.cpp
--- a/Engine/Source/Editor/UnrealEd/Private/Cooker/CookRequestCluster.cpp
+++ b/Engine/Source/Editor/UnrealEd/Private/Cooker/CookRequestCluster.cpp
@@ -336,7 +336,7 @@ void FRequestCluster::StartAsync(const FCookerTimer& CookerTimer, bool& bOutComp
 	}
 
 	FEditorDomain* EditorDomain = FEditorDomain::Get();
-	if (EditorDomain)
+	if (EditorDomain && EditorDomain->IsReadingPackages())
 	{
 		bool bBatchDownloadEnabled = true;
 		GConfig->GetBool(TEXT("EditorDomain"), TEXT("BatchDownloadEnabled"), bBatchDownloadEnabled, GEditorIni);
diff --git a/Engine/Source/Editor/UnrealEd/Private/EditorDomain/EditorDomain.cpp b/Engine/Source/Editor/UnrealEd/Private/EditorDomain/EditorDomain.cpp
--- a/Engine/Source/Editor/UnrealEd/Private/EditorDomain/EditorDomain.cpp
+++ b/Engine/Source/Editor/UnrealEd/Private/EditorDomain/EditorDomain.cpp
@@ -60,47 +60,71 @@ DEFINE_LOG_CATEGORY(LogEditorDomain);
  * for usage by any system that needs PackageDigests, with or without FEditorDomain.
  * If the EditorDomain is enabled, also sets PackageResourceManager to use the EditorDomain as the IPackageResourceManager
  */
-class FEditorDomainRegisterAsPackageResourceManager
+class FEditorDomainConstructor
 {
 public:
-	FEditorDomainRegisterAsPackageResourceManager()
+	FEditorDomainConstructor()
 	{
 		IPackageResourceManager::GetSetPackageResourceManagerDelegate().BindStatic(SetPackageResourceManager);
+		IPackageResourceManager::GetOnClearPackageResourceManagerDelegate().AddStatic(OnClearPackageResourceManager);
+	}
+
+	static void ConditionalConstruct()
+	{
+		ConditionalConstruct(IsEditorDomainEnabled());
+	}
+
+	static void ConditionalConstruct(EEditorDomainEnabled Enabled)
+	{
+		check(FEditorDomain::SingletonEditorDomain == nullptr);
+		if ((uint32)Enabled >= (uint32)EEditorDomainEnabled::Utilities)
+		{
+			UE::EditorDomain::UtilsInitialize();
+			UE::TargetDomain::UtilsInitialize();
+			FEditorDomain::SingletonEditorDomain = new FEditorDomain(Enabled);
+			UE::EditorDomain::IPackageDigestCache::Set(FEditorDomain::SingletonEditorDomain);
+		}
+		else
+		{
+			UE::EditorDomain::IPackageDigestCache::SetDefault();
+		}
 	}
 
 	static IPackageResourceManager* SetPackageResourceManager()
 	{
-		using namespace UE::EditorDomain;
+		EEditorDomainEnabled Enabled = IsEditorDomainEnabled();
+		ConditionalConstruct(Enabled);
 
-		bool bEditorDomainEnabled = IsEditorDomainEnabled();
 		if (GIsEditor)
 		{
-			UE_LOG(LogEditorDomain, Display, TEXT("EditorDomain is %s"), bEditorDomainEnabled ? TEXT("Enabled") : TEXT("Disabled"));
+			UE_LOG(LogEditorDomain, Display, TEXT("EditorDomain is %s"),
+				(uint32)Enabled >= (uint32)EEditorDomainEnabled::PackageResourceManager ? TEXT("Enabled") : TEXT("Disabled"));
 		}
-		UtilsInitialize();
-		if (bEditorDomainEnabled)
+		if ((uint32)Enabled >= (uint32)EEditorDomainEnabled::PackageResourceManager)
 		{
-			UE::TargetDomain::UtilsInitialize(bEditorDomainEnabled);
-			// Set values for config settings EditorDomain depends on
+			check(FEditorDomain::SingletonEditorDomain); // Should have been constructed by ConditionalConstruct
+			// Set values for config settings EditorDomain read/writes depends on
 			GAllowUnversionedContentInEditor = 1;
-
-			// Create the editor domain, register it as the IPackageDigestCache, and return it as the package resource manager
-			check(FEditorDomain::RegisteredEditorDomain == nullptr);
-			FEditorDomain::RegisteredEditorDomain = new FEditorDomain();
-			IPackageDigestCache::Set(FEditorDomain::RegisteredEditorDomain);
-			return FEditorDomain::RegisteredEditorDomain;
+			return FEditorDomain::SingletonEditorDomain;
 		}
 		else
 		{
-			IPackageDigestCache::SetDefault();
 			return nullptr;
 		}
 	}
+
+	static void OnClearPackageResourceManager()
+	{
+		// SetPackageResourceManager may have already deleted the singleton, but in the case
+		// where we don't register, we need to delete it ourselves.
+		delete FEditorDomain::SingletonEditorDomain;
+		check(FEditorDomain::SingletonEditorDomain == nullptr); // Destructor is responsible for clearing the pointer
+	}
 } GRegisterAsPackageResourceManager;
 
-FEditorDomain* FEditorDomain::RegisteredEditorDomain = nullptr;
+FEditorDomain* FEditorDomain::SingletonEditorDomain = nullptr;
 
-FEditorDomain::FEditorDomain()
+FEditorDomain::FEditorDomain(EEditorDomainEnabled EnableLevel)
 {
 	Locks = TRefCountPtr<FLocks>(new FLocks(*this));
 	Workspace.Reset(MakePackageResourceManagerFile());
@@ -116,20 +140,31 @@ FEditorDomain::FEditorDomain()
 	// without needing to call ScanPathsSynchronous
 	AssetRegistry->SearchAllAssets(false /* bSynchronousSearch */);
 
-	bEditorDomainReadEnabled = !FParse::Param(FCommandLine::Get(), TEXT("noeditordomainread"))
-	 && !FParse::Param(FCommandLine::Get(), TEXT("testeditordomaindeterminism"));
-	bEditorDomainWriteEnabled = !FParse::Param(FCommandLine::Get(), TEXT("noeditordomainwrite"));
-
-	ELoadingPhase::Type CurrentPhase = IPluginManager::Get().GetLastCompletedLoadingPhase();
-	if (CurrentPhase == ELoadingPhase::None || CurrentPhase < ELoadingPhase::PostEngineInit)
+	if (EnableLevel >= EEditorDomainEnabled::PackageResourceManager)
 	{
-		FCoreDelegates::OnPostEngineInit.AddRaw(this, &FEditorDomain::OnPostEngineInit);
+		bEditorDomainReadEnabled = !FParse::Param(FCommandLine::Get(), TEXT("noeditordomainread"))
+			&& !FParse::Param(FCommandLine::Get(), TEXT("testeditordomaindeterminism"));
+		bEditorDomainWriteEnabled = !FParse::Param(FCommandLine::Get(), TEXT("noeditordomainwrite"));
 	}
 	else
 	{
-		OnPostEngineInit();
+		bEditorDomainReadEnabled = false;
+		bEditorDomainWriteEnabled = false;
+	}
+
+	if (bEditorDomainWriteEnabled)
+	{
+		ELoadingPhase::Type CurrentPhase = IPluginManager::Get().GetLastCompletedLoadingPhase();
+		if (CurrentPhase == ELoadingPhase::None || CurrentPhase < ELoadingPhase::PostEngineInit)
+		{
+			FCoreDelegates::OnPostEngineInit.AddRaw(this, &FEditorDomain::OnPostEngineInit);
+		}
+		else
+		{
+			OnPostEngineInit();
+		}
+		FCoreUObjectDelegates::OnEndLoadPackage.AddRaw(this, &FEditorDomain::OnEndLoadPackage);
 	}
-	FCoreUObjectDelegates::OnEndLoadPackage.AddRaw(this, &FEditorDomain::OnEndLoadPackage);
 	UPackage::PackageSavedWithContextEvent.AddRaw(this, &FEditorDomain::OnPackageSavedWithContext);
 	AssetRegistry->OnAssetUpdatedOnDisk().AddRaw(this, &FEditorDomain::OnAssetUpdatedOnDisk);
 }
@@ -154,9 +189,9 @@ FEditorDomain::~FEditorDomain()
 	AssetRegistry = nullptr;
 	Workspace.Reset();
 
-	if (RegisteredEditorDomain == this)
+	if (SingletonEditorDomain == this)
 	{
-		RegisteredEditorDomain = nullptr;
+		SingletonEditorDomain = nullptr;
 	}
 	if (UE::EditorDomain::IPackageDigestCache::Get() == this)
 	{
@@ -166,7 +201,7 @@ FEditorDomain::~FEditorDomain()
 
 FEditorDomain* FEditorDomain::Get()
 {
-	return RegisteredEditorDomain;
+	return SingletonEditorDomain;
 }
 
 bool FEditorDomain::SupportsLocalOnlyPaths()
@@ -197,8 +232,6 @@ bool FEditorDomain::TryFindOrAddPackageSource(FScopeLock& ScopeLock, bool& bOutR
 	// Called within &Locks->Lock, ScopeLock is locked on that Lock
 	using namespace UE::EditorDomain;
 
-	// EDITOR_DOMAIN_TODO: Need to delete entries from PackageSources when the assetregistry reports the package is
-	// resaved on disk.
 	TRefCountPtr<FPackageSource>* PackageSourcePtr = PackageSources.Find(PackageName);
 	if (PackageSourcePtr && *PackageSourcePtr)
 	{
diff --git a/Engine/Source/Editor/UnrealEd/Private/EditorDomain/EditorDomain.h b/Engine/Source/Editor/UnrealEd/Private/EditorDomain/EditorDomain.h
--- a/Engine/Source/Editor/UnrealEd/Private/EditorDomain/EditorDomain.h
+++ b/Engine/Source/Editor/UnrealEd/Private/EditorDomain/EditorDomain.h
@@ -133,7 +133,7 @@ public:
 		Editor
 	};
 
-	FEditorDomain();
+	FEditorDomain(EEditorDomainEnabled EnableLevel);
 	virtual ~FEditorDomain();
 
 	/** Return the EditorDomain that is registered as the global PackageResourceManager, if there is one. */
@@ -178,6 +178,8 @@ public:
 	// EditorDomain interface
 	/** Fetch data from game-thread sources that is required to calculate the PackageDigest of the given PackageName. */
 	void PrecachePackageDigest(FName PackageName);
+	bool IsReadingPackages() const { return bEditorDomainReadEnabled; }
+	bool IsWritingPackages() const { return bEditorDomainWriteEnabled; }
 
 	/** Request the download of the given packages from the upstream DDC server. */
 	void BatchDownload(TArrayView<FName> PackageNames);
@@ -279,9 +281,9 @@ private:
 	 */
 	bool bSkipSavesUntilCatalogLoaded = false;
 
-	static FEditorDomain* RegisteredEditorDomain;
+	static FEditorDomain* SingletonEditorDomain;
 
-	friend class FEditorDomainRegisterAsPackageResourceManager;
+	friend class FEditorDomainConstructor;
 	friend class FEditorDomainPackageSegments;
 	friend class FEditorDomainReadArchive;
 	friend class FEditorDomainAsyncReadFileHandle;
diff --git a/Engine/Source/Editor/UnrealEd/Private/TargetDomain/TargetDomainUtils.cpp b/Engine/Source/Editor/UnrealEd/Private/TargetDomain/TargetDomainUtils.cpp
--- a/Engine/Source/Editor/UnrealEd/Private/TargetDomain/TargetDomainUtils.cpp
+++ b/Engine/Source/Editor/UnrealEd/Private/TargetDomain/TargetDomainUtils.cpp
@@ -669,20 +669,17 @@ void CommitEditorDomainCookAttachments(FName PackageName, TArrayView<IPackageWri
 	GEditorDomainOplog->CommitPackage(PackageName, Attachments);
 }
 
-void UtilsInitialize(bool bEditorDomainEnabled)
+void UtilsInitialize()
 {
-	if (bEditorDomainEnabled)
+	bool bCookAttachmentsEnabled = true;
+	GConfig->GetBool(TEXT("EditorDomain"), TEXT("CookAttachmentsEnabled"), bCookAttachmentsEnabled, GEditorIni);
+	if (bCookAttachmentsEnabled)
 	{
-		bool bCookAttachmentsEnabled = true;
-		GConfig->GetBool(TEXT("EditorDomain"), TEXT("CookAttachmentsEnabled"), bCookAttachmentsEnabled, GEditorIni);
-		if (bCookAttachmentsEnabled)
+		GEditorDomainOplog = MakeUnique<FEditorDomainOplog>();
+		if (!GEditorDomainOplog->IsValid())
 		{
-			GEditorDomainOplog = MakeUnique<FEditorDomainOplog>();
-			if (!GEditorDomainOplog->IsValid())
-			{
-				UE_LOG(LogEditorDomain, Display, TEXT("Failed to connect to ZenServer; EditorDomain oplog is unavailable."));
-				GEditorDomainOplog.Reset();
-			}
+			UE_LOG(LogEditorDomain, Display, TEXT("Failed to connect to ZenServer; EditorDomain oplog is unavailable."));
+			GEditorDomainOplog.Reset();
 		}
 	}
 }
diff --git a/Engine/Source/Editor/UnrealEd/Private/TargetDomain/TargetDomainUtils.h b/Engine/Source/Editor/UnrealEd/Private/TargetDomain/TargetDomainUtils.h
--- a/Engine/Source/Editor/UnrealEd/Private/TargetDomain/TargetDomainUtils.h
+++ b/Engine/Source/Editor/UnrealEd/Private/TargetDomain/TargetDomainUtils.h
@@ -19,7 +19,7 @@ template <typename FuncType> class TUniqueFunction;
 namespace UE::TargetDomain
 {
 
-void UtilsInitialize(bool bEditorDomainEnabled);
+void UtilsInitialize();
 
 /** Create the TargetDomainKey based on the EditorDomainKeys of the Package and its dependencies. */
 bool TryCreateKey(FName PackageName, TConstArrayView<FName> SortedBuildDependencies, FIoHash* OutHash, FString* OutErrorMessage);
diff --git a/Engine/Source/Runtime/CoreUObject/Private/Serialization/BulkDataRegistry.cpp b/Engine/Source/Runtime/CoreUObject/Private/Serialization/BulkDataRegistry.cpp
--- a/Engine/Source/Runtime/CoreUObject/Private/Serialization/BulkDataRegistry.cpp
+++ b/Engine/Source/Runtime/CoreUObject/Private/Serialization/BulkDataRegistry.cpp
@@ -130,7 +130,7 @@ void IBulkDataRegistry::Initialize()
 			GBulkDataRegistry = SetRegistryDelegate.Execute();
 		}
 	}
-	else if (IsEditorDomainEnabled())
+	else if (IsEditorDomainEnabled() >= EEditorDomainEnabled::PackageResourceManager)
 	{
 		GBulkDataRegistry = new FBulkDataRegistryTrackBulkDataToResave();
 	}
diff --git a/Engine/Source/Runtime/CoreUObject/Private/UObject/PackageResourceManager.cpp b/Engine/Source/Runtime/CoreUObject/Private/UObject/PackageResourceManager.cpp
--- a/Engine/Source/Runtime/CoreUObject/Private/UObject/PackageResourceManager.cpp
+++ b/Engine/Source/Runtime/CoreUObject/Private/UObject/PackageResourceManager.cpp
@@ -18,6 +18,7 @@ namespace
 {
 	IPackageResourceManager* GPackageResourceManager = nullptr;
 	FSetPackageResourceManager GSetPackageResourceManagerDelegate;
+	FOnClearPackageResourceManager GOnClearPackageResourceManagerDelegate;
 }
 
 IPackageResourceManager& IPackageResourceManager::Get()
@@ -31,6 +32,11 @@ FSetPackageResourceManager& IPackageResourceManager::GetSetPackageResourceManage
 	return GSetPackageResourceManagerDelegate;
 }
 
+FOnClearPackageResourceManager& IPackageResourceManager::GetOnClearPackageResourceManagerDelegate()
+{
+	return GOnClearPackageResourceManagerDelegate;
+}
+
 void IPackageResourceManager::Initialize()
 {
 	FSetPackageResourceManager& SetManagerDelegate = GetSetPackageResourceManagerDelegate(); 
@@ -57,6 +63,8 @@ void IPackageResourceManager::Shutdown()
 {
 	delete GPackageResourceManager;
 	GPackageResourceManager = nullptr;
+	FOnClearPackageResourceManager& OnClearManagerDelegate = GetOnClearPackageResourceManagerDelegate();
+	OnClearManagerDelegate.Broadcast();
 }
 
 bool IPackageResourceManager::DoesPackageExist(const FPackagePath& PackagePath, FPackagePath* OutUpdatedPath)
@@ -277,23 +285,23 @@ FOpenAsyncPackageResult::~FOpenAsyncPackageResult()
 }
 
 #if WITH_EDITOR
-bool IsEditorDomainEnabled()
+EEditorDomainEnabled IsEditorDomainEnabled()
 {
-	static bool bEnabled = []()
+	static EEditorDomainEnabled EnabledValue = []()
 	{
 		if (!GIsEditor || (IsRunningCommandlet() && !IsRunningCookCommandlet()))
 		{
-			return false;
+			return EEditorDomainEnabled::Disabled;
 		}
 
 		if (FParse::Param(FCommandLine::Get(), TEXT("NoEditorDomain")))
 		{
-			return false;
+			return EEditorDomainEnabled::Utilities;
 		}
 
 		if (FParse::Param(FCommandLine::Get(), TEXT("EditorDomain")))
 		{
-			return true;
+			return EEditorDomainEnabled::PackageResourceManager;
 		}
 
 		bool bEnabledByConfig = true;
@@ -303,8 +311,8 @@ bool IsEditorDomainEnabled()
 			UE_LOG(LogPackageResourceManager, Error,
 				TEXT("Editor.ini:[CookSettings]:EditorDomainEnabled is deprecated, use Editor.ini:[EditorDomain]:EditorDomainEnabled instead."));
 		}
-		return bEnabledByConfig;
+		return bEnabledByConfig ? EEditorDomainEnabled::PackageResourceManager : EEditorDomainEnabled::Utilities;
 	}();
-	return bEnabled;
+	return EnabledValue;
 }
 #endif
diff --git a/Engine/Source/Runtime/CoreUObject/Public/UObject/PackageResourceManager.h b/Engine/Source/Runtime/CoreUObject/Public/UObject/PackageResourceManager.h
--- a/Engine/Source/Runtime/CoreUObject/Public/UObject/PackageResourceManager.h
+++ b/Engine/Source/Runtime/CoreUObject/Public/UObject/PackageResourceManager.h
@@ -34,6 +34,7 @@ struct FFileStatData;
 
 DECLARE_LOG_CATEGORY_EXTERN(LogPackageResourceManager, Log, All);
 DECLARE_DELEGATE_RetVal(IPackageResourceManager*, FSetPackageResourceManager);
+DECLARE_MULTICAST_DELEGATE(FOnClearPackageResourceManager);
 
 /**
  * Format for a package payload reported by the PackageResourceManager.
@@ -335,6 +336,7 @@ public:
 	COREUOBJECT_API static void Shutdown();
 
 	COREUOBJECT_API static FSetPackageResourceManager& GetSetPackageResourceManagerDelegate();
+	COREUOBJECT_API static FOnClearPackageResourceManager& GetOnClearPackageResourceManagerDelegate();
 
 public:
 	// Internal API used by low-level PackageResourceManager users
@@ -427,9 +429,15 @@ private:
 };
 
 #if WITH_EDITOR
+enum class EEditorDomainEnabled : uint8
+{
+	Disabled,
+	Utilities,
+	PackageResourceManager,
+};
 /**
  * Report whether the EditorDomain is enabled by config, for systems that need to behave differently if editordomain
  * will be enabled and need to know before the PackageResourceManager is constructed.
  */
-COREUOBJECT_API bool IsEditorDomainEnabled();
+COREUOBJECT_API EEditorDomainEnabled IsEditorDomainEnabled();
 #endif
