Fixed static analysis warning about accessing a null pointer

#rnx

diff --git a/Engine/Source/Editor/UnrealEd/Private/EditorObject.cpp b/Engine/Source/Editor/UnrealEd/Private/EditorObject.cpp
--- a/Engine/Source/Editor/UnrealEd/Private/EditorObject.cpp
+++ b/Engine/Source/Editor/UnrealEd/Private/EditorObject.cpp
@@ -166,11 +166,11 @@ static void RemapProperty(FProperty* Property, int32 Index, const TMap<FSoftObje
 		if (StructProperty->Struct == TBaseStructure<FSoftObjectPath>::Get())
 		{
 			// If there's a concrete index, use that, otherwise iterate all array members (for the case that this property is inside a struct, or there is exactly one element)
-			const int32 Num = (Index == INDEX_NONE) ? ObjectProperty->ArrayDim : 1;
+			const int32 Num = (Index == INDEX_NONE) ? StructProperty->ArrayDim : 1;
 			const int32 StartIndex = (Index == INDEX_NONE) ? 0 : Index;
 			for (int32 Count = 0; Count < Num; Count++)
 			{
-				FSoftObjectPath* SoftObjectPathPtr = ObjectProperty->ContainerPtrToValuePtr<FSoftObjectPath>(DestData, StartIndex + Count);
+				FSoftObjectPath* SoftObjectPathPtr = StructProperty->ContainerPtrToValuePtr<FSoftObjectPath>(DestData, StartIndex + Count);
 				if (UObject* const* RemappedObject = ObjectRemapper.Find(*SoftObjectPathPtr))
 				{
 					*SoftObjectPathPtr = *RemappedObject;
