EditorDomain: Fix static analysis compile error.
#rnx

diff --git a/Engine/Source/Editor/UnrealEd/Private/EditorDomain/EditorDomain.cpp b/Engine/Source/Editor/UnrealEd/Private/EditorDomain/EditorDomain.cpp
--- a/Engine/Source/Editor/UnrealEd/Private/EditorDomain/EditorDomain.cpp
+++ b/Engine/Source/Editor/UnrealEd/Private/EditorDomain/EditorDomain.cpp
@@ -118,7 +118,8 @@ public:
 		// SetPackageResourceManager may have already deleted the singleton, but in the case
 		// where we don't register, we need to delete it ourselves.
 		delete FEditorDomain::SingletonEditorDomain;
-		check(FEditorDomain::SingletonEditorDomain == nullptr); // Destructor is responsible for clearing the pointer
+		// Constructor already sets it to null, but set it to null again to confirm and to avoid C6001: Using uninitialized memory 
+		FEditorDomain::SingletonEditorDomain = nullptr;
 	}
 } GRegisterAsPackageResourceManager;
 
