Cooker Logs: Suppress Rejected AlreadyCooked message in DLCs since they cause a lot of spam.
#rnx

diff --git a/Engine/Source/Editor/UnrealEd/Private/CookOnTheFlyServer.cpp b/Engine/Source/Editor/UnrealEd/Private/CookOnTheFlyServer.cpp
--- a/Engine/Source/Editor/UnrealEd/Private/CookOnTheFlyServer.cpp
+++ b/Engine/Source/Editor/UnrealEd/Private/CookOnTheFlyServer.cpp
@@ -2626,6 +2626,11 @@ void UCookOnTheFlyServer::DemoteToIdle(UE::Cook::FPackageData& PackageData, UE::
 
 		// Suppress the message in cases that cause large spam like NotInCurrentPlugin for DLC cooks.
 		bPrintDiagnostic &= (Reason != ESuppressCookReason::NotInCurrentPlugin);
+		if (bPrintDiagnostic && IsCookingDLC() && Reason == ESuppressCookReason::AlreadyCooked && LogCook.GetVerbosity() < ELogVerbosity::Verbose)
+		{
+			bPrintDiagnostic = false;
+		}
+
 		if (bPrintDiagnostic)
 		{
 			TStringBuilder<256> PackageNameStr(InPlace, PackageData.GetPackageName());
