[BugFix] false positive diff for UObject instances inside level blueprint's
(changed the syntax of DiffUtils::Identical to require outer objects for the properties)

diff --git a/Engine/Source/Editor/Kismet/Private/DetailsDiff.cpp b/Engine/Source/Editor/Kismet/Private/DetailsDiff.cpp
--- a/Engine/Source/Editor/Kismet/Private/DetailsDiff.cpp
+++ b/Engine/Source/Editor/Kismet/Private/DetailsDiff.cpp
@@ -264,10 +264,12 @@ void FDetailsDiff::DiffAgainst(const FDetailsDiff& Newer, TArray< FSingleObjectD
 		check(NewSelectedObjects.Num() == 1);
 		FResolvedProperty OldProperty = CommonProperty.Resolve(OldSelectedObjects[0].Get());
 		FResolvedProperty NewProperty = CommonProperty.Resolve(NewSelectedObjects[0].Get());
+		const UPackage* OldPackage = OldSelectedObjects[0]->GetPackage();
+		const UPackage* NewPackage = NewSelectedObjects[0]->GetPackage();
 
 		TArray<FPropertySoftPath> DifferingSubProperties;
 
-		if (!DiffUtils::Identical(OldProperty, NewProperty, CommonProperty, DifferingSubProperties))
+		if (!DiffUtils::Identical(OldProperty, NewProperty, OldPackage, NewPackage, CommonProperty, DifferingSubProperties))
 		{
 			for (const FPropertySoftPath& DifferingSubProperty : DifferingSubProperties)
 			{
diff --git a/Engine/Source/Editor/Kismet/Private/DiffControl.cpp b/Engine/Source/Editor/Kismet/Private/DiffControl.cpp
--- a/Engine/Source/Editor/Kismet/Private/DiffControl.cpp
+++ b/Engine/Source/Editor/Kismet/Private/DiffControl.cpp
@@ -146,7 +146,12 @@ void FMyBlueprintDiffControl::GenerateTreeEntries(TArray<TSharedPtr<FBlueprintDi
 				if (OldVarIndex != INDEX_NONE && NewVarIndex != INDEX_NONE)
 				{
 					TArray<FSingleObjectDiffEntry> DifferingProperties;
-					DiffUtils::CompareUnrelatedStructs(FBPVariableDescription::StaticStruct(), &OldBlueprint->NewVariables[OldVarIndex], FBPVariableDescription::StaticStruct(), &NewBlueprint->NewVariables[NewVarIndex], DifferingProperties);
+					static const UStruct* Struct = FBPVariableDescription::StaticStruct();
+					const void* OldVal = &OldBlueprint->NewVariables[OldVarIndex];
+					const void* NewVal = &NewBlueprint->NewVariables[NewVarIndex];
+					const UPackage* OldPackage = OldBlueprint->GetPackage();
+					const UPackage* NewPackage = NewBlueprint->GetPackage();
+					DiffUtils::CompareUnrelatedStructs(Struct, OldVal, OldPackage, Struct, NewVal, NewPackage, DifferingProperties);
 					for (const FSingleObjectDiffEntry& Difference : DifferingProperties)
 					{
 						TSharedPtr<FBlueprintDifferenceTreeEntry> Entry = MakeShared<FBlueprintDifferenceTreeEntry>(
diff --git a/Engine/Source/Editor/PropertyEditor/Private/AsyncDetailViewDiff.cpp b/Engine/Source/Editor/PropertyEditor/Private/AsyncDetailViewDiff.cpp
--- a/Engine/Source/Editor/PropertyEditor/Private/AsyncDetailViewDiff.cpp
+++ b/Engine/Source/Editor/PropertyEditor/Private/AsyncDetailViewDiff.cpp
@@ -68,8 +68,7 @@ bool TreeDiffSpecification::AreValuesEqual<TWeakPtr<FDetailTreeNode>>(const TWea
 
 			// use DiffUtils::Identical instead of FProperty::Identical so that sub-object pointers are considered identical
 			// if they're isomorphic rather than only if they have the same pointer values
-			TArray<FPropertySoftPath> Differing;
-			if(!DiffUtils::Identical(ResolvedA, ResolvedB, {}, Differing))
+			if(!DiffUtils::Identical(ResolvedA, ResolvedB, OwningObjectA->GetPackage(), OwningObjectB->GetPackage()))
 			{
 				return false;
 			}
diff --git a/Engine/Source/Editor/UnrealEd/Private/DiffUtils.cpp b/Engine/Source/Editor/UnrealEd/Private/DiffUtils.cpp
--- a/Engine/Source/Editor/UnrealEd/Private/DiffUtils.cpp
+++ b/Engine/Source/Editor/UnrealEd/Private/DiffUtils.cpp
@@ -639,7 +639,14 @@ const UObject* DiffUtils::GetCDO(const UBlueprint* ForBlueprint)
 	return ForBlueprint->GeneratedClass->ClassDefaultObject;
 }
 
-void DiffUtils::CompareUnrelatedStructs(const UStruct* StructA, const void* A, const UStruct* StructB, const void* B, TArray<FSingleObjectDiffEntry>& OutDifferingProperties)
+void DiffUtils::CompareUnrelatedStructs(const UStruct* StructA, const void* A, const UStruct* StructB, const void* B,
+	TArray<FSingleObjectDiffEntry>& OutDifferingProperties)
+{
+	CompareUnrelatedStructs(StructA, A, nullptr, StructB, B, nullptr, OutDifferingProperties);
+}
+
+void DiffUtils::CompareUnrelatedStructs(const UStruct* StructA, const void* A, const UObject* OwningOuterA, const UStruct* StructB, const void* B,
+                                        const UObject* OwningOuterB, TArray<FSingleObjectDiffEntry>& OutDifferingProperties)
 {
 	FPropertySoftPathSet PropertiesInA = UEDiffUtils_Private::GetPropertyNameSet(StructA);
 	FPropertySoftPathSet PropertiesInB = UEDiffUtils_Private::GetPropertyNameSet(StructB);
@@ -669,7 +676,7 @@ void DiffUtils::CompareUnrelatedStructs(const UStruct* StructA, const void* A, c
 
 			check(AProp != FResolvedProperty() && BProp != FResolvedProperty());
 			TArray<FPropertySoftPath> DifferingSubProperties;
-			if (!DiffUtils::Identical(AProp, BProp, PropertyName, DifferingSubProperties))
+			if (!DiffUtils::Identical(AProp, BProp, OwningOuterA, OwningOuterB, PropertyName, DifferingSubProperties))
 			{
 				for (int DifferingIndex = 0; DifferingIndex < DifferingSubProperties.Num(); DifferingIndex++)
 				{
@@ -684,7 +691,7 @@ void DiffUtils::CompareUnrelatedObjects(const UObject* A, const UObject* B, TArr
 {
 	if (A && B)
 	{
-		return CompareUnrelatedStructs(A->GetClass(), A, B->GetClass(), B, OutDifferingProperties);
+		return CompareUnrelatedStructs(A->GetClass(), A, A->GetPackage(), B->GetClass(), B, B->GetPackage(), OutDifferingProperties);
 	}
 }
 
@@ -784,7 +791,10 @@ static void AdvanceMapIterator( FScriptMapHelper& MapHelper, int32& Index)
 	while(Index < MapHelper.GetMaxIndex() && !MapHelper.IsValidIndex(Index));
 }
 
-static void IdenticalHelper(const FProperty* AProperty, const FProperty* BProperty, const void* AValue, const void* BValue, const FPropertySoftPath& RootPath, TArray<FPropertySoftPath>& DifferingSubProperties, bool bStaticArrayHandled = false)
+
+static void IdenticalHelper(const FProperty* AProperty, const FProperty* BProperty, const void* AValue, const void* BValue,
+	const UObject* OwningOuterA, const UObject* OwningOuterB, const FPropertySoftPath& RootPath,
+	TArray<FPropertySoftPath>& DifferingSubProperties, bool bStaticArrayHandled = false)
 {
 	if(AProperty == nullptr || BProperty == nullptr || AProperty->ArrayDim != BProperty->ArrayDim || AProperty->GetClass() != BProperty->GetClass())
 	{
@@ -802,7 +812,8 @@ static void IdenticalHelper(const FProperty* AProperty, const FProperty* BProper
 			const void* CurAValue = reinterpret_cast<const void*>(reinterpret_cast<const uint8*>(AValue) + Offset);
 			const void* CurBValue = reinterpret_cast<const void*>(reinterpret_cast<const uint8*>(BValue) + Offset);
 
-			IdenticalHelper(AProperty, BProperty, CurAValue, CurBValue, FPropertySoftPath(RootPath, I), DifferingSubProperties, true);
+			IdenticalHelper(AProperty, BProperty, CurAValue, CurBValue, OwningOuterA, OwningOuterB,
+				FPropertySoftPath(RootPath, I), DifferingSubProperties, true);
 		}
 
 		return;
@@ -832,7 +843,10 @@ static void IdenticalHelper(const FProperty* AProperty, const FProperty* BProper
 				for (TFieldIterator<FProperty> PropertyIt(APropAsStruct->Struct); PropertyIt; ++PropertyIt)
 				{
 					const FProperty* StructProp = *PropertyIt;
-					IdenticalHelper(StructProp, StructProp, StructProp->ContainerPtrToValuePtr<void>(AValue, 0), StructProp->ContainerPtrToValuePtr<void>(BValue, 0), FPropertySoftPath(RootPath, StructProp), DifferingSubProperties);
+					const void* SubValueA = StructProp->ContainerPtrToValuePtr<void>(AValue, 0);
+					const void* SubValueB = StructProp->ContainerPtrToValuePtr<void>(BValue, 0);
+					IdenticalHelper(StructProp, StructProp, SubValueA, SubValueB,
+						OwningOuterA, OwningOuterB, FPropertySoftPath(RootPath, StructProp), DifferingSubProperties);
 				}
 			}
 		}
@@ -852,7 +866,10 @@ static void IdenticalHelper(const FProperty* AProperty, const FProperty* BProper
 			// note any differences in contained types:
 			for (int32 ArrayIndex = 0; ArrayIndex < ArrayHelperA.Num() && ArrayIndex < ArrayHelperB.Num(); ArrayIndex++)
 			{
-				IdenticalHelper(APropAsArray->Inner, APropAsArray->Inner, ArrayHelperA.GetRawPtr(ArrayIndex), ArrayHelperB.GetRawPtr(ArrayIndex), FPropertySoftPath(RootPath, ArrayIndex), DifferingSubProperties);
+				const void* SubValueA = ArrayHelperA.GetRawPtr(ArrayIndex);
+				const void* SubValueB = ArrayHelperB.GetRawPtr(ArrayIndex);
+				IdenticalHelper(APropAsArray->Inner, BPropAsArray->Inner, SubValueA, SubValueB,
+					OwningOuterA, OwningOuterB, FPropertySoftPath(RootPath, ArrayIndex), DifferingSubProperties);
 			}
 
 			// note any size difference:
@@ -893,7 +910,10 @@ static void IdenticalHelper(const FProperty* AProperty, const FProperty* BProper
 
 			for (int32 VirtualIndex = 0; VirtualIndex < SetSizeA && VirtualIndex < SetSizeB; ++VirtualIndex)
 			{
-				IdenticalHelper(APropAsSet->ElementProp, APropAsSet->ElementProp, SetHelperA.GetElementPtr(SetIndexA), SetHelperB.GetElementPtr(SetIndexB), FPropertySoftPath(RootPath, VirtualIndex), DifferingSubProperties);
+				const void* SubValueA = SetHelperA.GetElementPtr(SetIndexA);
+				const void* SubValueB = SetHelperB.GetElementPtr(SetIndexB);
+				IdenticalHelper(APropAsSet->ElementProp, BPropAsSet->ElementProp, SubValueA, SubValueB,
+					OwningOuterA, OwningOuterB, FPropertySoftPath(RootPath, VirtualIndex), DifferingSubProperties);
 
 				// advance iterators in step:
 				AdvanceSetIterator(SetHelperA, SetIndexA);
@@ -931,8 +951,10 @@ static void IdenticalHelper(const FProperty* AProperty, const FProperty* BProper
 			
 			for (int32 VirtualIndex = 0; VirtualIndex < MapSizeA && VirtualIndex < MapSizeB; ++VirtualIndex)
 			{
-				IdenticalHelper(APropAsMap->KeyProp, APropAsMap->KeyProp, MapHelperA.GetKeyPtr(MapIndexA), MapHelperB.GetKeyPtr(MapIndexB), FPropertySoftPath(RootPath, VirtualIndex), DifferingSubProperties);
-				IdenticalHelper(APropAsMap->ValueProp, APropAsMap->ValueProp, MapHelperA.GetValuePtr(MapIndexA), MapHelperB.GetValuePtr(MapIndexB), FPropertySoftPath(RootPath, VirtualIndex), DifferingSubProperties);
+				IdenticalHelper(APropAsMap->KeyProp, APropAsMap->KeyProp, MapHelperA.GetKeyPtr(MapIndexA), MapHelperB.GetKeyPtr(MapIndexB),
+					OwningOuterA, OwningOuterB, FPropertySoftPath(RootPath, VirtualIndex), DifferingSubProperties);
+				IdenticalHelper(APropAsMap->ValueProp, APropAsMap->ValueProp, MapHelperA.GetValuePtr(MapIndexA), MapHelperB.GetValuePtr(MapIndexB),
+					OwningOuterA, OwningOuterB, FPropertySoftPath(RootPath, VirtualIndex), DifferingSubProperties);
 
 				AdvanceMapIterator(MapHelperA, MapIndexA);
 				AdvanceMapIterator(MapHelperB, MapIndexB);
@@ -956,8 +978,8 @@ static void IdenticalHelper(const FProperty* AProperty, const FProperty* BProper
 		const UObject* A = *((const UObject* const*)AValue);
 		const UObject* B = *((const UObject* const*)BValue);
 
-		// dig into the objects if they are in the same package as our initial object:
-		if(BPropAsObject->HasAnyPropertyFlags(CPF_InstancedReference) && APropAsObject->HasAnyPropertyFlags(CPF_InstancedReference) && A && B && A->GetClass() == B->GetClass())
+		// dig into the objects if they are in the same package as our initial object
+		if (A && B && A->GetClass() == B->GetClass() && OwningOuterA && OwningOuterB && A->IsIn(OwningOuterA) && B->IsIn(OwningOuterB))
 		{
 			const UClass* AClass = A->GetClass(); // BClass and AClass are identical!
 
@@ -968,13 +990,20 @@ static void IdenticalHelper(const FProperty* AProperty, const FProperty* BProper
 				for (TFieldIterator<FProperty> PropertyIt(AClass); PropertyIt; ++PropertyIt)
 				{
 					const FProperty* ClassProp = *PropertyIt;
-					IdenticalHelper(ClassProp, ClassProp, ClassProp->ContainerPtrToValuePtr<void>(A, 0), ClassProp->ContainerPtrToValuePtr<void>(B, 0), FPropertySoftPath(RootPath, ClassProp), DifferingSubProperties);
+					const void* SubValueA = ClassProp->ContainerPtrToValuePtr<void>(A, 0);
+					const void* SubValueB = ClassProp->ContainerPtrToValuePtr<void>(B, 0);
+					IdenticalHelper(ClassProp, ClassProp, SubValueA, SubValueB,
+						OwningOuterA, OwningOuterB, FPropertySoftPath(RootPath, ClassProp), DifferingSubProperties);
 				}
 			}
-			else if (A->GetFName() != B->GetFName())
+			else
 			{
-				// If the names don't match, report that as a difference as the object was likely changed
-				DifferingSubProperties.Push(RootPath);
+				const FString PathA = A->GetPathName(OwningOuterA);
+				const FString PathB = B->GetPathName(OwningOuterB);
+				if (PathA != PathB)
+				{
+					DifferingSubProperties.Push(RootPath);
+				}
 			}
 		}
 		else
@@ -1012,7 +1041,36 @@ bool DiffUtils::Identical(const FResolvedProperty& AProp, const FResolvedPropert
 	// Instead let's write a new function, specific to DiffUtils, that handles the sub properties
 	// NOTE: For Static Arrays, AValue and BValue were, and are, only references to the value at index 0.  So changes to values past index 0 didn't show up before and
 	// won't show up now.  Changes to index 0 will show up as a change to the entire array.
-	IdenticalHelper(AProp.Property, BProp.Property, AValue, BValue, RootPath, DifferingProperties);
+	IdenticalHelper(AProp.Property, BProp.Property, AValue, BValue, nullptr, nullptr, RootPath, DifferingProperties);
+	
+	return DifferingProperties.Num() == 0;
+}
+
+bool DiffUtils::Identical(const FResolvedProperty& AProp, const FResolvedProperty& BProp, const UObject* OwningOuterA,
+	const UObject* OwningOuterB)
+{
+	TArray<FPropertySoftPath> DifferingProperties;
+	return Identical(AProp, BProp, OwningOuterA, OwningOuterB, {}, DifferingProperties);
+}
+
+bool DiffUtils::Identical(const FResolvedProperty& AProp, const FResolvedProperty& BProp, const UObject* OwningOuterA,
+	const UObject* OwningOuterB, const FPropertySoftPath& RootPath, TArray<FPropertySoftPath>& DifferingProperties)
+{
+	if( AProp.Property == nullptr && BProp.Property == nullptr )
+	{
+		return true;
+	}
+	else if( AProp.Property == nullptr || BProp.Property == nullptr )
+	{
+		return false;
+	}
+
+	const void* AValue = AProp.Property->ContainerPtrToValuePtr<void>(AProp.Object);
+	const void* BValue = BProp.Property->ContainerPtrToValuePtr<void>(BProp.Object);
+
+	// note that we're not directly calling FProperty::Identical because sub-object properties should be weakly compared based on
+	// their paths instead of their pointers or data
+	IdenticalHelper(AProp.Property, BProp.Property, AValue, BValue, OwningOuterA, OwningOuterB, RootPath, DifferingProperties);
 	
 	return DifferingProperties.Num() == 0;
 }
diff --git a/Engine/Source/Editor/UnrealEd/Public/DiffUtils.h b/Engine/Source/Editor/UnrealEd/Public/DiffUtils.h
--- a/Engine/Source/Editor/UnrealEd/Public/DiffUtils.h
+++ b/Engine/Source/Editor/UnrealEd/Public/DiffUtils.h
@@ -246,10 +246,22 @@ struct FSCSDiffRoot
 namespace DiffUtils
 {
 	UNREALED_API const UObject* GetCDO(const UBlueprint* ForBlueprint);
+	UE_DEPRECATED(5.3, "DiffUtils now requires root objects so that object topology can be meaningfully compared.")
 	UNREALED_API void CompareUnrelatedStructs(const UStruct* StructA, const void* A, const UStruct* StructB, const void* B, TArray<FSingleObjectDiffEntry>& OutDifferingProperties);
+	UNREALED_API void CompareUnrelatedStructs(const UStruct* StructA, const void* A, const UObject* OwningOuterA, const UStruct* StructB, const void* B, const
+	                                          UObject* OwningOuterB, TArray<FSingleObjectDiffEntry>& OutDifferingProperties);
 	UNREALED_API void CompareUnrelatedObjects(const UObject* A, const UObject* B, TArray<FSingleObjectDiffEntry>& OutDifferingProperties);
 	UNREALED_API void CompareUnrelatedSCS(const UBlueprint* Old, const TArray< FSCSResolvedIdentifier >& OldHierarchy, const UBlueprint* New, const TArray< FSCSResolvedIdentifier >& NewHierarchy, FSCSDiffRoot& OutDifferingEntries );
+
+	UE_DEPRECATED(5.3, "DiffUtils now requires root objects so that object topology can be meaningfully compared.")
 	UNREALED_API bool Identical(const FResolvedProperty& AProp, const FResolvedProperty& BProp, const FPropertySoftPath& RootPath, TArray<FPropertySoftPath>& DifferingProperties); 
+	/**
+	 * DiffUtils now requires root objects so that object topology can be meaningfully compared.
+	 * DiffUtils::Identical works similar to FProperty::Identical except when a UObject is found, that is in OwningOuter*,
+	 * it's compared by topology instead. This allows sub-objects to diff correctly.
+	 */
+	UNREALED_API bool Identical(const FResolvedProperty& AProp, const FResolvedProperty& BProp, const UObject* OwningOuterA, const UObject* OwningOuterB);
+	UNREALED_API bool Identical(const FResolvedProperty& AProp, const FResolvedProperty& BProp, const UObject* OwningOuterA, const UObject* OwningOuterB, const FPropertySoftPath& RootPath, TArray<FPropertySoftPath>& DifferingProperties);
 	UNREALED_API TArray<FPropertySoftPath> GetVisiblePropertiesInOrderDeclared(const UStruct* ForStruct, const FPropertySoftPath& Scope = FPropertySoftPath());
 
 	UNREALED_API TArray<FPropertyPath> ResolveAll(const UObject* Object, const TArray<FPropertySoftPath>& InSoftProperties);
diff --git a/Engine/Source/Runtime/Engine/Private/EdGraph/EdGraphNode.cpp b/Engine/Source/Runtime/Engine/Private/EdGraph/EdGraphNode.cpp
--- a/Engine/Source/Runtime/Engine/Private/EdGraph/EdGraphNode.cpp
+++ b/Engine/Source/Runtime/Engine/Private/EdGraph/EdGraphNode.cpp
@@ -316,8 +316,7 @@ void UEdGraphNode::DiffProperties(UStruct* StructA, UStruct* StructB, uint8* Dat
 			continue;
 		}
 
-		TArray<FPropertySoftPath> DifferingSubProperties;
-		if (!DiffUtils::Identical(FResolvedProperty(DataA, Prop), FResolvedProperty(DataB, PropB), {}, DifferingSubProperties))
+		if (!DiffUtils::Identical(FResolvedProperty(DataA, Prop), FResolvedProperty(DataB, PropB), StructA, StructB))
 		{
 			// Only bother setting up the display data if we're storing the result
 			if (Results.CanStoreResults())
diff --git a/Engine/Source/Runtime/Engine/Private/World.cpp b/Engine/Source/Runtime/Engine/Private/World.cpp
--- a/Engine/Source/Runtime/Engine/Private/World.cpp
+++ b/Engine/Source/Runtime/Engine/Private/World.cpp
@@ -1415,7 +1415,8 @@ void UWorld::PostLoad()
 	if (GIsEditor)
 	{
 		// Avoid renaming PIE worlds and Instanced worlds except for the new maps (/Temp/Untitled) to preserve naming behavior
-		if (!GetPackage()->HasAnyPackageFlags(PKG_PlayInEditor) && (!IsInstanced() || GetPackage()->GetPathName().StartsWith(TEXT("/Temp/Untitled"))))
+		const bool bLoadedForDiff = GetPackage()->HasAnyPackageFlags(PKG_ForDiffing);
+		if (!GetPackage()->HasAnyPackageFlags(PKG_PlayInEditor) && !bLoadedForDiff && (!IsInstanced() || GetPackage()->GetPathName().StartsWith(TEXT("/Temp/Untitled"))))
 		{
 			// Needed for VER_UE4_WORLD_NAMED_AFTER_PACKAGE. If this file was manually renamed outside of the editor, this is needed anyway
 			const FString ShortPackageName = FPackageName::GetLongPackageAssetName(GetPackage()->GetName());
