AssetEditorSubsystem: Speculative fix for an ensure on trying to remove an item from the recent assets MRU list that isn't a valid package name

diff --git a/Engine/Source/Editor/UnrealEd/Private/Subsystems/AssetEditorSubsystem.cpp b/Engine/Source/Editor/UnrealEd/Private/Subsystems/AssetEditorSubsystem.cpp
--- a/Engine/Source/Editor/UnrealEd/Private/Subsystems/AssetEditorSubsystem.cpp
+++ b/Engine/Source/Editor/UnrealEd/Private/Subsystems/AssetEditorSubsystem.cpp
@@ -189,8 +189,8 @@ void UAssetEditorSubsystem::CullRecentAssetEditorsMap()
 	{
 		for (TMap<FString, FString>::TIterator It(RecentAssetToAssetEditorMap); It; ++It)
 		{
-			// Remove any entries that are not in the mru list
-			if(RecentAssetsList->FindMRUItemIdx(It->Key) == INDEX_NONE)
+			// Remove any entries that are not in the mru list or if the package isn't valid anymore
+			if(!FPackageName::IsValidLongPackageName(It->Key) || RecentAssetsList->FindMRUItemIdx(It->Key) == INDEX_NONE)
 			{
 				It.RemoveCurrent();
 			}
