AssetEditorToolkit: Fix a crash when trying to force delete an open DataFlow asset because you could register the same asset twice with the toolkit

diff --git a/Engine/Source/Editor/UnrealEd/Private/Toolkits/AssetEditorToolkit.cpp b/Engine/Source/Editor/UnrealEd/Private/Toolkits/AssetEditorToolkit.cpp
--- a/Engine/Source/Editor/UnrealEd/Private/Toolkits/AssetEditorToolkit.cpp
+++ b/Engine/Source/Editor/UnrealEd/Private/Toolkits/AssetEditorToolkit.cpp
@@ -588,6 +588,12 @@ void FAssetEditorToolkit::GetSaveableObjects(TArray<UObject*>& OutObjects) const
 
 void FAssetEditorToolkit::AddEditingObject(UObject* Object)
 {
+	// Don't allow adding the same object twice (or notify asset opened twice)
+	if(EditingObjects.Contains(Object))
+	{
+		return;
+	}
+	
 	EditingObjects.Add(Object);
 	GEditor->GetEditorSubsystem<UAssetEditorSubsystem>()->NotifyAssetOpened( Object, this );
 }
