Fix for a potential restricted editor crash in shipping builds.

diff --git a/Engine/Source/Editor/BlueprintGraph/Private/EdGraphSchema_K2.cpp b/Engine/Source/Editor/BlueprintGraph/Private/EdGraphSchema_K2.cpp
--- a/Engine/Source/Editor/BlueprintGraph/Private/EdGraphSchema_K2.cpp
+++ b/Engine/Source/Editor/BlueprintGraph/Private/EdGraphSchema_K2.cpp
@@ -4797,7 +4797,10 @@ void UEdGraphSchema_K2::BreakPinLinks(UEdGraphPin& TargetPin, bool bSendsNodeNot
 
 	Super::BreakPinLinks(TargetPin, bSendsNodeNotifcation);
 
-	FBlueprintEditorUtils::MarkBlueprintAsModified(Blueprint);
+	if (Blueprint != nullptr)
+	{
+		FBlueprintEditorUtils::MarkBlueprintAsModified(Blueprint);
+	}
 }
 
 void UEdGraphSchema_K2::BreakSinglePinLink(UEdGraphPin* SourcePin, UEdGraphPin* TargetPin) const
diff --git a/Engine/Source/Editor/UnrealEd/Private/Kismet2/BlueprintEditorUtils.cpp b/Engine/Source/Editor/UnrealEd/Private/Kismet2/BlueprintEditorUtils.cpp
--- a/Engine/Source/Editor/UnrealEd/Private/Kismet2/BlueprintEditorUtils.cpp
+++ b/Engine/Source/Editor/UnrealEd/Private/Kismet2/BlueprintEditorUtils.cpp
@@ -10466,18 +10466,38 @@ void FBlueprintEditorUtils::SanitizeRestrictedContent(UBlueprint* BP)
 	{
 		// Remove any restricted nodes and/or graphs (note: in restricted mode, it is implied that
 		// restricted items are not also user-facing, so any removals here should be opaque to users).
+		TSet<UEdGraph*> NodeSubgraphs;
 		TSet<UEdGraph*> GraphsToRemove;
 		for (UEdGraphNode* Node : RestrictedNodes)
 		{
+			// If we've already handled the node's graph/subgraph, no need to remove it again here.
+			const UEdGraph* NodeGraph = Node->GetGraph();
+			if (GraphsToRemove.Contains(NodeGraph)
+				|| NodeSubgraphs.Contains(NodeGraph))
+			{
+				continue;
+			}
+
 			// Remove the entire graph if this node represents an entry point for an event/function.
 			if (Node->IsA<UK2Node_Event>()
 				|| Node->IsA<UK2Node_FunctionEntry>())
 			{
 				GraphsToRemove.Add(Node->GetGraph());
 			}
-			
-			// Defer updating the skeleton class since we're going to do a full recompile below.
-			RemoveNode(BP, Node, /*bDontRecompile =*/ true);
+			else
+			{
+				// This may remove a subset of restricted nodes, so keep track of removed subgraphs.
+				NodeSubgraphs.Append(Node->GetSubGraphs());
+
+				// Defer updating the skeleton class since we're going to do a full recompile below.
+				RemoveNode(BP, Node, /*bDontRecompile =*/ true);
+
+				// Exclude any subgraphs that were not removed above from the set that we're tracking.
+				for (UEdGraph* SubGraph : Node->GetSubGraphs())
+				{
+					NodeSubgraphs.Remove(SubGraph);
+				}
+			}
 		}
 
 		// Note: On removal, graphs will be renamed and outered to the package containing the Blueprint.
