Skeleton REINST classes were not being marked as CLASS_NewerVersionExists when using bEnableSkelReinstUpdate, which is our default. This means that REINST skeleton classes could change layout in an unsafe manner. Setting and honoring the flag prevents this and fixes a potential crash when moving blueprint variables to a parent.

diff --git a/Engine/Source/Editor/UnrealEd/Private/Kismet2/KismetReinstanceUtilities.cpp b/Engine/Source/Editor/UnrealEd/Private/Kismet2/KismetReinstanceUtilities.cpp
--- a/Engine/Source/Editor/UnrealEd/Private/Kismet2/KismetReinstanceUtilities.cpp
+++ b/Engine/Source/Editor/UnrealEd/Private/Kismet2/KismetReinstanceUtilities.cpp
@@ -1739,7 +1739,7 @@ void FBlueprintCompileReinstancer::MoveDependentSkelToReinst(UClass* const Owner
 		UObject* OldCDO = CurClass->ClassDefaultObject;
 		const FName ReinstanceName = MakeUniqueObjectName(GetTransientPackage(), CurClass->GetClass(), *(FString(TEXT("REINST_")) + *CurClass->GetName()));
 
-		if (!IsValid(CurClass))
+		if (!IsValid(CurClass) || CurClass->HasAnyClassFlags(CLASS_NewerVersionExists))
 		{
 			if (UClass* const* NewSuper = NewSkeletonToOldSkeleton.Find(CurClass->GetSuperClass()))
 			{
@@ -1751,8 +1751,8 @@ void FBlueprintCompileReinstancer::MoveDependentSkelToReinst(UClass* const Owner
 		UClass* ReinstClass = CastChecked<UClass>(StaticDuplicateObject(CurClass, GetTransientPackage(), ReinstanceName, ~RF_Transactional));
 		
 		ReinstClass->RemoveFromRoot();
-		CurClass->ClassFlags |= CLASS_NewerVersionExists;
-		CurClass->ClassFlags &= ~CLASS_NewerVersionExists;
+		ReinstClass->ClassFlags |= CLASS_NewerVersionExists;
+
 		GIsDuplicatingClassForReinstancing = false;
 
 		UClass** OverridenParent = NewSkeletonToOldSkeleton.Find(ReinstClass->GetSuperClass());
