PR #9994: make simple assets display navigate and edit buttons for BP generated parents

diff --git a/Engine/Source/Editor/UnrealEd/Private/Toolkits/SimpleAssetEditor.cpp b/Engine/Source/Editor/UnrealEd/Private/Toolkits/SimpleAssetEditor.cpp
--- a/Engine/Source/Editor/UnrealEd/Private/Toolkits/SimpleAssetEditor.cpp
+++ b/Engine/Source/Editor/UnrealEd/Private/Toolkits/SimpleAssetEditor.cpp
@@ -8,6 +8,7 @@
 #include "IDetailsView.h"
 #include "EditorClassUtils.h"
 #include "Editor.h"
+#include "Widgets/Layout/SSpacer.h"
 #include "Widgets/Text/STextBlock.h"
 
 #define LOCTEXT_NAMESPACE "GenericEditor"
@@ -83,7 +84,7 @@ void FSimpleAssetEditor::InitEditor( const EToolkitMode::Type Mode, const TShare
 		SpawnToolkitTab(GetToolbarTabId(), FString(), EToolkitTabSpot::ToolBar);
 		PropertiesTab = SpawnToolkitTab( PropertiesTabId, TabInitializationPayload, EToolkitTabSpot::Details );
 	}*/
-	
+
 	// Get the list of objects to edit the details of
 	const TArray<UObject*> ObjectsToEditInDetailsView = ( GetDetailsViewObjects.IsBound() ) ? GetDetailsViewObjects.Execute( ObjectsToEdit ) : ObjectsToEdit;
 
@@ -297,6 +298,27 @@ TSharedRef<FSimpleAssetEditor> FSimpleAssetEditor::CreateEditor( const EToolkitM
 	return NewEditor;
 }
 
+FReply FSimpleAssetEditor::OnFindParentClassInContentBrowserClicked(TObjectPtr<UObject> SyncToClass) const
+{
+	if (SyncToClass)
+	{
+		TArray<UObject*> ObjectList { SyncToClass };
+		GEditor->SyncBrowserToObjects(ObjectList);
+	}
+
+	return FReply::Handled();
+}
+
+FReply FSimpleAssetEditor::OnEditParentClassClicked(TObjectPtr<UObject> EditClass) const
+{
+	if (EditClass)
+	{
+		GEditor->GetEditorSubsystem<UAssetEditorSubsystem>()->OpenEditorForAsset(EditClass);
+	}
+
+	return FReply::Handled();
+}
+
 void FSimpleAssetEditor::PostRegenerateMenusAndToolbars()
 {
 	// Find the common denominator class of the assets we're editing
@@ -313,26 +335,100 @@ void FSimpleAssetEditor::PostRegenerateMenusAndToolbars()
 	// Provide a hyperlink to view that native class
 	if (CommonDenominatorClass)
 	{
-		// build and attach the menu overlay
-		TSharedRef<SHorizontalBox> MenuOverlayBox = SNew(SHorizontalBox)
-			+ SHorizontalBox::Slot()
-			.AutoWidth()
-			.VAlign(VAlign_Center)
-			[
-				SNew(STextBlock)
-				.ColorAndOpacity(FSlateColor::UseSubduedForeground())
-				.ShadowOffset(FVector2D::UnitVector)
-				.Text(bNotAllSame ? LOCTEXT("SimpleAssetEditor_AssetType_Varied", "Common Asset Type: ") : LOCTEXT("SimpleAssetEditor_AssetType", "Asset Type: "))
-			]
-			+SHorizontalBox::Slot()
-			.AutoWidth()
-			.VAlign(VAlign_Center)
-			.Padding(0.0f, 0.0f, 8.0f, 0.0f)
-			[
-				FEditorClassUtils::GetSourceLink(CommonDenominatorClass)
-			];
-	
-		SetMenuOverlay(MenuOverlayBox);
+		// If the common denominator is a blueprint generated class, link to the BP instead of an inaccessible _c class
+		if (CommonDenominatorClass->ClassGeneratedBy)
+		{
+			// build and attach the menu overlay
+			TSharedRef<SHorizontalBox> MenuOverlayBox = SNew(SHorizontalBox)
+				+ SHorizontalBox::Slot()
+				.AutoWidth()
+				.VAlign(VAlign_Center)
+				[
+					SNew(STextBlock)
+					.ColorAndOpacity(FSlateColor::UseSubduedForeground())
+					.ShadowOffset(FVector2D::UnitVector)
+					.Text(LOCTEXT("BlueprintEditor_ParentClass", "Parent class: "))
+				]
+				+ SHorizontalBox::Slot()
+				.AutoWidth()
+				.VAlign(VAlign_Center)
+				[
+					SNew(SSpacer)
+					.Size(FVector2D(2.0f, 1.0f))
+				]
+				+ SHorizontalBox::Slot()
+				.AutoWidth()
+				.VAlign(VAlign_Center)
+				[
+					SNew(STextBlock)
+					.ShadowOffset(FVector2D::UnitVector)
+					.Text(FText::FromName(CommonDenominatorClass->ClassGeneratedBy->GetFName()))
+					.TextStyle(FAppStyle::Get(), "Common.InheritedFromBlueprintTextStyle")
+					.ToolTipText(LOCTEXT("ParentClassToolTip", "The class that the current Blueprint is based on. The parent provides the base definition, which the current Blueprint extends."))
+				]
+				+ SHorizontalBox::Slot()
+				.AutoWidth()
+				[
+					SNew(SButton)
+					.VAlign(VAlign_Center)
+					.ButtonStyle(FAppStyle::Get(), "HoverHintOnly")
+					.OnClicked(this, &FSimpleAssetEditor::OnFindParentClassInContentBrowserClicked, CommonDenominatorClass->ClassGeneratedBy)
+					.ToolTipText(LOCTEXT("FindParentInCBToolTip", "Find parent in Content Browser"))
+					.ContentPadding(4.0f)
+					.ForegroundColor(FSlateColor::UseForeground())
+					[
+						SNew(SImage)
+						.Image(FAppStyle::GetBrush("Icons.Search"))
+					]
+				]
+				+ SHorizontalBox::Slot()
+				.AutoWidth()
+				[
+					SNew(SButton)
+					.VAlign(VAlign_Center)
+					.ButtonStyle(FAppStyle::Get(), "HoverHintOnly")
+					.OnClicked(this, &FSimpleAssetEditor::OnEditParentClassClicked, CommonDenominatorClass->ClassGeneratedBy)
+					.ToolTipText(LOCTEXT("EditParentClassToolTip", "Open parent in editor"))
+					.ContentPadding(4.0f)
+					.ForegroundColor(FSlateColor::UseForeground())
+					[
+						SNew(SImage)
+						.Image(FAppStyle::GetBrush("Icons.Edit"))
+					]
+				]
+				+ SHorizontalBox::Slot()
+				.AutoWidth()
+				.VAlign(VAlign_Center)
+				[
+					SNew(SSpacer)
+					.Size(FVector2D(8.0f, 1.0f))
+				]
+				;
+			SetMenuOverlay(MenuOverlayBox);
+		}
+		else
+		{
+			// build and attach the menu overlay
+			TSharedRef<SHorizontalBox> MenuOverlayBox = SNew(SHorizontalBox)
+				+ SHorizontalBox::Slot()
+				.AutoWidth()
+				.VAlign(VAlign_Center)
+				[
+					SNew(STextBlock)
+					.ColorAndOpacity(FSlateColor::UseSubduedForeground())
+					.ShadowOffset(FVector2D::UnitVector)
+					.Text(bNotAllSame ? LOCTEXT("SimpleAssetEditor_AssetType_Varied", "Common Asset Type: ") : LOCTEXT("SimpleAssetEditor_AssetType", "Asset Type: "))
+				]
+				+SHorizontalBox::Slot()
+				.AutoWidth()
+				.VAlign(VAlign_Center)
+				.Padding(0.0f, 0.0f, 8.0f, 0.0f)
+				[
+					FEditorClassUtils::GetSourceLink(CommonDenominatorClass)
+				];
+		
+			SetMenuOverlay(MenuOverlayBox);
+		}
 	}
 }
 
diff --git a/Engine/Source/Editor/UnrealEd/Public/Toolkits/SimpleAssetEditor.h b/Engine/Source/Editor/UnrealEd/Public/Toolkits/SimpleAssetEditor.h
--- a/Engine/Source/Editor/UnrealEd/Public/Toolkits/SimpleAssetEditor.h
+++ b/Engine/Source/Editor/UnrealEd/Public/Toolkits/SimpleAssetEditor.h
@@ -45,12 +45,19 @@ public:
 
 	/** FAssetEditorToolkit interface */
 	UNREALED_API virtual void PostRegenerateMenusAndToolbars() override;
-	
+
 	/** Used to show or hide certain properties */
 	UNREALED_API void SetPropertyVisibilityDelegate(FIsPropertyVisible InVisibilityDelegate);
 	/** Can be used to disable the details view making it read-only */
 	UNREALED_API void SetPropertyEditingEnabledDelegate(FIsPropertyEditingEnabled InPropertyEditingDelegate);
 
+protected:
+	/** Handler for "Find parent class in CB" button */
+	FReply OnFindParentClassInContentBrowserClicked(TObjectPtr<UObject> SyncToClass) const;
+
+	/** Handler for "Edit parent class" button */
+	FReply OnEditParentClassClicked(TObjectPtr<UObject> EditClass) const;
+
 private:
 	/** Create the properties tab and its content */
 	UNREALED_API TSharedRef<SDockTab> SpawnPropertiesTab( const FSpawnTabArgs& Args );
