[CIS FIX] added missing SButton and SImage includes

diff --git a/Engine/Source/Editor/UnrealEd/Private/Toolkits/SimpleAssetEditor.cpp b/Engine/Source/Editor/UnrealEd/Private/Toolkits/SimpleAssetEditor.cpp
--- a/Engine/Source/Editor/UnrealEd/Private/Toolkits/SimpleAssetEditor.cpp
+++ b/Engine/Source/Editor/UnrealEd/Private/Toolkits/SimpleAssetEditor.cpp
@@ -10,6 +10,8 @@
 #include "Editor.h"
 #include "Widgets/Layout/SSpacer.h"
 #include "Widgets/Text/STextBlock.h"
+#include "Widgets/Input/SButton.h"
+#include "Widgets/Images/SImage.h"
 
 #define LOCTEXT_NAMESPACE "GenericEditor"
 
