Cooker DiffOnly: Refactor FDiffWriterArchive to reduce complexity. Previously it was complex because there was no central class that accumulated information between the two saves that diffonly does, because of some legacy constraints with SavePackage and public classes that no longer exist. The new version has a central FAccumulator class to accumulate the data, and the FArchive classes used for communication with SavePackage are thin proxies to this accumulator class.
#rnx

diff --git a/Engine/Source/Editor/UnrealEd/Private/Cooker/DiffPackageWriter.cpp b/Engine/Source/Editor/UnrealEd/Private/Cooker/DiffPackageWriter.cpp
--- a/Engine/Source/Editor/UnrealEd/Private/Cooker/DiffPackageWriter.cpp
+++ b/Engine/Source/Editor/UnrealEd/Private/Cooker/DiffPackageWriter.cpp
@@ -123,10 +123,9 @@ void FDiffPackageWriter::BeginPackage(const FBeginPackageInfo& Info)
 {
 	bIsDifferent = false;
 	bNewPackage = false;
-	bDiffCallstack = false;
 	bHasStartedSecondSave = false;
-	DiffMap[0].Reset();
-	DiffMap[1].Reset();
+	Accumulators[0].SafeRelease();
+	Accumulators[1].SafeRelease();
 
 	BeginInfo = Info;
 	ConditionallyDumpObjList();
@@ -136,7 +135,7 @@ void FDiffPackageWriter::BeginPackage(const FBeginPackageInfo& Info)
 
 void FDiffPackageWriter::CommitPackage(FCommitPackageInfo&& Info)
 {
-	if (bDiffCallstack && bSaveForDiff)
+	if (bHasStartedSecondSave && bSaveForDiff)
 	{
 		// Write the package to _ForDiff, but do not write any sidecars
 		EnumRemoveFlags(Info.WriteOptions, EWriteOptions::WriteSidecars);
@@ -152,55 +151,48 @@ void FDiffPackageWriter::CommitPackage(FCommitPackageInfo&& Info)
 void FDiffPackageWriter::WritePackageData(const FPackageInfo& Info, FLargeMemoryWriter& ExportsArchive,
 	const TArray<FFileRegion>& FileRegions)
 {
+	check(Info.MultiOutputIndex < 2);
+	check(Accumulators[Info.MultiOutputIndex].IsValid()); // Should have been constructed by CreateLinkerArchive
+	UE::DiffWriter::FAccumulator& Accumulator = *Accumulators[Info.MultiOutputIndex];
+	UE::DiffWriter::FDiffArchive& ExportsArchiveInternal =static_cast<UE::DiffWriter::FDiffArchive&>(ExportsArchive);
+	check(&ExportsArchiveInternal.GetAccumulator() == &Accumulator);
+
 	FPackageInfo LocalInfo(Info);
 	Inner->CompleteExportsArchiveForDiff(LocalInfo, ExportsArchive);
 
-	FDiffWriterArchive& Writer = static_cast<FDiffWriterArchive&>(ExportsArchive);
-	ICookedPackageWriter::FPreviousCookedBytesData PreviousInnerData;
-	if (!Inner->GetPreviousCookedBytes(LocalInfo, PreviousInnerData))
+	if (!bHasStartedSecondSave)
 	{
-		PreviousInnerData.Data.Reset();
-		PreviousInnerData.HeaderSize = 0;
-		PreviousInnerData.Size = 0;
-	}
-	check(PreviousInnerData.Data.Get() != nullptr || (PreviousInnerData.Size == 0 && PreviousInnerData.HeaderSize == 0));
-
-	FDiffWriterArchive::FPackageData PreviousPackageData;
-	PreviousPackageData.Data = PreviousInnerData.Data.Get();
-	PreviousPackageData.Size = PreviousInnerData.Size;
-	PreviousPackageData.HeaderSize = PreviousInnerData.HeaderSize;
-	PreviousPackageData.StartOffset = PreviousInnerData.StartOffset;
-
-	check(ExportsCallstacks.IsValid());
+		ICookedPackageWriter::FPreviousCookedBytesData PreviousInnerData;
+		if (!Inner->GetPreviousCookedBytes(LocalInfo, PreviousInnerData))
+		{
+			PreviousInnerData.Data.Reset();
+			PreviousInnerData.HeaderSize = 0;
+			PreviousInnerData.Size = 0;
+		}
+		check(PreviousInnerData.Data.Get() != nullptr || (PreviousInnerData.Size == 0 && PreviousInnerData.HeaderSize == 0));
 
-	if (bDiffCallstack)
+		bNewPackage = PreviousInnerData.Size == 0;
+		Accumulator.OnFirstSaveComplete(LocalInfo.LooseFilePath, LocalInfo.HeaderSize, MoveTemp(PreviousInnerData));
+		bIsDifferent = Accumulator.HasDifferences();
+	}
+	else
 	{
-		Writer.GetCallstacks().Append(*ExportsCallstacks);
+		Accumulator.OnSecondSaveComplete(LocalInfo.HeaderSize);
 
 		TMap<FName, FArchiveDiffStats> PackageDiffStats;
 		const TCHAR* CutoffString = TEXT("UEditorEngine::Save()");
-		Writer.CompareWith(PreviousPackageData, *LocalInfo.LooseFilePath, LocalInfo.HeaderSize, CutoffString,
-			MaxDiffsToLog, PackageDiffStats, Inner->GetCookCapabilities().HeaderFormat);
+		Accumulator.CompareWithPrevious(CutoffString, PackageDiffStats, Inner->GetCookCapabilities().HeaderFormat);
 
 		//COOK_STAT(FSavePackageStats::NumberOfDifferentPackages++);
 		//COOK_STAT(FSavePackageStats::MergeStats(PackageDiffStats));
 	}
-	else
-	{
-		Writer.GetCallstacks().Append(*ExportsCallstacks, LocalInfo.HeaderSize);
-		check(LocalInfo.MultiOutputIndex < 2);
-		ExportsDiffMapOffset[LocalInfo.MultiOutputIndex] = LocalInfo.HeaderSize;
-
-		bNewPackage = PreviousPackageData.Size == 0;
-		bIsDifferent = !Writer.GenerateDiffMap(PreviousPackageData, LocalInfo.HeaderSize, MaxDiffsToLog, DiffMap[LocalInfo.MultiOutputIndex]);
-	}
 
 	Inner->WritePackageData(LocalInfo, ExportsArchive, FileRegions);
 }
 
-UE::DiffWriterArchive::FMessageCallback FDiffPackageWriter::GetDiffWriterMessageCallback()
+UE::DiffWriter::FMessageCallback FDiffPackageWriter::GetDiffWriterMessageCallback()
 {
-	return UE::DiffWriterArchive::FMessageCallback([this](ELogVerbosity::Type Verbosity, FStringView Message)
+	return UE::DiffWriter::FMessageCallback([this](ELogVerbosity::Type Verbosity, FStringView Message)
 		{
 			this->OnDiffWriterMessage(Verbosity, Message);
 		});
@@ -215,64 +207,37 @@ FString FDiffPackageWriter::ResolveText(FStringView Message)
 {
 	FString ResolvedText(Message);
 	check(this->Indent && this->NewLine);
-	ResolvedText.ReplaceInline(UE::DiffWriterArchive::IndentToken, this->Indent);
-	ResolvedText.ReplaceInline(UE::DiffWriterArchive::NewLineToken, this->NewLine);
+	ResolvedText.ReplaceInline(UE::DiffWriter::IndentToken, this->Indent);
+	ResolvedText.ReplaceInline(UE::DiffWriter::NewLineToken, this->NewLine);
 	return ResolvedText;
 }
 
-TUniquePtr<FLargeMemoryWriter> FDiffPackageWriter::CreateLinkerArchive(FName PackageName, UObject* Asset, uint16 MultiOutputIndex)
+UE::DiffWriter::FAccumulator& FDiffPackageWriter::ConstructAccumulator(FName PackageName, UObject* Asset,
+	uint16 MultiOutputIndex)
 {
-	// The entire package will be serialized to memory and then compared against package on disk.
-	if (bDiffCallstack)
-	{
-		check(MultiOutputIndex < 2);
-		// Each difference will be logged with its Serialize call stack trace
-		return TUniquePtr<FLargeMemoryWriter>(new FDiffWriterArchive(Asset, *PackageName.ToString(),
-			GetDiffWriterMessageCallback(), true /* bInCollectCallstacks */, &DiffMap[MultiOutputIndex]));
-	}
-	else
+	check(MultiOutputIndex < 2);
+	TRefCountPtr<UE::DiffWriter::FAccumulator>& Accumulator = Accumulators[MultiOutputIndex];
+	if (!Accumulator.IsValid())
 	{
-		return TUniquePtr<FLargeMemoryWriter>(new FDiffWriterArchive(Asset, *PackageName.ToString(),
-			GetDiffWriterMessageCallback(), false /* bInCollectCallstacks */));
+		check(!bHasStartedSecondSave); // Accumulator should already exist from CreateLinkerArchive in the first save
+		Accumulator = new UE::DiffWriter::FAccumulator(Asset, *PackageName.ToString(), MaxDiffsToLog,
+			GetDiffWriterMessageCallback());
 	}
+	return *Accumulator;
 }
 
-TUniquePtr<FLargeMemoryWriter> FDiffPackageWriter::CreateLinkerExportsArchive(FName PackageName, UObject* Asset, uint16 MultiOutputIndex)
+TUniquePtr<FLargeMemoryWriter> FDiffPackageWriter::CreateLinkerArchive(FName PackageName,
+	UObject* Asset, uint16 MultiOutputIndex)
 {
-	// When cooking, exports are serialized into a separate archive. The serialization callstack offsets
-	// and stack traces are collected into a separate callstack collection and appended to the overall
-	// callstacks for the entire package. DiffOnly cooks saves a package twice. The first pass collects
-	// the serialization offsets without the stack traces and then creates a FDiffWriterDiffMap. In the second
-	// pass, the diff map is used to collect offsets AND the entire stack trace for mismatching package data.
-	// In the first pass, callstack offsets will be relative to the export archive and adjusted when appending
-	// the callstacks to the overall package callstacks. In the second pass, the offset (ExportsDiffMapOffset) is
-	// known and the callstack will be relative to the beginning of the package.
-	
-	ExportsCallstacks = MakeUnique<FDiffWriterCallstacks>(Asset);
-	const int64 PreAllocateBytes = 0;
-	const bool bIsPersistent = true;
+	UE::DiffWriter::FAccumulator& Accumulator = ConstructAccumulator(PackageName, Asset, MultiOutputIndex);
+	return TUniquePtr<FLargeMemoryWriter>(new UE::DiffWriter::FDiffArchiveForLinker(Accumulator));
+}
 
-	check(MultiOutputIndex < 2);
-	if (bDiffCallstack)
-	{
-		return MakeUnique<FDiffWriterArchiveMemoryWriter>(
-			*ExportsCallstacks,
-			&DiffMap[MultiOutputIndex],
-			ExportsDiffMapOffset[MultiOutputIndex],
-			PreAllocateBytes,
-			bIsPersistent,
-			*PackageName.ToString());
-	}
-	else
-	{
-		return MakeUnique<FDiffWriterArchiveMemoryWriter>(
-			*ExportsCallstacks,
-			nullptr,
-			0,
-			PreAllocateBytes, 
-			bIsPersistent, 
-			*PackageName.ToString());
-	}
+TUniquePtr<FLargeMemoryWriter> FDiffPackageWriter::CreateLinkerExportsArchive(FName PackageName,
+	UObject* Asset, uint16 MultiOutputIndex)
+{
+	UE::DiffWriter::FAccumulator& Accumulator = ConstructAccumulator(PackageName, Asset, MultiOutputIndex);
+	return TUniquePtr<FLargeMemoryWriter>(new UE::DiffWriter::FDiffArchiveForExports(Accumulator));
 }
 
 void FDiffPackageWriter::UpdateSaveArguments(FSavePackageArgs& SaveArgs)
@@ -307,8 +272,6 @@ bool FDiffPackageWriter::IsAnotherSaveNeeded(FSavePackageResultStruct& PreviousR
 		bHasStartedSecondSave = true;
 		if (PreviousResult.Result == ESavePackageResult::Success && bIsDifferent && !bNewPackage)
 		{
-			bDiffCallstack = true;
-
 			// The contract with the Inner is that Begin is paired with a single commit;
 			// send the old commit and the new begin
 			FCommitPackageInfo CommitInfo;
diff --git a/Engine/Source/Editor/UnrealEd/Private/Cooker/DiffPackageWriter.h b/Engine/Source/Editor/UnrealEd/Private/Cooker/DiffPackageWriter.h
--- a/Engine/Source/Editor/UnrealEd/Private/Cooker/DiffPackageWriter.h
+++ b/Engine/Source/Editor/UnrealEd/Private/Cooker/DiffPackageWriter.h
@@ -123,19 +123,18 @@ protected:
 	bool FilterPackageName(const FString& InWildcard);
 	void ConditionallyDumpObjList();
 	void ConditionallyDumpObjects();
-	UE::DiffWriterArchive::FMessageCallback GetDiffWriterMessageCallback();
+	UE::DiffWriter::FMessageCallback GetDiffWriterMessageCallback();
 	virtual void OnDiffWriterMessage(ELogVerbosity::Type Verbosity, FStringView Message);
 	FString ResolveText(FStringView Message);
+	UE::DiffWriter::FAccumulator& ConstructAccumulator(FName PackageName, UObject* Asset, uint16 MultiOutputIndex);
 
-	FDiffWriterDiffMap DiffMap[2];
-	TUniquePtr<FDiffWriterCallstacks> ExportsCallstacks;
+	TRefCountPtr<UE::DiffWriter::FAccumulator> Accumulators[2];
 	FBeginPackageInfo BeginInfo;
 	TUniquePtr<ICookedPackageWriter> Inner;
 	const TCHAR* Indent = nullptr;
 	const TCHAR* NewLine = nullptr;
 	FString DumpObjListParams;
 	FString PackageFilter;
-	int64 ExportsDiffMapOffset[2] = {0, 0};
 	int32 MaxDiffsToLog = 5;
 	bool bSaveForDiff = false;
 	bool bDiffOptional = false;
diff --git a/Engine/Source/Editor/UnrealEd/Private/Cooker/DiffWriterArchive.cpp b/Engine/Source/Editor/UnrealEd/Private/Cooker/DiffWriterArchive.cpp
--- a/Engine/Source/Editor/UnrealEd/Private/Cooker/DiffWriterArchive.cpp
+++ b/Engine/Source/Editor/UnrealEd/Private/Cooker/DiffWriterArchive.cpp
@@ -22,16 +22,30 @@
 #include "UObject/UObjectGlobals.h"
 #include "UObject/UObjectThreadContext.h"
 
-static const ANSICHAR* DebugDataStackMarker = "\r\nDebugDataStack:\r\n";
-namespace UE::DiffWriterArchive
+namespace UE::DiffWriter
 {
 
+static const ANSICHAR* DebugDataStackMarker = "\r\nDebugDataStack:\r\n";
 const TCHAR* const IndentToken = TEXT("%DWA%    ");
 const TCHAR* const NewLineToken = TEXT("%DWA%\n");
 
-}
+/** Logs any mismatching header data. */
+void DumpPackageHeaderDiffs(
+	const FPackageData& SourcePackage,
+	const FPackageData& DestPackage,
+	const FString& AssetFilename,
+	const int32 MaxDiffsToLog,
+	const EPackageHeaderFormat PackageHeaderFormat,
+	const FMessageCallback& MessageCallback);
+/** Returns a new linker for loading the specified package. */
+FLinkerLoad* CreateLinkerForPackage(
+	FUObjectSerializeContext* LoadContext,
+	const FString& InPackageName,
+	const FString& InFilename,
+	const FPackageData& PackageData);
+
 
-FDiffWriterCallstacks::FCallstackData::FCallstackData(TUniquePtr<ANSICHAR[]>&& InCallstack, UObject* InSerializedObject, FProperty* InSerializedProperty)
+FCallstacks::FCallstackData::FCallstackData(TUniquePtr<ANSICHAR[]>&& InCallstack, UObject* InSerializedObject, FProperty* InSerializedProperty)
 	: Callstack(MoveTemp(InCallstack))
 	, SerializedProp(InSerializedProperty)
 {
@@ -45,7 +59,7 @@ FDiffWriterCallstacks::FCallstackData::FCallstackData(TUniquePtr<ANSICHAR[]>&& I
 	}
 }
 
-FString FDiffWriterCallstacks::FCallstackData::ToString(const TCHAR* CallstackCutoffText) const
+FString FCallstacks::FCallstackData::ToString(const TCHAR* CallstackCutoffText) const
 {
 	FString HumanReadableString;
 
@@ -76,34 +90,34 @@ FString FDiffWriterCallstacks::FCallstackData::ToString(const TCHAR* CallstackCu
 				StackLine.MidInline(CutoffIndex + 1, MAX_int32, false);
 			}
 		}
-		HumanReadableString += UE::DiffWriterArchive::IndentToken;
+		HumanReadableString += IndentToken;
 		HumanReadableString += StackLine;
-		HumanReadableString += UE::DiffWriterArchive::NewLineToken;
+		HumanReadableString += NewLineToken;
 	}
 
 	if (!SerializedObjectName.IsEmpty())
 	{
-		HumanReadableString += UE::DiffWriterArchive::NewLineToken;
-		HumanReadableString += UE::DiffWriterArchive::IndentToken;
+		HumanReadableString += NewLineToken;
+		HumanReadableString += IndentToken;
 		HumanReadableString += TEXT("Serialized Object: ");
 		HumanReadableString += SerializedObjectName;
-		HumanReadableString += UE::DiffWriterArchive::NewLineToken;
+		HumanReadableString += NewLineToken;
 	}
 	if (!SerializedPropertyName.IsEmpty())
 	{
 		if (SerializedObjectName.IsEmpty())
 		{
-			HumanReadableString += UE::DiffWriterArchive::NewLineToken;
+			HumanReadableString += NewLineToken;
 		}
-		HumanReadableString += UE::DiffWriterArchive::IndentToken;
+		HumanReadableString += IndentToken;
 		HumanReadableString += TEXT("Serialized Property: ");
 		HumanReadableString += SerializedPropertyName;
-		HumanReadableString += UE::DiffWriterArchive::NewLineToken;
+		HumanReadableString += NewLineToken;
 	}
 	return HumanReadableString;
 }
 
-FDiffWriterCallstacks::FCallstackData FDiffWriterCallstacks::FCallstackData::Clone() const
+FCallstacks::FCallstackData FCallstacks::FCallstackData::Clone() const
 {
 	TUniquePtr<ANSICHAR[]> CallstackCopy;
 	if (const int32 Len = FCStringAnsi::Strlen(Callstack.Get()); Len > 0)
@@ -118,23 +132,27 @@ FDiffWriterCallstacks::FCallstackData FDiffWriterCallstacks::FCallstackData::Clo
 	return Clone;
 }
 
-FDiffWriterCallstacks::FDiffWriterCallstacks(UObject* InAsset)
-	: Asset(InAsset)
-	, bCallstacksDirty(true)
+FCallstacks::FCallstacks()
+	: bCallstacksDirty(true)
 	, StackTraceSize(65535)
 	, LastSerializeCallstack(nullptr)
-	, TotalSize(0)
+	, EndOffset(0)
 {
 	StackTrace = MakeUnique<ANSICHAR[]>(StackTraceSize);
 	StackTrace[0] = 0;
 }
 
-FName FDiffWriterCallstacks::GetAssetClass() const
+void FCallstacks::Reset()
 {
-	return Asset != nullptr ? Asset->GetClass()->GetFName() : NAME_None;
+	CallstackAtOffsetMap.Reset();
+	UniqueCallstacks.Reset();
+	bCallstacksDirty = true;
+	LastSerializeCallstack = nullptr;
+	StackTrace[0] = 0;
+	EndOffset = 0;
 }
 
-ANSICHAR* FDiffWriterCallstacks::AddUniqueCallstack(bool bIsCollectingCallstacks, UObject* SerializedObject, FProperty* SerializedProperty, uint32& OutCallstackCRC)
+ANSICHAR* FCallstacks::AddUniqueCallstack(bool bIsCollectingCallstacks, UObject* SerializedObject, FProperty* SerializedProperty, uint32& OutCallstackCRC)
 {
 	ANSICHAR* Callstack = nullptr;
 	if (bIsCollectingCallstacks)
@@ -161,7 +179,7 @@ ANSICHAR* FDiffWriterCallstacks::AddUniqueCallstack(bool bIsCollectingCallstacks
 	return Callstack;
 }
 
-void FDiffWriterCallstacks::Add(
+void FCallstacks::Add(
 	int64 Offset,
 	int64 Length,
 	UObject* SerializedObject,
@@ -177,7 +195,7 @@ void FDiffWriterCallstacks::Add(
 	}
 
 	const int64 CurrentOffset = Offset;
-	TotalSize = FMath::Max(TotalSize, CurrentOffset + Length); 
+	EndOffset = FMath::Max(EndOffset, CurrentOffset + Length); 
 
 	const bool bShouldCollectCallstack = bIsCollectingCallstacks && bCollectCurrentCallstack && !UE::ArchiveStackTrace::ShouldIgnoreDiff();
 	if (bShouldCollectCallstack)
@@ -189,7 +207,7 @@ void FDiffWriterCallstacks::Add(
 		{
 			FCStringAnsi::Strcat(StackTrace.Get(), StackTraceSize, DebugDataStackMarker);
 
-			const FString SubIndent = FString(UE::DiffWriterArchive::IndentToken) + FString(TEXT("    "));
+			const FString SubIndent = FString(IndentToken) + FString(TEXT("    "));
 
 			bool bIsIndenting = true;
 			for (const auto& DebugData : DebugDataStack)
@@ -221,7 +239,7 @@ void FDiffWriterCallstacks::Add(
 		{
 			// New data serialized at the end of archive buffer
 			LastSerializeCallstack = AddUniqueCallstack(bIsCollectingCallstacks, SerializedObject, SerializedProperty, CallstackCRC);
-			CallstackAtOffsetMap.Add(FCallstackAtOffset {CurrentOffset, CallstackCRC, UE::ArchiveStackTrace::ShouldIgnoreDiff()});
+			CallstackAtOffsetMap.Add(FCallstackAtOffset { CurrentOffset, CallstackCRC, UE::ArchiveStackTrace::ShouldIgnoreDiff()});
 		}
 		else
 		{
@@ -250,9 +268,9 @@ void FDiffWriterCallstacks::Add(
 	}
 }
 
-int32 FDiffWriterCallstacks::GetCallstackIndexAtOffset(int64 Offset, int32 MinOffsetIndex) const
+int32 FCallstacks::GetCallstackIndexAtOffset(int64 Offset, int32 MinOffsetIndex) const
 {
-	if (Offset < 0 || Offset > TotalSize || MinOffsetIndex < 0 || MinOffsetIndex >= CallstackAtOffsetMap.Num())
+	if (Offset < 0 || Offset >= EndOffset || MinOffsetIndex < 0 || MinOffsetIndex >= CallstackAtOffsetMap.Num())
 	{
 		return -1;
 	}
@@ -300,12 +318,12 @@ int32 FDiffWriterCallstacks::GetCallstackIndexAtOffset(int64 Offset, int32 MinOf
 	return OffsetForCallstackIndex;
 }
 
-void FDiffWriterCallstacks::Append(const FDiffWriterCallstacks& Other, int64 Offset)
+void FCallstacks::Append(const FCallstacks& Other, int64 OtherStartOffset)
 {
 	for (const FCallstackAtOffset& OtherOffset : Other.CallstackAtOffsetMap)
 	{
 		FCallstackAtOffset& New = CallstackAtOffsetMap.Add_GetRef(OtherOffset);
-		New.Offset += Offset;
+		New.Offset += OtherStartOffset;
 	}
 
 	CallstackAtOffsetMap.Sort([](const FCallstackAtOffset& LHS,const FCallstackAtOffset& RHS)
@@ -326,66 +344,79 @@ void FDiffWriterCallstacks::Append(const FDiffWriterCallstacks& Other, int64 Off
 		UniqueCallstacks.Emplace(Kv.Key, Kv.Value.Clone());
 	}
 
-	TotalSize = FMath::Max(TotalSize, Other.TotalSize);
+	EndOffset = FMath::Max(EndOffset, Other.EndOffset + OtherStartOffset);
 }
 
-FDiffWriterArchiveWriter::FDiffWriterArchiveWriter(
-	FArchive& InInner,
-	FDiffWriterCallstacks& InCallstacks,
-	const FDiffWriterDiffMap* InDiffMap,
-	int64 InDiffMapOffset)
-		: FArchiveProxy(InInner)
-		, Callstacks(InCallstacks)
-		, DiffMap(InDiffMap)
-		, DiffMapOffset(InDiffMapOffset)
-		, bInnerArchiveDisabled(false)
+void FCallstacks::RecordSerialize(EOffsetFrame OffsetFrame, int64 CurrentOffset, int64 Length,
+	const FAccumulator& Accumulator, FDiffArchive& Ar, int32 StackIgnoreCount)
 {
-}
+	int64 LinkerOffset = -1;
+	switch (OffsetFrame)
+	{
+	case EOffsetFrame::Linker:
+		LinkerOffset = CurrentOffset;
+		break;
+	case EOffsetFrame::Exports:
+		if (Accumulator.bFirstSaveComplete)
+		{
+			LinkerOffset = CurrentOffset + Accumulator.HeaderSize;
+		}
+		break;
+	default:
+		checkNoEntry();
+		break;
+	}
 
-FDiffWriterArchiveWriter::~FDiffWriterArchiveWriter() = default;
+	++StackIgnoreCount;
 
-void FDiffWriterArchiveWriter::Serialize(void* Data, int64 Length)
-{
-	static struct FBreakAtOffsetSettings
+	if (LinkerOffset >= 0)
 	{
-		FString PackageToBreakOn;
-		int64 OffsetToBreakOn;
-
-		FBreakAtOffsetSettings()
-			: OffsetToBreakOn(-1)
+		static struct FBreakAtOffsetSettings
 		{
-			if (!FParse::Param(FCommandLine::Get(), TEXT("cooksinglepackage")))
-			{
-				return;
-			}
+			FString PackageToBreakOn;
+			int64 OffsetToBreakOn;
 
-			FString Package;
-			if (!FParse::Value(FCommandLine::Get(), TEXT("map="), Package))
+			FBreakAtOffsetSettings()
+				: OffsetToBreakOn(-1)
 			{
-				return;
-			}
+				if (!FParse::Param(FCommandLine::Get(), TEXT("cooksinglepackage")) &&
+					!FParse::Param(FCommandLine::Get(), TEXT("cooksinglepackagenorefs")))
+				{
+					return;
+				}
 
-			int64 Offset;
-			if (!FParse::Value(FCommandLine::Get(), TEXT("diffonlybreakoffset="), Offset) || Offset <= 0)
-			{
-				return;
-			}
+				FString Package;
+				if (!FParse::Value(FCommandLine::Get(), TEXT("map="), Package) &&
+					!FParse::Value(FCommandLine::Get(), TEXT("package="), Package))
+				{
+					return;
+				}
 
-			OffsetToBreakOn = Offset;
-			PackageToBreakOn = TEXT("/") + FPackageName::GetShortName(Package);
-		}
-	} BreakAtOffsetSettings;
+				int64 Offset;
+				if (!FParse::Value(FCommandLine::Get(), TEXT("diffonlybreakoffset="), Offset) || Offset <= 0)
+				{
+					return;
+				}
 
-	const int64 CurrentOffset = DiffMapOffset + Tell();
+				OffsetToBreakOn = Offset;
+				PackageToBreakOn = TEXT("/") + FPackageName::GetShortName(Package);
+			}
+		} BreakAtOffsetSettings;
 
-	if (BreakAtOffsetSettings.OffsetToBreakOn >= 0 && BreakAtOffsetSettings.OffsetToBreakOn >= CurrentOffset && BreakAtOffsetSettings.OffsetToBreakOn < CurrentOffset + Length)
-	{
-		FString ArcName = GetArchiveName();
-		int32 SubnameIndex = ArcName.Find(BreakAtOffsetSettings.PackageToBreakOn, ESearchCase::IgnoreCase, ESearchDir::FromEnd);
-		if (SubnameIndex >= 0)
+		if (BreakAtOffsetSettings.OffsetToBreakOn >= 0 &&
+			LinkerOffset <= BreakAtOffsetSettings.OffsetToBreakOn && BreakAtOffsetSettings.OffsetToBreakOn < LinkerOffset + Length)
 		{
-			int32 SubnameEndIndex = SubnameIndex + BreakAtOffsetSettings.PackageToBreakOn.Len();
-			if (SubnameEndIndex == ArcName.Len() || ArcName[SubnameEndIndex] == TEXT('.'))
+			const FString& ArcName = Accumulator.Filename;
+			bool bFilenameMatches = false;
+			{
+				int32 SubnameIndex = ArcName.Find(BreakAtOffsetSettings.PackageToBreakOn, ESearchCase::IgnoreCase, ESearchDir::FromEnd);
+				if (SubnameIndex >= 0)
+				{
+					int32 SubnameEndIndex = SubnameIndex + BreakAtOffsetSettings.PackageToBreakOn.Len();
+					bFilenameMatches = SubnameEndIndex == ArcName.Len() || ArcName[SubnameEndIndex] == TEXT('.');
+				}
+			}
+			if (bFilenameMatches)
 			{
 				UE_DEBUG_BREAK();
 			}
@@ -394,174 +425,208 @@ void FDiffWriterArchiveWriter::Serialize(void* Data, int64 Length)
 
 	if (Length > 0)
 	{
-		UObject* SerializedObject = SerializeContext ? SerializeContext->SerializedObject : nullptr;
-		TArrayView<const FName> DebugStack;
-		DebugStack = DebugDataStack;
+		UObject* SerializedObject = Ar.GetSerializeContext() ? Ar.GetSerializeContext()->SerializedObject : nullptr;
+		TArrayView<const FName> DebugStack = Ar.GetDebugDataStack();
 
-		const bool bIsCollectingCallstacks = DiffMap != nullptr;
-		const bool bCollectCurrentCallstack = bIsCollectingCallstacks && DiffMap->ContainsOffset(CurrentOffset);
+		const bool bCollectingCallstacks = Accumulator.bFirstSaveComplete;
+		const bool bCollectCurrentCallstack = bCollectingCallstacks && LinkerOffset >= 0 &&
+			Accumulator.DiffMap.ContainsOffset(LinkerOffset);
 
-		Callstacks.Add(
-			CurrentOffset,
-			Length,
-			SerializedObject,
-			GetSerializedProperty(),
-			DebugStack,
-			bIsCollectingCallstacks,
-			bCollectCurrentCallstack,
-			StackIgnoreCount);
+		Add(CurrentOffset, Length, SerializedObject, Ar.GetSerializedProperty(), DebugStack, bCollectingCallstacks,
+			bCollectCurrentCallstack, StackIgnoreCount);
 	}
+}
 
-	if (bInnerArchiveDisabled == false)
-	{
-		InnerArchive.Serialize(Data, Length);
-	}
+FAccumulator::FAccumulator(UObject* InAsset, FName InPackageName, int32 InMaxDiffsToLog,
+	FMessageCallback&& InMessageCallback)
+	: LinkerCallstacks()
+	, ExportsCallstacks()
+	, MessageCallback(MoveTemp(InMessageCallback))
+	, PackageName(InPackageName)
+	, Asset(InAsset)
+	, MaxDiffsToLog(InMaxDiffsToLog)
+{
+}
+
+FAccumulator::~FAccumulator()
+{
+}
+
+void FAccumulator::SetHeaderSize(int64 InHeaderSize)
+{
+	HeaderSize = InHeaderSize;
 }
 
-void FDiffWriterArchiveWriter::SetSerializeContext(FUObjectSerializeContext* Context)
+FName FAccumulator::GetAssetClass() const
 {
-	SerializeContext = Context;
+	return Asset != nullptr ? Asset->GetClass()->GetFName() : NAME_None;
+}
+
+void FAccumulator::OnFirstSaveComplete(FStringView LooseFilePath, int64 InHeaderSize,
+	ICookedPackageWriter::FPreviousCookedBytesData&& InPreviousPackageData)
+{
+	Filename = LooseFilePath;
+	HeaderSize = InHeaderSize;
+	PreviousPackageData = MoveTemp(InPreviousPackageData);
+
+	LinkerCallstacks.Append(ExportsCallstacks, HeaderSize);
+	ExportsCallstacks.Reset();
 
-	if (bInnerArchiveDisabled == false)
+	GenerateDiffMap();
+
+	LinkerCallstacks.Reset();
+	bFirstSaveComplete = true;
+}
+
+void FAccumulator::OnSecondSaveComplete(int64 InHeaderSize)
+{
+	check(bFirstSaveComplete); // Should have been set in OnFirstSaveComplete
+	if (HeaderSize != InHeaderSize)
 	{
-		InnerArchive.SetSerializeContext(Context);
+		MessageCallback(ELogVerbosity::Error, FString::Printf(
+			TEXT("%s: Indeterministic header size. When saving the package twice into memory, first header size %d != second header size %d. Callstacks for indeterminism in the exports will be incorrect."),
+			*this->Filename, HeaderSize, InHeaderSize));
 	}
+
+	LinkerCallstacks.Append(ExportsCallstacks, HeaderSize);
+	ExportsCallstacks.Reset();
 }
 
-FUObjectSerializeContext* FDiffWriterArchiveWriter::GetSerializeContext()
+bool FAccumulator::HasDifferences() const
 {
-	return SerializeContext;
+	return bHasDifferences;
 }
 
-FDiffWriterArchiveMemoryWriter::FDiffWriterArchiveMemoryWriter(
-	FDiffWriterCallstacks& Callstacks,
-	const FDiffWriterDiffMap* DiffMap,
-	const int64 DiffMapOffset,
-	const int64 PreAllocateBytes,
-	bool bIsPersistent,
-	const TCHAR* Filename)
-		: FLargeMemoryWriter(PreAllocateBytes, bIsPersistent, Filename)
-		, StackTraceWriter(*this, Callstacks, DiffMap, DiffMapOffset)
+FDiffArchive::FDiffArchive(FAccumulator& InAccumulator)
+	: Accumulator(&InAccumulator)
 {
-	// Hack to prevent recursive calls to serialize when passing in this.
-	StackTraceWriter.SetDisableInnerArchive(true);
-	StackTraceWriter.SetStackIgnoreCount(StackTraceWriter.GetStackIgnoreCount() + 1);
+	SetIsPersistent(true /* bIsPersistent */);
 }
 
-void FDiffWriterArchiveMemoryWriter::Serialize(void* Memory, int64 Length)
+FString FDiffArchive::GetArchiveName() const
 {
-	StackTraceWriter.Serialize(Memory, Length);
-	FLargeMemoryWriter::Serialize(Memory, Length);
+	return Accumulator->Filename;
 }
 
-void FDiffWriterArchiveMemoryWriter::SetSerializeContext(FUObjectSerializeContext* Context)
+void FDiffArchive::PushDebugDataString(const FName& DebugData)
 {
-	StackTraceWriter.SetSerializeContext(Context);
+	DebugDataStack.Push(DebugData);
 }
 
-FUObjectSerializeContext* FDiffWriterArchiveMemoryWriter::GetSerializeContext()
+void FDiffArchive::PopDebugDataString()
 {
-	return StackTraceWriter.GetSerializeContext();
+	DebugDataStack.Pop();
 }
 
-FDiffWriterArchive::FDiffWriterArchive(UObject* InAsset, const TCHAR* InFilename, UE::DiffWriterArchive::FMessageCallback&& InMessageCallback,
-	bool bInCollectCallstacks, const FDiffWriterDiffMap* InDiffMap)
-	: FLargeMemoryWriter(0, false, InFilename)
-	, Callstacks(InAsset)
-	, StackTraceWriter(*this, Callstacks, InDiffMap, 0)
-	, MessageCallback(MoveTemp(InMessageCallback))
+TArray<FName>& FDiffArchive::GetDebugDataStack()
 {
-	// Hack to prevent recursive calls to serialize when passing in this.
-	StackTraceWriter.SetDisableInnerArchive(true);
-	StackTraceWriter.SetStackIgnoreCount(StackTraceWriter.GetStackIgnoreCount() + 1);
+	return DebugDataStack;
 }
 
-FDiffWriterArchive::~FDiffWriterArchive()
+FDiffArchiveForLinker::FDiffArchiveForLinker(FAccumulator& InAccumulator)
+	: FDiffArchive(InAccumulator)
 {
+	// SavePackage is supposed to destroy the archive before returning, we rely on that so that the linkerarchive
+	// in the second save can use the same data that was used in the first save
+	check(!Accumulator->LinkerArchive);
+	Accumulator->LinkerArchive = this;
 }
 
-void FDiffWriterArchive::Serialize(void* Memory, int64 Length)
+FDiffArchiveForLinker::~FDiffArchiveForLinker()
 {
-	StackTraceWriter.Serialize(Memory, Length);
-	FLargeMemoryWriter::Serialize(Memory, Length);
+	check(Accumulator->LinkerArchive == this)
+	Accumulator->LinkerArchive = nullptr;
 }
 
-void FDiffWriterArchive::SetSerializeContext(FUObjectSerializeContext* Context)
+void FDiffArchiveForLinker::Serialize(void* InData, int64 Num)
 {
-	StackTraceWriter.SetSerializeContext(Context);
+	int32 StackIgnoreCount = 1;
+	int64 CurrentOffset = Tell();
+	Accumulator->LinkerCallstacks.RecordSerialize(EOffsetFrame::Linker, CurrentOffset, Num, *Accumulator,
+		*this, StackIgnoreCount);
+	FDiffArchive::Serialize(InData, Num);
 }
 
-FUObjectSerializeContext* FDiffWriterArchive::GetSerializeContext()
+FDiffArchiveForExports::FDiffArchiveForExports(FAccumulator& InAccumulator)
+	: FDiffArchive(InAccumulator)
 {
-	return StackTraceWriter.GetSerializeContext();
+	// SavePackage is supposed to destroy the archive before returning, we rely on that so that the exportsarchive
+	// in the second save can use the same data that was used in the first save
+	check(!Accumulator->ExportsArchive);
+	Accumulator->ExportsArchive = this;
 }
 
-namespace
+FDiffArchiveForExports::~FDiffArchiveForExports()
 {
-	bool ShouldDumpPropertyValueState(FProperty* Prop)
-	{
-		if (Prop->IsA<FNumericProperty>()
-			|| Prop->IsA<FStrProperty>()
-			|| Prop->IsA<FBoolProperty>()
-			|| Prop->IsA<FNameProperty>())
-		{
-			return true;
-		}
+	check(Accumulator->ExportsArchive == this);
+	Accumulator->ExportsArchive = nullptr;
+}
 
-		if (FArrayProperty* ArrayProp = CastField<FArrayProperty>(Prop))
-		{
-			return ShouldDumpPropertyValueState(ArrayProp->Inner);
-		}
+void FDiffArchiveForExports::Serialize(void* InData, int64 Num)
+{
+	int32 StackIgnoreCount = 1;
+	int64 CurrentOffset = Tell();
+	Accumulator->ExportsCallstacks.RecordSerialize(EOffsetFrame::Exports, CurrentOffset, Num, *Accumulator,
+		*this, StackIgnoreCount);
+	FDiffArchive::Serialize(InData, Num);
+}
 
-		if (FMapProperty* MapProp = CastField<FMapProperty>(Prop))
-		{
-			return ShouldDumpPropertyValueState(MapProp->KeyProp) && ShouldDumpPropertyValueState(MapProp->ValueProp);
-		}
+bool ShouldDumpPropertyValueState(FProperty* Prop)
+{
+	if (Prop->IsA<FNumericProperty>()
+		|| Prop->IsA<FStrProperty>()
+		|| Prop->IsA<FBoolProperty>()
+		|| Prop->IsA<FNameProperty>())
+	{
+		return true;
+	}
 
-		if (FSetProperty* SetProp = CastField<FSetProperty>(Prop))
-		{
-			return ShouldDumpPropertyValueState(SetProp->ElementProp);
-		}
+	if (FArrayProperty* ArrayProp = CastField<FArrayProperty>(Prop))
+	{
+		return ShouldDumpPropertyValueState(ArrayProp->Inner);
+	}
 
-		if (FStructProperty* StructProp = CastField<FStructProperty>(Prop))
-		{
-			if (StructProp->Struct == TBaseStructure<FVector>::Get()
-				|| StructProp->Struct == TBaseStructure<FGuid>::Get())
-			{
-				return true;
-			}
-		}
+	if (FMapProperty* MapProp = CastField<FMapProperty>(Prop))
+	{
+		return ShouldDumpPropertyValueState(MapProp->KeyProp) && ShouldDumpPropertyValueState(MapProp->ValueProp);
+	}
 
-		if (FOptionalProperty* OptionalProp = CastField<FOptionalProperty>(Prop))
+	if (FSetProperty* SetProp = CastField<FSetProperty>(Prop))
+	{
+		return ShouldDumpPropertyValueState(SetProp->ElementProp);
+	}
+
+	if (FStructProperty* StructProp = CastField<FStructProperty>(Prop))
+	{
+		if (StructProp->Struct == TBaseStructure<FVector>::Get()
+			|| StructProp->Struct == TBaseStructure<FGuid>::Get())
 		{
-			return ShouldDumpPropertyValueState(OptionalProp->GetValueProperty());
+			return true;
 		}
+	}
 
-		return false;
+	if (FOptionalProperty* OptionalProp = CastField<FOptionalProperty>(Prop))
+	{
+		return ShouldDumpPropertyValueState(OptionalProp->GetValueProperty());
 	}
+
+	return false;
 }
 
-void FDiffWriterArchiveWriter::Compare(
-	const FPackageData& SourcePackage,
-	const FPackageData& DestPackage,
-	const FDiffWriterCallstacks& Callstacks,
-	const FDiffWriterDiffMap& DiffMap,
-	const TCHAR* AssetFilename,
-	const TCHAR* CallstackCutoffText,
-	const int64 MaxDiffsToLog,
-	int32& InOutDiffsLogged,
-	TMap<FName, FArchiveDiffStats>& OutStats,
-	const UE::DiffWriterArchive::FMessageCallback& MessageCallback)
+void FAccumulator::CompareWithPreviousForSection(const FPackageData& SourcePackage, const FPackageData& DestPackage,
+	const TCHAR* CallstackCutoffText, int32& InOutDiffsLogged, TMap<FName, FArchiveDiffStats>& OutStats,
+	const FString& SectionFilename)
 {
+	FCallstacks& Callstacks = this->LinkerCallstacks;
 	const int64 SourceSize = SourcePackage.Size - SourcePackage.StartOffset;
 	const int64 DestSize = DestPackage.Size - DestPackage.StartOffset;
 	const int64 SizeToCompare = FMath::Min(SourceSize, DestSize);
-	const FName AssetClass = Callstacks.GetAssetClass();
+	const FName AssetClass = GetAssetClass();
 	
 	if (SourceSize != DestSize)
 	{
 		MessageCallback(ELogVerbosity::Warning, FString::Printf(
-			TEXT("%s: Size mismatch: on disk: %lld vs memory: %lld"), AssetFilename, SourceSize, DestSize));
+			TEXT("%s: Size mismatch: on disk: %lld vs memory: %lld"), *SectionFilename, SourceSize, DestSize));
 		int64 SizeDiff = DestPackage.Size - SourcePackage.Size;
 		OutStats.FindOrAdd(AssetClass).DiffSize += SizeDiff;
 	}
@@ -605,7 +670,7 @@ void FDiffWriterArchiveWriter::Compare(
 			if (DifferenceCallstackoffsetIndex < 0)
 			{
 				MessageCallback(ELogVerbosity::Warning, FString::Printf(
-					TEXT("%s: Difference at offset %lld (absolute offset: %lld), unknown callstack"), AssetFilename, LocalOffset, DestAbsoluteOffset));
+					TEXT("%s: Difference at offset %lld (absolute offset: %lld), unknown callstack"), *SectionFilename, LocalOffset, DestAbsoluteOffset));
 				continue;
 			}
 
@@ -614,8 +679,8 @@ void FDiffWriterArchiveWriter::Compare(
 				continue;
 			}
 
-			const FDiffWriterCallstacks::FCallstackAtOffset& CallstackAtOffset = Callstacks.GetCallstack(DifferenceCallstackoffsetIndex);
-			const FDiffWriterCallstacks::FCallstackData& DifferenceCallstackData = Callstacks.GetCallstackData(CallstackAtOffset);
+			const FCallstacks::FCallstackAtOffset& CallstackAtOffset = Callstacks.GetCallstack(DifferenceCallstackoffsetIndex);
+			const FCallstacks::FCallstackData& DifferenceCallstackData = Callstacks.GetCallstackData(CallstackAtOffset);
 			FString DifferenceCallstackDataText = DifferenceCallstackData.ToString(CallstackCutoffText);
 			if (LastDifferenceCallstackDataText.Compare(DifferenceCallstackDataText, ESearchCase::CaseSensitive) == 0)
 			{
@@ -644,8 +709,8 @@ void FDiffWriterArchiveWriter::Compare(
 							}
 
 							const int32 CallstackIndex = Callstacks.GetCallstackIndexAtOffset(OffsetX - 1, 0);
-							const FDiffWriterCallstacks::FCallstackAtOffset& PreviousCallstack = Callstacks.GetCallstack(CallstackIndex);
-							const FDiffWriterCallstacks::FCallstackData& PreviousCallstackData = Callstacks.GetCallstackData(PreviousCallstack);
+							const FCallstacks::FCallstackAtOffset& PreviousCallstack = Callstacks.GetCallstack(CallstackIndex);
+							const FCallstacks::FCallstackData& PreviousCallstackData = Callstacks.GetCallstackData(PreviousCallstack);
 							if (PreviousCallstackData.SerializedProp != SerProp)
 							{
 								break;
@@ -675,8 +740,8 @@ void FDiffWriterArchiveWriter::Compare(
 				if (BeforePropertyVal != AfterPropertyVal)
 				{
 					DiffValues = FString::Printf(TEXT("\r\n%sBefore: %s\r\n%sAfter:  %s"),
-						UE::DiffWriterArchive::IndentToken, *BeforePropertyVal,
-						UE::DiffWriterArchive::IndentToken, *AfterPropertyVal);
+						IndentToken, *BeforePropertyVal,
+						IndentToken, *AfterPropertyVal);
 				}
 
 				FString DebugDataStackText;
@@ -685,20 +750,19 @@ void FDiffWriterArchiveWriter::Compare(
 				int32 DebugDataIndex = FullStackText.Find(ANSI_TO_TCHAR(DebugDataStackMarker), ESearchCase::CaseSensitive);
 				if (DebugDataIndex > 0)
 				{
-					DebugDataStackText = FString::Printf(TEXT("\r\n%s"),
-						UE::DiffWriterArchive::IndentToken)
+					DebugDataStackText = FString::Printf(TEXT("\r\n%s"), IndentToken)
 						+ FullStackText.RightChop(DebugDataIndex + 2);
 				}
 
 				MessageCallback(ELogVerbosity::Warning, FString::Printf(
 					TEXT("%s: Difference at offset %lld%s (absolute offset: %lld): byte %d on disk, byte %d in memory, callstack:%s%s%s%s%s"),
-					AssetFilename,
+					*SectionFilename,
 					CallstackAtOffset.Offset - DestPackage.StartOffset,
 					DestAbsoluteOffset > CallstackAtOffset.Offset ? *FString::Printf(TEXT("(+%lld)"), DestAbsoluteOffset - CallstackAtOffset.Offset) : TEXT(""),
 					DestAbsoluteOffset,
 					SourceByte, DestByte,
-					UE::DiffWriterArchive::NewLineToken,
-					UE::DiffWriterArchive::NewLineToken,
+					NewLineToken,
+					NewLineToken,
 					*DifferenceCallstackDataText,
 					*DiffValues,
 					*DebugDataStackText
@@ -707,7 +771,7 @@ void FDiffWriterArchiveWriter::Compare(
 				const int BytesToLog = 128;
 				MessageCallback(ELogVerbosity::Display, FString::Printf(
 					TEXT("%s: Logging %d bytes around absolute offset: %lld (%016X) in the on disk (existing) package, (which corresponds to offset %lld (%016X) in the in-memory package)"),
-					AssetFilename,
+					*SectionFilename,
 					BytesToLog,
 					SourceAbsoluteOffset,
 					SourceAbsoluteOffset,
@@ -723,7 +787,7 @@ void FDiffWriterArchiveWriter::Compare(
 
 				MessageCallback(ELogVerbosity::Display, FString::Printf(
 					TEXT("%s: Logging %d bytes around absolute offset: %lld (%016X) in the in memory (new) package"),
-					AssetFilename,
+					*SectionFilename,
 					BytesToLog,
 					DestAbsoluteOffset,
 					DestAbsoluteOffset
@@ -763,89 +827,87 @@ void FDiffWriterArchiveWriter::Compare(
 		{
 			MessageCallback(ELogVerbosity::Warning, FString::Printf(
 				TEXT("%s: %lld difference(s) not logged (first at absolute offset: %lld)."),
-				AssetFilename, NumDiffsLocal - NumDiffsLoggedLocal, FirstUnreportedDiffIndex));
+				*SectionFilename, NumDiffsLocal - NumDiffsLoggedLocal, FirstUnreportedDiffIndex));
 		}
 		else
 		{
 			MessageCallback(ELogVerbosity::Warning, FString::Printf(
-				TEXT("%s: %lld difference(s) not logged."), AssetFilename, NumDiffsLocal - NumDiffsLoggedLocal));
+				TEXT("%s: %lld difference(s) not logged."), *SectionFilename, NumDiffsLocal - NumDiffsLoggedLocal));
 		}
 	}
 }
 
-void FDiffWriterArchive::CompareWith(const TCHAR* InFilename, const int64 TotalHeaderSize, const TCHAR* CallstackCutoffText, const int32 MaxDiffsToLog, TMap<FName, FArchiveDiffStats>& OutStats)
+void FAccumulator::CompareWithPrevious(const TCHAR* CallstackCutoffText, TMap<FName, FArchiveDiffStats>&OutStats,
+	const EPackageHeaderFormat PackageHeaderFormat)
 {
-	TUniquePtr<uint8> SourcePackageBytes;
-	FPackageData SourcePackage;
-	UE::ArchiveStackTrace::LoadPackageIntoMemory(InFilename, SourcePackage, SourcePackageBytes);
-	CompareWith(SourcePackage, InFilename, TotalHeaderSize, CallstackCutoffText, MaxDiffsToLog, OutStats);
-}
+	// An FDiffArchiveForLinker should have been constructed by SavePackage and should still be in memory
+	check(LinkerArchive);
 
-void FDiffWriterArchive::CompareWith(const FPackageData& SourcePackage, const TCHAR* FileDisplayName, const int64 TotalHeaderSize,
-	const TCHAR* CallstackCutoffText, const int32 MaxDiffsToLog, TMap<FName, FArchiveDiffStats>&OutStats,
-	const FDiffWriterArchiveWriter::EPackageHeaderFormat PackageHeaderFormat /* = FDiffWriterArchiveWriter::EPackageHeaderFormat::PackageFileSummary */)
-{
-	const FName AssetClass = Callstacks.GetAssetClass();
-	OutStats.FindOrAdd(AssetClass).NewFileTotalSize = TotalSize();
-	if (SourcePackage.Size == 0)
+	const FName AssetClass = GetAssetClass();
+	OutStats.FindOrAdd(AssetClass).NewFileTotalSize = LinkerArchive->TotalSize();
+	if (PreviousPackageData.Size == 0)
 	{
-		MessageCallback(ELogVerbosity::Warning, FString::Printf(TEXT("New package: %s"), *GetArchiveName()));
+		MessageCallback(ELogVerbosity::Warning, FString::Printf(TEXT("New package: %s"), *Filename));
 		OutStats.FindOrAdd(AssetClass).DiffSize = OutStats.FindOrAdd(AssetClass).NewFileTotalSize;
 		return;
 	}
 
+	FPackageData SourcePackage;
+	SourcePackage.Data = PreviousPackageData.Data.Get();
+	SourcePackage.Size = PreviousPackageData.Size;
+	SourcePackage.HeaderSize = PreviousPackageData.HeaderSize;
+	SourcePackage.StartOffset = PreviousPackageData.StartOffset;
+
 	FPackageData DestPackage;
-	DestPackage.Data = GetData();
-	DestPackage.Size = TotalSize();
-	DestPackage.HeaderSize = TotalHeaderSize;
+	DestPackage.Data = LinkerArchive->GetData();
+	DestPackage.Size = LinkerArchive->TotalSize();
+	DestPackage.HeaderSize = HeaderSize;
 	DestPackage.StartOffset = 0;
 
-	MessageCallback(ELogVerbosity::Display, FString::Printf(TEXT("Comparing: %s"), *GetArchiveName()));
+	MessageCallback(ELogVerbosity::Display, FString::Printf(TEXT("Comparing: %s"), *Filename));
 	MessageCallback(ELogVerbosity::Warning, FString::Printf(TEXT("Asset class: %s"), *AssetClass.ToString()));
 
 	int32 NumLoggedDiffs = 0;
-
+	
 	FPackageData SourcePackageHeader = SourcePackage;
-	SourcePackageHeader.Size = SourcePackageHeader.HeaderSize;
+	SourcePackageHeader.Size = SourcePackage.HeaderSize;
 	SourcePackageHeader.HeaderSize = 0;
 	SourcePackageHeader.StartOffset = 0;
 
 	FPackageData DestPackageHeader = DestPackage;
-	DestPackageHeader.Size = TotalHeaderSize;
+	DestPackageHeader.Size = HeaderSize;
 	DestPackageHeader.HeaderSize = 0;
 	DestPackageHeader.StartOffset = 0;
 
-	FDiffWriterArchiveWriter::Compare(SourcePackageHeader, DestPackageHeader, Callstacks,
-		StackTraceWriter.GetDiffMap(), FileDisplayName, CallstackCutoffText, MaxDiffsToLog, NumLoggedDiffs,
-		OutStats, MessageCallback);
+	CompareWithPreviousForSection(SourcePackageHeader, DestPackageHeader, CallstackCutoffText,
+		NumLoggedDiffs, OutStats, Filename);
 
-	if (TotalHeaderSize > 0 && OutStats.FindOrAdd(AssetClass).NumDiffs > 0)
+	if (HeaderSize > 0 && OutStats.FindOrAdd(AssetClass).NumDiffs > 0)
 	{
-		FDiffWriterArchiveWriter::DumpPackageHeaderDiffs(SourcePackage, DestPackage, FileDisplayName, MaxDiffsToLog,
+		DumpPackageHeaderDiffs(SourcePackage, DestPackage, Filename, MaxDiffsToLog,
 			PackageHeaderFormat, MessageCallback);
 	}
 
 	FPackageData SourcePackageExports = SourcePackage;
 	SourcePackageExports.HeaderSize = 0;
-	SourcePackageExports.StartOffset = SourcePackage.HeaderSize;
+	SourcePackageExports.StartOffset = PreviousPackageData.HeaderSize;
 
 	FPackageData DestPackageExports = DestPackage;
 	DestPackageExports.HeaderSize = 0;
-	DestPackageExports.StartOffset = TotalHeaderSize;
+	DestPackageExports.StartOffset = HeaderSize;
 
-	FString AssetName;
+	FString ExportsFilename;
 	if (DestPackage.HeaderSize > 0)
 	{
-		AssetName = FPaths::ChangeExtension(FileDisplayName, TEXT("uexp"));
+		ExportsFilename = FPaths::ChangeExtension(Filename, TEXT("uexp"));
 	}
 	else
 	{
-		AssetName = FileDisplayName;
+		ExportsFilename = Filename;
 	}
 
-	FDiffWriterArchiveWriter::Compare(SourcePackageExports, DestPackageExports, Callstacks,
-		StackTraceWriter.GetDiffMap(), *AssetName, CallstackCutoffText, MaxDiffsToLog, NumLoggedDiffs,
-		OutStats, MessageCallback);
+	CompareWithPreviousForSection(SourcePackageExports, DestPackageExports, CallstackCutoffText,
+		NumLoggedDiffs, OutStats, ExportsFilename);
 
 	// Optionally save out any differences we detected.
 	const FArchiveDiffStats& Stats = OutStats.FindOrAdd(AssetClass);
@@ -871,7 +933,7 @@ void FDiffWriterArchive::CompareWith(const FPackageData& SourcePackage, const TC
 		// Only save out the differences if we have a -diffoutputdir set.
 		if (!DiffOutputSettings.DiffOutputDir.IsEmpty())
 		{
-			FString OutputFilename = FPaths::ConvertRelativePathToFull(FileDisplayName);
+			FString OutputFilename = FPaths::ConvertRelativePathToFull(Filename);
 			FString SavedDir       = FPaths::ConvertRelativePathToFull(FPaths::ProjectSavedDir());
 			if (OutputFilename.StartsWith(SavedDir))
 			{
@@ -881,7 +943,7 @@ void FDiffWriterArchive::CompareWith(const FPackageData& SourcePackage, const TC
 
 				// Copy the original asset as '.before.uasset'.
 				{
-					TUniquePtr<FArchive> DiffUAssetArchive(FileManager.CreateFileWriter(*FPaths::SetExtension(OutputFilename, TEXT(".before.") + FPaths::GetExtension(FileDisplayName))));
+					TUniquePtr<FArchive> DiffUAssetArchive(FileManager.CreateFileWriter(*FPaths::SetExtension(OutputFilename, TEXT(".before.") + FPaths::GetExtension(Filename))));
 					DiffUAssetArchive->Serialize(SourcePackageHeader.Data + SourcePackageHeader.StartOffset, SourcePackageHeader.Size - SourcePackageHeader.StartOffset);
 				}
 				{
@@ -891,7 +953,7 @@ void FDiffWriterArchive::CompareWith(const FPackageData& SourcePackage, const TC
 
 				// Save out the in-memory data as '.after.uasset'.
 				{
-					TUniquePtr<FArchive> DiffUAssetArchive(FileManager.CreateFileWriter(*FPaths::SetExtension(OutputFilename, TEXT(".after.") + FPaths::GetExtension(FileDisplayName))));
+					TUniquePtr<FArchive> DiffUAssetArchive(FileManager.CreateFileWriter(*FPaths::SetExtension(OutputFilename, TEXT(".after.") + FPaths::GetExtension(Filename))));
 					DiffUAssetArchive->Serialize(DestPackageHeader.Data + DestPackageHeader.StartOffset, DestPackageHeader.Size - DestPackageHeader.StartOffset);
 				}
 				{
@@ -908,11 +970,13 @@ void FDiffWriterArchive::CompareWith(const FPackageData& SourcePackage, const TC
 	}
 }
 
-bool FDiffWriterArchiveWriter::GenerateDiffMap(const FPackageData& SourcePackage, const FPackageData& DestPackage, const FDiffWriterCallstacks& Callstacks, int32 MaxDiffsToFind, FDiffWriterDiffMap& OutDiffMap)
+void FAccumulator::GenerateDiffMapForSection(const FPackageData& SourcePackage, const FPackageData& DestPackage,
+	bool& bOutSectionIdentical)
 {
+	FCallstacks& Callstacks = this->LinkerCallstacks;
 	bool bIdentical = true;
 	int32 LastDifferenceCallstackOffsetIndex = -1;
-	FDiffWriterCallstacks::FCallstackData* DifferenceCallstackData = nullptr;
+	FCallstacks::FCallstackData* DifferenceCallstackData = nullptr;
 
 	const int64 SourceSize = SourcePackage.Size - SourcePackage.StartOffset;
 	const int64 DestSize = DestPackage.Size - DestPackage.StartOffset;
@@ -925,18 +989,18 @@ bool FDiffWriterArchiveWriter::GenerateDiffMap(const FPackageData& SourcePackage
 		if (SourcePackage.Data[SourceAbsoluteOffset] != DestPackage.Data[DestAbsoluteOffset])
 		{
 			bIdentical = false;
-			if (OutDiffMap.Num() < MaxDiffsToFind)
+			if (DiffMap.Num() < MaxDiffsToLog)
 			{
 				const int32 DifferenceCallstackOffsetIndex = Callstacks.GetCallstackIndexAtOffset(DestAbsoluteOffset, FMath::Max<int32>(LastDifferenceCallstackOffsetIndex, 0));
 				if (DifferenceCallstackOffsetIndex >= 0 && DifferenceCallstackOffsetIndex != LastDifferenceCallstackOffsetIndex)
 				{
-					const FDiffWriterCallstacks::FCallstackAtOffset& CallstackAtOffset = Callstacks.GetCallstack(DifferenceCallstackOffsetIndex);
+					const FCallstacks::FCallstackAtOffset& CallstackAtOffset = Callstacks.GetCallstack(DifferenceCallstackOffsetIndex);
 					if (!CallstackAtOffset.bIgnore)
 					{
-						FDiffWriterDiffInfo OffsetAndSize;
+						FDiffInfo OffsetAndSize;
 						OffsetAndSize.Offset = CallstackAtOffset.Offset;
 						OffsetAndSize.Size = Callstacks.GetSerializedDataSizeForOffsetIndex(DifferenceCallstackOffsetIndex);
-						OutDiffMap.Add(OffsetAndSize);
+						DiffMap.Add(OffsetAndSize);
 					}
 				}
 				LastDifferenceCallstackOffsetIndex = DifferenceCallstackOffsetIndex;
@@ -949,18 +1013,18 @@ bool FDiffWriterArchiveWriter::GenerateDiffMap(const FPackageData& SourcePackage
 		bIdentical = false;
 
 		// Add all the remaining callstacks to the diff map
-		for (int32 OffsetIndex = LastDifferenceCallstackOffsetIndex + 1; OffsetIndex < Callstacks.Num() && OutDiffMap.Num() < MaxDiffsToFind; ++OffsetIndex)
+		for (int32 OffsetIndex = LastDifferenceCallstackOffsetIndex + 1; OffsetIndex < Callstacks.Num() && DiffMap.Num() < MaxDiffsToLog; ++OffsetIndex)
 		{
-			const FDiffWriterCallstacks::FCallstackAtOffset& CallstackAtOffset = Callstacks.GetCallstack(OffsetIndex);
+			const FCallstacks::FCallstackAtOffset& CallstackAtOffset = Callstacks.GetCallstack(OffsetIndex);
 			// Compare against the size without start offset as all callstack offsets are absolute (from the merged header + exports file)
 			if (CallstackAtOffset.Offset < DestPackage.Size)
 			{
 				if (!CallstackAtOffset.bIgnore)
 				{
-					FDiffWriterDiffInfo OffsetAndSize;
+					FDiffInfo OffsetAndSize;
 					OffsetAndSize.Offset = CallstackAtOffset.Offset;
 					OffsetAndSize.Size = Callstacks.GetSerializedDataSizeForOffsetIndex(OffsetIndex);
-					OutDiffMap.Add(OffsetAndSize);
+					DiffMap.Add(OffsetAndSize);
 				}
 			}
 			else
@@ -973,46 +1037,46 @@ bool FDiffWriterArchiveWriter::GenerateDiffMap(const FPackageData& SourcePackage
 	{
 		bIdentical = false;
 	}
-	return bIdentical;
+	bOutSectionIdentical = bIdentical;
 }
 
-bool FDiffWriterArchive::GenerateDiffMap(const TCHAR* InFilename, int64 TotalHeaderSize, int32 MaxDiffsToFind, FDiffWriterDiffMap& OutDiffMap)
+void FAccumulator::GenerateDiffMap()
 {
-	TUniquePtr<uint8> SourcePackageBytes;
-	FPackageData SourcePackage;
-	if (!UE::ArchiveStackTrace::LoadPackageIntoMemory(InFilename, SourcePackage, SourcePackageBytes))
-	{
-		return false;
-	}
-	return GenerateDiffMap(SourcePackage, TotalHeaderSize, MaxDiffsToFind, OutDiffMap);
-}
+	check(MaxDiffsToLog > 0);
+	// An FDiffArchiveForLinker should have been constructed by SavePackage and should still be in memory
+	check(LinkerArchive != nullptr);
 
-bool FDiffWriterArchive::GenerateDiffMap(const FPackageData& SourcePackage, int64 TotalHeaderSize, int32 MaxDiffsToFind, FDiffWriterDiffMap& OutDiffMap)
-{
-	check(MaxDiffsToFind > 0);
+	bHasDifferences = true;
+	DiffMap.Reset();
+
+	FPackageData SourcePackage;
+	SourcePackage.Data = PreviousPackageData.Data.Get();
+	SourcePackage.Size = PreviousPackageData.Size;
+	SourcePackage.HeaderSize = PreviousPackageData.HeaderSize;
+	SourcePackage.StartOffset = PreviousPackageData.StartOffset;
 
 	bool bIdentical = true;
 	bool bHeaderIdentical = true;
 	bool bExportsIdentical = true;
 
 	FPackageData DestPackage;
-	DestPackage.Data = GetData();
-	DestPackage.Size = TotalSize();
-	DestPackage.HeaderSize = TotalHeaderSize;
+	DestPackage.Data = LinkerArchive->GetData();
+	DestPackage.Size = LinkerArchive->TotalSize();
+	DestPackage.HeaderSize = HeaderSize;
 	DestPackage.StartOffset = 0;
 
 	{
 		FPackageData SourcePackageHeader = SourcePackage;
-		SourcePackageHeader.Size = SourcePackageHeader.HeaderSize;
+		SourcePackageHeader.Size = SourcePackage.HeaderSize;
 		SourcePackageHeader.HeaderSize = 0;
 		SourcePackageHeader.StartOffset = 0;
 
 		FPackageData DestPackageHeader = DestPackage;
-		DestPackageHeader.Size = TotalHeaderSize;
+		DestPackageHeader.Size = HeaderSize;
 		DestPackageHeader.HeaderSize = 0;
 		DestPackageHeader.StartOffset = 0;
 
-		bHeaderIdentical = FDiffWriterArchiveWriter::GenerateDiffMap(SourcePackageHeader, DestPackageHeader, Callstacks, MaxDiffsToFind, OutDiffMap);
+		GenerateDiffMapForSection(SourcePackageHeader, DestPackageHeader, bHeaderIdentical);
 	}
 
 	{
@@ -1022,43 +1086,16 @@ bool FDiffWriterArchive::GenerateDiffMap(const FPackageData& SourcePackage, int6
 
 		FPackageData DestPackageExports = DestPackage;
 		DestPackageExports.HeaderSize = 0;
-		DestPackageExports.StartOffset = TotalHeaderSize;
+		DestPackageExports.StartOffset = HeaderSize;
 
-		bExportsIdentical = FDiffWriterArchiveWriter::GenerateDiffMap(SourcePackageExports, DestPackageExports, Callstacks, MaxDiffsToFind, OutDiffMap);
+		GenerateDiffMapForSection(SourcePackageExports, DestPackageExports, bExportsIdentical);
 	}
 
 	bIdentical = bHeaderIdentical && bExportsIdentical;
-
-	return bIdentical;
-}
-
-bool FDiffWriterArchive::IsIdentical(const TCHAR* InFilename, int64 BufferSize, const uint8* BufferData)
-{
-	TUniquePtr<uint8> SourcePackageBytes;
-	FPackageData SourcePackage;
-	if (!UE::ArchiveStackTrace::LoadPackageIntoMemory(InFilename, SourcePackage, SourcePackageBytes))
-	{
-		return false;
-	}
-
-	return IsIdentical(SourcePackage, BufferSize, BufferData);
-}
-
-bool FDiffWriterArchive::IsIdentical(const FPackageData& SourcePackage, int64 BufferSize, const uint8* BufferData)
-{
-	bool bIdentical = false;
-	if (BufferSize == SourcePackage.Size)
-	{
-		bIdentical = (FMemory::Memcmp(SourcePackage.Data, BufferData, BufferSize) == 0);
-	}
-	else
-	{
-		bIdentical = false;
-	}
-	return bIdentical;
+	bHasDifferences = !bIdentical;
 }
 
-FLinkerLoad* FDiffWriterArchiveWriter::CreateLinkerForPackage(FUObjectSerializeContext* LoadContext, const FString& InPackageName, const FString& InFilename, const FPackageData& PackageData)
+FLinkerLoad* CreateLinkerForPackage(FUObjectSerializeContext* LoadContext, const FString& InPackageName, const FString& InFilename, const FPackageData& PackageData)
 {
 	// First create a temp package to associate the linker with
 	UPackage* Package = FindObjectFast<UPackage>(nullptr, *InPackageName);
@@ -1122,16 +1159,16 @@ static inline FString GetTableKeyForIndex(const FLinkerLoad* Linker, FPackageInd
 }
 
 static bool ComparePackageIndices(FLinkerLoad* SourceLinker, FLinkerLoad* DestLinker, const FPackageIndex& SourceIndex,
-	const FPackageIndex& DestIndex, const UE::DiffWriterArchive::FMessageCallback& MessageCallback);
+	const FPackageIndex& DestIndex, const FMessageCallback& MessageCallback);
 
 static bool CompareTableItem(FLinkerLoad* SourceLinker, FLinkerLoad* DestLinker,
-	const FName& SourceName, const FName& DestName, const UE::DiffWriterArchive::FMessageCallback& MessageCallback)
+	const FName& SourceName, const FName& DestName, const FMessageCallback& MessageCallback)
 {
 	return SourceName == DestName;
 }
 
 static bool CompareTableItem(FLinkerLoad* SourceLinker, FLinkerLoad* DestLinker,
-	FNameEntryId SourceName, FNameEntryId DestName, const UE::DiffWriterArchive::FMessageCallback& MessageCallback)
+	FNameEntryId SourceName, FNameEntryId DestName, const FMessageCallback& MessageCallback)
 {
 	return SourceName == DestName;
 }
@@ -1147,7 +1184,7 @@ static FString ConvertItemToText(FNameEntryId Id, FLinkerLoad* Linker)
 }
 
 static bool CompareTableItem(FLinkerLoad* SourceLinker, FLinkerLoad* DestLinker, const FObjectImport& SourceImport,
-	const FObjectImport& DestImport, const UE::DiffWriterArchive::FMessageCallback& MessageCallback)
+	const FObjectImport& DestImport, const FMessageCallback& MessageCallback)
 {
 	if (SourceImport.ObjectName != DestImport.ObjectName ||
 		SourceImport.ClassName != DestImport.ClassName ||
@@ -1173,7 +1210,7 @@ static FString ConvertItemToText(const FObjectImport& Import, FLinkerLoad* Linke
 }
 
 static bool CompareTableItem(FLinkerLoad* SourceLinker, FLinkerLoad* DestLinker, const FObjectExport& SourceExport,
-	const FObjectExport& DestExport, const UE::DiffWriterArchive::FMessageCallback& MessageCallback)
+	const FObjectExport& DestExport, const FMessageCallback& MessageCallback)
 {
 	if (SourceExport.ObjectName != DestExport.ObjectName ||
 		SourceExport.PackageFlags != DestExport.PackageFlags ||
@@ -1204,7 +1241,7 @@ static bool CompareTableItem(FLinkerLoad* SourceLinker, FLinkerLoad* DestLinker,
 }
 
 static bool IsImportMapIdentical(FLinkerLoad* SourceLinker, FLinkerLoad* DestLinker,
-	const UE::DiffWriterArchive::FMessageCallback& MessageCallback)
+	const FMessageCallback& MessageCallback)
 {
 	bool bIdentical = (SourceLinker->ImportMap.Num() == DestLinker->ImportMap.Num());
 	if (bIdentical)
@@ -1223,7 +1260,7 @@ static bool IsImportMapIdentical(FLinkerLoad* SourceLinker, FLinkerLoad* DestLin
 }
 
 static bool ComparePackageIndices(FLinkerLoad* SourceLinker, FLinkerLoad* DestLinker, const FPackageIndex& SourceIndex,
-	const FPackageIndex& DestIndex, const UE::DiffWriterArchive::FMessageCallback& MessageCallback)
+	const FPackageIndex& DestIndex, const FMessageCallback& MessageCallback)
 {
 	if (SourceIndex.IsNull() && DestIndex.IsNull())
 	{
@@ -1297,7 +1334,7 @@ static FString ConvertItemToText(const FObjectExport& Export, FLinkerLoad* Linke
 }
 
 static bool IsExportMapIdentical(FLinkerLoad* SourceLinker, FLinkerLoad* DestLinker,
-	const UE::DiffWriterArchive::FMessageCallback& MessageCallback)
+	const FMessageCallback& MessageCallback)
 {
 	bool bIdentical = (SourceLinker->ExportMap.Num() == DestLinker->ExportMap.Num());
 	if (bIdentical)
@@ -1354,7 +1391,7 @@ static void DumpTableDifferences(
 	const TCHAR* AssetFilename,
 	const TCHAR* ItemName,
 	const int32 MaxDiffsToLog,
-	const UE::DiffWriterArchive::FMessageCallback& MessageCallback
+	const FMessageCallback& MessageCallback
 )
 {
 	FString HumanReadableString;
@@ -1403,15 +1440,15 @@ static void DumpTableDifferences(
 	// Dump all changes
 	for (const TTableItem<T>& RemovedItem : RemovedItems)
 	{
-		HumanReadableString += UE::DiffWriterArchive::IndentToken;
+		HumanReadableString += IndentToken;
 		HumanReadableString += FString::Printf(TEXT("-[%d] %s"), RemovedItem.Index, *ConvertItemToText(*RemovedItem.Item, SourceLinker));
-		HumanReadableString += UE::DiffWriterArchive::NewLineToken;
+		HumanReadableString += NewLineToken;
 	}
 	for (const TTableItem<T>& AddedItem : AddedItems)
 	{
-		HumanReadableString += UE::DiffWriterArchive::IndentToken;
+		HumanReadableString += IndentToken;
 		HumanReadableString += FString::Printf(TEXT("+[%d] %s"), AddedItem.Index, *ConvertItemToText(*AddedItem.Item, DestLinker));
-		HumanReadableString += UE::DiffWriterArchive::NewLineToken;
+		HumanReadableString += NewLineToken;
 	}
 
 	// For now just log everything out. When this becomes too spammy, respect the MaxDiffsToLog parameter
@@ -1420,9 +1457,9 @@ static void DumpTableDifferences(
 
 	if (NumDiffs > LoggedDiffs)
 	{
-		HumanReadableString += UE::DiffWriterArchive::IndentToken;
+		HumanReadableString += IndentToken;
 		HumanReadableString += FString::Printf(TEXT("+ %d differences not logged."), (NumDiffs - LoggedDiffs));
-		HumanReadableString += UE::DiffWriterArchive::NewLineToken;
+		HumanReadableString += NewLineToken;
 	}
 
 	MessageCallback(ELogVerbosity::Warning, FString::Printf(
@@ -1433,16 +1470,16 @@ static void DumpTableDifferences(
 		ItemName,
 		DestTable.Num(),
 		ItemName,
-		UE::DiffWriterArchive::NewLineToken,
+		NewLineToken,
 		*HumanReadableString));
 }
 
-static void DumpPackageHeaderDiffs_LinkerLoad(
-	const FDiffWriterArchiveWriter::FPackageData& SourcePackage,
-	const FDiffWriterArchiveWriter::FPackageData& DestPackage,
+void DumpPackageHeaderDiffs_LinkerLoad(
+	const FPackageData& SourcePackage,
+	const FPackageData& DestPackage,
 	const FString& AssetFilename,
 	const int32 MaxDiffsToLog,
-	const UE::DiffWriterArchive::FMessageCallback& MessageCallback)
+	const FMessageCallback& MessageCallback)
 {
 	FString AssetPathName = FPaths::Combine(*FPaths::GetPath(AssetFilename.Mid(AssetFilename.Find(TEXT(":"), ESearchCase::CaseSensitive) + 1)), *FPaths::GetBaseFilename(AssetFilename));
 	// The root directory could have a period in it (d:/Release5.0/EngineTest/Saved/Cooked),
@@ -1465,14 +1502,14 @@ static void DumpPackageHeaderDiffs_LinkerLoad(
 	{
 		TRefCountPtr<FUObjectSerializeContext> LinkerLoadContext(FUObjectThreadContext::Get().GetSerializeContext());
 		BeginLoad(LinkerLoadContext);
-		SourceLinker = FDiffWriterArchiveWriter::CreateLinkerForPackage(LinkerLoadContext, SourceAssetPackageName, AssetFilename, SourcePackage);
+		SourceLinker = CreateLinkerForPackage(LinkerLoadContext, SourceAssetPackageName, AssetFilename, SourcePackage);
 		EndLoad(SourceLinker ? SourceLinker->GetSerializeContext() : LinkerLoadContext.GetReference());
 	}
 	
 	{
 		TRefCountPtr<FUObjectSerializeContext> LinkerLoadContext(FUObjectThreadContext::Get().GetSerializeContext());
 		BeginLoad(LinkerLoadContext);
-		DestLinker = FDiffWriterArchiveWriter::CreateLinkerForPackage(LinkerLoadContext, DestAssetPackageName, AssetFilename, DestPackage);
+		DestLinker = CreateLinkerForPackage(LinkerLoadContext, DestAssetPackageName, AssetFilename, DestPackage);
 		EndLoad(DestLinker ? DestLinker->GetSerializeContext() : LinkerLoadContext.GetReference());
 	}
 
@@ -1507,23 +1544,23 @@ static void DumpPackageHeaderDiffs_LinkerLoad(
 	}
 }
 
-static void DumpPackageHeaderDiffs_ZenPackage(
-	const FDiffWriterArchiveWriter::FPackageData& SourcePackage,
-	const FDiffWriterArchiveWriter::FPackageData& DestPackage,
+void DumpPackageHeaderDiffs_ZenPackage(
+	const FPackageData& SourcePackage,
+	const FPackageData& DestPackage,
 	const FString& AssetFilename,
 	const int32 MaxDiffsToLog,
-	const UE::DiffWriterArchive::FMessageCallback& MessageCallback)
+	const FMessageCallback& MessageCallback)
 {
 	// TODO: Fill in detailed diffing of Zen Package Summary
 }
 
-void FDiffWriterArchiveWriter::DumpPackageHeaderDiffs(
+void DumpPackageHeaderDiffs(
 	const FPackageData& SourcePackage,
 	const FPackageData& DestPackage,
 	const FString& AssetFilename,
 	const int32 MaxDiffsToLog,
 	const EPackageHeaderFormat PackageHeaderFormat,
-	const UE::DiffWriterArchive::FMessageCallback& MessageCallback)
+	const FMessageCallback& MessageCallback)
 {
 	switch (PackageHeaderFormat)
 	{
@@ -1537,3 +1574,5 @@ void FDiffWriterArchiveWriter::DumpPackageHeaderDiffs(
 		unimplemented();
 	}
 }
+
+} // namespace UE::DiffWriter
\ No newline at end of file
diff --git a/Engine/Source/Editor/UnrealEd/Private/Cooker/DiffWriterArchive.h b/Engine/Source/Editor/UnrealEd/Private/Cooker/DiffWriterArchive.h
--- a/Engine/Source/Editor/UnrealEd/Private/Cooker/DiffWriterArchive.h
+++ b/Engine/Source/Editor/UnrealEd/Private/Cooker/DiffWriterArchive.h
@@ -18,6 +18,7 @@
 #include "Templates/UniquePtr.h"
 #include "UObject/NameTypes.h"
 
+class FDiffWriterArchiveTestsCallstacks;
 class FLinkerLoad;
 class FProperty;
 class FUObjectThreadContext;
@@ -25,38 +26,49 @@ class UObject;
 struct FUObjectSerializeContext;
 
 
-namespace UE::DiffWriterArchive
+namespace UE::DiffWriter
 {
 
+class FAccumulator;
+class FDiffArchive;
+
+typedef TUniqueFunction<void(ELogVerbosity::Type, FStringView)> FMessageCallback;
+using EPackageHeaderFormat = ICookedPackageWriter::EPackageHeaderFormat;
+using FPackageData = UE::ArchiveStackTrace::FPackageData;
+
 extern const TCHAR* const IndentToken;
 extern const TCHAR* const NewLineToken;
-typedef TUniqueFunction<void(ELogVerbosity::Type, FStringView)> FMessageCallback;
 
-}
+enum class EOffsetFrame
+{
+	Linker,
+	Exports,
+};
 
-struct FDiffWriterDiffInfo
+struct FDiffInfo
 {
 	int64 Offset;
 	int64 Size;
-	FDiffWriterDiffInfo()
+
+	FDiffInfo()
 		: Offset(0)
 		, Size(0)
 	{
 	}
-	FDiffWriterDiffInfo(int64 InOffset, int64 InSize)
+	FDiffInfo(int64 InOffset, int64 InSize)
 		: Offset(InOffset)
 		, Size(InSize)
 	{
 	}
-	bool operator==(const FDiffWriterDiffInfo& InOther) const
+	bool operator==(const FDiffInfo& InOther) const
 	{
 		return Offset == InOther.Offset;
 	}
-	bool operator<(const FDiffWriterDiffInfo& InOther) const
+	bool operator<(const FDiffInfo& InOther) const
 	{
 		return Offset < InOther.Offset;
 	}
-	friend FArchive& operator << (FArchive& Ar, FDiffWriterDiffInfo& InDiffInfo)
+	friend FArchive& operator << (FArchive& Ar, FDiffInfo& InDiffInfo)
 	{
 		Ar << InDiffInfo.Offset;
 		Ar << InDiffInfo.Size;
@@ -64,12 +76,12 @@ struct FDiffWriterDiffInfo
 	}
 };
 
-class FDiffWriterDiffMap : public TArray<FDiffWriterDiffInfo>
+class FDiffMap : public TArray<FDiffInfo>
 {
 public:
 	bool ContainsOffset(int64 Offset) const
 	{
-		for (const FDiffWriterDiffInfo& Diff : *this)
+		for (const FDiffInfo& Diff : *this)
 		{
 			if (Diff.Offset <= Offset && Offset < (Diff.Offset + Diff.Size))
 			{
@@ -82,7 +94,7 @@ public:
 };
 
 /** Holds offsets to captured callstacks. */
-class FDiffWriterCallstacks
+class FCallstacks
 {
 public:
 	/** Offset and callstack pair */
@@ -123,16 +135,17 @@ public:
 		FCallstackData Clone() const;
 	};
 	
-	explicit FDiffWriterCallstacks(UObject* InAsset);
-
-	/** Returns the asset class name. */
-	FName GetAssetClass() const;
+	FCallstacks();
 
 	/** Returns the total number of callstacks. */
 	int32 Num() const
 	{
 		return CallstackAtOffsetMap.Num();
 	}
+	void Reset();
+
+	FORCENOINLINE void RecordSerialize(EOffsetFrame OffsetFrame, int64 CurrentOffset, int64 Length,
+		const FAccumulator& Accumulator, FDiffArchive& Ar, int32 StackIgnoreCount);
 
 	/** Capture and append the current callstack. */
 	void Add(
@@ -146,7 +159,7 @@ public:
 		int32 StackIgnoreCount);
 
 	/** Append other callstacks. */
-	void Append(const FDiffWriterCallstacks& Other, int64 Offset = 0);
+	void Append(const FCallstacks& Other, int64 OtherStartOffset);
 
 	/** Finds a callstack associated with data at the specified offset */
 	int32 GetCallstackIndexAtOffset(int64 Offset, int32 MinOffsetIndex = 0) const;
@@ -171,22 +184,19 @@ public:
 		}
 		else
 		{
-			return TotalSize - CallstackAtOffsetMap[InOffsetIndex].Offset;
+			return EndOffset - CallstackAtOffsetMap[InOffsetIndex].Offset;
 		}
 	}
 
-	/** Returns total serialized bytes. */
-	int64 TotalCapturedSize() const
+	int64 GetEndOffset() const
 	{
-		return TotalSize;
+		return EndOffset;
 	}
 
 private:
 	/** Adds a unique callstack to UniqueCallstacks map */
 	ANSICHAR* AddUniqueCallstack(bool bIsCollectingCallstacks, UObject* SerializedObject, FProperty* SerializedProperty, uint32& OutCallstackCRC);
 
-	/** The asset being serialized */
-	UObject* Asset;
 	/** List of offsets and their respective callstacks */
 	TArray<FCallstackAtOffset> CallstackAtOffsetMap;
 	/** Contains all unique callstacks for all Serialize calls */
@@ -200,170 +210,115 @@ private:
 	/** Callstack associated with the previous Serialize call */
 	ANSICHAR* LastSerializeCallstack;
 	/** Total serialized bytes */
-	int64 TotalSize;
+	int64 EndOffset;
 };
 
-/** Archive proxy that captures callstacks for each serialize call. */
-class FDiffWriterArchiveWriter
-	: public FArchiveProxy
+/**
+ * Collects the memory version of a saved package, compares it with an existing package on disk, and reports callstack
+ * for the Serialize call at each offset where they differ.
+ * 
+ * It works by saving a package twice. The first pass collects the serialization offsets without the stack traces and
+ * creates a FDiffMap, and records sizes necessary to remap offsets during Serialize to the final offset in the package
+ * on disk. In the second pass, the diff map is read during each call to Serialize to decide whether we need to collect
+ * the stack trace for that call.
+ */
+class FAccumulator : public FRefCountBase
 {
 public:
-	FDiffWriterArchiveWriter(
-		FArchive& InInner,
-		FDiffWriterCallstacks& InCallstacks,
-		const FDiffWriterDiffMap* InDiffMap = nullptr,
-		int64 InDiffMapStartOffset = 0);
+	FAccumulator(UObject* InAsset, FName InPackageName, int32 InMaxDiffsToLog,
+		FMessageCallback&& InMessageCallback);
+	virtual ~FAccumulator();
 
-	virtual ~FDiffWriterArchiveWriter() override;
+	void OnFirstSaveComplete(FStringView InLooseFilePath, int64 InHeaderSize,
+		ICookedPackageWriter::FPreviousCookedBytesData&& InPreviousPackageData);
+	void OnSecondSaveComplete(int64 InHeaderSize);
+	bool HasDifferences() const;
 
-	FORCENOINLINE virtual void Serialize(void* Data, int64 Length) override; // FORCENOINLINE so it can be counted during StackTrace
-	virtual void SetSerializeContext(FUObjectSerializeContext* Context) override;
-	virtual FUObjectSerializeContext* GetSerializeContext() override;
-	
-	virtual void PushDebugDataString(const FName& DebugData) override
-	{
-		DebugDataStack.Push(DebugData);
-	}
+	/** Compares results from the second save with the previous cook results in PreviousPackagedata.  */
+	void CompareWithPrevious(const TCHAR* CallstackCutoffText, TMap<FName,FArchiveDiffStats>& OutStats,
+		const EPackageHeaderFormat PackageHeaderFormat);
 
-	virtual void PopDebugDataString() override
-	{
-		DebugDataStack.Pop();
-	}
+	void SetHeaderSize(int64 InHeaderSize);
+	void SetCollectingCallstacks(bool bInCollectingCallstacks);
+	FName GetAssetClass() const;
 
-	const FDiffWriterDiffMap& GetDiffMap() const
-	{
-		static FDiffWriterDiffMap Empty;
-		return DiffMap != nullptr ? *DiffMap : Empty;
-	}
 
-	void SetDisableInnerArchive(bool bDisable)
-	{
-		bInnerArchiveDisabled = bDisable;
-	}
+private:
+	void GenerateDiffMapForSection(const FPackageData& SourcePackage, const FPackageData& DestPackage, bool& bOutSectionIdentical);
+	void GenerateDiffMap();
+	/** Compares two packages and logs the differences and calltacks. */
+	void CompareWithPreviousForSection(const FPackageData& SourcePackage, const FPackageData& DestPackage,
+		const TCHAR* CallstackCutoffText, int32& InOutLoggedDiffs,TMap<FName, FArchiveDiffStats>& OutStats,
+		const FString& SectionFilename);
 
-	int32 GetStackIgnoreCount() const { return StackIgnoreCount; }
-	void SetStackIgnoreCount(const int32 IgnoreCount) { StackIgnoreCount = IgnoreCount; }
+private:
+	FCallstacks LinkerCallstacks;
+	FCallstacks ExportsCallstacks;
+	ICookedPackageWriter::FPreviousCookedBytesData PreviousPackageData;
+	FDiffArchive* LinkerArchive = nullptr;
+	FDiffArchive* ExportsArchive = nullptr;
+
+	FDiffMap DiffMap;
+	FMessageCallback MessageCallback;
+	FName PackageName;
+	FString Filename;
+	UObject* Asset = nullptr;
+	int64 HeaderSize = 0;
+	int32 MaxDiffsToLog = 5;
+	bool bFirstSaveComplete = false;
+	bool bHasDifferences = false;
+
+	friend class FCallstacks;
+	friend class FDiffArchive;
+	friend class FDiffArchiveForLinker;
+	friend class FDiffArchiveForExports;
+	friend class ::FDiffWriterArchiveTestsCallstacks;
+};
 
-	using FPackageData = UE::ArchiveStackTrace::FPackageData;
+class FDiffArchive : public FLargeMemoryWriter
+{
+public:
+	FDiffArchive(FAccumulator& InAccumulator);
 
-	using EPackageHeaderFormat = ICookedPackageWriter::EPackageHeaderFormat;
+	// FLargeMemoryWriterBase interface
+	virtual FString GetArchiveName() const override;
+	virtual void PushDebugDataString(const FName& DebugData) override;
+	virtual void PopDebugDataString() override;
 
-	/** Compares two packages and logs the differences and calltacks. */
-	static void Compare(
-		const FPackageData& SourcePackage,
-		const FPackageData& DestPackage,
-		const FDiffWriterCallstacks& Callstacks,
-		const FDiffWriterDiffMap& DiffMap,
-		const TCHAR* AssetFilename,
-		const TCHAR* CallstackCutoffText,
-		const int64 MaxDiffsToLog,
-		int32& InOutDiffsLogged,
-		TMap<FName, FArchiveDiffStats>& OutStats,
-		const UE::DiffWriterArchive::FMessageCallback& MessageCallback);
-
-	/** Creates map with mismatching callstacks. */
-	static bool GenerateDiffMap(
-		const FPackageData& SourcePackage,
-		const FPackageData& DestPackage,
-		const FDiffWriterCallstacks& Callstacks,
-		int32 MaxDiffsToFind,
-		FDiffWriterDiffMap& OutDiffMap);
-
-	/** Logs any mismatching header data. */
-	static void DumpPackageHeaderDiffs(
-		const FPackageData& SourcePackage,
-		const FPackageData& DestPackage,
-		const FString& AssetFilename,
-		const int32 MaxDiffsToLog,
-		const EPackageHeaderFormat PackageHeaderFormat,
-		const UE::DiffWriterArchive::FMessageCallback& MessageCallback);
-
-	/** Returns a new linker for loading the specified package. */
-	static FLinkerLoad* CreateLinkerForPackage(
-		FUObjectSerializeContext* LoadContext,
-		const FString& InPackageName,
-		const FString& InFilename,
-		const FPackageData& PackageData);
+	// FDiffArchive interface
+	FAccumulator& GetAccumulator() { return *Accumulator; }
+	TArray<FName>& GetDebugDataStack();
 
-private:
-	FDiffWriterCallstacks& Callstacks;
-	const FDiffWriterDiffMap* DiffMap;
-	TRefCountPtr<FUObjectSerializeContext> SerializeContext;
+
+protected:
 	TArray<FName> DebugDataStack;
-	int64 DiffMapOffset;
-	int32 StackIgnoreCount = 2;
-	bool bInnerArchiveDisabled;
+	TRefCountPtr<FAccumulator> Accumulator;
 };
 
 /**
- * Memory backed stack trace writer.
+ * The archive written to when SavePackage is writing the Serialize blobs for exports.
+ * When cooking, exports are serialized into a separate archive. We collect the serialization
+ * callstack offsets and stack traces are into a separate callstack collection and append it
+ * at the proper offset to the overall callstacks for the entire linker archive.
  */
-class FDiffWriterArchiveMemoryWriter final
-	: public FLargeMemoryWriter
+class FDiffArchiveForLinker : public FDiffArchive
 {
 public:
-	FDiffWriterArchiveMemoryWriter(
-		FDiffWriterCallstacks& Callstacks,
-		const FDiffWriterDiffMap* DiffMap = nullptr,
-		const int64 DiffMapOffset = 0,
-		const int64 PreAllocateBytes = 0,
-		bool bIsPersistent = false,
-		const TCHAR* Filename = nullptr);
-
-	FORCENOINLINE virtual void Serialize(void* Memory, int64 Length) override; // FORCENOINLINE so it can be counted during StackTrace
-	virtual void SetSerializeContext(FUObjectSerializeContext* Context) override;
-	virtual FUObjectSerializeContext* GetSerializeContext() override;
+	FDiffArchiveForLinker(FAccumulator& InAccumulator);
+	~FDiffArchiveForLinker();
 
-private:
-	FDiffWriterArchiveWriter StackTraceWriter;
+	// FLargeMemoryWriter interface
+	FORCENOINLINE virtual void Serialize(void* InData, int64 Num) override; // FORCENOINLINE so it can be counted during StackTrace
 };
 
-/**
- * Archive that stores a callstack for each of the Serialize calls and has the ability to compare itself to an existing
- * package on disk and dump all the differences to log.
- */
-class FDiffWriterArchive
-	: public FLargeMemoryWriter
+class FDiffArchiveForExports : public FDiffArchive
 {
 public:
-	using FPackageData = FDiffWriterArchiveWriter::FPackageData;
-
-	FDiffWriterArchive(UObject* InAsset, const TCHAR* InFilename, UE::DiffWriterArchive::FMessageCallback&& InMessageCallback,
-		bool bInCollectCallstacks = true, const FDiffWriterDiffMap* InDiffMap = nullptr);
-	virtual ~FDiffWriterArchive();
+	FDiffArchiveForExports(FAccumulator& InAccumulator);
+	~FDiffArchiveForExports();
 
-	FDiffWriterCallstacks& GetCallstacks()
-	{
-		return Callstacks;
-	}
-
-	const FDiffWriterCallstacks& GetCallstacks() const
-	{
-		return Callstacks;
-	}
-
-	FORCENOINLINE virtual void Serialize(void* Memory, int64 Length) override; // FORCENOINLINE so it can be counted during StackTrace
-	virtual void SetSerializeContext(FUObjectSerializeContext* Context) override;
-	virtual FUObjectSerializeContext* GetSerializeContext() override;
-
-	/** Compares this archive with the given bytes from disk or FPackageData. Dumps all differences to log. */
-	void CompareWith(const TCHAR* InFilename, const int64 TotalHeaderSize, const TCHAR* CallstackCutoffText,
-		const int32 MaxDiffsToLog, TMap<FName, FArchiveDiffStats>& OutStats);
-	void CompareWith(const FPackageData& SourcePackage, const TCHAR* FileDisplayName, const int64 TotalHeaderSize,
-		const TCHAR* CallstackCutoffText, const int32 MaxDiffsToLog, TMap<FName, FArchiveDiffStats>& OutStats,
-		const FDiffWriterArchiveWriter::EPackageHeaderFormat PackageHeaderFormat = FDiffWriterArchiveWriter::EPackageHeaderFormat::PackageFileSummary);
-
-	/** Generates a map of all differences between this archive and the given bytes from disk or FPackageData. */
-	bool GenerateDiffMap(const TCHAR* InFilename, int64 TotalHeaderSize, int32 MaxDiffsToFind, FDiffWriterDiffMap& OutDiffMap);
-	bool GenerateDiffMap(const FPackageData& SourcePackage, int64 TotalHeaderSize, int32 MaxDiffsToFind,
-		FDiffWriterDiffMap& OutDiffMap);
-
-	/** Compares the provided buffer with the given bytes from disk or FPackageData. */
-	static bool IsIdentical(const TCHAR* InFilename, int64 BufferSize, const uint8* BufferData);
-	static bool IsIdentical(const FPackageData& SourcePackage, int64 BufferSize, const uint8* BufferData);
-
-private:
-	FDiffWriterCallstacks Callstacks;
-	FDiffWriterArchiveWriter StackTraceWriter;
-	UE::DiffWriterArchive::FMessageCallback MessageCallback;
+	// FLargeMemoryWriter interface
+	FORCENOINLINE virtual void Serialize(void* InData, int64 Num) override; // FORCENOINLINE so it can be counted during StackTrace
 };
+
+} // namespace UE::DiffWriter
\ No newline at end of file
diff --git a/Engine/Source/Editor/UnrealEd/Private/Tests/DiffWriterArchiveTests.cpp b/Engine/Source/Editor/UnrealEd/Private/Tests/DiffWriterArchiveTests.cpp
--- a/Engine/Source/Editor/UnrealEd/Private/Tests/DiffWriterArchiveTests.cpp
+++ b/Engine/Source/Editor/UnrealEd/Private/Tests/DiffWriterArchiveTests.cpp
@@ -10,7 +10,7 @@
 #include "Serialization/LargeMemoryWriter.h"
 #include "Serialization/MemoryReader.h"
 
-namespace UE::DiffWriterArchive::Tests
+namespace UE::DiffWriter::Tests
 {
 
 struct FBasicDiffState
@@ -27,34 +27,37 @@ struct FBasicDiffState
 	}
 };
 
-FDiffWriterArchiveWriter::FPackageData ToPackageData(FLargeMemoryWriter& Ar)
+ICookedPackageWriter::FPreviousCookedBytesData ToPackageData(FLargeMemoryWriter& Ar)
 {
-	return FDiffWriterArchiveWriter::FPackageData
+	int64 Size = Ar.TotalSize(); // Cache this before calling ReleaseOwnership
+	return ICookedPackageWriter::FPreviousCookedBytesData
 	{
-		Ar.GetData(),
-		Ar.TotalSize()
+		TUniquePtr<uint8>(Ar.ReleaseOwnership()), Size, 0 /* HeaderSize */, 0 /* StartOffset */
 	};
 }
 
-} // namespace UE::DiffWriterArchive::Tests
+} // namespace UE::DiffWriter::Tests
 
 IMPLEMENT_SIMPLE_AUTOMATION_TEST(FDiffWriterArchiveTestsCallstacks, "System.Core.Cooker.DiffWriterArchive.Callstacks",
 	EAutomationTestFlags::EditorContext | EAutomationTestFlags::EngineFilter)
 
 bool FDiffWriterArchiveTestsCallstacks::RunTest(const FString& Parameters)
 {
-	using namespace UE::DiffWriterArchive::Tests;
+	using namespace UE::DiffWriter;
+	using namespace UE::DiffWriter::Tests;
 
 	const int32 MaxDiffsToLog = 1000;
-	FDiffWriterCallstacks Callstacks(nullptr);
+	FName PackageName(TEXT("PackageName"));
+	TRefCountPtr<FAccumulator> Accumulator = new FAccumulator(nullptr, PackageName, MaxDiffsToLog,
+		[](ELogVerbosity::Type Verbosity, FStringView Message) {});
+	FBasicDiffState State = FBasicDiffState{ 1,2,3 };
+	FDiffArchiveForLinker Ar(*Accumulator);
 
-	FLargeMemoryWriter Memory;
-	FBasicDiffState State = FBasicDiffState {1,2,3};
-	FDiffWriterArchiveWriter StackTraceWriter(Memory, Callstacks, 0);
-	State.Serialize(StackTraceWriter);
+	State.Serialize(Ar);
 
+	FCallstacks& Callstacks = Accumulator->LinkerCallstacks;
 	TestTrueExpr(Callstacks.Num() == 3);
-	TestTrueExpr(Callstacks.TotalCapturedSize() == sizeof(FBasicDiffState));
+	TestTrueExpr(Callstacks.GetEndOffset() == sizeof(FBasicDiffState));
 	TestTrueExpr(Callstacks.GetCallstack(0).Offset == 0);
 	TestTrueExpr(Callstacks.GetCallstack(0).Callstack == 0);
 	TestTrueExpr(Callstacks.GetCallstack(0).bIgnore == false);
@@ -72,14 +75,17 @@ IMPLEMENT_SIMPLE_AUTOMATION_TEST(FDiffWriterArchiveTestsBasic, "System.Core.Cook
 
 bool FDiffWriterArchiveTestsBasic::RunTest(const FString& Parameters)
 {
-	using namespace UE::DiffWriterArchive::Tests;
+	using namespace UE::DiffWriter;
+	using namespace UE::DiffWriter::Tests;
 
 	const int32 MaxDiffsToLog = 1000;
+	FName PackageName(TEXT("PackageName"));
+	FString Filename(TEXT("Filename"));
 
 	// SECTION("DiffMap - identical")
 	{
-		FDiffWriterCallstacks Callstacks(nullptr);
-		FDiffWriterDiffMap DiffMap;
+		TRefCountPtr<FAccumulator> Accumulator = new FAccumulator(nullptr, PackageName, MaxDiffsToLog,
+			[](ELogVerbosity::Type Verbosity, FStringView Message) {});
 
 		FLargeMemoryWriter InitialState;
 		{
@@ -87,23 +93,22 @@ bool FDiffWriterArchiveTestsBasic::RunTest(const FString& Parameters)
 			State.Serialize(InitialState);
 		}
 
-		FLargeMemoryWriter NewState;
+		FDiffArchiveForLinker NewState(*Accumulator);
 		{
 			FBasicDiffState State = FBasicDiffState {1,2,3};
-			FDiffWriterArchiveWriter StackTraceWriter(NewState, Callstacks);
-			State.Serialize(StackTraceWriter);
+			State.Serialize(NewState);
 		}
 
-		FDiffWriterArchiveWriter::FPackageData InitialData = ToPackageData(InitialState);
-		FDiffWriterArchiveWriter::FPackageData NewData = ToPackageData(NewState);
-		const bool bIsIdentical = FDiffWriterArchiveWriter::GenerateDiffMap(InitialData, NewData, Callstacks, MaxDiffsToLog, DiffMap);
-		TestTrueExpr(bIsIdentical);
+		ICookedPackageWriter::FPreviousCookedBytesData InitialData = ToPackageData(InitialState);
+		Accumulator->OnFirstSaveComplete(Filename, 0, MoveTemp(InitialData));
+
+		TestTrueExpr(!Accumulator->HasDifferences());
 	}
 
 	// SECTION("DiffMap - mismatch")
 	{
-		FDiffWriterCallstacks Callstacks(nullptr);
-		FDiffWriterDiffMap DiffMap;
+		TRefCountPtr<FAccumulator> Accumulator = new FAccumulator(nullptr, PackageName, MaxDiffsToLog,
+			[](ELogVerbosity::Type Verbosity, FStringView Message) {});
 
 		FLargeMemoryWriter InitialState;
 		{
@@ -111,76 +116,50 @@ bool FDiffWriterArchiveTestsBasic::RunTest(const FString& Parameters)
 			State.Serialize(InitialState);
 		}
 
-		FLargeMemoryWriter NewState;
+		FDiffArchiveForLinker NewState(*Accumulator);
 		{
 			FBasicDiffState State = FBasicDiffState {1,2,4};
-			FDiffWriterArchiveWriter StackTraceWriter(NewState, Callstacks);
-			State.Serialize(StackTraceWriter);
+			State.Serialize(NewState);
 		}
 
-		FDiffWriterArchiveWriter::FPackageData InitialData = ToPackageData(InitialState);
-		FDiffWriterArchiveWriter::FPackageData NewData = ToPackageData(NewState);
-		const bool bIsIdentical = FDiffWriterArchiveWriter::GenerateDiffMap(InitialData, NewData, Callstacks, MaxDiffsToLog, DiffMap);
-		TestTrueExpr(bIsIdentical == false);
+		ICookedPackageWriter::FPreviousCookedBytesData InitialData = ToPackageData(InitialState);
+		Accumulator->OnFirstSaveComplete(Filename, 0, MoveTemp(InitialData));
+
+		TestTrueExpr(Accumulator->HasDifferences());
 	}
 
 	// SECTION("Compare - mismatch")
 	{
-		FDiffWriterCallstacks Callstacks(nullptr);
-		FDiffWriterCallstacks CallstacksWithStackTrace(nullptr);
-		FDiffWriterDiffMap DiffMap;
+		TRefCountPtr<FAccumulator> Accumulator = new FAccumulator(nullptr, PackageName, MaxDiffsToLog,
+			[](ELogVerbosity::Type Verbosity, FStringView Message) {});
 
 		FBasicDiffState InitialState = FBasicDiffState {1,2,3};
 		FLargeMemoryWriter InitialMemory;
 		InitialMemory.SetByteSwapping(false);
-
 		InitialState.Serialize(InitialMemory);
+		ICookedPackageWriter::FPreviousCookedBytesData InitialData = ToPackageData(InitialMemory);
 
 		FBasicDiffState NewState = FBasicDiffState {1,22,33};
-		TUniquePtr<FLargeMemoryWriter> NewMemory = MakeUnique<FLargeMemoryWriter>();
+		TUniquePtr<FDiffArchiveForLinker> NewMemory = MakeUnique<FDiffArchiveForLinker>(*Accumulator);
 		NewMemory->SetByteSwapping(false);
-
-		{
-			FDiffWriterArchiveWriter StackTraceWriter(*NewMemory, Callstacks);
-			NewState.Serialize(StackTraceWriter);
-		}
+		NewState.Serialize(*NewMemory);
 
 		// Generate diff map
 		{
-			FDiffWriterArchiveWriter::FPackageData InitialData = ToPackageData(InitialMemory);
-			FDiffWriterArchiveWriter::FPackageData NewData = ToPackageData(*NewMemory);
-			const bool bIsIdentical = FDiffWriterArchiveWriter::GenerateDiffMap(InitialData, NewData, Callstacks, MaxDiffsToLog, DiffMap);
-			TestTrueExpr(bIsIdentical == false);
+			Accumulator->OnFirstSaveComplete(Filename, 0, MoveTemp(InitialData));
+			TestTrueExpr(Accumulator->HasDifferences());
 		}
 		
-		NewMemory = MakeUnique<FLargeMemoryWriter>();
+		NewMemory.Reset();
+		NewMemory = MakeUnique<FDiffArchiveForLinker>(*Accumulator);
 		NewMemory->SetByteSwapping(false);
-		{
-			FDiffWriterArchiveWriter StackTraceWriter(*NewMemory, CallstacksWithStackTrace, &DiffMap);
-			NewState.Serialize(StackTraceWriter);
-		}
+		NewState.Serialize(*NewMemory);
+		Accumulator->OnSecondSaveComplete(0);
 
 		// Compare using callstacks and diffmap
 		{
-			FDiffWriterArchiveWriter::FPackageData InitialData = ToPackageData(InitialMemory);
-			FDiffWriterArchiveWriter::FPackageData NewData = ToPackageData(*NewMemory);
-			int32 DiffsLogged = 0;
 			TMap<FName, FArchiveDiffStats> DiffStats;
-			UE::DiffWriterArchive::FMessageCallback MessageCallback =
-				[](ELogVerbosity::Type Verbosity, FStringView Message)
-			{};
-			FDiffWriterArchiveWriter::Compare(
-				InitialData,
-				NewData,
-				CallstacksWithStackTrace,
-				DiffMap,
-				TEXT("BasicDiffTest"),
-				nullptr,
-				MaxDiffsToLog,
-				DiffsLogged,
-				DiffStats,
-				MessageCallback);
-
+			Accumulator->CompareWithPrevious(TEXT("BasicDiffTest"), DiffStats, EPackageHeaderFormat::PackageFileSummary);
 			TestTrueExpr(DiffStats[NAME_None].NumDiffs == 2);
 		}
 	}
