Allow both hard and soft actor references to fix-up when moving (via copy/paste) actors to a new world

We need to consider both the found actor (which will be referenced by hard-references) and the previous ExportPath (which will be referenced by soft-references) when building the list of actors to fix-up post-import.

#rnx

diff --git a/Engine/Source/Editor/UnrealEd/Private/Factories/EditorFactories.cpp b/Engine/Source/Editor/UnrealEd/Private/Factories/EditorFactories.cpp
--- a/Engine/Source/Editor/UnrealEd/Private/Factories/EditorFactories.cpp
+++ b/Engine/Source/Editor/UnrealEd/Private/Factories/EditorFactories.cpp
@@ -1078,7 +1078,7 @@ UObject* ULevelFactory::FactoryCreateText
 								{
 									ExistingToNewMap.Add(Found, NewActor);
 								}
-								else if (!ExportedActorFullName.IsEmpty())
+								if (!ExportedActorFullName.IsEmpty())
 								{
 									ExistingToNewMap.Add(ExportedActorFullName, NewActor);
 								}
