World Partition - HLOD: Fix distributed HLOD builds when working with collections
* When building many single HLOD maps, make sure workload is distributed properly over multiple machines

diff --git a/Engine/Source/Editor/UnrealEd/Private/WorldPartition/WorldPartitionBuilder.cpp b/Engine/Source/Editor/UnrealEd/Private/WorldPartition/WorldPartitionBuilder.cpp
--- a/Engine/Source/Editor/UnrealEd/Private/WorldPartition/WorldPartitionBuilder.cpp
+++ b/Engine/Source/Editor/UnrealEd/Private/WorldPartition/WorldPartitionBuilder.cpp
@@ -81,6 +81,12 @@ bool UWorldPartitionBuilder::RunBuilder(UWorld* World)
 		return false;
 	}
 
+	// Verify if the builder want to skip processing the given world.
+	if (!ShouldProcessWorld(World))
+	{
+		return true;
+	}
+
 	FPackageSourceControlHelper SCCHelper;
 
 	// Perform builder pre world initialisation
diff --git a/Engine/Source/Editor/UnrealEd/Private/WorldPartition/WorldPartitionHLODsBuilder.cpp b/Engine/Source/Editor/UnrealEd/Private/WorldPartition/WorldPartitionHLODsBuilder.cpp
--- a/Engine/Source/Editor/UnrealEd/Private/WorldPartition/WorldPartitionHLODsBuilder.cpp
+++ b/Engine/Source/Editor/UnrealEd/Private/WorldPartition/WorldPartitionHLODsBuilder.cpp
@@ -250,12 +250,43 @@ bool UWorldPartitionHLODsBuilder::ValidateParams() const
 	return true;
 }
 
+FString GetDistributedBuildWorkingDir(UWorld* InWorld)
+{
+	uint32 WorldPackageHash = GetTypeHash(InWorld->GetPackage()->GetFullName());
+	return FString::Printf(TEXT("%s/%s/%08x"), *FPaths::RootDir(), *DistributedBuildWorkingDirName, WorldPackageHash);
+}
+
+bool UWorldPartitionHLODsBuilder::ShouldProcessWorld(UWorld* InWorld) const
+{
+	bool bShouldProcessWorld = true;
+
+	// When building HLODs in a distributed build, if there is no config section for the given builder index
+	// it means that the builder can skip processing this world altogether.
+	if (bDistributedBuild && ShouldRunStep(EHLODBuildStep::HLOD_Build))
+	{
+		const FString BuildManifestDirName = GetDistributedBuildWorkingDir(InWorld);
+		const FString BuildManifestFileName = BuildManifestDirName / DistributedBuildManifestName;
+
+		FConfigFile ConfigFile;
+		ConfigFile.Read(BuildManifestFileName);
+
+		FString SectionName = GetHLODBuilderFolderName(BuilderIdx);
+
+		const FConfigSection* ConfigSection = ConfigFile.Find(SectionName);
+		if (!ConfigSection || ConfigSection->IsEmpty())
+		{
+			bShouldProcessWorld = false;
+		}
+	}
+
+	return bShouldProcessWorld;
+}
+
 bool UWorldPartitionHLODsBuilder::PreWorldInitialization(UWorld* InWorld, FPackageSourceControlHelper& PackageHelper)
 {
 	if (bDistributedBuild)
 	{
-		uint32 WorldPackageHash = GetTypeHash(InWorld->GetPackage()->GetFullName());
-		DistributedBuildWorkingDir = FString::Printf(TEXT("%s/%s/%08x"), *FPaths::RootDir(), *DistributedBuildWorkingDirName, WorldPackageHash);
+		DistributedBuildWorkingDir = GetDistributedBuildWorkingDir(InWorld);
 		DistributedBuildManifest = DistributedBuildWorkingDir / DistributedBuildManifestName;
 
 		if (!BuildManifest.IsEmpty())
@@ -329,51 +360,49 @@ bool UWorldPartitionHLODsBuilder::RunInternal(UWorld* InWorld, const FCellInfo&
 bool UWorldPartitionHLODsBuilder::SetupHLODActors()
 {
 	// No setup needed for non partitioned worlds
-	if (!WorldPartition)
-	{
-		return true;
-	}
-
-	auto ActorFolderAddedDelegateHandle = GEngine->OnActorFolderAdded().AddLambda([this](UActorFolder* InActorFolder)
+	if (WorldPartition)
 	{
-		UPackage* ActorFolderPackage = InActorFolder->GetPackage();
-		const bool bIsTempPackage = FPackageName::IsTempPackage(ActorFolderPackage->GetName());
-		if (!bIsTempPackage)
+		auto ActorFolderAddedDelegateHandle = GEngine->OnActorFolderAdded().AddLambda([this](UActorFolder* InActorFolder)
 		{
-			// We don't want the HLOD folders to be expanded by default
-			InActorFolder->SetIsInitiallyExpanded(false);
-			SourceControlHelper->Save(InActorFolder->GetPackage());
-		}
-	});
+			UPackage* ActorFolderPackage = InActorFolder->GetPackage();
+			const bool bIsTempPackage = FPackageName::IsTempPackage(ActorFolderPackage->GetName());
+			if (!bIsTempPackage)
+			{
+				// We don't want the HLOD folders to be expanded by default
+				InActorFolder->SetIsInitiallyExpanded(false);
+				SourceControlHelper->Save(InActorFolder->GetPackage());
+			}
+		});
 	
-	ON_SCOPE_EXIT
-	{
-		GEngine->OnActorFolderAdded().Remove(ActorFolderAddedDelegateHandle);
-	};
+		ON_SCOPE_EXIT
+		{
+			GEngine->OnActorFolderAdded().Remove(ActorFolderAddedDelegateHandle);
+		};
 
-	const bool bCreateActorsOnly = true;
-	WorldPartition->GenerateHLOD(SourceControlHelper, bCreateActorsOnly);
+		const bool bCreateActorsOnly = true;
+		WorldPartition->GenerateHLOD(SourceControlHelper, bCreateActorsOnly);
 
-	// When performing a distributed build, ensure our work folder is empty
-	if (IsDistributedBuild())
-	{
-		IFileManager::Get().DeleteDirectory(*DistributedBuildWorkingDir, false, true);
-	}
+		// When performing a distributed build, ensure our work folder is empty
+		if (IsDistributedBuild())
+		{
+			IFileManager::Get().DeleteDirectory(*DistributedBuildWorkingDir, false, true);
+		}
 
-	UE_LOG(LogWorldPartitionHLODsBuilder, Display, TEXT("#### World HLOD actors ####"));
+		UE_LOG(LogWorldPartitionHLODsBuilder, Display, TEXT("#### World HLOD actors ####"));
 
-	int32 NumActors = 0;
-	for (FActorDescContainerCollection::TIterator<AWorldPartitionHLOD> HLODIterator(WorldPartition); HLODIterator; ++HLODIterator)
-	{
-		FWorldPartitionActorDesc* HLODActorDesc = *HLODIterator;
-		FString PackageName = HLODActorDesc->GetActorPackage().ToString();
+		int32 NumActors = 0;
+		for (FActorDescContainerCollection::TIterator<AWorldPartitionHLOD> HLODIterator(WorldPartition); HLODIterator; ++HLODIterator)
+		{
+			FWorldPartitionActorDesc* HLODActorDesc = *HLODIterator;
+			FString PackageName = HLODActorDesc->GetActorPackage().ToString();
 
-		UE_LOG(LogWorldPartitionHLODsBuilder, Display, TEXT("    [%d] %s"), NumActors, *PackageName);
+			UE_LOG(LogWorldPartitionHLODsBuilder, Display, TEXT("    [%d] %s"), NumActors, *PackageName);
 
-		NumActors++;
-	}
+			NumActors++;
+		}
 
-	UE_LOG(LogWorldPartitionHLODsBuilder, Display, TEXT("#### World contains %d HLOD actors ####"), NumActors);
+		UE_LOG(LogWorldPartitionHLODsBuilder, Display, TEXT("#### World contains %d HLOD actors ####"), NumActors);
+	}
 
 	if (IsUsingBuildManifest())
 	{
@@ -390,8 +419,11 @@ bool UWorldPartitionHLODsBuilder::SetupHLODActors()
 			// Ensure we don't hold on to packages of always loaded actors
 			// When running distributed builds, we wanna leave the machine clean, so added files are deleted, check'd out files are reverted
 			// and deleted files are restored.
-			WorldPartition->Uninitialize();
-			FWorldPartitionHelpers::DoCollectGarbage();
+			if (WorldPartition)
+			{
+				WorldPartition->Uninitialize();
+				FWorldPartitionHelpers::DoCollectGarbage();
+			}
 
 			ModifiedFiles.Append(SourceControlHelper->GetModifiedFiles());
 
@@ -497,7 +529,7 @@ bool UWorldPartitionHLODsBuilder::BuildHLODActors()
 			UE_LOG(LogWorldPartitionHLODsBuilder, Display, TEXT("[%d / %d] Building HLOD actor %s..."), CurrentActor + 1, HLODActorsToBuild.Num(), *HLODActor->GetActorLabel());
 
 			// Simulate an engine tick to make sure engine & render resources that are queued for deletion are processed.
-			FWorldPartitionHelpers::FakeEngineTick(WorldPartition->GetWorld());
+			FWorldPartitionHelpers::FakeEngineTick(World);
 
 			HLODActor->BuildHLOD(bForceBuild);
 
@@ -519,23 +551,31 @@ bool UWorldPartitionHLODsBuilder::BuildHLODActors()
 	}
 	else
 	{
+		IWorldPartitionHLODProvider::FBuildHLODActorParams BuildHLODActorParams;
+		BuildHLODActorParams.bForceRebuild = bForceBuild;
+		BuildHLODActorParams.OnPackageModified.BindLambda([this](UPackage* ModifiedPackage)
+		{
+			bool bSaved = SourceControlHelper->Save(ModifiedPackage);
+			if (!bSaved)
+			{
+				UE_LOG(LogWorldPartitionHLODsBuilder, Error, TEXT("Failed to save %s, exiting..."), *USourceControlHelpers::PackageFilename(ModifiedPackage));
+			}
+
+			return bSaved;
+		});
+
 		uint32 NumHLODActors = 0;
 		for (TActorIterator<AActor> ActorIt(World); ActorIt; ++ActorIt)
 		{
 			if (IWorldPartitionHLODProvider* HLODProvider = Cast<IWorldPartitionHLODProvider>(*ActorIt))
 			{
-				AWorldPartitionHLOD* HLODActor = HLODProvider->CreateHLODActor();
-				if (HLODActor)
+				bool bBuildResult = HLODProvider->BuildHLODActor(BuildHLODActorParams);
+				if (!bBuildResult)
 				{
-					HLODActor->BuildHLOD(bForceBuild);
-
-					bool bSaved = SaveHLODActor(HLODActor);
-					if (!bSaved)
-					{
-						return false;
-					}
-					NumHLODActors++;
+					return false;
 				}
+
+				NumHLODActors++;
 			}
 		}
 		UE_LOG(LogWorldPartitionHLODsBuilder, Display, TEXT("#### Built %d HLOD actor ####"), NumHLODActors);
@@ -548,8 +588,11 @@ bool UWorldPartitionHLODsBuilder::BuildHLODActors()
 		// Ensure we don't hold on to packages of always loaded actors
 		// When running distributed builds, we wanna leave the machine clean, so added files are deleted, check'd out files are reverted
 		// and deleted files are restored.
-		WorldPartition->Uninitialize();
-		FWorldPartitionHelpers::DoCollectGarbage();
+		if (WorldPartition)
+		{
+			WorldPartition->Uninitialize();
+			FWorldPartitionHelpers::DoCollectGarbage();
+		}
 
 		ModifiedFiles.Append(SourceControlHelper->GetModifiedFiles());
 
@@ -617,14 +660,14 @@ bool UWorldPartitionHLODsBuilder::SubmitHLODActors()
 	ModifiedFiles.Append(SourceControlHelper->GetModifiedFiles());
 
 	// Check in all modified files
-	const FString ChangeDescription = FString::Printf(TEXT("Rebuilt HLODs for %s"), *WorldPartition->GetWorld()->GetPackage()->GetName());
+	const FString ChangeDescription = FString::Printf(TEXT("Rebuilt HLODs for %s"), *World->GetPackage()->GetName());
 	return OnFilesModified(ModifiedFiles.GetAllFiles(), ChangeDescription);
 }
 
 bool UWorldPartitionHLODsBuilder::DumpStats()
 {
 	const FString HLODStatsOutputFilename = FPaths::ProjectSavedDir() / TEXT("WorldPartition") / FString::Printf(TEXT("HLODStats-%08x.csv"), FPlatformProcess::GetCurrentProcessId());
-	return UHLODSubsystem::WriteHLODStatsCSV(WorldPartition->GetWorld(), HLODStatsOutputFilename);
+	return UHLODSubsystem::WriteHLODStatsCSV(World, HLODStatsOutputFilename);
 }
 
 bool UWorldPartitionHLODsBuilder::GetHLODActorsToBuild(TArray<FGuid>& HLODActorsToBuild) const
@@ -677,7 +720,10 @@ bool UWorldPartitionHLODsBuilder::GetHLODActorsToBuild(TArray<FGuid>& HLODActors
 
 TArray<TArray<FGuid>> UWorldPartitionHLODsBuilder::GetHLODWorkloads(int32 NumWorkloads) const
 {
-	check(WorldPartition);
+	if (!WorldPartition)
+	{
+		return { { FGuid() } };
+	}
 
 	// Build a mapping of 1 HLOD[Level] -> N HLOD[Level - 1]
 	TMap<FGuid, TArray<FGuid>>	HLODParenting;
@@ -799,8 +845,6 @@ bool UWorldPartitionHLODsBuilder::ValidateWorkload(const TArray<FGuid>& Workload
 
 bool UWorldPartitionHLODsBuilder::GenerateBuildManifest(TMap<FString, int32>& FilesToBuilderMap) const
 {
-	check(WorldPartition);
-
 	TArray<TArray<FGuid>> BuildersWorkload = GetHLODWorkloads(BuilderCount);
 
 	FConfigFile ConfigFile;
@@ -809,27 +853,44 @@ bool UWorldPartitionHLODsBuilder::GenerateBuildManifest(TMap<FString, int32>& Fi
 	GeneralSection.Add(TEXT("BuilderCount"), FString::FromInt(BuilderCount));
 	GeneralSection.Add(TEXT("EngineVersion"), FEngineVersion::Current().ToString());
 
+	// When processing multiple maps, ensure that the worldload is distributed evenly between builders.
+	// Otherwise, maps with a single HLOD would all end up being processed by the first builder, while the others would have no work.
+	static int32 BuilderDispatchOffset = 0;
+
 	for(int32 Idx = 0; Idx < BuilderCount; Idx++)
 	{
-		FString SectionName = GetHLODBuilderFolderName(Idx);
+		const int32 WorkloadIndex = Idx;
+		const int32 BuilderIndex = (BuilderDispatchOffset + Idx) % BuilderCount;
 
+		if (!BuildersWorkload.IsValidIndex(WorkloadIndex) || BuildersWorkload[WorkloadIndex].IsEmpty())
+		{
+			continue;
+		}
+
+		FString SectionName = GetHLODBuilderFolderName(BuilderIndex);
 		FConfigSection& Section = ConfigFile.Add(SectionName);
-		for(const FGuid& ActorGuid : BuildersWorkload[Idx])
+
+		for(const FGuid& ActorGuid : BuildersWorkload[WorkloadIndex])
 		{
 			Section.Add(TEXT("+HLODActorGuid"), ActorGuid.ToString(EGuidFormats::Digits));
 
-			// Track which builder is responsible to handle each actor
-			const FWorldPartitionActorDesc* ActorDesc = WorldPartition->GetActorDesc(ActorGuid);
-			if (!ActorDesc)
+			if (WorldPartition)
 			{
-				UE_LOG(LogWorldPartitionHLODsBuilder, Error, TEXT("Invalid actor GUID found while generating the HLOD build manifest, exiting..."));
-				return false;
+				// Track which builder is responsible to handle each actor
+				const FWorldPartitionActorDesc* ActorDesc = WorldPartition->GetActorDesc(ActorGuid);
+				if (!ActorDesc)
+				{
+					UE_LOG(LogWorldPartitionHLODsBuilder, Error, TEXT("Invalid actor GUID found while generating the HLOD build manifest, exiting..."));
+					return false;
+				}
+				FString ActorPackageFilename = USourceControlHelpers::PackageFilename(ActorDesc->GetActorPackage().ToString());
+				FilesToBuilderMap.Emplace(ActorPackageFilename, BuilderIndex);
 			}
-			FString ActorPackageFilename = USourceControlHelpers::PackageFilename(ActorDesc->GetActorPackage().ToString());
-			FilesToBuilderMap.Emplace(ActorPackageFilename, Idx);
 		}
 	}
 
+	BuilderDispatchOffset++;
+
 	ConfigFile.Dirty = true;
 
 	bool bRet = ConfigFile.Write(BuildManifest);
diff --git a/Engine/Source/Editor/UnrealEd/Public/WorldPartition/WorldPartitionBuilder.h b/Engine/Source/Editor/UnrealEd/Public/WorldPartition/WorldPartitionBuilder.h
--- a/Engine/Source/Editor/UnrealEd/Public/WorldPartition/WorldPartitionBuilder.h
+++ b/Engine/Source/Editor/UnrealEd/Public/WorldPartition/WorldPartitionBuilder.h
@@ -94,6 +94,11 @@ protected:
 	 */
 	virtual bool CanProcessNonPartitionedWorlds() const { return false; }
 
+	/**
+	 * Some builders may decide to skip processing some worlds before initializing them.
+	 */
+	virtual bool ShouldProcessWorld(UWorld* World) const { return true; }
+
 	UNREALED_API bool OnFilesModified(const TArray<FString>& InModifiedFiles, const FString& InChangelistDescription) const;
 	UNREALED_API bool OnPackagesModified(const TArray<UPackage*>& InModifiedPackages, const FString& InChangelistDescription) const;
 
diff --git a/Engine/Source/Editor/UnrealEd/Public/WorldPartition/WorldPartitionHLODsBuilder.h b/Engine/Source/Editor/UnrealEd/Public/WorldPartition/WorldPartitionHLODsBuilder.h
--- a/Engine/Source/Editor/UnrealEd/Public/WorldPartition/WorldPartitionHLODsBuilder.h
+++ b/Engine/Source/Editor/UnrealEd/Public/WorldPartition/WorldPartitionHLODsBuilder.h
@@ -88,6 +88,7 @@ public:
 protected:
 	UNREALED_API virtual bool RunInternal(UWorld* World, const FCellInfo& InCellInfo, FPackageSourceControlHelper& PackageHelper) override;
 	virtual bool CanProcessNonPartitionedWorlds() const override { return true; }
+	UNREALED_API virtual bool ShouldProcessWorld(UWorld* World) const override;
 	// UWorldPartitionBuilder interface end
 
 	bool IsDistributedBuild() const { return bDistributedBuild; }
diff --git a/Engine/Source/Runtime/Engine/Public/WorldPartition/HLOD/HLODProviderInterface.h b/Engine/Source/Runtime/Engine/Public/WorldPartition/HLOD/HLODProviderInterface.h
--- a/Engine/Source/Runtime/Engine/Public/WorldPartition/HLOD/HLODProviderInterface.h
+++ b/Engine/Source/Runtime/Engine/Public/WorldPartition/HLOD/HLODProviderInterface.h
@@ -3,6 +3,7 @@
 #pragma once
 
 #include "CoreMinimal.h"
+#include "Delegates/Delegate.h"
 #include "UObject/Interface.h"
 #include "HLODProviderInterface.generated.h"
 
@@ -16,11 +17,18 @@ class UWorldPartitionHLODProvider : public UInterface
 	GENERATED_UINTERFACE_BODY()
 };
 
-
 class IWorldPartitionHLODProvider
 {
 	GENERATED_IINTERFACE_BODY()
 
 public:
-	virtual AWorldPartitionHLOD* CreateHLODActor() = 0;
+	DECLARE_DELEGATE_RetVal_OneParam(bool, FPackageModifiedDelegate, UPackage*)
+
+	struct FBuildHLODActorParams
+	{
+		bool bForceRebuild;
+		FPackageModifiedDelegate OnPackageModified;
+	};
+
+	virtual bool BuildHLODActor(const FBuildHLODActorParams& BuildHLODActorParams) = 0;
 };
