Fixed an issue where the data validation was not able to execute some editor utility validator because some module where not loaded by default that those editor utility might depends on.

Extracted an function on the default EditorEngine to load the module that are loaded by default when it is initialized. This can now be used to load those module in some commandlet without the need to initiliaze the EditorEngine.

The DataValidationCommandlet now load those editor modules if it found an validator that is made via an EditorUtility.

diff --git a/Engine/Plugins/Editor/DataValidation/Source/DataValidation/Private/DataValidationCommandlet.cpp b/Engine/Plugins/Editor/DataValidation/Source/DataValidation/Private/DataValidationCommandlet.cpp
--- a/Engine/Plugins/Editor/DataValidation/Source/DataValidation/Private/DataValidationCommandlet.cpp
+++ b/Engine/Plugins/Editor/DataValidation/Source/DataValidation/Private/DataValidationCommandlet.cpp
@@ -5,12 +5,16 @@
 #include "AssetRegistry/AssetRegistryHelpers.h"
 #include "AssetRegistry/AssetRegistryModule.h"
 #include "AssetRegistry/IAssetRegistry.h"
+#include "Blueprint/BlueprintSupport.h"
 #include "DataValidationModule.h"
 #include "Editor.h"
+#include "EditorUtilityBlueprint.h"
+#include "EditorValidatorBase.h"
 #include "EditorValidatorSubsystem.h"
 #include "Misc/PackageName.h"
-#include "Misc/Paths.h"
 #include "Misc/PathViews.h"
+#include "Misc/Paths.h"
+#include "UObject/UnrealType.h"
 
 #include UE_INLINE_GENERATED_CPP_BY_NAME(DataValidationCommandlet)
 
@@ -20,6 +24,9 @@ DEFINE_LOG_CATEGORY_STATIC(LogDataValidation, Warning, All);
 // Commandlet for validating data
 int32 UDataValidationCommandlet::Main(const FString& FullCommandLine)
 {
+	// This commandlet won't work properly when use outside of an editor executable
+	check(GEditor);
+
 	UE_LOG(LogDataValidation, Log, TEXT("--------------------------------------------------------------------------------------------"));
 	UE_LOG(LogDataValidation, Log, TEXT("Running DataValidation Commandlet"));
 
@@ -96,6 +103,33 @@ bool UDataValidationCommandlet::ValidateData(const FString& FullCommandLine)
 			});
 	}
 
+
+	if (!GEditor->IsInitialized())
+	{
+		// Check if we have some BP validator that were created using an editor utility
+		const FTopLevelAssetPath EditorUtilityClassPath = UEditorUtilityBlueprint::StaticClass()->GetClassPathName();
+		const FString EditorValidatorBaseClassExportPath = FObjectPropertyBase::GetExportPath(UEditorValidatorBase::StaticClass());
+		const bool bHasAnEditorUtilityDataValidator = AssetDataList.ContainsByPredicate([EditorUtilityClassPath, &EditorValidatorBaseClassExportPath](const FAssetData& AssetData)
+			{
+				if (AssetData.AssetClassPath == EditorUtilityClassPath)
+				{
+					if (AssetData.TagsAndValues.ContainsKeyValue(FBlueprintTags::NativeParentClassPath, EditorValidatorBaseClassExportPath))
+					{
+						return true;
+					}
+				}
+
+				return false;
+			});
+
+		if (bHasAnEditorUtilityDataValidator)
+		{
+			// Those Editor Utilities Validator might have an dependency to an editor module that is loaded during the editor initialization.
+			GEditor->LoadDefaultEditorModules();
+		}
+	}
+
+
 	UEditorValidatorSubsystem* EditorValidationSubsystem = GEditor->GetEditorSubsystem<UEditorValidatorSubsystem>();
 	check(EditorValidationSubsystem);
 
diff --git a/Engine/Plugins/Editor/DataValidation/Source/DataValidation/Private/EditorValidatorSubsystem.cpp b/Engine/Plugins/Editor/DataValidation/Source/DataValidation/Private/EditorValidatorSubsystem.cpp
--- a/Engine/Plugins/Editor/DataValidation/Source/DataValidation/Private/EditorValidatorSubsystem.cpp
+++ b/Engine/Plugins/Editor/DataValidation/Source/DataValidation/Private/EditorValidatorSubsystem.cpp
@@ -3,22 +3,23 @@
 #include "EditorValidatorSubsystem.h"
 
 #include "AssetRegistry/ARFilter.h"
-#include "Editor.h"
 #include "AssetRegistry/AssetRegistryModule.h"
 #include "Blueprint/BlueprintSupport.h"
-#include "IDirectoryWatcher.h"
+#include "DataValidationChangelist.h"
 #include "DirectoryWatcherModule.h"
+#include "Editor.h"
 #include "EditorUtilityBlueprint.h"
 #include "EditorValidatorBase.h"
+#include "Engine/BlueprintGeneratedClass.h"
+#include "Engine/Level.h"
+#include "IDirectoryWatcher.h"
+#include "ISourceControlModule.h"
 #include "Logging/MessageLog.h"
 #include "Misc/App.h"
-#include "Misc/ScopedSlowTask.h"
 #include "Misc/DataValidation.h"
-#include "ISourceControlModule.h"
-#include "DataValidationChangelist.h"
-#include "Engine/Level.h"
-#include "Editor.h"
 #include "Misc/PackageName.h"
+#include "Misc/ScopedSlowTask.h"
+#include "UObject/TopLevelAssetPath.h"
 
 #include UE_INLINE_GENERATED_CPP_BY_NAME(EditorValidatorSubsystem)
 
@@ -109,19 +110,7 @@ void UEditorValidatorSubsystem::RegisterBlueprintValidators()
 				}
 			}
 
-			// If this object isn't currently loaded, load it
-			UObject* ValidatorObject = BPAssetData.ToSoftObjectPath().ResolveObject();
-			if (ValidatorObject == nullptr)
-			{
-				FSoftObjectPathSerializationScope SerializationScope(NAME_None, NAME_None, ESoftObjectPathCollectType::EditorOnlyCollect, ESoftObjectPathSerializeType::AlwaysSerialize);
-				ValidatorObject = BPAssetData.ToSoftObjectPath().TryLoad();
-			}
-			if (ValidatorObject)
-			{
-				UEditorUtilityBlueprint* ValidatorBlueprint = Cast<UEditorUtilityBlueprint>(ValidatorObject);
-				UEditorValidatorBase* Validator = NewObject<UEditorValidatorBase>(GetTransientPackage(), ValidatorBlueprint->GeneratedClass);
-				AddValidator(Validator);
-			}
+			AddValidator(BPAssetData);
 		}
 	}
 }
@@ -140,7 +129,27 @@ void UEditorValidatorSubsystem::AddValidator(UEditorValidatorBase* InValidator)
 {
 	if (InValidator)
 	{
-		Validators.Add(InValidator->GetClass()->GetPathName(), InValidator);
+		UClass* Class = InValidator->GetClass();
+		if (Cast<UBlueprintGeneratedClass>(Class))
+		{
+			if(UObject* ClassGenerator = Class->ClassGeneratedBy)
+			{
+				Validators.Add(FTopLevelAssetPath(ClassGenerator), InValidator);
+			}
+		}
+		else
+		{ 
+			Validators.Add(InValidator->GetClass()->GetClassPathName(), InValidator);
+		}
+	}
+}
+
+void UEditorValidatorSubsystem::AddValidator(const FAssetData& InValidatorAssetData)
+{
+	if (InValidatorAssetData.IsValid())
+	{
+		Validators.Add(FTopLevelAssetPath(InValidatorAssetData.PackageName, InValidatorAssetData.AssetName), nullptr);
+		bNeedLoadingOfValidators = true;
 	}
 }
 
@@ -163,7 +172,10 @@ EDataValidationResult UEditorValidatorSubsystem::IsObjectValid(UObject* InObject
 		// If the asset is still valid or there wasn't a class-level validation, keep validating with custom validators
 		if (Result != EDataValidationResult::Invalid)
 		{
-			for (auto ValidatorPair : Validators)
+			// This doesn't change the logical view of the subsystem
+			const_cast<UEditorValidatorSubsystem&>(*this).LoadValidators();
+
+			for (const TPair<FTopLevelAssetPath, TObjectPtr<UEditorValidatorBase>>& ValidatorPair : Validators)
 			{
 				if (ValidatorPair.Value && ValidatorPair.Value->IsEnabled() && ValidatorPair.Value->CanValidate(InValidationUsecase) && ValidatorPair.Value->CanValidateAsset(InObject))
 				{
@@ -509,5 +521,38 @@ void UEditorValidatorSubsystem::ValidateChangelistPreSubmit(FSourceControlChange
 	Changelist->RemoveFromRoot();
 }
 
+void UEditorValidatorSubsystem::LoadValidators()
+{
+	if (bNeedLoadingOfValidators)
+	{
+		for (TPair<FTopLevelAssetPath, TObjectPtr<UEditorValidatorBase>>& ValidatorPair : Validators)
+		{
+			if (!ValidatorPair.Value)
+			{
+				FSoftObjectPath ValidatorObjectSoftPath(ValidatorPair.Key);
+				UObject* ValidatorObject = ValidatorObjectSoftPath.ResolveObject();
+
+				// If this object isn't currently loaded, load it
+				if (ValidatorObject == nullptr)
+				{
+					FSoftObjectPathSerializationScope SerializationScope(NAME_None, NAME_None, ESoftObjectPathCollectType::EditorOnlyCollect, ESoftObjectPathSerializeType::AlwaysSerialize);
+					ValidatorObject = ValidatorObjectSoftPath.TryLoad();
+				}
+
+				if (ValidatorObject)
+				{
+					if (UEditorUtilityBlueprint* ValidatorBlueprint = Cast<UEditorUtilityBlueprint>(ValidatorObject))
+					{
+						UEditorValidatorBase* Validator = NewObject<UEditorValidatorBase>(GetTransientPackage(), ValidatorBlueprint->GeneratedClass);
+						ValidatorPair.Value = Validator;
+					}
+				}
+			}
+		}
+
+		bNeedLoadingOfValidators = false;
+	}
+}
+
 #undef LOCTEXT_NAMESPACE
 
diff --git a/Engine/Plugins/Editor/DataValidation/Source/DataValidation/Public/EditorValidatorSubsystem.h b/Engine/Plugins/Editor/DataValidation/Source/DataValidation/Public/EditorValidatorSubsystem.h
--- a/Engine/Plugins/Editor/DataValidation/Source/DataValidation/Public/EditorValidatorSubsystem.h
+++ b/Engine/Plugins/Editor/DataValidation/Source/DataValidation/Public/EditorValidatorSubsystem.h
@@ -217,6 +217,13 @@ protected:
 	 */
 	void ValidateChangelistPreSubmit(FSourceControlChangelistPtr Changelist, EDataValidationResult& OutResult, TArray<FText>& ValidationErrors, TArray<FText>& ValidationWarnings) const;
 
+	void LoadValidators();
+
+	/*
+	* Adds a validator to the list making sure it is a unique instance, but wait for the first use to load the assets.
+	*/
+	void AddValidator(const FAssetData&	InValidatorAssetData);
+
 protected:
 	/**
 	 * Directories to ignore for data validation. Useful for test assets
@@ -234,7 +241,7 @@ protected:
 	TArray<FName> SavedPackagesToValidate;
 
 	UPROPERTY(Transient)
-	TMap<FString, TObjectPtr<UEditorValidatorBase>> Validators;
+	TMap<FTopLevelAssetPath, TObjectPtr<UEditorValidatorBase>> Validators;
 
 	/** Specifies whether or not to validate assets on save when saving for a cook */
 	UPROPERTY(config)
@@ -244,6 +251,7 @@ protected:
 	UPROPERTY(config)
 	bool bAllowBlueprintValidators;
 
+	bool bNeedLoadingOfValidators = false;
 };
 
 
diff --git a/Engine/Source/Editor/UnrealEd/Classes/Editor/EditorEngine.h b/Engine/Source/Editor/UnrealEd/Classes/Editor/EditorEngine.h
--- a/Engine/Source/Editor/UnrealEd/Classes/Editor/EditorEngine.h
+++ b/Engine/Source/Editor/UnrealEd/Classes/Editor/EditorEngine.h
@@ -759,6 +759,12 @@ public:
 	/**	Broadcasts that an object has been reimported. THIS SHOULD NOT BE PUBLIC */
 	UNREALED_API void BroadcastObjectReimported(UObject* InObject);
 
+	/**
+	 * Load the editor module that are loaded by default in the editor even when the target doesn't have an dependency on those 
+	 * Note: This is useful for the commandlets that may depend on those modules.
+	 */
+	UNREALED_API void LoadDefaultEditorModules();
+
 	//~ Begin UObject Interface.
 	UNREALED_API virtual void FinishDestroy() override;
 	UNREALED_API virtual void PostEditChangeProperty(FPropertyChangedEvent& PropertyChangedEvent) override;
diff --git a/Engine/Source/Editor/UnrealEd/Private/EditorEngine.cpp b/Engine/Source/Editor/UnrealEd/Private/EditorEngine.cpp
--- a/Engine/Source/Editor/UnrealEd/Private/EditorEngine.cpp
+++ b/Engine/Source/Editor/UnrealEd/Private/EditorEngine.cpp
@@ -162,6 +162,7 @@
 #include "EditorWorldExtension.h"
 
 #if PLATFORM_WINDOWS
+	#include "Windows/WindowsHWrapper.h"
 // For WAVEFORMATEXTENSIBLE
 	#include "Windows/AllowWindowsPlatformTypes.h"
 #include <mmreg.h>
@@ -1108,120 +1109,7 @@
 	SlowTask.EnterProgressFrame(50);
 
 	// Load all editor modules here
-	{
-		static const TCHAR* ModuleNames[] =
-		{
-			TEXT("Documentation"),
-			TEXT("WorkspaceMenuStructure"),
-			TEXT("MainFrame"),
-			TEXT("OutputLog"),
-			TEXT("SourceControl"),
-			TEXT("SourceControlWindows"),
-			TEXT("SourceControlWindowExtender"),
-			TEXT("UncontrolledChangelists"),
-			TEXT("TextureCompressor"),
-			TEXT("MeshUtilities"),
-			TEXT("MovieSceneTools"),
-			TEXT("ClassViewer"),
-			TEXT("StructViewer"),
-			TEXT("ContentBrowser"),
-			TEXT("AssetTools"),
-			TEXT("GraphEditor"),
-			TEXT("KismetCompiler"),
-			TEXT("Kismet"),
-			TEXT("Persona"),
-			TEXT("AnimationBlueprintEditor"),
-			TEXT("LevelEditor"),
-			TEXT("MainFrame"),
-			TEXT("PropertyEditor"),
-			TEXT("PackagesDialog"),
-			// TEXT("AssetRegistry"), // Loaded in constructor
-			TEXT("DetailCustomizations"),
-			TEXT("ComponentVisualizers"),
-			TEXT("Layers"),
-			TEXT("AutomationWindow"),
-			TEXT("AutomationController"),
-			TEXT("DeviceManager"),
-			TEXT("ProfilerClient"),
-			TEXT("SessionFrontend"),
-			TEXT("ProjectLauncher"),
-			TEXT("SettingsEditor"),
-			TEXT("EditorSettingsViewer"),
-			TEXT("ProjectSettingsViewer"),
-			TEXT("Blutility"),
-			TEXT("ScriptableEditorWidgets"),
-			TEXT("XmlParser"),
-			TEXT("UndoHistory"),
-			TEXT("DeviceProfileEditor"),
-			TEXT("SourceCodeAccess"),
-			TEXT("BehaviorTreeEditor"),
-			TEXT("HardwareTargeting"),
-			TEXT("LocalizationDashboard"),
-			TEXT("MergeActors"),
-			TEXT("InputBindingEditor"),
-			TEXT("AudioEditor"),
-			TEXT("EditorInteractiveToolsFramework"),
-			TEXT("TraceInsights"),
-			TEXT("StaticMeshEditor"),
-			TEXT("EditorFramework"),
-			TEXT("WorldPartitionEditor"),
-			TEXT("EditorConfig"),
-			TEXT("DerivedDataEditor"),
-			TEXT("CSVtoSVG"),
-			TEXT("GeometryFramework"),
-			TEXT("VirtualizationEditor"),
-			TEXT("AnimationSettings"),
-			TEXT("GameplayDebuggerEditor"),
-			TEXT("RenderResourceViewer"),
-		};
+	LoadDefaultEditorModules();
 
-		FScopedSlowTask ModuleSlowTask((float)UE_ARRAY_COUNT(ModuleNames));
-		for (const TCHAR* ModuleName : ModuleNames)
-		{
-			ModuleSlowTask.EnterProgressFrame(1);
-			FModuleManager::Get().LoadModule(ModuleName);
-		}
-
-		{
-			// Load platform runtime settings modules
-			TArray<FName> Modules;
-			FModuleManager::Get().FindModules( TEXT( "*RuntimeSettings" ), Modules );
-
-			for( int32 Index = 0; Index < Modules.Num(); Index++ )
-			{
-				FModuleManager::Get().LoadModule( Modules[Index] );
-			}
-		}
-
-		{
-			// Load platform editor modules
-			TArray<FName> Modules;
-			FModuleManager::Get().FindModules( TEXT( "*PlatformEditor" ), Modules );
-
-			for( int32 Index = 0; Index < Modules.Num(); Index++ )
-			{
-				if( Modules[Index] != TEXT("ProjectTargetPlatformEditor") )
-				{
-					FModuleManager::Get().LoadModule( Modules[Index] );
-				}
-			}
-		}
-
-		if( FParse::Param( FCommandLine::Get(),TEXT( "PListEditor" ) ) )
-		{
-			FModuleManager::Get().LoadModule(TEXT("PListEditor"));
-		}
-
-		FModuleManager::Get().LoadModule(TEXT("LogVisualizer"));
-		FModuleManager::Get().LoadModule(TEXT("WidgetRegistration"));
-		FModuleManager::Get().LoadModule(TEXT("HotReload"));
-
-		FModuleManager::Get().LoadModuleChecked(TEXT("ClothPainter"));
-
-		// Load VR Editor support
-		FModuleManager::Get().LoadModuleChecked( TEXT( "ViewportInteraction" ) );
-		FModuleManager::Get().LoadModuleChecked( TEXT( "VREditor" ) );
-	}
-
 	SlowTask.EnterProgressFrame(10);
 
@@ -1350,6 +1238,122 @@
 	GetEditorSubsystem<UImportSubsystem>()->BroadcastAssetReimport(InObject);
 }
 
+void UEditorEngine::LoadDefaultEditorModules()
+{
+	static const TCHAR* ModuleNames[] =
+		{
+			TEXT("Documentation"),
+			TEXT("WorkspaceMenuStructure"),
+			TEXT("MainFrame"),
+			TEXT("OutputLog"),
+			TEXT("SourceControl"),
+			TEXT("SourceControlWindows"),
+			TEXT("SourceControlWindowExtender"),
+			TEXT("UncontrolledChangelists"),
+			TEXT("TextureCompressor"),
+			TEXT("MeshUtilities"),
+			TEXT("MovieSceneTools"),
+			TEXT("ClassViewer"),
+			TEXT("StructViewer"),
+			TEXT("ContentBrowser"),
+			TEXT("AssetTools"),
+			TEXT("GraphEditor"),
+			TEXT("KismetCompiler"),
+			TEXT("Kismet"),
+			TEXT("Persona"),
+			TEXT("AnimationBlueprintEditor"),
+			TEXT("LevelEditor"),
+			TEXT("MainFrame"),
+			TEXT("PropertyEditor"),
+			TEXT("PackagesDialog"),
+			// TEXT("AssetRegistry"), // Loaded in constructor
+			TEXT("DetailCustomizations"),
+			TEXT("ComponentVisualizers"),
+			TEXT("Layers"),
+			TEXT("AutomationWindow"),
+			TEXT("AutomationController"),
+			TEXT("DeviceManager"),
+			TEXT("ProfilerClient"),
+			TEXT("SessionFrontend"),
+			TEXT("ProjectLauncher"),
+			TEXT("SettingsEditor"),
+			TEXT("EditorSettingsViewer"),
+			TEXT("ProjectSettingsViewer"),
+			TEXT("Blutility"),
+			TEXT("ScriptableEditorWidgets"),
+			TEXT("XmlParser"),
+			TEXT("UndoHistory"),
+			TEXT("DeviceProfileEditor"),
+			TEXT("SourceCodeAccess"),
+			TEXT("BehaviorTreeEditor"),
+			TEXT("HardwareTargeting"),
+			TEXT("LocalizationDashboard"),
+			TEXT("MergeActors"),
+			TEXT("InputBindingEditor"),
+			TEXT("AudioEditor"),
+			TEXT("EditorInteractiveToolsFramework"),
+			TEXT("TraceInsights"),
+			TEXT("StaticMeshEditor"),
+			TEXT("EditorFramework"),
+			TEXT("WorldPartitionEditor"),
+			TEXT("EditorConfig"),
+			TEXT("DerivedDataEditor"),
+			TEXT("CSVtoSVG"),
+			TEXT("GeometryFramework"),
+			TEXT("VirtualizationEditor"),
+			TEXT("AnimationSettings"),
+			TEXT("GameplayDebuggerEditor"),
+			TEXT("RenderResourceViewer"),
+		};
+
+	FScopedSlowTask ModuleSlowTask((float)UE_ARRAY_COUNT(ModuleNames));
+	for (const TCHAR* ModuleName : ModuleNames)
+	{
+		ModuleSlowTask.EnterProgressFrame(1);
+		FModuleManager::Get().LoadModule(ModuleName);
+	}
+
+	{
+		// Load platform runtime settings modules
+		TArray<FName> Modules;
+		FModuleManager::Get().FindModules( TEXT( "*RuntimeSettings" ), Modules );
+
+		for( int32 Index = 0; Index < Modules.Num(); Index++ )
+		{
+			FModuleManager::Get().LoadModule( Modules[Index] );
+		}
+	}
+
+	{
+		// Load platform editor modules
+		TArray<FName> Modules;
+		FModuleManager::Get().FindModules( TEXT( "*PlatformEditor" ), Modules );
+
+		for( int32 Index = 0; Index < Modules.Num(); Index++ )
+		{
+			if( Modules[Index] != TEXT("ProjectTargetPlatformEditor") )
+			{
+				FModuleManager::Get().LoadModule( Modules[Index] );
+			}
+		}
+	}
+
+	if( FParse::Param( FCommandLine::Get(),TEXT( "PListEditor" ) ) )
+	{
+		FModuleManager::Get().LoadModule(TEXT("PListEditor"));
+	}
+
+	FModuleManager::Get().LoadModule(TEXT("LogVisualizer"));
+	FModuleManager::Get().LoadModule(TEXT("WidgetRegistration"));
+	FModuleManager::Get().LoadModule(TEXT("HotReload"));
+
+	FModuleManager::Get().LoadModuleChecked(TEXT("ClothPainter"));
+
+	// Load VR Editor support
+	FModuleManager::Get().LoadModuleChecked( TEXT( "ViewportInteraction" ) );
+	FModuleManager::Get().LoadModuleChecked( TEXT( "VREditor" ) );
+}
+
 void UEditorEngine::PreExit()
 {
 	if (!HasAnyFlags(RF_ClassDefaultObject))
