PR #10738: Fixed late join toolbar button to match style

diff --git a/Engine/Source/Editor/UnrealEd/Private/Kismet2/DebuggerCommands.cpp b/Engine/Source/Editor/UnrealEd/Private/Kismet2/DebuggerCommands.cpp
--- a/Engine/Source/Editor/UnrealEd/Private/Kismet2/DebuggerCommands.cpp
+++ b/Engine/Source/Editor/UnrealEd/Private/Kismet2/DebuggerCommands.cpp
@@ -646,7 +646,9 @@ void FPlayWorldCommands::BuildToolbar(FToolMenuSection& InSection, bool bInclude
 	InSection.AddEntry(StopEntry);
 
 	// Late Join
-	InSection.AddEntry(FToolMenuEntry::InitToolBarButton(FPlayWorldCommands::Get().LateJoinSession, TAttribute<FText>(), TAttribute<FText>(), TAttribute<FSlateIcon>(), FName(TEXT("LateJoinSession"))));
+	FToolMenuEntry LateJoinEntry = FToolMenuEntry::InitToolBarButton(FPlayWorldCommands::Get().LateJoinSession, TAttribute<FText>(), TAttribute<FText>(), TAttribute<FSlateIcon>(), FName(TEXT("LateJoinSession")));
+	LateJoinEntry.StyleNameOverride = FName("Toolbar.BackplateCenter");
+	InSection.AddEntry(LateJoinEntry);
 
 	// Eject/possess toggle
 
