IncrementalCooking: Move initialization of the EditorDomainOplog to after engine startup, and only when cooking, to avoid crashing because project name is not available.
#rnx

diff --git a/Engine/Config/BaseEditor.ini b/Engine/Config/BaseEditor.ini
--- a/Engine/Config/BaseEditor.ini
+++ b/Engine/Config/BaseEditor.ini
@@ -408,7 +408,6 @@ CookProcessCount=1
 [EditorDomain]
 EditorDomainEnabled=false
 SaveUnversioned=true
-CookAttachmentsEnabled=false
 ; +ClassBlockList=/Long/PackageName.ClassName
 ; +PackageBlockList=/Long/PackageName OR ../../../Standard/Filename/NoExtension OR D:\Absolute\Filename\NoExtension
 
diff --git a/Engine/Source/Editor/UnrealEd/Classes/CookOnTheSide/CookOnTheFlyServer.h b/Engine/Source/Editor/UnrealEd/Classes/CookOnTheSide/CookOnTheFlyServer.h
--- a/Engine/Source/Editor/UnrealEd/Classes/CookOnTheSide/CookOnTheFlyServer.h
+++ b/Engine/Source/Editor/UnrealEd/Classes/CookOnTheSide/CookOnTheFlyServer.h
@@ -896,6 +896,13 @@ private:
 	UNREALED_API void BeginCookEditorSystems();
 	UNREALED_API void BeginCookPackageWriters(FBeginCookContext& BeginContext);
 	UNREALED_API void BeginCookDirector(FBeginCookContext& BeginContext);
+	/**
+	 * Initialization for systems that persist across CookSessions but that are not available during Initialize(which
+	 * can occur during EngineStartup if !bDisableCookInEditor). We initialize these as late as possible: at the
+	 * beginning of the first session, or for CookOnTheFly at the first cook request.
+	 */
+	UNREALED_API void InitializeAtFirstSession();
+	/** Initialize steps that are reexecuted at the beginning of every cook session. */
 	UNREALED_API void InitializeSession();
 
 	//////////////////////////////////////////////////////////////////////////
diff --git a/Engine/Source/Editor/UnrealEd/Private/CookOnTheFlyServer.cpp b/Engine/Source/Editor/UnrealEd/Private/CookOnTheFlyServer.cpp
--- a/Engine/Source/Editor/UnrealEd/Private/CookOnTheFlyServer.cpp
+++ b/Engine/Source/Editor/UnrealEd/Private/CookOnTheFlyServer.cpp
@@ -6533,10 +6533,12 @@ void UCookOnTheFlyServer::Initialize( ECookMode::Type DesiredCookMode, ECookInit
 #if WITH_ADDITIONAL_CRASH_CONTEXTS
 	FGenericCrashContext::OnAdditionalCrashContextDelegate().AddUObject(this, &UCookOnTheFlyServer::DumpCrashContext);
 #endif
-
-	UE::EditorDomain::UtilsTargetDomainInit();
 }
 
+void UCookOnTheFlyServer::InitializeAtFirstSession()
+{
+	UE::EditorDomain::UtilsCookInitialize();
+}
 
 void UCookOnTheFlyServer::LoadInitializeConfigSettings(const FString& InOutputDirectoryOverride)
 {
@@ -6886,6 +6888,23 @@ bool UCookOnTheFlyServer::TryInitializeCookWorker()
 
 void UCookOnTheFlyServer::InitializeSession()
 {
+	if (!bFirstCookInThisProcessInitialized)
+	{
+		// This is the first cook; set bFirstCookInThisProcess=true for the entire cook until SetBeginCookConfigSettings is called to mark the second cook
+		bFirstCookInThisProcessInitialized = true;
+		bFirstCookInThisProcess = true;
+	}
+	else
+	{
+		// We have cooked before; set bFirstCookInThisProcess=false
+		bFirstCookInThisProcess = false;
+	}
+
+	if (bFirstCookInThisProcess)
+	{
+		InitializeAtFirstSession();
+	}
+
 	NumObjectsHistory.Initialize(GUObjectArray.GetObjectArrayNumMinusAvailable());
 	VirtualMemoryHistory.Initialize(FPlatformMemory::GetStats().UsedVirtual);
 }
@@ -10293,17 +10312,6 @@ void UCookOnTheFlyServer::SetBeginCookConfigSettings(FBeginCookContext& BeginCon
 	bHybridIterativeAllowAllClasses = Settings.bHybridIterativeAllowAllClasses;
 	PackageDatas->SetBeginCookConfigSettings(Settings.CookShowInstigator);
 	SetNeverCookPackageConfigSettings(BeginContext, Settings);
-	if (!bFirstCookInThisProcessInitialized)
-	{
-		// This is the first cook; set bFirstCookInThisProcess=true for the entire cook until SetBeginCookConfigSettings is called to mark the second cook
-		bFirstCookInThisProcessInitialized = true;
-		bFirstCookInThisProcess = true;
-	}
-	else
-	{
-		// We have cooked before; set bFirstCookInThisProcess=false
-		bFirstCookInThisProcess = false;
-	}
 }
 
 namespace UE::Cook
@@ -10737,6 +10745,7 @@ void UCookOnTheFlyServer::StartCookByTheBook( const FCookByTheBookStartupOptions
 
 	// Initialize systems and settings that the rest of StartCookByTheBook depends on
 	// Functions in this section are ordered and can depend on the functions before them
+	InitializeSession();
 	FBeginCookContext BeginContext = CreateBeginCookByTheBookContext(EffectiveStartupOptions);
 	BlockOnAssetRegistry();
 	CreateSandboxFile(BeginContext);
@@ -10913,6 +10922,7 @@ void UCookOnTheFlyServer::StartCookAsCookWorker()
 
 	// Initialize systems and settings that the rest of StartCookAsCookWorker depends on
 	// Functions in this section are ordered and can depend on the functions before them
+	InitializeSession();
 	FBeginCookContext BeginContext = CreateCookWorkerContext();
 	// MPCOOKTODO: Load serialized AssetRegistry from Director
 	BlockOnAssetRegistry();
@@ -11438,7 +11448,6 @@ void UCookOnTheFlyServer::BeginCookPackageWriters(FBeginCookContext& BeginContex
 
 void UCookOnTheFlyServer::SelectSessionPlatforms(FBeginCookContext& BeginContext)
 {
-	InitializeSession();
 	PlatformManager->SelectSessionPlatforms(*this, BeginContext.TargetPlatforms);
 
 	FindOrCreateSaveContexts(BeginContext.TargetPlatforms);
diff --git a/Engine/Source/Editor/UnrealEd/Private/EditorDomain/EditorDomain.cpp b/Engine/Source/Editor/UnrealEd/Private/EditorDomain/EditorDomain.cpp
--- a/Engine/Source/Editor/UnrealEd/Private/EditorDomain/EditorDomain.cpp
+++ b/Engine/Source/Editor/UnrealEd/Private/EditorDomain/EditorDomain.cpp
@@ -80,7 +80,6 @@ public:
 		if ((uint32)Enabled >= (uint32)EEditorDomainEnabled::Utilities)
 		{
 			UE::EditorDomain::UtilsInitialize();
-			UE::TargetDomain::UtilsInitialize();
 			FEditorDomain::SingletonEditorDomain = new FEditorDomain(Enabled);
 			UE::EditorDomain::IPackageDigestCache::Set(FEditorDomain::SingletonEditorDomain);
 		}
diff --git a/Engine/Source/Editor/UnrealEd/Private/EditorDomain/EditorDomainUtils.cpp b/Engine/Source/Editor/UnrealEd/Private/EditorDomain/EditorDomainUtils.cpp
--- a/Engine/Source/Editor/UnrealEd/Private/EditorDomain/EditorDomainUtils.cpp
+++ b/Engine/Source/Editor/UnrealEd/Private/EditorDomain/EditorDomainUtils.cpp
@@ -199,7 +199,7 @@ TMultiMap<FTopLevelAssetPath, FTopLevelAssetPath> GConstructClasses;
 TSet<FTopLevelAssetPath> GTargetDomainClassBlockList;
 TArray<FTopLevelAssetPath> GGlobalConstructClasses;
 bool bGGlobalConstructClassesInitialized = false;
-bool bGUtilsTargetDomainInitialized = false;
+bool bGUtilsCookInitialized = false;
 FBlake3Hash GGlobalConstructClassesHash;
 int64 GMaxBulkDataSize = -1;
 
@@ -1212,14 +1212,15 @@ void UtilsInitialize()
 	COOK_STAT(UE::EditorDomain::CookStats::Register());
 }
 
-void UtilsTargetDomainInit()
+void UtilsCookInitialize()
 {
-	if (bGUtilsTargetDomainInitialized)
+	if (bGUtilsCookInitialized)
 	{
 		return;
 	}
-	bGUtilsTargetDomainInitialized = true;
+	bGUtilsCookInitialized = true;
 	ConstructTargetIterativeClassAllowList();
+	UE::TargetDomain::CookInitialize();
 }
 
 UE::DerivedData::FCacheKey GetEditorDomainPackageKey(const FIoHash& EditorDomainHash)
diff --git a/Engine/Source/Editor/UnrealEd/Private/EditorDomain/EditorDomainUtils.h b/Engine/Source/Editor/UnrealEd/Private/EditorDomain/EditorDomainUtils.h
--- a/Engine/Source/Editor/UnrealEd/Private/EditorDomain/EditorDomainUtils.h
+++ b/Engine/Source/Editor/UnrealEd/Private/EditorDomain/EditorDomainUtils.h
@@ -115,8 +115,8 @@ FClassDigestMap& GetClassDigests();
 /** Initializes some global config-driven values used by the EditorDomain and TargetDomain. */
 void UtilsInitialize();
 
-/** Initializes global config-driven values that are only needed by the TargetDomain. */
-void UtilsTargetDomainInit();
+/** Initializes global config-driven values that are only needed by the cooker. */
+void UtilsCookInitialize();
 
 #if ENABLE_COOK_STATS
 namespace CookStats
diff --git a/Engine/Source/Editor/UnrealEd/Private/TargetDomain/TargetDomainUtils.cpp b/Engine/Source/Editor/UnrealEd/Private/TargetDomain/TargetDomainUtils.cpp
--- a/Engine/Source/Editor/UnrealEd/Private/TargetDomain/TargetDomainUtils.cpp
+++ b/Engine/Source/Editor/UnrealEd/Private/TargetDomain/TargetDomainUtils.cpp
@@ -669,7 +669,7 @@ void CommitEditorDomainCookAttachments(FName PackageName, TArrayView<IPackageWri
 	GEditorDomainOplog->CommitPackage(PackageName, Attachments);
 }
 
-void UtilsInitialize()
+void CookInitialize()
 {
 	bool bCookAttachmentsEnabled = true;
 	GConfig->GetBool(TEXT("EditorDomain"), TEXT("CookAttachmentsEnabled"), bCookAttachmentsEnabled, GEditorIni);
diff --git a/Engine/Source/Editor/UnrealEd/Private/TargetDomain/TargetDomainUtils.h b/Engine/Source/Editor/UnrealEd/Private/TargetDomain/TargetDomainUtils.h
--- a/Engine/Source/Editor/UnrealEd/Private/TargetDomain/TargetDomainUtils.h
+++ b/Engine/Source/Editor/UnrealEd/Private/TargetDomain/TargetDomainUtils.h
@@ -19,7 +19,7 @@ template <typename FuncType> class TUniqueFunction;
 namespace UE::TargetDomain
 {
 
-void UtilsInitialize();
+void CookInitialize();
 
 /** Create the TargetDomainKey based on the EditorDomainKeys of the Package and its dependencies. */
 bool TryCreateKey(FName PackageName, TConstArrayView<FName> SortedBuildDependencies, FIoHash* OutHash, FString* OutErrorMessage);
