Remove long deprecated function
#rnx

diff --git a/Engine/Source/Editor/UnrealEd/Private/Kismet2/KismetReinstanceUtilities.cpp b/Engine/Source/Editor/UnrealEd/Private/Kismet2/KismetReinstanceUtilities.cpp
--- a/Engine/Source/Editor/UnrealEd/Private/Kismet2/KismetReinstanceUtilities.cpp
+++ b/Engine/Source/Editor/UnrealEd/Private/Kismet2/KismetReinstanceUtilities.cpp
@@ -783,10 +783,6 @@ TSharedPtr<FReinstanceFinalizer> FBlueprintCompileReinstancer::ReinstanceInner(b
 	return Finalizer;
 }
 
-void FBlueprintCompileReinstancer::EnlistDependentBlueprintToRecompile(UBlueprint* BP, bool bBytecodeOnly)
-{
-}
-
 void FBlueprintCompileReinstancer::BlueprintWasRecompiled(UBlueprint* BP, bool bBytecodeOnly)
 {
 }
diff --git a/Engine/Source/Editor/UnrealEd/Public/Kismet2/KismetReinstanceUtilities.h b/Engine/Source/Editor/UnrealEd/Public/Kismet2/KismetReinstanceUtilities.h
--- a/Engine/Source/Editor/UnrealEd/Public/Kismet2/KismetReinstanceUtilities.h
+++ b/Engine/Source/Editor/UnrealEd/Public/Kismet2/KismetReinstanceUtilities.h
@@ -173,8 +173,6 @@ public:
 
 	static UNREALED_API void OptionallyRefreshNodes(UBlueprint* BP);
 
-	UE_DEPRECATED(5.0, "This method performs no function and isn't invoked by the base class.  Remove all calls to the base class method.")
-	UNREALED_API virtual void EnlistDependentBlueprintToRecompile(UBlueprint* BP, bool bBytecodeOnly);
 	UNREALED_API virtual void BlueprintWasRecompiled(UBlueprint* BP, bool bBytecodeOnly);
 
 	static TSharedPtr<FBlueprintCompileReinstancer> Create(UClass* InClassToReinstance, EBlueprintCompileReinstancerFlags Flags = EBlueprintCompileReinstancerFlags::AutoInferSaveOnCompile)
