Level Instance: Fix failed check caused by uninitializing wrong world (in cases where World is duplicated when saved)

#rnx

diff --git a/Engine/Source/Editor/UnrealEd/Private/EditorLevelUtils.cpp b/Engine/Source/Editor/UnrealEd/Private/EditorLevelUtils.cpp
--- a/Engine/Source/Editor/UnrealEd/Private/EditorLevelUtils.cpp
+++ b/Engine/Source/Editor/UnrealEd/Private/EditorLevelUtils.cpp
@@ -769,13 +769,12 @@ ULevelStreaming* UEditorLevelUtils::CreateNewStreamingLevelForWorld(UWorld& InWo
 	if (bNewWorldSaved && !NewPackageName.IsEmpty())
 	{
 		NewPackageName = FPackageName::FilenameToLongPackageName(NewPackageName);
-		if (!NewLevelWorld)
+
+		// Find or Load package and re-assign NewLevelWorld in case it was duplicated by Save
+		UPackage* NewPackage = LoadPackage(nullptr, *NewPackageName, LOAD_None);
+		if (NewPackage)
 		{
-			UPackage* NewPackage = LoadPackage(nullptr, *NewPackageName, LOAD_None);
-			if (NewPackage)
-			{
-				NewLevelWorld = UWorld::FindWorldInPackage(NewPackage);
-			}
+			NewLevelWorld = UWorld::FindWorldInPackage(NewPackage);
 		}
 	}
 
