CookOnTheFlyServer: Fix missing "not" in comment.
#rnx

diff --git a/Engine/Source/Editor/UnrealEd/Private/CookOnTheFlyServer.cpp b/Engine/Source/Editor/UnrealEd/Private/CookOnTheFlyServer.cpp
--- a/Engine/Source/Editor/UnrealEd/Private/CookOnTheFlyServer.cpp
+++ b/Engine/Source/Editor/UnrealEd/Private/CookOnTheFlyServer.cpp
@@ -6669,7 +6669,8 @@ void UCookOnTheFlyServer::SetInitializeConfigSettings(UE::Cook::FInitializeConfi
 	{
 		PluginsToRemap = IPluginManager::Get().GetEnabledPlugins();
 		TArray<FString> AdditionalPluginDirs = Project->GetAdditionalPluginDirectories();
-		// Remove any plugin that is in the additional directories since they are handled normally and don't need remapping
+		// Remove all plugins that are not in the additional directories. Plugins not in additional directories
+		// are under ProjectRoot or EngineRoot and do not need remapping.
 		for (int32 Index = PluginsToRemap.Num() - 1; Index >= 0; Index--)
 		{
 			bool bRemove = true;
