Delete Actor: Avoid warning about Soft/Hard references if that reference is also being deleted

diff --git a/Engine/Source/Editor/UnrealEd/Private/EditorActor.cpp b/Engine/Source/Editor/UnrealEd/Private/EditorActor.cpp
--- a/Engine/Source/Editor/UnrealEd/Private/EditorActor.cpp
+++ b/Engine/Source/Editor/UnrealEd/Private/EditorActor.cpp
@@ -992,12 +992,20 @@ bool UUnrealEdEngine::DeleteActors(const TArray<AActor*>& InActorsToDelete, UWor
 				// Remove any references from object types marked to be ignored
 				for (int32 i = SoftReferencingObjects->Num() - 1; i >= 0; --i)
 				{
-					for (const TObjectPtr<UClass>& ClassToIgnore : ClassesToIgnoreDeleteReferenceWarning)
+					const UObject* SoftReferencingObject = (*SoftReferencingObjects)[i];
+					if (const AActor* ReferencingActor = Cast<AActor>(SoftReferencingObject); ReferencingActor && ActorsToDelete.Contains(ReferencingActor))
 					{
-						if ((*SoftReferencingObjects)[i]->IsA(ClassToIgnore))
+						SoftReferencingObjects->RemoveAt(i);
+					}
+					else
+					{
+						for (const TObjectPtr<UClass>& ClassToIgnore : ClassesToIgnoreDeleteReferenceWarning)
 						{
-							SoftReferencingObjects->RemoveAt(i);
-							break;
+							if (SoftReferencingObject->IsA(ClassToIgnore))
+							{
+								SoftReferencingObjects->RemoveAt(i);
+								break;
+							}
 						}
 					}
 				}
@@ -1012,7 +1020,7 @@ bool UUnrealEdEngine::DeleteActors(const TArray<AActor*>& InActorsToDelete, UWor
 			for (AActor* ReferencingActor : *ReferencingActors)
 			{
 				// Skip to next if we are referencing ourselves
-				if (ReferencingActor == Actor)
+				if (ReferencingActor == Actor || ActorsToDelete.Contains(ReferencingActor))
 				{
 					continue;
 				}
