World Partition World Builders: saves up to 20 seconds (builder startup time) looking for an incomplete world package name twice.

diff --git a/Engine/Source/Editor/UnrealEd/Private/Commandlets/WorldPartitionBuilderCommandlet.cpp b/Engine/Source/Editor/UnrealEd/Private/Commandlets/WorldPartitionBuilderCommandlet.cpp
--- a/Engine/Source/Editor/UnrealEd/Private/Commandlets/WorldPartitionBuilderCommandlet.cpp
+++ b/Engine/Source/Editor/UnrealEd/Private/Commandlets/WorldPartitionBuilderCommandlet.cpp
@@ -79,9 +79,10 @@ int32 UWorldPartitionBuilderCommandlet::Main(const FString& Params)
 	TArray<FString> MapPackagesNames;
 
 	// Parse map name or maps collection
-	if (FPackageName::SearchForPackageOnDisk(Tokens[0]))
+	FString MapLongPackageName;
+	if (FPackageName::SearchForPackageOnDisk(Tokens[0], &MapLongPackageName))
 	{
-		MapPackagesNames = { Tokens[0] };
+		MapPackagesNames = { MapLongPackageName };
 	}	
 	else if (CollectionManager.CollectionExists(FName(Tokens[0]), ECollectionShareType::CST_All))
 	{
