Auto Save
- Fix restore dialog to use asset labels when possible.
- Fix actor folders not taken into account for auto save.

diff --git a/Engine/Source/Editor/UnrealEd/Private/PackageAutoSaver.cpp b/Engine/Source/Editor/UnrealEd/Private/PackageAutoSaver.cpp
--- a/Engine/Source/Editor/UnrealEd/Private/PackageAutoSaver.cpp
+++ b/Engine/Source/Editor/UnrealEd/Private/PackageAutoSaver.cpp
@@ -32,6 +32,8 @@
 #include "IVREditorModule.h"
 #include "LevelEditorViewport.h"
 #include "AssetCompilingManager.h"
+#include "AssetDefinition.h"
+#include "AssetDefinitionRegistry.h"
 
 namespace PackageAutoSaverJson
 {
@@ -43,11 +45,12 @@ namespace PackageAutoSaverJson
 	typedef TJsonReader<CharType> FJsonReader;
 	typedef TJsonReaderFactory<CharType> FJsonReaderFactory;
 
-	static const FString TagRestoreEnabled	= TEXT("RestoreEnabled");
-	static const FString TagPackages		= TEXT("Packages");
-	static const FString TagPackagePathName = TEXT("PackagePathName");
-	static const FString TagAutoSavePath	= TEXT("AutoSavePath");
-	static const FString RestoreFilename	= TEXT("PackageRestoreData.json");
+	static const FString TagRestoreEnabled	 = TEXT("RestoreEnabled");
+	static const FString TagPackages		 = TEXT("Packages");
+	static const FString TagPackagePathName  = TEXT("PackagePathName");
+	static const FString TagPackageAssetName = TEXT("PackageAssetName");
+	static const FString TagAutoSavePath	 = TEXT("AutoSavePath");
+	static const FString RestoreFilename	 = TEXT("PackageRestoreData.json");
 
 	/**
 	 * @param bEnsurePath True to ensure that the directory for the restore file exists
@@ -61,7 +64,7 @@ namespace PackageAutoSaverJson
 	 *
 	 * @return The packages that have auto-saves that they can be restored from
 	 */
-	TMap<FString, FString> LoadRestoreFile();
+	TMap<FString, TPair<FString, FString>> LoadRestoreFile();
 
 	/**
 	 * Save the file on disk that's used to restore auto-saved packages in the event of a crash
@@ -550,7 +553,7 @@ void FPackageAutoSaver::UpdateDirtyListsForPackage(UPackage* Pkg)
 		}
 
 		// Add package into the appropriate list (map or content)
-		if (UWorld::IsWorldOrExternalActorPackage(Pkg))
+		if (UObject* PackageAsset = Pkg->FindAssetInPackage(); PackageAsset && (Cast<UWorld>(PackageAsset) || PackageAsset->GetTypedOuter<UWorld>()))
 		{
 			DirtyMapsForAutoSave.Add(Pkg);
 		}
@@ -876,9 +879,9 @@ FString PackageAutoSaverJson::GetRestoreFilename(const bool bEnsurePath)
 	return Filename;
 }
 
-TMap<FString, FString> PackageAutoSaverJson::LoadRestoreFile()
+TMap<FString, TPair<FString, FString>> PackageAutoSaverJson::LoadRestoreFile()
 {
-	TMap<FString, FString> PackagesThatCanBeRestored;
+	TMap<FString, TPair<FString, FString>> PackagesThatCanBeRestored;
 
 	const FString Filename = GetRestoreFilename(false);
 	FArchive* const FileAr = IFileManager::Get().CreateFileReader(*Filename);
@@ -909,7 +912,11 @@ TMap<FString, FString> PackageAutoSaverJson::LoadRestoreFile()
 
 		const FString PackagePathName = EntryObject->GetStringField(TagPackagePathName);
 		const FString AutoSavePath = EntryObject->GetStringField(TagAutoSavePath);
-		PackagesThatCanBeRestored.Add(PackagePathName, AutoSavePath);
+
+		FString PackageAssetName;
+		EntryObject->TryGetStringField(TagPackageAssetName, PackageAssetName);
+
+		PackagesThatCanBeRestored.Add(PackagePathName, TPair<FString, FString>(PackageAssetName, AutoSavePath));
 	}
 
 	return PackagesThatCanBeRestored;
@@ -938,9 +945,22 @@ void PackageAutoSaverJson::SaveRestoreFile(const bool bRestoreEnabled, const TMa
 			if(PackagePtr && !AutoSavePath.IsEmpty())
 			{
 				const FString& PackagePathName = PackagePtr->GetPathName();
-
 				TSharedPtr<FJsonObject> EntryObject = MakeShareable(new FJsonObject);
 				EntryObject->SetStringField(TagPackagePathName, PackagePathName);
+
+				if (UObject* PackageAsset = PackagePtr->FindAssetInPackage())
+				{
+					if (UAssetDefinitionRegistry* AssetDefinitionRegistry = UAssetDefinitionRegistry::Get())
+					{
+						const UAssetDefinition* AssetDefinition = AssetDefinitionRegistry->GetAssetDefinitionForClass(PackageAsset->GetClass());
+						if (AssetDefinition)
+						{
+							const FString PackagAssetName = AssetDefinition->GetObjectDisplayNameText(PackageAsset).ToString();
+							EntryObject->SetStringField(TagPackageAssetName, PackagAssetName);
+						}
+					}					
+				}				
+				
 				EntryObject->SetStringField(TagAutoSavePath, AutoSavePath);
 
 				TSharedPtr<FJsonValue> EntryValue = MakeShareable(new FJsonValueObject(EntryObject));
diff --git a/Engine/Source/Editor/UnrealEd/Private/PackageAutoSaver.h b/Engine/Source/Editor/UnrealEd/Private/PackageAutoSaver.h
--- a/Engine/Source/Editor/UnrealEd/Private/PackageAutoSaver.h
+++ b/Engine/Source/Editor/UnrealEd/Private/PackageAutoSaver.h
@@ -158,5 +158,5 @@ private:
 	TSet<TWeakObjectPtr<UPackage>, TWeakObjectPtrSetKeyFuncs<TWeakObjectPtr<UPackage>>> DirtyContentForAutoSave;
 
 	/** Restore information that was loaded following a crash */
-	TMap<FString, FString> PackagesThatCanBeRestored;
+	TMap<FString, TPair<FString, FString>> PackagesThatCanBeRestored;
 };
diff --git a/Engine/Source/Editor/UnrealEd/Private/PackageRestore.cpp b/Engine/Source/Editor/UnrealEd/Private/PackageRestore.cpp
--- a/Engine/Source/Editor/UnrealEd/Private/PackageRestore.cpp
+++ b/Engine/Source/Editor/UnrealEd/Private/PackageRestore.cpp
@@ -33,9 +33,10 @@ namespace PackageRestore
 	class FPackageRestoreItem : public TSharedFromThis<FPackageRestoreItem>
 	{
 	public:
-		FPackageRestoreItem(const FString& InPackageName, const FString& InPackageFilename, const FString& InAutoSaveFilename, const bool bInIsExistingPackage)
+		FPackageRestoreItem(const FString& InPackageName, const FString& InPackageLabel, const FString& InPackageFilename, const FString& InAutoSaveFilename, const bool bInIsExistingPackage)
 			: PackageName(InPackageName)
-			, PackageFilename(InPackageFilename)
+			, PackageLabel(InPackageLabel)
+			, PackageFilename(InPackageFilename)			
 			, AutoSaveFilename(InAutoSaveFilename)
 			, bIsExistingPackage(bInIsExistingPackage)
 			, State(ECheckBoxState::Unchecked)
@@ -48,6 +49,12 @@ namespace PackageRestore
 			return PackageName;
 		}
 
+		/** @return The package label for this item */
+		const FString& GetPackageLabel() const
+		{
+			return PackageLabel.IsEmpty() ? PackageName : PackageLabel;
+		}
+
 		/** @return The package filename for this item */
 		const FString& GetPackageFilename() const
 		{
@@ -127,6 +134,7 @@ namespace PackageRestore
 
 	private:
 		FString PackageName;
+		FString PackageLabel;
 		FString PackageFilename;
 		FString AutoSaveFilename;
 		bool bIsExistingPackage;
@@ -208,7 +216,7 @@ namespace PackageRestore
 					.FillWidth(1)
 					[
 						SNew(STextBlock)
-						.Text(FText::FromString(Item->GetPackageName()))
+						.Text(FText::FromString(Item->GetPackageLabel()))
 					];
 			}
 			else if (ColumnName == ColumnID_FileLabel)
@@ -597,7 +605,7 @@ namespace PackageRestore
 	}
 }
 
-FEditorFileUtils::EPromptReturnCode PackageRestore::PromptToRestorePackages(const TMap<FString, FString>& PackagesToRestore, TArray<FString>* OutFailedPackages)
+FEditorFileUtils::EPromptReturnCode PackageRestore::PromptToRestorePackages(const TMap<FString, TPair<FString, FString>>& PackagesToRestore, TArray<FString>* OutFailedPackages)
 {
 	const FString AutoSaveDir = AutoSaveUtils::GetAutoSaveDir();
 
@@ -607,12 +615,13 @@ FEditorFileUtils::EPromptReturnCode PackageRestore::PromptToRestorePackages(cons
 	for(auto It = PackagesToRestore.CreateConstIterator(); It; ++It)
 	{
 		const FString& PackageFullPath = It.Key();
-		const FString& AutoSavePath = It.Value();
+		const FString& PackageAssetLabel = It.Value().Key;
+		const FString& AutoSavePath = It.Value().Value;
 
 		FString PackageFilename;
 		if(FPackageName::DoesPackageExist(PackageFullPath, &PackageFilename))
 		{
-			FPackageRestoreItemPtr PackageItemPtr = MakeShared<FPackageRestoreItem>(PackageFullPath, PackageFilename, AutoSaveDir / AutoSavePath, true/*bIsExistingPackage*/);
+			FPackageRestoreItemPtr PackageItemPtr = MakeShared<FPackageRestoreItem>(PackageFullPath, PackageAssetLabel, PackageFilename, AutoSaveDir / AutoSavePath, true/*bIsExistingPackage*/);
 			PackageRestoreItems.Add(PackageItemPtr);
 		}
 		else
@@ -622,7 +631,7 @@ FEditorFileUtils::EPromptReturnCode PackageRestore::PromptToRestorePackages(cons
 			{
 				PackageFilename += FPaths::GetExtension(AutoSavePath, true/*bIncludeDot*/);
 
-				FPackageRestoreItemPtr PackageItemPtr = MakeShared<FPackageRestoreItem>(PackageFullPath, PackageFilename, AutoSaveDir / AutoSavePath, false/*bIsExistingPackage*/);
+				FPackageRestoreItemPtr PackageItemPtr = MakeShared<FPackageRestoreItem>(PackageFullPath, PackageAssetLabel, PackageFilename, AutoSaveDir / AutoSavePath, false/*bIsExistingPackage*/);
 				PackageRestoreItems.Add(PackageItemPtr);
 			}
 		}
diff --git a/Engine/Source/Editor/UnrealEd/Private/PackageRestore.h b/Engine/Source/Editor/UnrealEd/Private/PackageRestore.h
--- a/Engine/Source/Editor/UnrealEd/Private/PackageRestore.h
+++ b/Engine/Source/Editor/UnrealEd/Private/PackageRestore.h
@@ -20,5 +20,5 @@ namespace PackageRestore
 	 *				is possible some packages have been successfully restored. If the user opts the "Skip Restore" option on the dialog, the return code will indicate the user has 
 	 *				declined out of the prompt. This way calling code can distinguish between a decline and a cancel and then proceed as planned, or abort its operation accordingly.
 	 */
-	FEditorFileUtils::EPromptReturnCode PromptToRestorePackages(const TMap<FString, FString>& PackagesToRestore, TArray<FString>* OutFailedPackages = nullptr);
+	FEditorFileUtils::EPromptReturnCode PromptToRestorePackages(const  TMap<FString, TPair<FString, FString>>& PackagesToRestore, TArray<FString>* OutFailedPackages = nullptr);
 }
diff --git a/Engine/Source/Runtime/Engine/Classes/Engine/World.h b/Engine/Source/Runtime/Engine/Classes/Engine/World.h
--- a/Engine/Source/Runtime/Engine/Classes/Engine/World.h
+++ b/Engine/Source/Runtime/Engine/Classes/Engine/World.h
@@ -4082,9 +4082,6 @@ public:
 	/** Given a package, locate the UWorld contained within if one exists */
 	static UWorld* FindWorldInPackage(UPackage* Package);
 
-	/** Given a package, return if package contains UWorld or External Actor */
-	static bool IsWorldOrExternalActorPackage(UPackage* Package);
-
 	/** If the specified package contains a redirector to a UWorld, that UWorld is returned. Otherwise, nullptr is returned. */
 	static UWorld* FollowWorldRedirectorInPackage(UPackage* Package, UObjectRedirector** OptionalOutRedirector = nullptr);
 
diff --git a/Engine/Source/Runtime/Engine/Private/World.cpp b/Engine/Source/Runtime/Engine/Private/World.cpp
--- a/Engine/Source/Runtime/Engine/Private/World.cpp
+++ b/Engine/Source/Runtime/Engine/Private/World.cpp
@@ -8206,30 +8206,6 @@ UWorld* UWorld::FindWorldInPackage(UPackage* Package)
 	return RetVal;
 }
 
-bool UWorld::IsWorldOrExternalActorPackage(UPackage* Package)
-{
-	bool bRetVal = false;
-	ForEachObjectWithPackage(Package, [&bRetVal](UObject* PackageObject)
-	{
-		if (PackageObject->IsA<UWorld>())
-		{
-			bRetVal = true;
-			return false;
-		}
-		else if (AActor* Actor = Cast<AActor>(PackageObject))
-		{
-			if (Actor->IsPackageExternal())
-			{
-				bRetVal = true;
-				return false;
-			}
-		}
-		return true;
-	}, false);
-
-	return bRetVal;
-}
-
 UWorld* UWorld::FollowWorldRedirectorInPackage(UPackage* Package, UObjectRedirector** OptionalOutRedirector)
 {
 	UWorld* RetVal = nullptr;
