Cooker: DiffOnly in ZenStore: Fix UE::DiffWriter to not break at offsets in the header when the requested offset happens to be within (OriginalHeaderSize - CompressedHeaderSize) of the beginning of the .uexp file.
#rnx

diff --git a/Engine/Source/Editor/UnrealEd/Private/Cooker/DiffPackageWriter.cpp b/Engine/Source/Editor/UnrealEd/Private/Cooker/DiffPackageWriter.cpp
--- a/Engine/Source/Editor/UnrealEd/Private/Cooker/DiffPackageWriter.cpp
+++ b/Engine/Source/Editor/UnrealEd/Private/Cooker/DiffPackageWriter.cpp
@@ -172,7 +172,8 @@ void FDiffPackageWriter::WritePackageData(const FPackageInfo& Info, FLargeMemory
 		check(PreviousInnerData.Data.Get() != nullptr || (PreviousInnerData.Size == 0 && PreviousInnerData.HeaderSize == 0));
 
 		bNewPackage = PreviousInnerData.Size == 0;
-		Accumulator.OnFirstSaveComplete(LocalInfo.LooseFilePath, LocalInfo.HeaderSize, MoveTemp(PreviousInnerData));
+		Accumulator.OnFirstSaveComplete(LocalInfo.LooseFilePath, LocalInfo.HeaderSize, Info.HeaderSize,
+			MoveTemp(PreviousInnerData));
 		bIsDifferent = Accumulator.HasDifferences();
 	}
 	else
@@ -181,7 +182,7 @@ void FDiffPackageWriter::WritePackageData(const FPackageInfo& Info, FLargeMemory
 
 		TMap<FName, FArchiveDiffStats> PackageDiffStats;
 		const TCHAR* CutoffString = TEXT("UEditorEngine::Save()");
-		Accumulator.CompareWithPrevious(CutoffString, PackageDiffStats, Inner->GetCookCapabilities().HeaderFormat);
+		Accumulator.CompareWithPrevious(CutoffString, PackageDiffStats);
 
 		//COOK_STAT(FSavePackageStats::NumberOfDifferentPackages++);
 		//COOK_STAT(FSavePackageStats::MergeStats(PackageDiffStats));
@@ -221,7 +222,7 @@ UE::DiffWriter::FAccumulator& FDiffPackageWriter::ConstructAccumulator(FName Pac
 	{
 		check(!bHasStartedSecondSave); // Accumulator should already exist from CreateLinkerArchive in the first save
 		Accumulator = new UE::DiffWriter::FAccumulator(Asset, *PackageName.ToString(), MaxDiffsToLog,
-			GetDiffWriterMessageCallback());
+			GetDiffWriterMessageCallback(), Inner->GetCookCapabilities().HeaderFormat);
 	}
 	return *Accumulator;
 }
diff --git a/Engine/Source/Editor/UnrealEd/Private/Cooker/DiffWriterArchive.cpp b/Engine/Source/Editor/UnrealEd/Private/Cooker/DiffWriterArchive.cpp
--- a/Engine/Source/Editor/UnrealEd/Private/Cooker/DiffWriterArchive.cpp
+++ b/Engine/Source/Editor/UnrealEd/Private/Cooker/DiffWriterArchive.cpp
@@ -193,6 +193,7 @@ void FCallstacks::Add(
 	{
 		return;
 	}
+	++StackIgnoreCount;
 
 	const int64 CurrentOffset = Offset;
 	EndOffset = FMath::Max(EndOffset, CurrentOffset + Length); 
@@ -270,13 +271,14 @@ void FCallstacks::Add(
 
 int32 FCallstacks::GetCallstackIndexAtOffset(int64 Offset, int32 MinOffsetIndex) const
 {
-	if (Offset < 0 || Offset >= EndOffset || MinOffsetIndex < 0 || MinOffsetIndex >= CallstackAtOffsetMap.Num())
+	if (Offset < 0 || Offset >= EndOffset || MinOffsetIndex >= CallstackAtOffsetMap.Num())
 	{
 		return -1;
 	}
 
 	// Find the index of the offset the InOffset maps to
 	int32 OffsetForCallstackIndex = -1;
+	MinOffsetIndex = FMath::Max(MinOffsetIndex, 0);
 	int32 MaxOffsetIndex = CallstackAtOffsetMap.Num() - 1;
 
 	// Binary search
@@ -318,6 +320,14 @@ int32 FCallstacks::GetCallstackIndexAtOffset(int64 Offset, int32 MinOffsetIndex)
 	return OffsetForCallstackIndex;
 }
 
+void FCallstacks::RemoveRange(int64 StartOffset, int64 Length)
+{
+	CallstackAtOffsetMap.RemoveAll([StartOffset, Length](const FCallstackAtOffset& Entry)
+		{
+			return StartOffset <= Entry.Offset && Entry.Offset < StartOffset + Length;
+		});
+}
+
 void FCallstacks::Append(const FCallstacks& Other, int64 OtherStartOffset)
 {
 	for (const FCallstackAtOffset& OtherOffset : Other.CallstackAtOffsetMap)
@@ -347,79 +357,112 @@ void FCallstacks::Append(const FCallstacks& Other, int64 OtherStartOffset)
 	EndOffset = FMath::Max(EndOffset, Other.EndOffset + OtherStartOffset);
 }
 
+struct FBreakAtOffsetSettings
+{
+	FString PackageToBreakOn;
+	int64 OffsetToBreakOn = -1;
+	bool bInitialized = false;
+
+	void Initialize()
+	{
+		if (bInitialized)
+		{
+			return;
+		}
+		bInitialized = true;
+		OffsetToBreakOn = -1;
+		PackageToBreakOn.Empty();
+
+		if (!FParse::Param(FCommandLine::Get(), TEXT("cooksinglepackage")) &&
+			!FParse::Param(FCommandLine::Get(), TEXT("cooksinglepackagenorefs")))
+		{
+			return;
+		}
+
+		FString Package;
+		if (!FParse::Value(FCommandLine::Get(), TEXT("map="), Package) &&
+			!FParse::Value(FCommandLine::Get(), TEXT("package="), Package))
+		{
+			return;
+		}
+
+		int64 Offset;
+		// DiffOnlyBreakOffset should be the Combined/DiffBreak Offset from the warning message
+		if (!FParse::Value(FCommandLine::Get(), TEXT("diffonlybreakoffset="), Offset) || Offset <= 0)
+		{
+			return;
+		}
+
+		OffsetToBreakOn = Offset;
+		PackageToBreakOn = TEXT("/") + FPackageName::GetShortName(Package);
+	}
+
+	bool MatchesFilename(const FString& Filename) const
+	{
+		int32 SubnameIndex = Filename.Find(PackageToBreakOn, ESearchCase::IgnoreCase, ESearchDir::FromEnd);
+		if (SubnameIndex < 0)
+		{
+			return false;
+		}
+		int32 SubnameEndIndex = SubnameIndex + PackageToBreakOn.Len();
+		return SubnameEndIndex == Filename.Len() || Filename[SubnameEndIndex] == TEXT('.');
+	}
+
+} GBreakAtOffsetSettings;
+
+
 void FCallstacks::RecordSerialize(EOffsetFrame OffsetFrame, int64 CurrentOffset, int64 Length,
 	const FAccumulator& Accumulator, FDiffArchive& Ar, int32 StackIgnoreCount)
 {
 	int64 LinkerOffset = -1;
-	switch (OffsetFrame)
+
+	// If the writer is using postsave transforms, then we need to know what segment we are in before deciding whether
+	// to allow use of the LinkerOffset for diffonlybreakoffset or for recording the LinkerOffset for each callstack.
+	// The segment information is only known after the first save.
+	if (!Accumulator.IsWriterUsingPostSaveTransforms() || Accumulator.bFirstSaveComplete)
 	{
-	case EOffsetFrame::Linker:
-		LinkerOffset = CurrentOffset;
-		break;
-	case EOffsetFrame::Exports:
-		if (Accumulator.bFirstSaveComplete)
+		switch (OffsetFrame)
 		{
-			LinkerOffset = CurrentOffset + Accumulator.HeaderSize;
+		case EOffsetFrame::Linker:
+			LinkerOffset = CurrentOffset;
+			break;
+		case EOffsetFrame::Exports:
+			if (Accumulator.bFirstSaveComplete)
+			{
+				LinkerOffset = CurrentOffset + Accumulator.HeaderSize;
+			}
+			break;
+		default:
+			checkNoEntry();
+			break;
 		}
-		break;
-	default:
-		checkNoEntry();
-		break;
 	}
 
 	++StackIgnoreCount;
 
 	if (LinkerOffset >= 0)
 	{
-		static struct FBreakAtOffsetSettings
+		if (GBreakAtOffsetSettings.OffsetToBreakOn >= 0 &&
+			LinkerOffset <= GBreakAtOffsetSettings.OffsetToBreakOn && GBreakAtOffsetSettings.OffsetToBreakOn < LinkerOffset + Length)
 		{
-			FString PackageToBreakOn;
-			int64 OffsetToBreakOn;
-
-			FBreakAtOffsetSettings()
-				: OffsetToBreakOn(-1)
+			if (GBreakAtOffsetSettings.MatchesFilename(Accumulator.Filename))
 			{
-				if (!FParse::Param(FCommandLine::Get(), TEXT("cooksinglepackage")) &&
-					!FParse::Param(FCommandLine::Get(), TEXT("cooksinglepackagenorefs")))
+				if (Accumulator.IsWriterUsingPostSaveTransforms() && OffsetFrame == EOffsetFrame::Linker &&
+					LinkerOffset < Accumulator.PreTransformHeaderSize)
 				{
-					return;
+					// Do not break at this serialize call; it's in the pre-transformed header. If the requested
+					// breakoffset is not within the post-transformed header, then this offset is in the wrong segment
+					// and is not a match. If the breakoffset is within the post-transformed header we break and give
+					// instructions during OnFirstSaveComplete.
 				}
-
-				FString Package;
-				if (!FParse::Value(FCommandLine::Get(), TEXT("map="), Package) &&
-					!FParse::Value(FCommandLine::Get(), TEXT("package="), Package))
-				{
-					return;
-				}
-
-				int64 Offset;
-				if (!FParse::Value(FCommandLine::Get(), TEXT("diffonlybreakoffset="), Offset) || Offset <= 0)
-				{
-					return;
-				}
-
-				OffsetToBreakOn = Offset;
-				PackageToBreakOn = TEXT("/") + FPackageName::GetShortName(Package);
-			}
-		} BreakAtOffsetSettings;
-
-		if (BreakAtOffsetSettings.OffsetToBreakOn >= 0 &&
-			LinkerOffset <= BreakAtOffsetSettings.OffsetToBreakOn && BreakAtOffsetSettings.OffsetToBreakOn < LinkerOffset + Length)
-		{
-			const FString& ArcName = Accumulator.Filename;
-			bool bFilenameMatches = false;
-			{
-				int32 SubnameIndex = ArcName.Find(BreakAtOffsetSettings.PackageToBreakOn, ESearchCase::IgnoreCase, ESearchDir::FromEnd);
-				if (SubnameIndex >= 0)
+				else
 				{
-					int32 SubnameEndIndex = SubnameIndex + BreakAtOffsetSettings.PackageToBreakOn.Len();
-					bFilenameMatches = SubnameEndIndex == ArcName.Len() || ArcName[SubnameEndIndex] == TEXT('.');
+					if (!UE::ArchiveStackTrace::ShouldBypassDiff() && !UE::ArchiveStackTrace::ShouldIgnoreDiff())
+					{
+						UE_DEBUG_BREAK();
+					}
 				}
 			}
-			if (bFilenameMatches)
-			{
-				UE_DEBUG_BREAK();
-			}
 		}
 	}
 
@@ -430,6 +473,7 @@ void FCallstacks::RecordSerialize(EOffsetFrame OffsetFrame, int64 CurrentOffset,
 
 		const bool bCollectingCallstacks = Accumulator.bFirstSaveComplete;
 		const bool bCollectCurrentCallstack = bCollectingCallstacks && LinkerOffset >= 0 &&
+			(!Accumulator.IsWriterUsingPostSaveTransforms() || LinkerOffset >= Accumulator.HeaderSize) &&
 			Accumulator.DiffMap.ContainsOffset(LinkerOffset);
 
 		Add(CurrentOffset, Length, SerializedObject, Ar.GetSerializedProperty(), DebugStack, bCollectingCallstacks,
@@ -438,14 +482,16 @@ void FCallstacks::RecordSerialize(EOffsetFrame OffsetFrame, int64 CurrentOffset,
 }
 
 FAccumulator::FAccumulator(UObject* InAsset, FName InPackageName, int32 InMaxDiffsToLog,
-	FMessageCallback&& InMessageCallback)
+	FMessageCallback&& InMessageCallback, EPackageHeaderFormat InPackageHeaderFormat)
 	: LinkerCallstacks()
 	, ExportsCallstacks()
 	, MessageCallback(MoveTemp(InMessageCallback))
 	, PackageName(InPackageName)
 	, Asset(InAsset)
 	, MaxDiffsToLog(InMaxDiffsToLog)
+	, PackageHeaderFormat(InPackageHeaderFormat)
 {
+	GBreakAtOffsetSettings.Initialize();
 }
 
 FAccumulator::~FAccumulator()
@@ -462,13 +508,24 @@ FName FAccumulator::GetAssetClass() const
 	return Asset != nullptr ? Asset->GetClass()->GetFName() : NAME_None;
 }
 
-void FAccumulator::OnFirstSaveComplete(FStringView LooseFilePath, int64 InHeaderSize,
+bool FAccumulator::IsWriterUsingPostSaveTransforms() const
+{
+	return PackageHeaderFormat != EPackageHeaderFormat::PackageFileSummary;
+}
+
+void FAccumulator::OnFirstSaveComplete(FStringView LooseFilePath, int64 InHeaderSize, int64 InPreTransformHeaderSize,
 	ICookedPackageWriter::FPreviousCookedBytesData&& InPreviousPackageData)
 {
 	Filename = LooseFilePath;
 	HeaderSize = InHeaderSize;
+	PreTransformHeaderSize = InPreTransformHeaderSize;
 	PreviousPackageData = MoveTemp(InPreviousPackageData);
 
+	if (IsWriterUsingPostSaveTransforms())
+	{
+		// The header has been transformed, so all callstacks in it are invalid; remove them
+		LinkerCallstacks.RemoveRange(0, PreTransformHeaderSize);
+	}
 	LinkerCallstacks.Append(ExportsCallstacks, HeaderSize);
 	ExportsCallstacks.Reset();
 
@@ -476,6 +533,20 @@ void FAccumulator::OnFirstSaveComplete(FStringView LooseFilePath, int64 InHeader
 
 	LinkerCallstacks.Reset();
 	bFirstSaveComplete = true;
+
+
+	if (IsWriterUsingPostSaveTransforms() &&
+		GBreakAtOffsetSettings.MatchesFilename(Filename) &&
+		GBreakAtOffsetSettings.OffsetToBreakOn < HeaderSize)
+	{
+		// The package writer used for this cook transforms the header before saving it to disk, for e.g. compression.
+		// This means that we don't in general know which offsets in the pre-transform header correspond to
+		// the offsets in the header on disk, and we only know the callstack for the offsets in the pre-transform
+		// header. So we don't know where to break during serialization of the header.
+		// If you specified settings to break at an offset in the pre-transform header, you need instead to debug
+		// using the information reported by DumpPackageHeaderDiffs.
+		UE_DEBUG_BREAK();
+	}
 }
 
 void FAccumulator::OnSecondSaveComplete(int64 InHeaderSize)
@@ -488,6 +559,11 @@ void FAccumulator::OnSecondSaveComplete(int64 InHeaderSize)
 			*this->Filename, HeaderSize, InHeaderSize));
 	}
 
+	if (IsWriterUsingPostSaveTransforms())
+	{
+		// The header has been transformed, so all callstacks in it are invalid; remove them
+		LinkerCallstacks.RemoveRange(0, PreTransformHeaderSize);
+	}
 	LinkerCallstacks.Append(ExportsCallstacks, HeaderSize);
 	ExportsCallstacks.Reset();
 }
@@ -649,196 +725,187 @@ void FAccumulator::CompareWithPreviousForSection(const FPackageData& SourcePacka
 			continue;
 		}
 
-		bool bDifferenceLogged = false;
-		ON_SCOPE_EXIT
+		constexpr int BytesToLog = 128;
+		int32 DifferenceCallstackOffsetIndex = Callstacks.GetCallstackIndexAtOffset(DestAbsoluteOffset,
+			LastDifferenceCallstackOffsetIndex < 0 ? 0 : LastDifferenceCallstackOffsetIndex);
+
+		// Skip reporting the difference if we are still within the same Serialize call, or for any bytes after
+		// the first different byte if we don't know the callstack.
+		if (DifferenceCallstackOffsetIndex == LastDifferenceCallstackOffsetIndex)
 		{
-			if (bDifferenceLogged)
-			{
-				InOutDiffsLogged++;
-				NumDiffsLoggedLocal++;
-			}
-		};
+			continue;
+		}
+		LastDifferenceCallstackOffsetIndex = DifferenceCallstackOffsetIndex;
 
-		if (DiffMap.ContainsOffset(DestAbsoluteOffset))
+		// Also skip reporting the difference if it is another occurrence of the last reported callstack,
+		// or if the callstack is marked as ignore
+		const FCallstacks::FCallstackAtOffset* CallstackAtOffsetPtr = nullptr;
+		const FCallstacks::FCallstackData* DifferenceCallstackDataPtr = nullptr;
+		if (DifferenceCallstackOffsetIndex >= 0)
 		{
-			int32 DifferenceCallstackoffsetIndex = Callstacks.GetCallstackIndexAtOffset(DestAbsoluteOffset, FMath::Max(LastDifferenceCallstackOffsetIndex, 0));
-			ON_SCOPE_EXIT
+			CallstackAtOffsetPtr = &Callstacks.GetCallstack(DifferenceCallstackOffsetIndex);
+			if (CallstackAtOffsetPtr->bIgnore)
 			{
-				LastDifferenceCallstackOffsetIndex = DifferenceCallstackoffsetIndex;
-			};
+				continue;
+			}
 
-			if (DifferenceCallstackoffsetIndex < 0)
+			DifferenceCallstackDataPtr = &Callstacks.GetCallstackData(*CallstackAtOffsetPtr);
+			FString CurrentDifferenceCallstackDataText = DifferenceCallstackDataPtr->ToString(CallstackCutoffText);
+			if (LastDifferenceCallstackDataText.Compare(CurrentDifferenceCallstackDataText, ESearchCase::CaseSensitive) == 0)
 			{
-				MessageCallback(ELogVerbosity::Warning, FString::Printf(
-					TEXT("%s: Difference at offset %lld (absolute offset: %lld), unknown callstack"), *SectionFilename, LocalOffset, DestAbsoluteOffset));
 				continue;
 			}
+			LastDifferenceCallstackDataText = MoveTemp(CurrentDifferenceCallstackDataText);
+		}
+
+		// Update counter for number of existing diffs
+		OutStats.FindOrAdd(AssetClass).NumDiffs++;
+		NumDiffsLocal++;
 
-			if (DifferenceCallstackoffsetIndex == LastDifferenceCallstackOffsetIndex)
+		// Skip reporting the difference if we are over the limit
+		if (MaxDiffsToLog >= 0 && InOutDiffsLogged >= MaxDiffsToLog)
+		{
+			if (FirstUnreportedDiffIndex == -1)
 			{
-				continue;
+				FirstUnreportedDiffIndex = LocalOffset;
 			}
+			continue;
+		}
+
+		// Update counter for number of logged diffs
+		InOutDiffsLogged++;
+		NumDiffsLoggedLocal++;
 
-			const FCallstacks::FCallstackAtOffset& CallstackAtOffset = Callstacks.GetCallstack(DifferenceCallstackoffsetIndex);
-			const FCallstacks::FCallstackData& DifferenceCallstackData = Callstacks.GetCallstackData(CallstackAtOffset);
-			FString DifferenceCallstackDataText = DifferenceCallstackData.ToString(CallstackCutoffText);
-			if (LastDifferenceCallstackDataText.Compare(DifferenceCallstackDataText, ESearchCase::CaseSensitive) == 0)
+		if (DifferenceCallstackOffsetIndex < 0)
+		{
+			if (IsWriterUsingPostSaveTransforms() && DestAbsoluteOffset < HeaderSize)
 			{
-				continue;
+				MessageCallback(ELogVerbosity::Warning, FString::Printf(
+					TEXT("%s: Difference at offset %lld (Combined/DiffBreak Offset: %" INT64_FMT "): OnDisk %d != %d InMemory.%s")
+					TEXT("Callstack is unknown because the offset is in the header and the header has been optimized. See the output of DumpPackageHeaderDiffs to debug this difference."),
+					*SectionFilename, LocalOffset, DestAbsoluteOffset, SourceByte, DestByte, NewLineToken
+				));
 			}
-			ON_SCOPE_EXIT
+			else
 			{
-				LastDifferenceCallstackDataText = MoveTemp(DifferenceCallstackDataText);
-			};
+				MessageCallback(ELogVerbosity::Warning, FString::Printf(
+					TEXT("%s: Difference at offset %lld (Combined/DiffBreak Offset: %" INT64_FMT "): OnDisk %d != %d InMemory.%s")
+					TEXT("Callstack is unknown."),
+					*SectionFilename, LocalOffset, DestAbsoluteOffset, SourceByte, DestByte, NewLineToken
+				));
+			}
+		}
+		else
+		{
+			const FCallstacks::FCallstackAtOffset& CallstackAtOffset = *CallstackAtOffsetPtr;
+			const FCallstacks::FCallstackData& DifferenceCallstackData = *DifferenceCallstackDataPtr;
+			const FString& DifferenceCallstackDataText = LastDifferenceCallstackDataText;
 
-			if (!CallstackAtOffset.bIgnore && (MaxDiffsToLog < 0 || InOutDiffsLogged < MaxDiffsToLog))
+			FString BeforePropertyVal;
+			FString AfterPropertyVal;
+			if (FProperty* SerProp = DifferenceCallstackData.SerializedProp)
 			{
-				FString BeforePropertyVal;
-				FString AfterPropertyVal;
-				if (FProperty* SerProp = DifferenceCallstackData.SerializedProp)
+				if (SourceSize == DestSize && ShouldDumpPropertyValueState(SerProp))
 				{
-					if (SourceSize == DestSize && ShouldDumpPropertyValueState(SerProp))
+					// Walk backwards until we find a callstack which wasn't from the given property
+					int64 OffsetX = DestAbsoluteOffset;
+					for (;;)
 					{
-						// Walk backwards until we find a callstack which wasn't from the given property
-						int64 OffsetX = DestAbsoluteOffset;
-						for (;;)
+						if (OffsetX == 0)
 						{
-							if (OffsetX == 0)
-							{
-								break;
-							}
-
-							const int32 CallstackIndex = Callstacks.GetCallstackIndexAtOffset(OffsetX - 1, 0);
-							const FCallstacks::FCallstackAtOffset& PreviousCallstack = Callstacks.GetCallstack(CallstackIndex);
-							const FCallstacks::FCallstackData& PreviousCallstackData = Callstacks.GetCallstackData(PreviousCallstack);
-							if (PreviousCallstackData.SerializedProp != SerProp)
-							{
-								break;
-							}
-
-							--OffsetX;
+							break;
 						}
 
-						FPropertyTempVal SourceVal(SerProp);
-						FPropertyTempVal DestVal  (SerProp);
-
-						FStaticMemoryReader SourceReader(&SourcePackage.Data[SourceAbsoluteOffset - (DestAbsoluteOffset - OffsetX)], SourcePackage.Size - SourceAbsoluteOffset);
-						FStaticMemoryReader DestReader(&DestPackage.Data[OffsetX], DestPackage.Size - DestAbsoluteOffset);
-
-						SourceVal.Serialize(SourceReader);
-						DestVal  .Serialize(DestReader);
-
-									if (!SourceReader.IsError() && !DestReader.IsError())
-									{
-										SourceVal.ExportText(BeforePropertyVal);
-										DestVal  .ExportText(AfterPropertyVal);
-									}
-								}
-							}
+						const int32 CallstackIndex = Callstacks.GetCallstackIndexAtOffset(OffsetX - 1, 0);
+						const FCallstacks::FCallstackAtOffset& PreviousCallstack = Callstacks.GetCallstack(CallstackIndex);
+						const FCallstacks::FCallstackData& PreviousCallstackData = Callstacks.GetCallstackData(PreviousCallstack);
+						if (PreviousCallstackData.SerializedProp != SerProp)
+						{
+							break;
+						}
 
-				FString DiffValues;
-				if (BeforePropertyVal != AfterPropertyVal)
-				{
-					DiffValues = FString::Printf(TEXT("\r\n%sBefore: %s\r\n%sAfter:  %s"),
-						IndentToken, *BeforePropertyVal,
-						IndentToken, *AfterPropertyVal);
-				}
+						--OffsetX;
+					}
 
-				FString DebugDataStackText;
-				//check for a debug data stack as part of the unique stack entry, and log it out if we find it.
-				FString FullStackText = DifferenceCallstackData.Callstack.Get();
-				int32 DebugDataIndex = FullStackText.Find(ANSI_TO_TCHAR(DebugDataStackMarker), ESearchCase::CaseSensitive);
-				if (DebugDataIndex > 0)
-				{
-					DebugDataStackText = FString::Printf(TEXT("\r\n%s"), IndentToken)
-						+ FullStackText.RightChop(DebugDataIndex + 2);
-				}
+					FPropertyTempVal SourceVal(SerProp);
+					FPropertyTempVal DestVal(SerProp);
 
-				MessageCallback(ELogVerbosity::Warning, FString::Printf(
-					TEXT("%s: Difference at offset %lld%s (absolute offset: %lld): byte %d on disk, byte %d in memory, callstack:%s%s%s%s%s"),
-					*SectionFilename,
-					CallstackAtOffset.Offset - DestPackage.StartOffset,
-					DestAbsoluteOffset > CallstackAtOffset.Offset ? *FString::Printf(TEXT("(+%lld)"), DestAbsoluteOffset - CallstackAtOffset.Offset) : TEXT(""),
-					DestAbsoluteOffset,
-					SourceByte, DestByte,
-					NewLineToken,
-					NewLineToken,
-					*DifferenceCallstackDataText,
-					*DiffValues,
-					*DebugDataStackText
-				));
+					FStaticMemoryReader SourceReader(&SourcePackage.Data[SourceAbsoluteOffset - (DestAbsoluteOffset - OffsetX)], SourcePackage.Size - SourceAbsoluteOffset);
+					FStaticMemoryReader DestReader(&DestPackage.Data[OffsetX], DestPackage.Size - DestAbsoluteOffset);
 
-				const int BytesToLog = 128;
-				MessageCallback(ELogVerbosity::Display, FString::Printf(
-					TEXT("%s: Logging %d bytes around absolute offset: %lld (%016X) in the on disk (existing) package, (which corresponds to offset %lld (%016X) in the in-memory package)"),
-					*SectionFilename,
-					BytesToLog,
-					SourceAbsoluteOffset,
-					SourceAbsoluteOffset,
-					DestAbsoluteOffset,
-					DestAbsoluteOffset
-				));
-				TArray<FString> HexDumpLines = FCompressionUtil::HexDumpLines(SourcePackage.Data, SourcePackage.Size,
-					SourceAbsoluteOffset - BytesToLog / 2, SourceAbsoluteOffset + BytesToLog / 2);
-				for (FString& Line : HexDumpLines)
-				{
-					MessageCallback(ELogVerbosity::Display, Line);
-				}
+					SourceVal.Serialize(SourceReader);
+					DestVal.Serialize(DestReader);
 
-				MessageCallback(ELogVerbosity::Display, FString::Printf(
-					TEXT("%s: Logging %d bytes around absolute offset: %lld (%016X) in the in memory (new) package"),
-					*SectionFilename,
-					BytesToLog,
-					DestAbsoluteOffset,
-					DestAbsoluteOffset
-				));
-				HexDumpLines = FCompressionUtil::HexDumpLines(DestPackage.Data, DestPackage.Size, DestAbsoluteOffset - BytesToLog / 2,
-					DestAbsoluteOffset + BytesToLog / 2);
-				for (FString& Line : HexDumpLines)
-				{
-					MessageCallback(ELogVerbosity::Display, Line);
+					if (!SourceReader.IsError() && !DestReader.IsError())
+					{
+						SourceVal.ExportText(BeforePropertyVal);
+						DestVal.ExportText(AfterPropertyVal);
+					}
 				}
-
-				bDifferenceLogged = true;
 			}
-			else if (FirstUnreportedDiffIndex == -1)
+
+			FString DiffValues;
+			if (BeforePropertyVal != AfterPropertyVal)
 			{
-				FirstUnreportedDiffIndex = DestAbsoluteOffset;
+				DiffValues = FString::Printf(TEXT("\r\n%sBefore: %s\r\n%sAfter:  %s"),
+					IndentToken, *BeforePropertyVal,
+					IndentToken, *AfterPropertyVal);
 			}
-			OutStats.FindOrAdd(AssetClass).NumDiffs++;
-			NumDiffsLocal++;
-		}
-		else
-		{
-			// Each byte will count as a difference but without callstack data there's no way around it
-			OutStats.FindOrAdd(AssetClass).NumDiffs++;
-			NumDiffsLocal++;
-			if (FirstUnreportedDiffIndex == -1)
+
+			FString DebugDataStackText;
+			//check for a debug data stack as part of the unique stack entry, and log it out if we find it.
+			FString FullStackText = DifferenceCallstackData.Callstack.Get();
+			int32 DebugDataIndex = FullStackText.Find(ANSI_TO_TCHAR(DebugDataStackMarker), ESearchCase::CaseSensitive);
+			if (DebugDataIndex > 0)
 			{
-				FirstUnreportedDiffIndex = DestAbsoluteOffset;
+				DebugDataStackText = FString::Printf(TEXT("\r\n%s"), IndentToken)
+					+ FullStackText.RightChop(DebugDataIndex + 2);
 			}
+
+			MessageCallback(ELogVerbosity::Warning, FString::Printf(
+				TEXT("%s: Difference at offset %lld (Combined/DiffBreak Offset: %" INT64_FMT "): OnDisk %d != %d InMemory.%s")
+					TEXT("Difference occurs at index %lld within Serialize call at callstack:%s%s%s%s"),
+				*SectionFilename, LocalOffset, DestAbsoluteOffset, SourceByte, DestByte, NewLineToken,
+				DestAbsoluteOffset - CallstackAtOffset.Offset, NewLineToken,
+				*DifferenceCallstackDataText, *DiffValues, *DebugDataStackText
+			));
 		}
-		OutStats.FindOrAdd(AssetClass).DiffSize++;
-	}
 
-	if (MaxDiffsToLog >= 0 && NumDiffsLocal > NumDiffsLoggedLocal)
-	{
-		if (FirstUnreportedDiffIndex != -1)
+		MessageCallback(ELogVerbosity::Display, FString::Printf(
+			TEXT("%s: Logging %d bytes around offset: %lld (%016X) in the OnDisk package:"),
+			*SectionFilename,
+			BytesToLog, LocalOffset, LocalOffset
+		));
+		TArray<FString> HexDumpLines = FCompressionUtil::HexDumpLines(SourcePackage.Data + SourcePackage.StartOffset,
+			SourcePackage.Size - SourcePackage.StartOffset,
+			LocalOffset - BytesToLog / 2, LocalOffset + BytesToLog / 2);
+		for (FString& Line : HexDumpLines)
 		{
-			MessageCallback(ELogVerbosity::Warning, FString::Printf(
-				TEXT("%s: %lld difference(s) not logged (first at absolute offset: %lld)."),
-				*SectionFilename, NumDiffsLocal - NumDiffsLoggedLocal, FirstUnreportedDiffIndex));
+			MessageCallback(ELogVerbosity::Display, Line);
 		}
-		else
+
+		MessageCallback(ELogVerbosity::Display, FString::Printf(
+			TEXT("%s: Logging %d bytes around offset: %lld (%016X) in the InMemory package:"),
+			*SectionFilename, BytesToLog, LocalOffset, LocalOffset
+		));
+		HexDumpLines = FCompressionUtil::HexDumpLines(DestPackage.Data + DestPackage.StartOffset,
+			DestPackage.Size - DestPackage.StartOffset,
+			LocalOffset - BytesToLog / 2, LocalOffset + BytesToLog / 2);
+		for (FString& Line : HexDumpLines)
 		{
-			MessageCallback(ELogVerbosity::Warning, FString::Printf(
-				TEXT("%s: %lld difference(s) not logged."), *SectionFilename, NumDiffsLocal - NumDiffsLoggedLocal));
+			MessageCallback(ELogVerbosity::Display, Line);
 		}
 	}
+
+	if (MaxDiffsToLog >= 0 && NumDiffsLocal > NumDiffsLoggedLocal)
+	{
+		MessageCallback(ELogVerbosity::Warning, FString::Printf(
+			TEXT("%s: %lld difference(s) not logged (first at offset: %lld)."),
+			*SectionFilename, NumDiffsLocal - NumDiffsLoggedLocal, FirstUnreportedDiffIndex));
+	}
 }
 
-void FAccumulator::CompareWithPrevious(const TCHAR* CallstackCutoffText, TMap<FName, FArchiveDiffStats>&OutStats,
-	const EPackageHeaderFormat PackageHeaderFormat)
+void FAccumulator::CompareWithPrevious(const TCHAR* CallstackCutoffText, TMap<FName, FArchiveDiffStats>&OutStats)
 {
 	// An FDiffArchiveForLinker should have been constructed by SavePackage and should still be in memory
 	check(LinkerArchive);
diff --git a/Engine/Source/Editor/UnrealEd/Private/Cooker/DiffWriterArchive.h b/Engine/Source/Editor/UnrealEd/Private/Cooker/DiffWriterArchive.h
--- a/Engine/Source/Editor/UnrealEd/Private/Cooker/DiffWriterArchive.h
+++ b/Engine/Source/Editor/UnrealEd/Private/Cooker/DiffWriterArchive.h
@@ -158,7 +158,13 @@ public:
 		bool bCollectCurrentCallstack,
 		int32 StackIgnoreCount);
 
-	/** Append other callstacks. */
+	/**
+	 * Remove offset->callstack entries reported for a range of offsets. Only removes entries that start within the
+	 * range, does not remove entries that start before the range but end in or after it.
+	 */
+	void RemoveRange(int64 StartOffset, int64 Length);
+
+	/** Append other offset->callstacks entries and callstacks they refer to. */
 	void Append(const FCallstacks& Other, int64 OtherStartOffset);
 
 	/** Finds a callstack associated with data at the specified offset */
@@ -226,22 +232,21 @@ class FAccumulator : public FRefCountBase
 {
 public:
 	FAccumulator(UObject* InAsset, FName InPackageName, int32 InMaxDiffsToLog,
-		FMessageCallback&& InMessageCallback);
+		FMessageCallback&& InMessageCallback, EPackageHeaderFormat InPackageHeaderFormat);
 	virtual ~FAccumulator();
 
-	void OnFirstSaveComplete(FStringView InLooseFilePath, int64 InHeaderSize,
+	void OnFirstSaveComplete(FStringView InLooseFilePath, int64 InHeaderSize, int64 InPreTransformHeaderSize,
 		ICookedPackageWriter::FPreviousCookedBytesData&& InPreviousPackageData);
 	void OnSecondSaveComplete(int64 InHeaderSize);
 	bool HasDifferences() const;
 
 	/** Compares results from the second save with the previous cook results in PreviousPackagedata.  */
-	void CompareWithPrevious(const TCHAR* CallstackCutoffText, TMap<FName,FArchiveDiffStats>& OutStats,
-		const EPackageHeaderFormat PackageHeaderFormat);
+	void CompareWithPrevious(const TCHAR* CallstackCutoffText, TMap<FName,FArchiveDiffStats>& OutStats);
 
 	void SetHeaderSize(int64 InHeaderSize);
 	void SetCollectingCallstacks(bool bInCollectingCallstacks);
 	FName GetAssetClass() const;
-
+	bool IsWriterUsingPostSaveTransforms() const;
 
 private:
 	void GenerateDiffMapForSection(const FPackageData& SourcePackage, const FPackageData& DestPackage, bool& bOutSectionIdentical);
@@ -264,7 +269,9 @@ private:
 	FString Filename;
 	UObject* Asset = nullptr;
 	int64 HeaderSize = 0;
+	int64 PreTransformHeaderSize = 0;
 	int32 MaxDiffsToLog = 5;
+	EPackageHeaderFormat PackageHeaderFormat = EPackageHeaderFormat::PackageFileSummary;
 	bool bFirstSaveComplete = false;
 	bool bHasDifferences = false;
 
diff --git a/Engine/Source/Editor/UnrealEd/Private/Tests/DiffWriterArchiveTests.cpp b/Engine/Source/Editor/UnrealEd/Private/Tests/DiffWriterArchiveTests.cpp
--- a/Engine/Source/Editor/UnrealEd/Private/Tests/DiffWriterArchiveTests.cpp
+++ b/Engine/Source/Editor/UnrealEd/Private/Tests/DiffWriterArchiveTests.cpp
@@ -48,8 +48,9 @@ bool FDiffWriterArchiveTestsCallstacks::RunTest(const FString& Parameters)
 
 	const int32 MaxDiffsToLog = 1000;
 	FName PackageName(TEXT("PackageName"));
+	EPackageHeaderFormat HeaderFormat(EPackageHeaderFormat::PackageFileSummary);
 	TRefCountPtr<FAccumulator> Accumulator = new FAccumulator(nullptr, PackageName, MaxDiffsToLog,
-		[](ELogVerbosity::Type Verbosity, FStringView Message) {});
+		[](ELogVerbosity::Type Verbosity, FStringView Message) {}, HeaderFormat);
 	FBasicDiffState State = FBasicDiffState{ 1,2,3 };
 	FDiffArchiveForLinker Ar(*Accumulator);
 
@@ -81,11 +82,12 @@ bool FDiffWriterArchiveTestsBasic::RunTest(const FString& Parameters)
 	const int32 MaxDiffsToLog = 1000;
 	FName PackageName(TEXT("PackageName"));
 	FString Filename(TEXT("Filename"));
+	EPackageHeaderFormat HeaderFormat(EPackageHeaderFormat::PackageFileSummary);
 
 	// SECTION("DiffMap - identical")
 	{
 		TRefCountPtr<FAccumulator> Accumulator = new FAccumulator(nullptr, PackageName, MaxDiffsToLog,
-			[](ELogVerbosity::Type Verbosity, FStringView Message) {});
+			[](ELogVerbosity::Type Verbosity, FStringView Message) {}, HeaderFormat);
 
 		FLargeMemoryWriter InitialState;
 		{
@@ -100,7 +102,7 @@ bool FDiffWriterArchiveTestsBasic::RunTest(const FString& Parameters)
 		}
 
 		ICookedPackageWriter::FPreviousCookedBytesData InitialData = ToPackageData(InitialState);
-		Accumulator->OnFirstSaveComplete(Filename, 0, MoveTemp(InitialData));
+		Accumulator->OnFirstSaveComplete(Filename, 0, 0, MoveTemp(InitialData));
 
 		TestTrueExpr(!Accumulator->HasDifferences());
 	}
@@ -108,7 +110,7 @@ bool FDiffWriterArchiveTestsBasic::RunTest(const FString& Parameters)
 	// SECTION("DiffMap - mismatch")
 	{
 		TRefCountPtr<FAccumulator> Accumulator = new FAccumulator(nullptr, PackageName, MaxDiffsToLog,
-			[](ELogVerbosity::Type Verbosity, FStringView Message) {});
+			[](ELogVerbosity::Type Verbosity, FStringView Message) {}, HeaderFormat);
 
 		FLargeMemoryWriter InitialState;
 		{
@@ -123,7 +125,7 @@ bool FDiffWriterArchiveTestsBasic::RunTest(const FString& Parameters)
 		}
 
 		ICookedPackageWriter::FPreviousCookedBytesData InitialData = ToPackageData(InitialState);
-		Accumulator->OnFirstSaveComplete(Filename, 0, MoveTemp(InitialData));
+		Accumulator->OnFirstSaveComplete(Filename, 0, 0, MoveTemp(InitialData));
 
 		TestTrueExpr(Accumulator->HasDifferences());
 	}
@@ -131,7 +133,7 @@ bool FDiffWriterArchiveTestsBasic::RunTest(const FString& Parameters)
 	// SECTION("Compare - mismatch")
 	{
 		TRefCountPtr<FAccumulator> Accumulator = new FAccumulator(nullptr, PackageName, MaxDiffsToLog,
-			[](ELogVerbosity::Type Verbosity, FStringView Message) {});
+			[](ELogVerbosity::Type Verbosity, FStringView Message) {}, HeaderFormat);
 
 		FBasicDiffState InitialState = FBasicDiffState {1,2,3};
 		FLargeMemoryWriter InitialMemory;
@@ -146,7 +148,7 @@ bool FDiffWriterArchiveTestsBasic::RunTest(const FString& Parameters)
 
 		// Generate diff map
 		{
-			Accumulator->OnFirstSaveComplete(Filename, 0, MoveTemp(InitialData));
+			Accumulator->OnFirstSaveComplete(Filename, 0, 0, MoveTemp(InitialData));
 			TestTrueExpr(Accumulator->HasDifferences());
 		}
 		
@@ -159,7 +161,7 @@ bool FDiffWriterArchiveTestsBasic::RunTest(const FString& Parameters)
 		// Compare using callstacks and diffmap
 		{
 			TMap<FName, FArchiveDiffStats> DiffStats;
-			Accumulator->CompareWithPrevious(TEXT("BasicDiffTest"), DiffStats, EPackageHeaderFormat::PackageFileSummary);
+			Accumulator->CompareWithPrevious(TEXT("BasicDiffTest"), DiffStats);
 			TestTrueExpr(DiffStats[NAME_None].NumDiffs == 2);
 		}
 	}
diff --git a/Engine/Source/Runtime/Core/Private/Compression/CompressionUtil.cpp b/Engine/Source/Runtime/Core/Private/Compression/CompressionUtil.cpp
--- a/Engine/Source/Runtime/Core/Private/Compression/CompressionUtil.cpp
+++ b/Engine/Source/Runtime/Core/Private/Compression/CompressionUtil.cpp
@@ -96,23 +96,36 @@ namespace FCompressionUtil
 
 	TArray<FString> CORE_API HexDumpLines(const uint8* Bytes, int64 BytesNum, int64 OffsetStart, int64 OffsetEnd, int32 BytesPerLine)
 	{
-		OffsetStart = FMath::Max(0ll, OffsetStart);
-		OffsetEnd = FMath::Min(BytesNum, OffsetEnd);
-		TArray<FString> Results;
+ 		TArray<FString> Results;
 
 		for (int64 Idx = OffsetStart; Idx < OffsetEnd;)
 		{
-			int64 LineOffset = OffsetStart;
 			FString HexString;
+			HexString.Reserve(128);
+			int64 LineOffset = OffsetStart;
 			for (int64 Idx2 = 0; Idx2 < BytesPerLine && Idx < OffsetEnd; ++Idx, ++Idx2, ++OffsetStart)
 			{
-				HexString += FString::Printf(TEXT("%02X "), Bytes[Idx]);
+				if (0 <= Idx && Idx < BytesNum)
+				{
+					HexString += FString::Printf(TEXT("%02X "), Bytes[Idx]);
+				}
+				else
+				{
+					HexString += TEXT("-- ");
+				}
 				if ((Idx2 & 7) == 7)
 				{
 					HexString += TEXT(" ");
 				}
 			}
-			Results.Add(FString::Printf(TEXT("%016X: %s"), LineOffset, *HexString));
+			if (LineOffset >= 0)
+			{
+				Results.Add(FString::Printf(TEXT("%016X: %s"), LineOffset, *HexString));
+			}
+			else
+			{
+				Results.Add(FString::Printf(TEXT(" -%014X: %s"), -LineOffset, *HexString));
+			}
 		}
 		return Results;
 	}
