Cooker DiffOnly with ZenStore: Fix failed assertion when there is a diff in the header.
#rnx

diff --git a/Engine/Source/Editor/UnrealEd/Private/Cooker/DiffWriterArchive.cpp b/Engine/Source/Editor/UnrealEd/Private/Cooker/DiffWriterArchive.cpp
--- a/Engine/Source/Editor/UnrealEd/Private/Cooker/DiffWriterArchive.cpp
+++ b/Engine/Source/Editor/UnrealEd/Private/Cooker/DiffWriterArchive.cpp
@@ -312,9 +312,11 @@ int32 FCallstacks::GetCallstackIndexAtOffset(int64 Offset, int32 MinOffsetIndex)
 				break;
 			}
 		}
-		check(OffsetForCallstackIndex != -1);
-		check(CallstackAtOffsetMap[OffsetForCallstackIndex].Offset < Offset);
-		check(OffsetForCallstackIndex == (CallstackAtOffsetMap.Num() - 1) || CallstackAtOffsetMap[OffsetForCallstackIndex + 1].Offset > Offset);
+		if (OffsetForCallstackIndex != -1)
+		{
+			check(CallstackAtOffsetMap[OffsetForCallstackIndex].Offset < Offset);
+			check(OffsetForCallstackIndex == (CallstackAtOffsetMap.Num() - 1) || CallstackAtOffsetMap[OffsetForCallstackIndex + 1].Offset > Offset);
+		}
 	}
 
 	return OffsetForCallstackIndex;
@@ -817,6 +819,10 @@ void FAccumulator::CompareWithPreviousForSection(const FPackageData& SourcePacka
 						}
 
 						const int32 CallstackIndex = Callstacks.GetCallstackIndexAtOffset(OffsetX - 1, 0);
+						if (CallstackIndex < 0)
+						{
+							break;
+						}
 						const FCallstacks::FCallstackAtOffset& PreviousCallstack = Callstacks.GetCallstack(CallstackIndex);
 						const FCallstacks::FCallstackData& PreviousCallstackData = Callstacks.GetCallstackData(PreviousCallstack);
 						if (PreviousCallstackData.SerializedProp != SerProp)
