MPCook: Fix crash when CookWorker loses socket connection to CookDirector but still tries to run some of its tick functions that look at the socket before it shuts down.
#rnx

diff --git a/Engine/Source/Editor/UnrealEd/Private/Cooker/CompactBinaryTCP.cpp b/Engine/Source/Editor/UnrealEd/Private/Cooker/CompactBinaryTCP.cpp
--- a/Engine/Source/Editor/UnrealEd/Private/Cooker/CompactBinaryTCP.cpp
+++ b/Engine/Source/Editor/UnrealEd/Private/Cooker/CompactBinaryTCP.cpp
@@ -251,6 +251,8 @@ EConnectionStatus FCompactBinaryTCPImpl::PollSendBytes(FSocket* Socket, const vo
 EConnectionStatus FCompactBinaryTCPImpl::TryWritePacket(FSocket* Socket, FSendBuffer& Buffer,
 	TArrayView<FMarshalledMessage>&& AppendMessages, uint64 MaxPacketSize)
 {
+	// Copy AppendMessages into the SendBuffer before any early exit; we are responsible for holding
+	// a reference to them now.
 	if (!AppendMessages.IsEmpty())
 	{
 		Buffer.PendingMessages.Reserve(Buffer.PendingMessages.Num() + AppendMessages.Num());
@@ -259,6 +261,12 @@ EConnectionStatus FCompactBinaryTCPImpl::TryWritePacket(FSocket* Socket, FSendBu
 			Buffer.PendingMessages.Add(MoveTemp(NewMessage));
 		}
 	}
+
+	if (!Socket)
+	{
+		return EConnectionStatus::Terminated;
+	}
+
 	MaxPacketSize = MaxPacketSize == 0 ? MaxOSPacketSize : FMath::Min(MaxPacketSize, MaxOSPacketSize);
 
 	for (;;)
