Fix ConformRemovedNativeComponents to handle components that have changed type, this prevents ConformRemovedNativeComponents from trying to destroy components that have been inherited from a native class, but some non-native class has overridden the type of the inherited component. In this case a child of that type would attempt to destroy its inherited component if it was trying to match the native parent that declared the component instead of the blueprint that overrode the component type.

diff --git a/Engine/Source/Editor/UnrealEd/Private/Kismet2/Kismet2.cpp b/Engine/Source/Editor/UnrealEd/Private/Kismet2/Kismet2.cpp
--- a/Engine/Source/Editor/UnrealEd/Private/Kismet2/Kismet2.cpp
+++ b/Engine/Source/Editor/UnrealEd/Private/Kismet2/Kismet2.cpp
@@ -804,7 +804,7 @@ bool FKismetEditorUtilities::GenerateBlueprintSkeleton(UBlueprint* BlueprintObj,
 namespace ConformComponentsUtils
 {
 	static void ConformRemovedNativeComponents(UObject* BpCdo);
-	static UObject* FindNativeArchetype(UActorComponent* Component);
+	static UObject* FindNativeArchetype(const UObject* NativeCDO, UActorComponent* Component);
 };
 
 static void ConformComponentsUtils::ConformRemovedNativeComponents(UObject* BpCdo)
@@ -829,7 +829,7 @@ static void ConformComponentsUtils::ConformRemovedNativeComponents(UObject* BpCd
 	TSet<UObject*> DestroyedComponents;
 	for (UActorComponent* Component : OldNativeComponents)
 	{
-		UObject* NativeArchetype = FindNativeArchetype(Component);
+		UObject* NativeArchetype = FindNativeArchetype(NativeCDO, Component);
 		if ((NativeArchetype == nullptr) || !NativeArchetype->HasAnyFlags(RF_ClassDefaultObject))
 		{
 			// Keep track of components inherited from the native super class that are still valid.
@@ -949,7 +949,7 @@ static void ConformComponentsUtils::ConformRemovedNativeComponents(UObject* BpCd
 			// If the component in the Blueprint CDO was attached to a component that's been removed, update the Blueprint's component instance to match the archetype in the native parent CDO.
 			if (DestroyedComponents.Contains(SceneComponent->GetAttachParent()))
 			{
-				if (USceneComponent* NativeArchetype = Cast<USceneComponent>(FindNativeArchetype(SceneComponent)))
+				if (USceneComponent* NativeArchetype = Cast<USceneComponent>(FindNativeArchetype(NativeCDO, SceneComponent)))
 				{
 					USceneComponent* NewAttachParent = NativeArchetype->GetAttachParent();
 					if (NewAttachParent)
@@ -968,27 +968,21 @@ static void ConformComponentsUtils::ConformRemovedNativeComponents(UObject* BpCd
 	}
 }
 
-static UObject* ConformComponentsUtils::FindNativeArchetype(UActorComponent* Component)
+static UObject* ConformComponentsUtils::FindNativeArchetype(const UObject* NativeCDO, UActorComponent* Component)
 {
-	UActorComponent* Archetype = Cast<UActorComponent>(Component->GetArchetype());
-	if (Archetype == nullptr)
+	UObject* NativeSubobject = StaticFindObjectFast(UObject::StaticClass(), const_cast<UObject*>(NativeCDO), Component->GetFName());
+	if (!NativeSubobject)
 	{
 		return nullptr;
 	}
-
-	UObject* ArchetypeOwner = Archetype->GetOuter();
-	UClass* OwnerClass = ArchetypeOwner->GetClass();
-
-	const bool bOwnerIsNative = OwnerClass->HasAnyClassFlags(CLASS_Native);
-	if (bOwnerIsNative)
+	else if (!Component->IsA(NativeSubobject->GetClass()))
 	{
-		return Archetype;
+		return nullptr;
 	}
-	if (Archetype == Component)
+	else
 	{
-		return nullptr;
+		return NativeSubobject;
 	}
-	return FindNativeArchetype(Archetype);
 }
 
 /** Tries to make sure that a blueprint is conformed to its native parent, in case any native class flags have changed */
