Cooker: AdditionalPluginDirectories: Fix RemappedPlugins to handle filepaths in formats other than FPaths::MakeStandardFilename by normalizing the plugin's content directory and the tested filenames before doing the string comparisons to check for childpath under a parentpath.
#rnx

diff --git a/Engine/Source/Editor/UnrealEd/Private/Cooker/CookSandbox.cpp b/Engine/Source/Editor/UnrealEd/Private/Cooker/CookSandbox.cpp
--- a/Engine/Source/Editor/UnrealEd/Private/Cooker/CookSandbox.cpp
+++ b/Engine/Source/Editor/UnrealEd/Private/Cooker/CookSandbox.cpp
@@ -15,7 +15,6 @@ namespace UE::Cook
 {
 
 FCookSandbox::FCookSandbox(FStringView OutputDirectory, TArray<TSharedRef<IPlugin>>& InPluginsToRemap)
-	: PluginsToRemap(InPluginsToRemap)
 {
 	// Local sandbox file wrapper. This will be used to handle path conversions, but will not be used to actually
 	// write/read files so we can safely use [Platform] token in the sandbox directory name and then replace it
@@ -23,6 +22,15 @@ FCookSandbox::FCookSandbox(FStringView OutputDirectory, TArray<TSharedRef<IPlugi
 	SandboxFile = FSandboxPlatformFile::Create(false);
 	SandboxFile->Initialize(&FPlatformFileManager::Get().GetPlatformFile(),
 		*FString::Printf(TEXT("-sandbox=\"%.*s\""), OutputDirectory.Len(), OutputDirectory.GetData()));
+
+	PluginsToRemap.Reserve(InPluginsToRemap.Num());
+	for (TSharedRef<IPlugin>& Plugin : InPluginsToRemap)
+	{
+		FPluginData& Data = PluginsToRemap.Emplace_GetRef();
+		Data.Plugin = Plugin;
+		Data.NormalizedContentDir = Plugin->GetContentDir();
+		FPaths::MakeStandardFilename(Data.NormalizedContentDir);
+	}
 }
 
 const FString& FCookSandbox::GetSandboxDirectory() const
@@ -79,24 +87,27 @@ bool FCookSandbox::TryConvertUncookedFilenameToCookedRemappedPluginFilename(FStr
 	{
 		PlatformSandboxRootDir = GetSandboxDirectory();
 	}
+	FString NormalizedFileName(FileName);
+	FPaths::MakeStandardFilename(NormalizedFileName);
+	constexpr FStringView ContentFolderName(TEXTVIEW("Content"));
 
-	for (const TSharedRef<IPlugin>& Plugin : PluginsToRemap)
+	for (const FPluginData& Data : PluginsToRemap)
 	{
 		// If these match, then this content is part of plugin that gets remapped when packaged/staged
-		if (FileName.StartsWith(Plugin->GetContentDir()))
+		FStringView ContentRelPath;
+		if (FPathViews::TryMakeChildPathRelativeTo(NormalizedFileName, Data.NormalizedContentDir, ContentRelPath))
 		{
-			FString SearchFor;
-			SearchFor /= Plugin->GetName() / TEXT("Content");
-			int32 FoundAt = UE::String::FindLast(FileName, SearchFor, ESearchCase::IgnoreCase);
-			check(FoundAt != INDEX_NONE);
-			// Strip off everything but <PluginName/Content/<remaing path to file>
-			FStringView SnippedOffPath = FileName.RightChop(FoundAt);
-			// Put this is in <sandbox path>/RemappedPlugins/<PluginName>/Content/<remaining path to file>
+			const FString& PluginName = Data.Plugin->GetName();
+
+			// Put this is in <sandbox path>/RemappedPlugins/<PluginName>/Content/ContentRelPath
 			constexpr FStringView RemappedPluginsDirName(REMAPPED_PLUGINS);
-			OutCookedFileName.Reserve(PlatformSandboxRootDir.Len() + RemappedPluginsDirName.Len() + SnippedOffPath.Len() + 2);
+			OutCookedFileName.Reserve(PlatformSandboxRootDir.Len() + RemappedPluginsDirName.Len() +
+				PluginName.Len() + ContentFolderName.Len() + ContentRelPath.Len() + 4);
 			OutCookedFileName = PlatformSandboxRootDir;
 			OutCookedFileName /= REMAPPED_PLUGINS;
-			OutCookedFileName /= SnippedOffPath;
+			OutCookedFileName /= Data.Plugin->GetName();
+			OutCookedFileName /= ContentFolderName;
+			OutCookedFileName /= ContentRelPath;
 			return true;
 		}
 	}
@@ -194,13 +205,15 @@ FString& FCookSandbox::ConvertCookedPathToUncookedPath(FStringView CookedPath,
 			// Snip everything up through the RemappedPlugins/ off so we can find the plugin it corresponds to
 			FStringView PluginPath = CookedPath.RightChop(RemappedIndex + RemappedPluginsFolder.Len() + 1);
 			// Find the plugin that owns this content
-			for (const TSharedRef<IPlugin>& Plugin : PluginsToRemap)
+			FString ExpectedRemainingRoot;
+			for (const FPluginData& Data: PluginsToRemap)
 			{
-				if (PluginPath.StartsWith(Plugin->GetName()))
+				ExpectedRemainingRoot = Data.Plugin->GetName() + TEXT("/") + ContentFolder;
+				FStringView RelPathFromContentFolder;
+				if (FPathViews::TryMakeChildPathRelativeTo(PluginPath, ExpectedRemainingRoot, RelPathFromContentFolder))
 				{
-					UncookedFileName = Plugin->GetContentDir();
-					// Chop off the pluginName/Content since it's part of the full path
-					UncookedFileName /= PluginPath.RightChop(Plugin->GetName().Len() + ContentFolder.Len());
+					UncookedFileName = Data.NormalizedContentDir;
+					UncookedFileName /= RelPathFromContentFolder;
 					break;
 				}
 			}
diff --git a/Engine/Source/Editor/UnrealEd/Private/Cooker/CookSandbox.h b/Engine/Source/Editor/UnrealEd/Private/Cooker/CookSandbox.h
--- a/Engine/Source/Editor/UnrealEd/Private/Cooker/CookSandbox.h
+++ b/Engine/Source/Editor/UnrealEd/Private/Cooker/CookSandbox.h
@@ -57,8 +57,14 @@ public:
 
 
 private:
+	struct FPluginData
+	{
+		TSharedPtr<IPlugin> Plugin;
+		FString NormalizedContentDir;
+	};
+
 	TUniquePtr<FSandboxPlatformFile> SandboxFile;
-	TArray<TSharedRef<IPlugin>>& PluginsToRemap;
+	TArray<FPluginData> PluginsToRemap;
 };
 
 struct FCookSandboxConvertCookedPathToPackageNameContext
