Cooker: DiffWriterArchive: Give a hint to static analysis that a pointer can not be null.
#rnx

diff --git a/Engine/Source/Editor/UnrealEd/Private/Cooker/DiffWriterArchive.cpp b/Engine/Source/Editor/UnrealEd/Private/Cooker/DiffWriterArchive.cpp
--- a/Engine/Source/Editor/UnrealEd/Private/Cooker/DiffWriterArchive.cpp
+++ b/Engine/Source/Editor/UnrealEd/Private/Cooker/DiffWriterArchive.cpp
@@ -799,6 +799,7 @@ void FAccumulator::CompareWithPreviousForSection(const FPackageData& SourcePacka
 		}
 		else
 		{
+			check(CallstackAtOffsetPtr && DifferenceCallstackDataPtr); // These were set up above.
 			const FCallstacks::FCallstackAtOffset& CallstackAtOffset = *CallstackAtOffsetPtr;
 			const FCallstacks::FCallstackData& DifferenceCallstackData = *DifferenceCallstackDataPtr;
 			const FString& DifferenceCallstackDataText = LastDifferenceCallstackDataText;
