Sequencer: Control Rig FBX import.
Transform channels are being negated twice (ty, ry/rz).
Script doesn't seem to pick-up the correct values for other controls (they match if looking at the curves but the script is warning about mismatch's and empty curves)
Fix from AnatoleM but out of office hours.

diff --git a/Engine/Source/Editor/MovieSceneTools/Private/MovieSceneToolHelpers.cpp b/Engine/Source/Editor/MovieSceneTools/Private/MovieSceneToolHelpers.cpp
--- a/Engine/Source/Editor/MovieSceneTools/Private/MovieSceneToolHelpers.cpp
+++ b/Engine/Source/Editor/MovieSceneTools/Private/MovieSceneToolHelpers.cpp
@@ -1755,12 +1755,12 @@ static bool ImportFBXTransformToChannels(FString NodeName, const UMovieSceneUser
 		{
 			FloatChannels[i]->SetDefault(DefaultTransformValues[i]);
 			ChannelEnumIndex = (int)FControlRigChannelEnum::PositionX;
-			bool bNegate = i == 1 || i == 4 || i == 5;
+			bool bNegate = false;
 			const FControlRigChannelEnum ChannelEnum = (FControlRigChannelEnum)(ChannelEnumIndex + i);
 			if (const FControlToTransformMappings* ChannelMapping = ImportFBXControlRigSettings->ControlChannelMappings.FindByPredicate(
 				[ChannelEnum](const FControlToTransformMappings& Mapping) { return ChannelEnum == Mapping.ControlChannel; }))
 			{
-				bNegate = ChannelMapping->bNegate ? !bNegate : bNegate;
+				bNegate = ChannelMapping->bNegate;
 			}
 			ImportTransformChannelToFloat(Transform[i], FloatChannels[i], FrameRate, bNegate, false, StartFrame, bNegate);
 		}
@@ -2609,7 +2609,8 @@ bool MovieSceneToolHelpers::ImportFBXIntoControlRigChannelsWithDialog(const TSha
 	TSharedRef<SWindow> Window = SNew(SWindow)
 		.Title(TitleText)
 		.HasCloseButton(true)
-		.SizingRule(ESizingRule::Autosized)
+		.SizingRule(ESizingRule::UserSized)
+		.ClientSize(FVector2D(500.f, 700.f))
 		.AutoCenter(EAutoCenter::PreferredWorkArea)
 		.SupportsMinimize(false);
 
@@ -2692,7 +2693,8 @@ bool MovieSceneToolHelpers::ExportFBXFromControlRigChannelsWithDialog(const TSha
 	const TSharedRef<SWindow> Window = SNew(SWindow)
 		.Title(TitleText)
 		.HasCloseButton(true)
-		.SizingRule(ESizingRule::Autosized)
+		.SizingRule(ESizingRule::UserSized)
+		.ClientSize(FVector2D(450.0f, 300.0f))
 		.AutoCenter(EAutoCenter::PreferredWorkArea)
 		.SupportsMinimize(false);
 
@@ -5377,4 +5379,4 @@ bool MovieSceneToolHelpers::AddTransformKeys(
 	}
 
 	return true;
-}
\ No newline at end of file
+}
diff --git a/Engine/Source/Editor/UnrealEd/Private/Fbx/FbxMainExport.cpp b/Engine/Source/Editor/UnrealEd/Private/Fbx/FbxMainExport.cpp
--- a/Engine/Source/Editor/UnrealEd/Private/Fbx/FbxMainExport.cpp
+++ b/Engine/Source/Editor/UnrealEd/Private/Fbx/FbxMainExport.cpp
@@ -2091,7 +2091,7 @@ bool FFbxExporter::ExportControlRigSection(const UMovieSceneSection* Section, co
 	const TArray<FString> TransformPropertyArray = {"Location", "Rotation", "Scale"};
 
 	// Cached transform channels for a single property to bake and export the 3 channels in one go
-	FMovieSceneFloatChannel* CachedTransformChannels[3] = {nullptr, nullptr, nullptr};
+	TPair<FMovieSceneFloatChannel*, bool> CachedTransformChannels[3] = {{nullptr, false}, {nullptr, false}, {nullptr, false}};
 	// Cached transform property index keep track of which property we are currently caching to bake & export Transforms & TransformNoScale
 	int CurrentTmPropertyIndex = 0;
 	
@@ -2164,9 +2164,16 @@ bool FFbxExporter::ExportControlRigSection(const UMovieSceneSection* Section, co
 			const FControlRigFbxNodeMapping* NodeMapping = ChannelsMapping.FindByPredicate(
 				[ChannelTypeName, FbxCurveData, TmComponentIndex](const FControlRigFbxNodeMapping& AMapping)
 			{
+				const bool bAttrMatch = AMapping.ChannelAttrIndex == TmComponentIndex;
+				const bool bControlTypeOk = AMapping.ControlType == FbxCurveData.ControlType;
+				// todo: am: need to allow mapping of position/rotation/scale for transforms & euler transforms to allow export of negated transform curves
+				const bool bChannelTypeOk = AMapping.ChannelType == ChannelTypeName;
+
 				// If the control is treated as an attribute of another control, we don't need to check the attribute to remap
-				const bool bRemapControl = AMapping.ChannelType == ChannelTypeName && AMapping.ControlType == FbxCurveData.ControlType;
-				return bRemapControl && (!FbxCurveData.IsControlNode() || AMapping.ChannelAttrIndex == TmComponentIndex);
+				const bool bAttrOk = !FbxCurveData.IsControlNode() || bAttrMatch;
+					
+				const bool bRemapControl = bControlTypeOk && bChannelTypeOk && bAttrOk;
+				return bRemapControl;
 			});
 			
 			if (TmComponentIndex != INDEX_NONE)
@@ -2199,14 +2206,15 @@ bool FFbxExporter::ExportControlRigSection(const UMovieSceneSection* Section, co
 					}
 
 					// Cache a single transform property channel
-					CachedTransformChannels[TmComponentIndex] = FloatChannel;
+					CachedTransformChannels[TmComponentIndex].Key = FloatChannel;
+					CachedTransformChannels[TmComponentIndex].Value = bNegate;
 
 					// Once all 3 channels in the current property have been cached, export baked
 					if (bExportCachedTransforms)
 					{
 						for (int i = 0; i < 3; i++)
 						{
-							if (CachedTransformChannels[i] && CachedTransformChannels[i]->GetNumKeys() > 0)
+							if (CachedTransformChannels[i].Key && CachedTransformChannels[i].Key->GetNumKeys() > 0)
 							{
 								// Only export baked if at least one curve has keys
 								ExportTransformChannelsToFbxCurve(ControlFbxNode, CachedTransformChannels[0],
@@ -2226,7 +2234,7 @@ bool FFbxExporter::ExportControlRigSection(const UMovieSceneSection* Section, co
 					continue;
 				}
 			}
-			// Remap non channels if required (apart from rotations & transforms, remapped separately)
+			// Remap non transform channels if required (rotations & transforms remapped separately)
 			if (NodeMapping)
 			{
 				FbxCurveData.AttributePropertyName = TransformPropertyArray[NodeMapping->FbxAttrIndex / 3];
@@ -2293,12 +2301,10 @@ bool FFbxExporter::ExportControlRigSection(const UMovieSceneSection* Section, co
 			{
 				continue;
 			}
-
-			// If AttributePropertyName was found (i.e Scale), use AttributeName as the property channel (i.e. X) 
-			const FString FbxCurveChannel = FbxCurveData.AttributePropertyName.IsEmpty() ? "" : FbxCurveData.AttributeName;
-
-			// Create the anim curve, optionally of the given name if the attribute had a component 
-			FbxAnimCurve* AnimCurve = Property.GetCurve(BaseLayer, StringCast<char>(*FbxCurveChannel).Get(), true);
+			
+			// Create the anim curve, optionally of the given name if the attribute had a component
+			FbxAnimCurveNode* AnimCurveNode = Property.GetCurveNode(BaseLayer, true);
+			FbxAnimCurve* AnimCurve = Property.GetCurve(BaseLayer, FbxCurveData.AttributePropertyName.IsEmpty() ? nullptr : StringCast<char>(*FbxCurveData.AttributeName).Get(), true);
 			if (!AnimCurve)
 			{
 				continue;
@@ -2330,7 +2336,7 @@ bool FFbxExporter::ExportControlRigSection(const UMovieSceneSection* Section, co
 	return true;
 }
 
-void FFbxExporter::ExportTransformChannelsToFbxCurve(FbxNode* InFbxNode, FMovieSceneFloatChannel* ChannelX, FMovieSceneFloatChannel* ChannelY, FMovieSceneFloatChannel* ChannelZ, int TmPropertyIndex, const UMovieSceneTrack* Track, const FMovieSceneSequenceTransform& RootToLocalTransform)
+void FFbxExporter::ExportTransformChannelsToFbxCurve(FbxNode* InFbxNode, TPair<FMovieSceneFloatChannel*, bool> ChannelX, TPair<FMovieSceneFloatChannel*, bool> ChannelY, TPair<FMovieSceneFloatChannel*, bool> ChannelZ, int TmPropertyIndex, const UMovieSceneTrack* Track, const FMovieSceneSequenceTransform& RootToLocalTransform)
 {
 	FbxAnimLayer* BaseLayer = AnimStack->GetMember<FbxAnimLayer>(0);
 
@@ -2361,17 +2367,29 @@ void FFbxExporter::ExportTransformChannelsToFbxCurve(FbxNode* InFbxNode, FMovieS
 		FFrameTime LocalTime = FFrameRate::TransformTime(FFrameTime(LocalFrame), DisplayRate, TickResolution);
 
 		FVector3f Vec = FVector3f::ZeroVector;
-		if (ChannelX)
+		if (ChannelX.Key)
 		{
-			ChannelX->Evaluate(LocalTime, Vec.X);
+			ChannelX.Key->Evaluate(LocalTime, Vec.X);
+			if (ChannelX.Value)
+			{
+				Vec.X = -Vec.X;
+			}
 		}
-		if (ChannelY)
+		if (ChannelY.Key)
 		{
-			ChannelY->Evaluate(LocalTime, Vec.Y);
+			ChannelY.Key->Evaluate(LocalTime, Vec.Y);
+			if (ChannelY.Value)
+			{
+				Vec.Y = -Vec.Y;
+			}
 		}
-		if (ChannelZ)
+		if (ChannelZ.Key)
 		{
-			ChannelZ->Evaluate(LocalTime, Vec.Z);
+			ChannelZ.Key->Evaluate(LocalTime, Vec.Z);
+			if (ChannelZ.Value)
+			{
+				Vec.Z = -Vec.Z;
+			}
 		}
 
 		FbxVector4 KeyVec;
diff --git a/Engine/Source/Editor/UnrealEd/Private/FbxExporter.h b/Engine/Source/Editor/UnrealEd/Private/FbxExporter.h
--- a/Engine/Source/Editor/UnrealEd/Private/FbxExporter.h
+++ b/Engine/Source/Editor/UnrealEd/Private/FbxExporter.h
@@ -471,7 +471,7 @@ private:
 	template <class ChannelType, typename T>
 	void ExportConstantChannelToFbxCurve(FbxAnimCurve& InFbxCurve, const ChannelType& InChannel, FFrameRate TickResolution, const FMovieSceneSequenceTransform& RootToLocalTransform);
 
-	UNREALED_API void ExportTransformChannelsToFbxCurve(FbxNode* InFbxNode, FMovieSceneFloatChannel* ChannelX, FMovieSceneFloatChannel* ChannelY, FMovieSceneFloatChannel* ChannelZ, int TmPropertyIndex, const UMovieSceneTrack* Track, const FMovieSceneSequenceTransform& RootToLocalTransform);
+	UNREALED_API void ExportTransformChannelsToFbxCurve(FbxNode* InFbxNode, TPair<FMovieSceneFloatChannel*, bool> ChannelX, TPair<FMovieSceneFloatChannel*, bool> ChannelY, TPair<FMovieSceneFloatChannel*, bool> ChannelZ, int TmPropertyIndex, const UMovieSceneTrack* Track, const FMovieSceneSequenceTransform& RootToLocalTransform);
 
 	/**
 	 * Finds the given actor in the already-exported list of structures
