New FExternalContentDragDropOp: Subclasses FAssetDragDropOp with unique Guid.
New FEditorDelegates::FOnExternalContentResolved(const FGuid& Guid, const FAssetData& PlaceholderAsset, const FAssetData& ResolvedAsset)

#rnx

diff --git a/Engine/Source/Editor/UnrealEd/Private/Editor.cpp b/Engine/Source/Editor/UnrealEd/Private/Editor.cpp
--- a/Engine/Source/Editor/UnrealEd/Private/Editor.cpp
+++ b/Engine/Source/Editor/UnrealEd/Private/Editor.cpp
@@ -171,6 +171,7 @@ FEditorDelegates::FOnViewAssetIdentifiers				FEditorDelegates::OnEditAssetIdenti
 FEditorDelegates::FOnRestartRequested					FEditorDelegates::OnRestartRequested;
 FEditorDelegates::FOnEditorBoot							FEditorDelegates::OnEditorBoot;
 FEditorDelegates::FOnEditorInitialized					FEditorDelegates::OnEditorInitialized;
+FEditorDelegates::FOnExternalContentResolved			FEditorDelegates::OnExternalContentResolved;
 
 PRAGMA_ENABLE_DEPRECATION_WARNINGS
 
diff --git a/Engine/Source/Editor/UnrealEd/Public/DragAndDrop/ExternalContentDragDropOp.h b/Engine/Source/Editor/UnrealEd/Public/DragAndDrop/ExternalContentDragDropOp.h
new file mode 100644
--- /dev/null
+++ b/Engine/Source/Editor/UnrealEd/Public/DragAndDrop/ExternalContentDragDropOp.h
@@ -0,0 +1,21 @@
+// Copyright Epic Games, Inc. All Rights Reserved.
+
+#pragma once
+
+#include "CoreMinimal.h"
+#include "DragAndDrop/AssetDragDropOp.h"
+#include "Misc/Guid.h"
+
+class FExternalContentDragDropOp : public FAssetDragDropOp
+{
+public:
+	DRAG_DROP_OPERATOR_TYPE(FExternalContentDragDropOp, FAssetDragDropOp)
+
+	FExternalContentDragDropOp() : Guid(FGuid::NewGuid()) {}
+	virtual ~FExternalContentDragDropOp() {}
+
+	const FGuid& GetGuid() const { return Guid; }
+
+private:
+	FGuid Guid;
+};
\ No newline at end of file
diff --git a/Engine/Source/Editor/UnrealEd/Public/Editor.h b/Engine/Source/Editor/UnrealEd/Public/Editor.h
--- a/Engine/Source/Editor/UnrealEd/Public/Editor.h
+++ b/Engine/Source/Editor/UnrealEd/Public/Editor.h
@@ -22,6 +22,7 @@
 class AStaticMeshActor;
 class FEdMode;
 class UFactory;
+struct FGuid;
 
 /** The shorthand identifier used for editor modes */
 typedef FName FEditorModeID;
@@ -173,6 +174,8 @@ struct FEditorDelegates
 	DECLARE_MULTICAST_DELEGATE_OneParam(FOnEditorBoot, double Duration);
 	/** delegate for when the editor has fully initialized */
 	DECLARE_MULTICAST_DELEGATE_OneParam(FOnEditorInitialized, double Duration);
+	/** delegate when external content resolves and can replace placeholder data */
+	DECLARE_MULTICAST_DELEGATE_ThreeParams(FOnExternalContentResolved, const FGuid&, const FAssetData& /*PlaceholderAsset*/, const FAssetData& /*ResolvedAsset*/);
 
 	/** Called when the CurrentLevel is switched to a new level.  Note that this event won't be fired for temporary
 		changes to the current level, such as when copying/pasting actors. */
@@ -357,6 +360,8 @@ struct FEditorDelegates
 	static UNREALED_API FOnEditorBoot OnEditorBoot;
 	/** Called when the editor has initialized */
 	static UNREALED_API FOnEditorInitialized OnEditorInitialized;
+	/** Called when external content gets resolved */
+	static UNREALED_API FOnExternalContentResolved OnExternalContentResolved;
 };
 
 /**
