Cooker: AdditionalPluginDirectories: Fix staging to look in the RemappedPlugins directory when cooking DLC plugins that are in AdditionalPluginDirectories.
Fix cooker to store all metadata in the RemappedPlugins.
Fix cooker for even normal DLCs to store all metadata in the DLC's directory rather than the project's directory.
#rn Bugfix, Core

diff --git a/Engine/Source/Editor/UnrealEd/Classes/CookOnTheSide/CookOnTheFlyServer.h b/Engine/Source/Editor/UnrealEd/Classes/CookOnTheSide/CookOnTheFlyServer.h
--- a/Engine/Source/Editor/UnrealEd/Classes/CookOnTheSide/CookOnTheFlyServer.h
+++ b/Engine/Source/Editor/UnrealEd/Classes/CookOnTheSide/CookOnTheFlyServer.h
@@ -1256,6 +1256,8 @@ private:
 
 	UNREALED_API FString GetContentDirectoryForDLC() const;
 
+	FString GetMetadataDirectory() const;
+
 	/**
 	 * Is the local CookOnTheFlyServer cooking a Project+Engine+EmbeddedPlugin Release that can be used as a
 	 * base for future DLC or Patches?
diff --git a/Engine/Source/Editor/UnrealEd/Private/CookOnTheFlyServer.cpp b/Engine/Source/Editor/UnrealEd/Private/CookOnTheFlyServer.cpp
--- a/Engine/Source/Editor/UnrealEd/Private/CookOnTheFlyServer.cpp
+++ b/Engine/Source/Editor/UnrealEd/Private/CookOnTheFlyServer.cpp
@@ -1200,6 +1200,12 @@ FString UCookOnTheFlyServer::GetContentDirectoryForDLC() const
 	return GetBaseDirectoryForDLC() / TEXT("Content");
 }
 
+FString UCookOnTheFlyServer::GetMetadataDirectory() const
+{
+	FString ProjectOrPluginRoot = !IsCookingDLC() ? FPaths::ProjectDir() : GetBaseDirectoryForDLC();
+	return ProjectOrPluginRoot / TEXT("Metadata");
+}
+
 // allow for a command line to start async preloading a Development AssetRegistry if requested
 static FEventRef GPreloadAREvent(EEventMode::ManualReset);
 static FEventRef GPreloadARInfoEvent(EEventMode::ManualReset);
@@ -7392,7 +7398,7 @@ bool UCookOnTheFlyServer::GetCurrentIniVersionStrings( const ITargetPlatform* Ta
 
 bool UCookOnTheFlyServer::GetCookedIniVersionStrings(const ITargetPlatform* TargetPlatform, UE::Cook::FIniSettingContainer& OutIniSettings, TMap<FString,FString>& OutAdditionalSettings) const
 {
-	const FString EditorIni = FPaths::ProjectDir() / TEXT("Metadata") / TEXT("CookedIniVersion.txt");
+	const FString EditorIni = GetMetadataDirectory() / TEXT("CookedIniVersion.txt");
 	const FString SandboxEditorIni = ConvertToFullSandboxPath(*EditorIni, true);
 
 
@@ -7691,7 +7697,7 @@ TMap<FName, FString> UCookOnTheFlyServer::CalculateCookSettingStrings() const
 
 FString UCookOnTheFlyServer::GetCookSettingsFileName(const ITargetPlatform* TargetPlatform) const
 {
-	FString CookedSettingsIni = FPaths::ProjectDir() / TEXT("Metadata") / TEXT("CookedSettings.txt");
+	FString CookedSettingsIni = GetMetadataDirectory() / TEXT("CookedSettings.txt");
 	return ConvertToFullSandboxPath(*CookedSettingsIni, true, TargetPlatform->PlatformName());
 }
 
@@ -7928,7 +7934,7 @@ bool UCookOnTheFlyServer::SaveCurrentIniSettings(const ITargetPlatform* TargetPl
 	UE::Cook::FIniSettingContainer CurrentIniSettings;
 	GetCurrentIniVersionStrings(TargetPlatform, CurrentIniSettings);
 
-	const FString EditorIni = FPaths::ProjectDir() / TEXT("Metadata") / TEXT("CookedIniVersion.txt");
+	const FString EditorIni = GetMetadataDirectory() / TEXT("CookedIniVersion.txt");
 	const FString SandboxEditorIni = ConvertToFullSandboxPath(*EditorIni, true);
 
 
@@ -8935,14 +8941,7 @@ FString UCookOnTheFlyServer::GetCookedAssetRegistryFilename(const FString& Platf
 
 FString UCookOnTheFlyServer::GetCookedCookMetadataFilename(const FString& PlatformName )
 {
-	if (IsCookingDLC())
-	{
-		check(IsDirectorCookByTheBook());
-		const FString Filename = GetBaseDirectoryForDLC() / TEXT("Metadata") / UE::Cook::GetCookMetadataFilename();
-		return ConvertToFullSandboxPath(*Filename, true).Replace(TEXT("[Platform]"), *PlatformName);
-	}
-
-	const FString MetadataFilename = FPaths::ProjectDir() / TEXT("Metadata") / UE::Cook::GetCookMetadataFilename();
+	const FString MetadataFilename = GetMetadataDirectory() / UE::Cook::GetCookMetadataFilename();
 	return ConvertToFullSandboxPath(*MetadataFilename, true).Replace(TEXT("[Platform]"), *PlatformName);
 }
 
@@ -9307,7 +9306,7 @@ void UCookOnTheFlyServer::CreatePipelineCache(const ITargetPlatform* TargetPlatf
 					// copy the resulting file to metadata for easier examination later
 					if (IFileManager::Get().FileExists(*PCPath))
 					{
-						const FString RootPipelineCacheMetadataPath = FPaths::ProjectDir() / TEXT("Metadata") / TEXT("PipelineCaches");
+						const FString RootPipelineCacheMetadataPath = GetMetadataDirectory() / TEXT("PipelineCaches");
 						const FString PipelineCacheMetadataPathSB = ConvertToFullSandboxPath(*RootPipelineCacheMetadataPath, true);
 						const FString PipelineCacheMetadataPath = PipelineCacheMetadataPathSB.Replace(TEXT("[Platform]"), *TargetPlatform->PlatformName());
 						const FString PipelineCacheMetadataFileName = PipelineCacheMetadataPath / OutFilename;
@@ -9350,7 +9349,7 @@ void UCookOnTheFlyServer::GetShaderLibraryPaths(const ITargetPlatform* TargetPla
 	const FString BasePath = (!IsCookingDLC() || bUseProjectDirForDLC) ? FPaths::ProjectContentDir() : GetContentDirectoryForDLC();
 	OutShaderCodeDir = ConvertToFullSandboxPath(*BasePath, true, TargetPlatform->PlatformName());
 
-	const FString RootMetaDataPath = FPaths::ProjectDir() / TEXT("Metadata") / TEXT("PipelineCaches");
+	const FString RootMetaDataPath = GetMetadataDirectory() / TEXT("PipelineCaches");
 	OutMetaDataPath = ConvertToFullSandboxPath(*RootMetaDataPath, true, *TargetPlatform->PlatformName()); 
 }
 
@@ -10534,6 +10533,8 @@ void UCookOnTheFlyServer::LoadBeginCookIterativeFlagsLocal(FBeginCookContext& Be
 		bool bIterateSharedBuild = false;
 		if (bIterative && bIsSharedIterativeCook && !PlatformData->bIsSandboxInitialized)
 		{
+			checkf(!IsCookingDLC(), TEXT("SharedIterativeCook is not implemented for DLC")); // The paths below look in the ProjectSavedDir, we don't have a save dir per dlc
+
 			// see if the shared build is newer then the current cooked content in the local directory
 			FString SharedCookedAssetRegistry = FPaths::Combine(*FPaths::ProjectSavedDir(), TEXT("SharedIterativeBuild"),
 				*TargetPlatform->PlatformName(), TEXT("Metadata"), GetDevelopmentAssetRegistryFilename());
@@ -10691,15 +10692,7 @@ UE::Cook::FCookSavePackageContext* UCookOnTheFlyServer::CreateSaveContext(const
 	checkf(SandboxFile, TEXT("SaveContexts cannot be created until after CreateSandboxFile has been called from a StartCook function."));
 
 	const FString RootPathSandbox = ConvertToFullSandboxPath(FPaths::RootDir(), true);
-	FString MetadataPathSandbox;
-	if (IsCookingDLC())
-	{
-		MetadataPathSandbox = ConvertToFullSandboxPath(GetBaseDirectoryForDLC() / "Metadata", true);
-	}
-	else
-	{
-		MetadataPathSandbox = ConvertToFullSandboxPath(FPaths::ProjectDir() / "Metadata", true);
-	}
+	FString MetadataPathSandbox = ConvertToFullSandboxPath(GetMetadataDirectory(), true);
 	const FString PlatformString = TargetPlatform->PlatformName();
 	const FString ResolvedRootPath = RootPathSandbox.Replace(TEXT("[Platform]"), *PlatformString);
 	const FString ResolvedMetadataPath = MetadataPathSandbox.Replace(TEXT("[Platform]"), *PlatformString);
diff --git a/Engine/Source/Editor/UnrealEd/Private/Cooker/CookSandbox.cpp b/Engine/Source/Editor/UnrealEd/Private/Cooker/CookSandbox.cpp
--- a/Engine/Source/Editor/UnrealEd/Private/Cooker/CookSandbox.cpp
+++ b/Engine/Source/Editor/UnrealEd/Private/Cooker/CookSandbox.cpp
@@ -28,8 +28,8 @@ FCookSandbox::FCookSandbox(FStringView OutputDirectory, TArray<TSharedRef<IPlugi
 	{
 		FPluginData& Data = PluginsToRemap.Emplace_GetRef();
 		Data.Plugin = Plugin;
-		Data.NormalizedContentDir = Plugin->GetContentDir();
-		FPaths::MakeStandardFilename(Data.NormalizedContentDir);
+		Data.NormalizedRootDir = Plugin->GetBaseDir();
+		FPaths::MakeStandardFilename(Data.NormalizedRootDir);
 	}
 }
 
@@ -94,20 +94,19 @@ bool FCookSandbox::TryConvertUncookedFilenameToCookedRemappedPluginFilename(FStr
 	for (const FPluginData& Data : PluginsToRemap)
 	{
 		// If these match, then this content is part of plugin that gets remapped when packaged/staged
-		FStringView ContentRelPath;
-		if (FPathViews::TryMakeChildPathRelativeTo(NormalizedFileName, Data.NormalizedContentDir, ContentRelPath))
+		FStringView PluginRootRelPath;
+		if (FPathViews::TryMakeChildPathRelativeTo(NormalizedFileName, Data.NormalizedRootDir, PluginRootRelPath))
 		{
 			const FString& PluginName = Data.Plugin->GetName();
 
-			// Put this is in <sandbox path>/RemappedPlugins/<PluginName>/Content/ContentRelPath
+			// Put this is in <sandbox path>/RemappedPlugins/<PluginName>/PluginRootRelPath
 			constexpr FStringView RemappedPluginsDirName(REMAPPED_PLUGINS);
 			OutCookedFileName.Reserve(PlatformSandboxRootDir.Len() + RemappedPluginsDirName.Len() +
-				PluginName.Len() + ContentFolderName.Len() + ContentRelPath.Len() + 4);
+				PluginName.Len() + PluginRootRelPath.Len() + 3);
 			OutCookedFileName = PlatformSandboxRootDir;
 			OutCookedFileName /= REMAPPED_PLUGINS;
 			OutCookedFileName /= Data.Plugin->GetName();
-			OutCookedFileName /= ContentFolderName;
-			OutCookedFileName /= ContentRelPath;
+			OutCookedFileName /= PluginRootRelPath;
 			return true;
 		}
 	}
@@ -208,12 +207,12 @@ FString& FCookSandbox::ConvertCookedPathToUncookedPath(FStringView CookedPath,
 			FString ExpectedRemainingRoot;
 			for (const FPluginData& Data: PluginsToRemap)
 			{
-				ExpectedRemainingRoot = Data.Plugin->GetName() + TEXT("/") + ContentFolder;
-				FStringView RelPathFromContentFolder;
-				if (FPathViews::TryMakeChildPathRelativeTo(PluginPath, ExpectedRemainingRoot, RelPathFromContentFolder))
+				ExpectedRemainingRoot = Data.Plugin->GetName();
+				FStringView RelPathFromPluginRoot;
+				if (FPathViews::TryMakeChildPathRelativeTo(PluginPath, ExpectedRemainingRoot, RelPathFromPluginRoot))
 				{
-					UncookedFileName = Data.NormalizedContentDir;
-					UncookedFileName /= RelPathFromContentFolder;
+					UncookedFileName = Data.NormalizedRootDir;
+					UncookedFileName /= RelPathFromPluginRoot;
 					break;
 				}
 			}
diff --git a/Engine/Source/Editor/UnrealEd/Private/Cooker/CookSandbox.h b/Engine/Source/Editor/UnrealEd/Private/Cooker/CookSandbox.h
--- a/Engine/Source/Editor/UnrealEd/Private/Cooker/CookSandbox.h
+++ b/Engine/Source/Editor/UnrealEd/Private/Cooker/CookSandbox.h
@@ -60,7 +60,7 @@ private:
 	struct FPluginData
 	{
 		TSharedPtr<IPlugin> Plugin;
-		FString NormalizedContentDir;
+		FString NormalizedRootDir;
 	};
 
 	TUniquePtr<FSandboxPlatformFile> SandboxFile;
diff --git a/Engine/Source/Programs/AutomationTool/AutomationUtils/DeploymentContext.cs b/Engine/Source/Programs/AutomationTool/AutomationUtils/DeploymentContext.cs
--- a/Engine/Source/Programs/AutomationTool/AutomationUtils/DeploymentContext.cs
+++ b/Engine/Source/Programs/AutomationTool/AutomationUtils/DeploymentContext.cs
@@ -167,6 +167,12 @@ public class DeploymentContext //: ProjectParams
 	/// </summary>
 	public DirectoryReference DLCRoot;
 
+	/// <summary>
+	/// The list of AdditionalPluginDirectories from the project.uproject. Files in plugins in these
+	/// directories are staged into <StageRoot>/RemappedPlugins/PluginName.
+	/// </summary>
+	public List<DirectoryReference> AdditionalPluginDirectories;
+
 	/// <summary>
 	///  raw name used for platform subdirectories Win32
 	/// </summary>
@@ -407,7 +413,8 @@ public class DeploymentContext //: ProjectParams
 		bool IsClientInsteadOfNoEditor,
         bool InForceChunkManifests,
 		bool InSeparateDebugStageDirectory,
-		DirectoryReference InDLCRoot
+		DirectoryReference InDLCRoot,
+		List<DirectoryReference> InAdditionalPluginDirectories
 		)
 	{
 		bStageCrashReporter = InStageCrashReporter;
@@ -424,6 +431,7 @@ public class DeploymentContext //: ProjectParams
 		Stage = InStage;
 		Archive = InArchive;
 		DLCRoot = InDLCRoot;
+		AdditionalPluginDirectories = InAdditionalPluginDirectories;
 
         if (CookSourcePlatform != null && InCooked)
         {
@@ -738,34 +746,32 @@ public class DeploymentContext //: ProjectParams
 	public StagedFileReference GetStagedFileLocation(FileReference InputFile)
 	{
 		StagedFileReference OutputFile;
-		if(InputFile.IsUnderDirectory(ProjectRoot))
+		foreach (DirectoryReference AdditionalPluginDir in AdditionalPluginDirectories)
 		{
-			OutputFile = StagedFileReference.Combine(RelativeProjectRootForStage, InputFile.MakeRelativeTo(ProjectRoot));
-		}
-        else if (InputFile.HasExtension(".uplugin"))
-        {
-			DirectoryReference EnterpriseRoot = DirectoryReference.Combine(EngineRoot, "..", "Enterprise"); // Enterprise plugins aren't under the project additional plugin directories, so they shouldn't be remapped
-            if (InputFile.IsUnderDirectory(EngineRoot) || InputFile.IsUnderDirectory(EnterpriseRoot))
-			{
-				OutputFile = new StagedFileReference(InputFile.MakeRelativeTo(LocalRoot));
-			}
-            else
+			if (InputFile.IsUnderDirectory(AdditionalPluginDir))
 			{
 				// This is a plugin that lives outside of the Engine/Plugins or Game/Plugins directory so needs to be remapped for staging/packaging
-				// We need to remap C:\SomePath\PluginName\PluginName.uplugin to RemappedPlugins\PluginName\PluginName.uplugin
-				OutputFile = new StagedFileReference(String.Format("RemappedPlugins/{0}/{1}", InputFile.GetFileNameWithoutExtension(), InputFile.GetFileName()));
+				// We need to remap C:\SomePath\PluginName\RelativePath to RemappedPlugins\PluginName\RelativePath
+				OutputFile = new StagedFileReference(
+					String.Format("RemappedPlugins/{0}", InputFile.MakeRelativeTo(AdditionalPluginDir)));
+				return OutputFile;
 			}
-        }
-        else if (InputFile.IsUnderDirectory(LocalRoot))
-        {
+		}
+
+		if (InputFile.IsUnderDirectory(ProjectRoot))
+		{
+			OutputFile = StagedFileReference.Combine(RelativeProjectRootForStage, InputFile.MakeRelativeTo(ProjectRoot));
+		}
+		else if (InputFile.IsUnderDirectory(LocalRoot))
+		{
 			OutputFile = new StagedFileReference(InputFile.MakeRelativeTo(LocalRoot));
-        }
-        else if (DLCRoot != null && InputFile.IsUnderDirectory(DLCRoot))
+		}
+		else if (DLCRoot != null && InputFile.IsUnderDirectory(DLCRoot))
 		{
 			OutputFile = new StagedFileReference(InputFile.MakeRelativeTo(DLCRoot));
 		}
 		else
-        {
+		{
 			throw new AutomationException("Can't deploy {0} because it doesn't start with {1} or {2}", InputFile, ProjectRoot, LocalRoot);
 		}
 		return OutputFile;
diff --git a/Engine/Source/Programs/AutomationTool/AutomationUtils/ProjectParams.cs b/Engine/Source/Programs/AutomationTool/AutomationUtils/ProjectParams.cs
--- a/Engine/Source/Programs/AutomationTool/AutomationUtils/ProjectParams.cs
+++ b/Engine/Source/Programs/AutomationTool/AutomationUtils/ProjectParams.cs
@@ -443,6 +443,7 @@ namespace AutomationTool
 			this.TraceHost = InParams.TraceHost;
 			this.TraceFile = InParams.TraceFile;
 			this.SessionLabel = InParams.SessionLabel;
+			this.ProjectDescriptor = InParams.ProjectDescriptor;
 		}
 
 		/// <summary>
@@ -620,6 +621,12 @@ namespace AutomationTool
 			//
 
 			this.RawProjectPath = RawProjectPath;
+			try
+			{
+				this.ProjectDescriptor = ProjectDescriptor.FromFile(RawProjectPath);
+			}
+			catch { this.ProjectDescriptor = new ProjectDescriptor(); }
+
 			if (DirectoriesToCook != null)
 			{
 				this.DirectoriesToCook = DirectoriesToCook;
@@ -711,7 +718,8 @@ namespace AutomationTool
 				}
 				else
 				{
-					List<PluginInfo> CandidatePlugins = Plugins.ReadAvailablePlugins(Unreal.EngineDirectory, DirectoryReference.FromFile(RawProjectPath), null);
+					List<PluginInfo> CandidatePlugins = Plugins.ReadAvailablePlugins(Unreal.EngineDirectory,
+						DirectoryReference.FromFile(RawProjectPath), AdditionalPluginDirectories);
 					PluginInfo DLCPlugin = CandidatePlugins.FirstOrDefault(x => String.Equals(x.Name, DLCName, StringComparison.InvariantCultureIgnoreCase));
 					if (DLCPlugin == null)
 					{
@@ -2323,6 +2331,7 @@ namespace AutomationTool
 		private List<SingleTargetProperties> DetectedTargets;
 		private Dictionary<UnrealTargetPlatform, ConfigHierarchy> LoadedEngineConfigs;
 		private Dictionary<UnrealTargetPlatform, ConfigHierarchy> LoadedGameConfigs;
+		private ProjectDescriptor ProjectDescriptor;
 
 		private List<String> TargetNamesOfType(TargetType DesiredType)
 		{
@@ -2739,7 +2748,53 @@ namespace AutomationTool
 			get { return ProjectUtils.GetShortProjectName(RawProjectPath); }
 		}
 
-  		/// <summary>
+		/// <summary>
+		/// AdditionalPluginDirectories from the project.uproject file
+		/// </summary>
+		public List<DirectoryReference> AdditionalPluginDirectories
+		{
+			get { return ProjectDescriptor.AdditionalPluginDirectories; }
+		}
+
+		/// <summary>
+		/// Get the relative path to the DLC plugin's cooked output from the deployment
+		/// root of the DLC. e.g. <ProjectName>\Plugins\<PluginName> for plugins under the Project's plugin
+		/// directories.
+		/// </summary>
+		public string FindPluginRelativePathFromPlatformCookDir(FileReference PluginFile,
+			DirectoryReference ProjectRoot, DirectoryReference EngineRoot, DirectoryReference LocalRoot, string ShortProjectName)
+		{
+			if (DLCOverrideCookedSubDir != null)
+			{
+				return DLCOverrideCookedSubDir;
+			}
+
+			foreach (DirectoryReference AdditionalPluginDir in AdditionalPluginDirectories)
+			{
+				if (PluginFile.IsUnderDirectory(AdditionalPluginDir))
+				{
+					// This is a plugin that lives outside of the Engine/Plugins or Game/Plugins directory so needs to be remapped for staging/packaging
+					// The deployment path for plugins in AdditionalPluginDirectories is RemappedPlugins\PluginName
+					return String.Format("RemappedPlugins/{0}", PluginFile.GetFileNameWithoutExtension());
+				}
+			}
+
+			DirectoryReference DLCRoot = PluginFile.Directory;
+			if (DLCRoot.IsUnderDirectory(EngineRoot))
+			{
+				return Path.Combine("Engine", DLCRoot.MakeRelativeTo(EngineRoot));
+			}
+			else if (DLCRoot.IsUnderDirectory(ProjectRoot))
+			{
+				return Path.Combine(ShortProjectName, DLCRoot.MakeRelativeTo(ProjectRoot));
+			}
+			else
+			{
+				return DLCRoot.MakeRelativeTo(LocalRoot);
+			}
+		}
+
+		/// <summary>
 		/// True if this project contains source code.
 		/// </summary>	
 		public bool IsCodeBasedProject
diff --git a/Engine/Source/Programs/AutomationTool/Scripts/CopyBuildToStagingDirectory.Automation.cs b/Engine/Source/Programs/AutomationTool/Scripts/CopyBuildToStagingDirectory.Automation.cs
--- a/Engine/Source/Programs/AutomationTool/Scripts/CopyBuildToStagingDirectory.Automation.cs
+++ b/Engine/Source/Programs/AutomationTool/Scripts/CopyBuildToStagingDirectory.Automation.cs
@@ -689,10 +689,12 @@ namespace AutomationScripts
 		/// <param name="ProjectFile">The project being built</param>
 		/// <param name="Targets">List of targets being staged</param>
 		/// <returns>List of plugin files that should be staged</returns>
-		private static List<FileReference> GetPluginsForContentProject(FileReference ProjectFile, List<TargetReceipt> Targets)
+		private static List<FileReference> GetPluginsForContentProject(FileReference ProjectFile,
+			List<TargetReceipt> Targets, List<DirectoryReference> AdditionalPluginDirectories)
 		{
 			ProjectDescriptor Project = ProjectDescriptor.FromFile(ProjectFile);
-			List<PluginInfo> AvailablePlugins = UnrealBuildTool.Plugins.ReadAvailablePlugins(Unreal.EngineDirectory, ProjectFile.Directory, null);
+			List<PluginInfo> AvailablePlugins = UnrealBuildTool.Plugins.ReadAvailablePlugins(Unreal.EngineDirectory,
+				ProjectFile.Directory, AdditionalPluginDirectories);
 
 			HashSet<FileReference> Plugins = new HashSet<FileReference>();
 			foreach (TargetReceipt Target in Targets)
@@ -882,29 +884,13 @@ namespace AutomationScripts
 			}
 			if (Params.HasDLCName)
 			{
-				// Making a plugin
-				DirectoryReference DLCRoot = Params.DLCFile.Directory;
-				string DLCCookedSubDir;
-				if (Params.DLCOverrideCookedSubDir != null)
-				{
-					DLCCookedSubDir = Params.DLCOverrideCookedSubDir;
-				}
-				else if (DLCRoot.IsUnderDirectory(SC.EngineRoot))
-				{
-					DLCCookedSubDir = Path.Combine("Engine", DLCRoot.MakeRelativeTo(SC.EngineRoot));
-				}
-				else if (DLCRoot.IsUnderDirectory(SC.ProjectRoot))
-				{
-					DLCCookedSubDir = Path.Combine(SC.ShortProjectName, DLCRoot.MakeRelativeTo(SC.ProjectRoot));
-				}
-				else
-				{
-					DLCCookedSubDir = DLCRoot.MakeRelativeTo(SC.LocalRoot);
-				}
+				// We are making a plugin
+				string DLCCookedSubDir = Params.FindPluginRelativePathFromPlatformCookDir(Params.DLCFile, SC.ProjectRoot, SC.EngineRoot, SC.LocalRoot, SC.ShortProjectName);
+
 				// Put all of the cooked dir into the staged dir
 				if (String.IsNullOrEmpty(Params.CookOutputDir))
 				{
-					SC.PlatformCookDir = DirectoryReference.Combine(DLCRoot, "Saved", "Cooked", SC.CookPlatform);
+					SC.PlatformCookDir = DirectoryReference.Combine(Params.DLCFile.Directory, "Saved", "Cooked", SC.CookPlatform);
 				}
 				else
 				{
@@ -1142,7 +1128,8 @@ namespace AutomationScripts
 					// Stage any content-only plugins for content-only projects. We don't have a custom executable for these.
 					if (!Params.IsCodeBasedProject)
 					{
-						List<FileReference> PluginFiles = GetPluginsForContentProject(Params.RawProjectPath, SC.StageTargets.ConvertAll(x => x.Receipt));
+						List<FileReference> PluginFiles = GetPluginsForContentProject(Params.RawProjectPath,
+							SC.StageTargets.ConvertAll(x => x.Receipt), Params.AdditionalPluginDirectories);
 						SC.StageFiles(StagedFileType.UFS, PluginFiles);
 					}
 				}
@@ -1234,7 +1221,8 @@ namespace AutomationScripts
 					List<PluginInfo> AvailablePlugins = new List<PluginInfo>();
 					if (StagedPlugins.Count > 0 && PlatformExtensionsToStage.Count > 0)
 					{
-						AvailablePlugins = UnrealBuildTool.Plugins.ReadAvailablePlugins(Unreal.EngineDirectory, Params.RawProjectPath.Directory, null);
+						AvailablePlugins = UnrealBuildTool.Plugins.ReadAvailablePlugins(Unreal.EngineDirectory,
+							Params.RawProjectPath.Directory, Params.AdditionalPluginDirectories);
 					}
 					foreach (KeyValuePair<StagedFileReference, FileReference> StagedPlugin in StagedPlugins)
 					{
@@ -5044,7 +5032,8 @@ namespace AutomationScripts
 					Params.Client,
 					Params.Manifests,
 					Params.SeparateDebugInfo,
-					Params.HasDLCName ? Params.DLCFile.Directory : null
+					Params.HasDLCName ? Params.DLCFile.Directory : null,
+					Params.AdditionalPluginDirectories
 					);
 				LogDeploymentContext(SC);
 
