Fix camera speed slider jittering

diff --git a/Engine/Source/Editor/UnrealEd/Private/SEditorViewportToolBarMenu.cpp b/Engine/Source/Editor/UnrealEd/Private/SEditorViewportToolBarMenu.cpp
--- a/Engine/Source/Editor/UnrealEd/Private/SEditorViewportToolBarMenu.cpp
+++ b/Engine/Source/Editor/UnrealEd/Private/SEditorViewportToolBarMenu.cpp
@@ -95,9 +95,9 @@ void SEditorViewportToolbarMenu::Construct( const FArguments& Declaration )
 
 	ChildSlot
 	[
-		SAssignNew( MenuAnchor, SMenuAnchor )
-		.Placement( MenuPlacement_BelowAnchor )
-		.OnGetMenuContent( Declaration._OnGetMenuContent )
+		SAssignNew(MenuAnchor, SMenuAnchor)
+		.Placement(Declaration._MenuPlacement)
+		.OnGetMenuContent(Declaration._OnGetMenuContent)
 	];
 
 	MenuAnchor->SetContent(
diff --git a/Engine/Source/Editor/UnrealEd/Private/STransformViewportToolbar.cpp b/Engine/Source/Editor/UnrealEd/Private/STransformViewportToolbar.cpp
--- a/Engine/Source/Editor/UnrealEd/Private/STransformViewportToolbar.cpp
+++ b/Engine/Source/Editor/UnrealEd/Private/STransformViewportToolbar.cpp
@@ -422,6 +422,8 @@ TSharedRef< SWidget > STransformViewportToolBar::MakeTransformToolBar( const TSh
 				.ToolTipText(LOCTEXT("CameraSpeed_ToolTip", "Camera Speed"))
 				.LabelIcon(FAppStyle::Get().GetBrush("EditorViewport.CamSpeedSetting"))
 				.Label(this, &STransformViewportToolBar::GetCameraSpeedLabel)
+				// Anchor to the right, otherwise the slider in this menu will jitter when the label width changes
+				.MenuPlacement(MenuPlacement_BelowRightAnchor)
 				.OnGetMenuContent(this, &STransformViewportToolBar::FillCameraSpeedMenu),
 			CameraSpeedName,
 			false,
@@ -471,17 +473,25 @@ TSharedRef<SWidget> STransformViewportToolBar::FillCameraSpeedMenu()
 			.FillWidth(1)
 			.Padding( FMargin(0.0f, 2.0f) )
 			[
-				SAssignNew(CamSpeedSlider, SSlider)
-				.Value(this, &STransformViewportToolBar::GetCamSpeedSliderPosition)
-				.OnValueChanged(this, &STransformViewportToolBar::OnSetCamSpeed)
+				SNew( SBox )
+				.MinDesiredWidth(220)
+				[
+					SAssignNew(CamSpeedSlider, SSlider)
+					.Value(this, &STransformViewportToolBar::GetCamSpeedSliderPosition)
+					.OnValueChanged(this, &STransformViewportToolBar::OnSetCamSpeed)
+				]
 			]
 			+SHorizontalBox::Slot()
 			.AutoWidth()
 			.Padding( 8.0f, 2.0f, 0.0f, 2.0f)
 			[
-				SNew( STextBlock )
-				.Text(this, &STransformViewportToolBar::GetCameraSpeedLabel )
-				.Font( FAppStyle::GetFontStyle( TEXT( "MenuItem.Font" ) ) )
+				SNew( SBox )
+				.WidthOverride(40)
+				[
+					SNew( STextBlock )
+					.Text(this, &STransformViewportToolBar::GetCameraSpeedLabel )
+					.Font( FAppStyle::GetFontStyle( TEXT( "MenuItem.Font" ) ) )
+				]
 			]
 		] // Camera Speed Scalar
 		+ SVerticalBox::Slot()
diff --git a/Engine/Source/Editor/UnrealEd/Public/SEditorViewportToolBarMenu.h b/Engine/Source/Editor/UnrealEd/Public/SEditorViewportToolBarMenu.h
--- a/Engine/Source/Editor/UnrealEd/Public/SEditorViewportToolBarMenu.h
+++ b/Engine/Source/Editor/UnrealEd/Public/SEditorViewportToolBarMenu.h
@@ -60,6 +60,8 @@ public:
 		SLATE_EVENT(FOnGetContent, OnGetMenuContent)
 		/** The foreground color of the content */
 		SLATE_ATTRIBUTE(FSlateColor, ForegroundColor)
+		/** Where the menu is spawned relative to the menu button */
+		SLATE_ARGUMENT_DEFAULT(EMenuPlacement, MenuPlacement) = EMenuPlacement::MenuPlacement_BelowAnchor;
 	SLATE_END_ARGS()
 
 	/**
