Add "Browse to Asset" context action to the global asset picker

diff --git a/Engine/Source/Editor/UnrealEd/Private/Toolkits/SGlobalOpenAssetDialog.cpp b/Engine/Source/Editor/UnrealEd/Private/Toolkits/SGlobalOpenAssetDialog.cpp
--- a/Engine/Source/Editor/UnrealEd/Private/Toolkits/SGlobalOpenAssetDialog.cpp
+++ b/Engine/Source/Editor/UnrealEd/Private/Toolkits/SGlobalOpenAssetDialog.cpp
@@ -12,20 +12,29 @@
 #include "ContentBrowserModule.h"
 
 #include "Subsystems/AssetEditorSubsystem.h"
+#include "Toolkits/GlobalEditorCommonCommands.h"
+#include "ToolMenus.h"
 #include "Editor.h"
 
 #define LOCTEXT_NAMESPACE "SGlobalOpenAssetDialog"
 
+namespace GlobalOpenAssetDialogUtils
+{
+	static const FName ContextMenuName("GlobalOpenAssetDialog.ContextMenu");
+}
+
 //////////////////////////////////////////////////////////////////////////
 // SGlobalOpenAssetDialog
 
 void SGlobalOpenAssetDialog::Construct(const FArguments& InArgs, FVector2D InSize)
 {
-	FContentBrowserModule& ContentBrowserModule = FModuleManager::Get().LoadModuleChecked<FContentBrowserModule>(TEXT("ContentBrowser"));
+	IContentBrowserSingleton& ContentBrowser = GetContentBrowser();
 
 	FAssetPickerConfig AssetPickerConfig;
 	AssetPickerConfig.OnAssetDoubleClicked = FOnAssetSelected::CreateSP(this, &SGlobalOpenAssetDialog::OnAssetSelectedFromPicker);
 	AssetPickerConfig.OnAssetEnterPressed = FOnAssetEnterPressed::CreateSP(this, &SGlobalOpenAssetDialog::OnPressedEnterOnAssetsInPicker);
+	AssetPickerConfig.OnGetAssetContextMenu = FOnGetAssetContextMenu::CreateSP(this, &SGlobalOpenAssetDialog::OnGetAssetContextMenu);
+	AssetPickerConfig.GetCurrentSelectionDelegates.Add(&GetCurrentSelectionDelegate);
 	AssetPickerConfig.InitialAssetViewType = EAssetViewType::List;
 	AssetPickerConfig.bAllowNullSelection = false;
 	AssetPickerConfig.bShowBottomToolbar = true;
@@ -34,6 +43,19 @@ void SGlobalOpenAssetDialog::Construct(const FArguments& InArgs, FVector2D InSiz
 	AssetPickerConfig.bAddFilterUI = true;
 	AssetPickerConfig.SaveSettingsName = TEXT("GlobalAssetPicker");
 
+	if (!UToolMenus::Get()->IsMenuRegistered(GlobalOpenAssetDialogUtils::ContextMenuName))
+	{
+		UToolMenu* Menu = UToolMenus::Get()->RegisterMenu(GlobalOpenAssetDialogUtils::ContextMenuName);
+		FToolMenuSection& AssetSection = Menu->AddSection("Asset", LOCTEXT("AssetSectionLabel", "Asset"));
+		AssetSection.AddMenuEntry(FGlobalEditorCommonCommands::Get().FindInContentBrowser);
+	}
+
+	Commands = MakeShared<FUICommandList>();
+	Commands->MapAction(FGlobalEditorCommonCommands::Get().FindInContentBrowser, FUIAction(
+		FExecuteAction::CreateSP(this, &SGlobalOpenAssetDialog::FindInContentBrowser),
+		FCanExecuteAction::CreateSP(this, &SGlobalOpenAssetDialog::AreAnyAssetsSelected)
+	));
+
 	ChildSlot
 	[
 		SNew(SBox)
@@ -44,7 +66,7 @@ void SGlobalOpenAssetDialog::Construct(const FArguments& InArgs, FVector2D InSiz
 			+ SVerticalBox::Slot()
 			.FillHeight(1.0f)
 			[
-				ContentBrowserModule.Get().CreateAssetPicker(AssetPickerConfig)
+				ContentBrowser.CreateAssetPicker(AssetPickerConfig)
 			]
 		]
 	];
@@ -62,6 +84,16 @@ FReply SGlobalOpenAssetDialog::OnPreviewKeyDown(const FGeometry& MyGeometry, con
 	return FReply::Unhandled();
 }
 
+FReply SGlobalOpenAssetDialog::OnKeyDown(const FGeometry& MyGeometry, const FKeyEvent& InKeyEvent)
+{
+	if (Commands->ProcessCommandBindings(InKeyEvent))
+	{
+		return FReply::Handled();
+	}
+
+	return SCompoundWidget::OnKeyDown(MyGeometry, InKeyEvent);
+}
+
 void SGlobalOpenAssetDialog::OnAssetSelectedFromPicker(const FAssetData& AssetData)
 {
 	if (UObject* ObjectToEdit = AssetData.GetAsset())
@@ -93,6 +125,32 @@ void SGlobalOpenAssetDialog::RequestCloseAssetPicker()
 	}
 }
 
+TSharedPtr<SWidget> SGlobalOpenAssetDialog::OnGetAssetContextMenu(const TArray<FAssetData>& SelectedAssets) const
+{
+	if (SelectedAssets.IsEmpty())
+	{
+		return nullptr;
+	}
+
+	return UToolMenus::Get()->GenerateWidget(GlobalOpenAssetDialogUtils::ContextMenuName, FToolMenuContext(Commands));
+}
+
+IContentBrowserSingleton& SGlobalOpenAssetDialog::GetContentBrowser() const
+{
+	return FModuleManager::Get().LoadModuleChecked<FContentBrowserModule>(TEXT("ContentBrowser")).Get();
+}
+
+void SGlobalOpenAssetDialog::FindInContentBrowser()
+{
+	GetContentBrowser().SyncBrowserToAssets(GetCurrentSelectionDelegate.Execute());
+	RequestCloseAssetPicker();
+}
+
+bool SGlobalOpenAssetDialog::AreAnyAssetsSelected() const
+{
+	return !GetCurrentSelectionDelegate.Execute().IsEmpty();
+}
+
 
 //////////////////////////////////////////////////////////////////////////
 
diff --git a/Engine/Source/Editor/UnrealEd/Private/Toolkits/SGlobalOpenAssetDialog.h b/Engine/Source/Editor/UnrealEd/Private/Toolkits/SGlobalOpenAssetDialog.h
--- a/Engine/Source/Editor/UnrealEd/Private/Toolkits/SGlobalOpenAssetDialog.h
+++ b/Engine/Source/Editor/UnrealEd/Private/Toolkits/SGlobalOpenAssetDialog.h
@@ -3,6 +3,7 @@
 #pragma once
 
 #include "CoreMinimal.h"
+#include "ContentBrowserDelegates.h"
 #include "Input/Reply.h"
 #include "Widgets/DeclarativeSyntaxSupport.h"
 #include "Widgets/SCompoundWidget.h"
@@ -21,10 +22,20 @@ public:
 
 	// SWidget interface
 	virtual FReply OnPreviewKeyDown(const FGeometry& MyGeometry, const FKeyEvent& InKeyEvent) override;
+	virtual FReply OnKeyDown(const FGeometry& MyGeometry, const FKeyEvent& InKeyEvent) override;
 	// End of SWidget interface
 
 protected:
 	void OnAssetSelectedFromPicker(const struct FAssetData& AssetData);
 	void OnPressedEnterOnAssetsInPicker(const TArray<struct FAssetData>& SelectedAssets);
 	void RequestCloseAssetPicker();
+
+	TSharedPtr<SWidget> OnGetAssetContextMenu(const TArray<FAssetData>& SelectedAssets) const;
+	class IContentBrowserSingleton& GetContentBrowser() const;
+
+	void FindInContentBrowser();
+	bool AreAnyAssetsSelected() const;
+	
+	FGetCurrentSelectionDelegate GetCurrentSelectionDelegate;
+	TSharedPtr<FUICommandList> Commands;
 };
