Only prompt to recompile Blueprints when starting PIE when Blueprints are actually dirty

diff --git a/Engine/Source/Editor/UnrealEd/Private/PlayLevel.cpp b/Engine/Source/Editor/UnrealEd/Private/PlayLevel.cpp
--- a/Engine/Source/Editor/UnrealEd/Private/PlayLevel.cpp
+++ b/Engine/Source/Editor/UnrealEd/Private/PlayLevel.cpp
@@ -1252,7 +1252,7 @@ int32 FInternalPlayLevelUtils::ResolveDirtyBlueprints(const bool bPromptForCompi
 	}
 
 	bool bRunCompilation = bAutoCompile;
-	if (bPromptForCompile)
+	if (bPromptForCompile && (InNeedOfRecompile.Num() > 0))
 	{
 		FFormatNamedArguments Args;
 		Args.Add(TEXT("DirtyBlueprints"), FText::FromString(PromptDirtyList));
