Avoid partitioning foliage and landscape when using OnlyMergeSubLevels.

diff --git a/Engine/Source/Editor/UnrealEd/Private/Commandlets/WorldPartitionConvertCommandlet.cpp b/Engine/Source/Editor/UnrealEd/Private/Commandlets/WorldPartitionConvertCommandlet.cpp
--- a/Engine/Source/Editor/UnrealEd/Private/Commandlets/WorldPartitionConvertCommandlet.cpp
+++ b/Engine/Source/Editor/UnrealEd/Private/Commandlets/WorldPartitionConvertCommandlet.cpp
@@ -1146,27 +1146,30 @@ int32 UWorldPartitionConvertCommandlet::Main(const FString& Params)
 			}
 		}
 
-		// do loop after as it may modify Level->Actors
-		if (IFAs.Num())
+		if (!bOnlyMergeSubLevels)
 		{
-			UE_SCOPED_TIMER(TEXT("PartitionFoliage"), LogWorldPartitionConvertCommandlet, Display);
-
-			for (AInstancedFoliageActor* IFA : IFAs)
+			// do loop after as it may modify Level->Actors
+			if (IFAs.Num())
 			{
-				if (!PartitionFoliage(IFA))
+				UE_SCOPED_TIMER(TEXT("PartitionFoliage"), LogWorldPartitionConvertCommandlet, Display);
+
+				for (AInstancedFoliageActor* IFA : IFAs)
 				{
-					return false;
+					if (!PartitionFoliage(IFA))
+					{
+						return false;
+					}
 				}
 			}
-		}
-
-		if (LandscapeInfos.Num())
-		{
-			UE_SCOPED_TIMER(TEXT("PartitionLandscape"), LogWorldPartitionConvertCommandlet, Display);
 
-			for (ULandscapeInfo* LandscapeInfo : LandscapeInfos)
+			if (LandscapeInfos.Num())
 			{
-				PartitionLandscape(LandscapeInfo);
+				UE_SCOPED_TIMER(TEXT("PartitionLandscape"), LogWorldPartitionConvertCommandlet, Display);
+
+				for (ULandscapeInfo* LandscapeInfo : LandscapeInfos)
+				{
+					PartitionLandscape(LandscapeInfo);
+				}
 			}
 		}
 
