Fix dragging of physics objects in PIE and paused PIE

diff --git a/Engine/Source/Editor/UnrealEd/Private/PhysicsManipulationMode.cpp b/Engine/Source/Editor/UnrealEd/Private/PhysicsManipulationMode.cpp
--- a/Engine/Source/Editor/UnrealEd/Private/PhysicsManipulationMode.cpp
+++ b/Engine/Source/Editor/UnrealEd/Private/PhysicsManipulationMode.cpp
@@ -10,6 +10,11 @@
 
 DEFINE_LOG_CATEGORY_STATIC(LogEditorPhysMode, Log, All);
 
+// Disable drag handles for PIE manipulation and just teleport (handy for some internal physics debugging modes)
+bool bDisableEditorPhysicsHandle = false;
+FAutoConsoleVariableRef CVarHackMaxVelocity(TEXT("p.DisableEditorPhysicsHandle"), bDisableEditorPhysicsHandle, TEXT("When true, disable the physics spring for dragging objects in PIE. Use a teleport instead."));
+
+
 void FPhysicsManipulationEdModeFactory::OnSelectionChanged(FEditorModeTools& Tools, UObject* ItemUndergoingChange) const
 {
 	USelection* Selection = GEditor->GetSelectedActors();
@@ -81,6 +86,15 @@ bool FPhysicsManipulationEdMode::InputDelta( FEditorViewportClient* InViewportCl
 		HandleComp->SetTargetLocation(HandleTargetLocation);
 		HandleComp->SetTargetRotation(HandleTargetRotation);
 
+		// We cannot use a spring to move the object when the world is paused. Teleport instead
+		const bool bTeleportObject = (GetWorld() && GetWorld()->IsPaused()) || bDisableEditorPhysicsHandle;
+		if (bTeleportObject)
+		{
+			HandleComp->GrabbedComponent->SetWorldLocationAndRotation(HandleTargetLocation, HandleTargetRotation, false, nullptr, ETeleportType::TeleportPhysics);
+			HandleComp->GrabbedComponent->SetPhysicsLinearVelocity(FVector::Zero());
+			HandleComp->GrabbedComponent->SetPhysicsAngularVelocityInRadians(FVector::Zero());
+		}
+
 		return true;
 	}
 	else
