Remove ULightWeightInstanceFactory::MakeNewLightWeightInstanceManager
Never called, was creating twice the blueprint, and was always returning null (invalid cast from UBlueprint to ALightWeightInstanceManager)
#rnx

diff --git a/Engine/Source/Editor/UnrealEd/Classes/Factories/LightWeightInstanceFactory.h b/Engine/Source/Editor/UnrealEd/Classes/Factories/LightWeightInstanceFactory.h
--- a/Engine/Source/Editor/UnrealEd/Classes/Factories/LightWeightInstanceFactory.h
+++ b/Engine/Source/Editor/UnrealEd/Classes/Factories/LightWeightInstanceFactory.h
@@ -20,7 +20,6 @@ class ULightWeightInstanceFactory : public UFactory
 	//~ Begin UFactory Interface
 
 protected:
-	UNREALED_API virtual ALightWeightInstanceManager* MakeNewLightWeightInstanceManager(UObject* InParent, FName Name, EObjectFlags Flags);
 
 	// The parent class of the created blueprint
 	UPROPERTY()
diff --git a/Engine/Source/Editor/UnrealEd/Private/Factories/LightWeightInstanceFactory.cpp b/Engine/Source/Editor/UnrealEd/Private/Factories/LightWeightInstanceFactory.cpp
--- a/Engine/Source/Editor/UnrealEd/Private/Factories/LightWeightInstanceFactory.cpp
+++ b/Engine/Source/Editor/UnrealEd/Private/Factories/LightWeightInstanceFactory.cpp
@@ -103,19 +103,4 @@ UObject* ULightWeightInstanceFactory::FactoryCreateNew(UClass* Class, UObject* I
 	return LWIManager;
 }
 
-ALightWeightInstanceManager* ULightWeightInstanceFactory::MakeNewLightWeightInstanceManager(UObject* InParent, FName Name, EObjectFlags Flags)
-{
-	UClass* BlueprintClass = nullptr;
-	UClass* BlueprintGeneratedClass = nullptr;
-
-	IKismetCompilerInterface& KismetCompilerModule = FModuleManager::LoadModuleChecked<IKismetCompilerInterface>("KismetCompiler");
-	KismetCompilerModule.GetBlueprintTypesForClass(ParentClass, BlueprintClass, BlueprintGeneratedClass);
-
-	UBlueprint* NewBP = FKismetEditorUtilities::CreateBlueprint(ParentClass, InParent, Name, BPTYPE_Normal, BlueprintClass, BlueprintGeneratedClass);
-
-	return Cast<ALightWeightInstanceManager>(FKismetEditorUtilities::CreateBlueprint(ParentClass, InParent, Name, BPTYPE_Normal, BlueprintClass, BlueprintGeneratedClass));
-//	return FKismetEditorUtilities::CreateBlueprint(ParentClass, InParent, Name, BPTYPE_Normal, UEditorUtilityBlueprint::StaticClass(), UBlueprintGeneratedClass::StaticClass());
-//	return NewObject<ALightWeightInstanceManager>(InParent, Name, Flags);
-}
-
 #undef LOCTEXT_NAMESPACE // "LightWeightInstanceFactory"
