Cooker: DiffOnly: Fix bug from refactor: header diffs are not being printed.
The bug was caused by incorrect handling of FCallstackAtOffset::bIgnore; that variable indicates logging should be suppressed, it does not indicate that the recording of the diff should be suppressed. The recording of the diff in the bIgnore scope needs to be kept for stats and control flow.
Fixed the usage to keep the recording of the diff, and renamed the variable for clarity.
Also fixed spurious logs being printed in the zenstore case when the header differences are suppressed.
#rnx

diff --git a/Engine/Source/Editor/UnrealEd/Private/Cooker/DiffPackageWriter.cpp b/Engine/Source/Editor/UnrealEd/Private/Cooker/DiffPackageWriter.cpp
--- a/Engine/Source/Editor/UnrealEd/Private/Cooker/DiffPackageWriter.cpp
+++ b/Engine/Source/Editor/UnrealEd/Private/Cooker/DiffPackageWriter.cpp
@@ -222,7 +222,7 @@ UE::DiffWriter::FAccumulator& FDiffPackageWriter::ConstructAccumulator(FName Pac
 	{
 		check(!bHasStartedSecondSave); // Accumulator should already exist from CreateLinkerArchive in the first save
 		Accumulator = new UE::DiffWriter::FAccumulator(Asset, *PackageName.ToString(), MaxDiffsToLog,
-			GetDiffWriterMessageCallback(), Inner->GetCookCapabilities().HeaderFormat);
+			bIgnoreHeaderDiffs, GetDiffWriterMessageCallback(), Inner->GetCookCapabilities().HeaderFormat);
 	}
 	return *Accumulator;
 }
diff --git a/Engine/Source/Editor/UnrealEd/Private/Cooker/DiffWriterArchive.cpp b/Engine/Source/Editor/UnrealEd/Private/Cooker/DiffWriterArchive.cpp
--- a/Engine/Source/Editor/UnrealEd/Private/Cooker/DiffWriterArchive.cpp
+++ b/Engine/Source/Editor/UnrealEd/Private/Cooker/DiffWriterArchive.cpp
@@ -236,11 +236,12 @@ void FCallstacks::Add(
 	if (LastSerializeCallstack == nullptr || (bCallstacksDirty && FCStringAnsi::Strcmp(LastSerializeCallstack, StackTrace.Get()) != 0))
 	{
 		uint32 CallstackCRC = 0;
+		bool bSuppressLogging = UE::ArchiveStackTrace::ShouldIgnoreDiff();
 		if (CallstackAtOffsetMap.Num() == 0 || CurrentOffset > CallstackAtOffsetMap.Last().Offset)
 		{
 			// New data serialized at the end of archive buffer
 			LastSerializeCallstack = AddUniqueCallstack(bIsCollectingCallstacks, SerializedObject, SerializedProperty, CallstackCRC);
-			CallstackAtOffsetMap.Add(FCallstackAtOffset { CurrentOffset, CallstackCRC, UE::ArchiveStackTrace::ShouldIgnoreDiff()});
+			CallstackAtOffsetMap.Add(FCallstackAtOffset { CurrentOffset, CallstackCRC, bSuppressLogging });
 		}
 		else
 		{
@@ -257,7 +258,7 @@ void FCallstacks::Add(
 			{
 				// Insert a new callstack
 				check(CallstackToUpdate.Offset < CurrentOffset);
-				CallstackAtOffsetMap.Insert(FCallstackAtOffset {CurrentOffset, CallstackCRC, UE::ArchiveStackTrace::ShouldIgnoreDiff()}, CallstackToUpdateIndex + 1);
+				CallstackAtOffsetMap.Insert(FCallstackAtOffset {CurrentOffset, CallstackCRC, bSuppressLogging }, CallstackToUpdateIndex + 1);
 			}
 		}
 		check(CallstackCRC != 0 || !bShouldCollectCallstack);
@@ -483,7 +484,7 @@ void FCallstacks::RecordSerialize(EOffsetFrame OffsetFrame, int64 CurrentOffset,
 	}
 }
 
-FAccumulator::FAccumulator(UObject* InAsset, FName InPackageName, int32 InMaxDiffsToLog,
+FAccumulator::FAccumulator(UObject* InAsset, FName InPackageName, int32 InMaxDiffsToLog, bool bInIgnoreHeaderDiffs,
 	FMessageCallback&& InMessageCallback, EPackageHeaderFormat InPackageHeaderFormat)
 	: LinkerCallstacks()
 	, ExportsCallstacks()
@@ -492,6 +493,7 @@ FAccumulator::FAccumulator(UObject* InAsset, FName InPackageName, int32 InMaxDif
 	, Asset(InAsset)
 	, MaxDiffsToLog(InMaxDiffsToLog)
 	, PackageHeaderFormat(InPackageHeaderFormat)
+	, bIgnoreHeaderDiffs(bInIgnoreHeaderDiffs)
 {
 	GBreakAtOffsetSettings.Initialize();
 }
@@ -712,6 +714,7 @@ void FAccumulator::CompareWithPreviousForSection(const FPackageData& SourcePacka
 	FString LastDifferenceCallstackDataText;
 	int32 LastDifferenceCallstackOffsetIndex = -1;
 	int64 NumDiffsLocal = 0;
+	int64 NumDiffsForLogStatLocal = 0;
 	int64 NumDiffsLoggedLocal = 0;
 	int64 FirstUnreportedDiffIndex = -1;
 
@@ -733,39 +736,55 @@ void FAccumulator::CompareWithPreviousForSection(const FPackageData& SourcePacka
 
 		// Skip reporting the difference if we are still within the same Serialize call, or for any bytes after
 		// the first different byte if we don't know the callstack.
-		if (DifferenceCallstackOffsetIndex == LastDifferenceCallstackOffsetIndex)
+		if (NumDiffsLocal > 0 &&
+			DifferenceCallstackOffsetIndex == LastDifferenceCallstackOffsetIndex)
 		{
 			continue;
 		}
-		LastDifferenceCallstackOffsetIndex = DifferenceCallstackOffsetIndex;
 
-		// Also skip reporting the difference if it is another occurrence of the last reported callstack,
-		// or if the callstack is marked as ignore
+		// Also skip reporting the difference if it is another occurrence of the last reported callstack
 		const FCallstacks::FCallstackAtOffset* CallstackAtOffsetPtr = nullptr;
 		const FCallstacks::FCallstackData* DifferenceCallstackDataPtr = nullptr;
+		FString DifferenceCallstackDataText;
+		bool bCallstackSuppressLogging = false;
 		if (DifferenceCallstackOffsetIndex >= 0)
 		{
 			CallstackAtOffsetPtr = &Callstacks.GetCallstack(DifferenceCallstackOffsetIndex);
-			if (CallstackAtOffsetPtr->bIgnore)
-			{
-				continue;
-			}
+			bCallstackSuppressLogging = CallstackAtOffsetPtr->bSuppressLogging;
 
 			DifferenceCallstackDataPtr = &Callstacks.GetCallstackData(*CallstackAtOffsetPtr);
-			FString CurrentDifferenceCallstackDataText = DifferenceCallstackDataPtr->ToString(CallstackCutoffText);
-			if (LastDifferenceCallstackDataText.Compare(CurrentDifferenceCallstackDataText, ESearchCase::CaseSensitive) == 0)
+			DifferenceCallstackDataText = DifferenceCallstackDataPtr->ToString(CallstackCutoffText);
+			if (NumDiffsLocal > 0 && 
+				LastDifferenceCallstackDataText.Compare(DifferenceCallstackDataText, ESearchCase::CaseSensitive) == 0)
 			{
 				continue;
 			}
-			LastDifferenceCallstackDataText = MoveTemp(CurrentDifferenceCallstackDataText);
 		}
 
 		// Update counter for number of existing diffs
 		OutStats.FindOrAdd(AssetClass).NumDiffs++;
 		NumDiffsLocal++;
+		// Update LastReported fields
+		LastDifferenceCallstackOffsetIndex = DifferenceCallstackOffsetIndex;
+		LastDifferenceCallstackDataText = MoveTemp(DifferenceCallstackDataText);
+
+		// Skip logging of the difference if the callstack has a suppresslogging scope 
+		if (bCallstackSuppressLogging)
+		{
+			return;
+		}
+		// Skip logging of header differences if requested
+		if (bIgnoreHeaderDiffs && DestAbsoluteOffset < HeaderSize)
+		{
+			continue;
+		}
+
+		// Update counter for number of diffs that should be reported as existing when over the limit
+		// Ignored header diffs and suppressed callstacks do not contribute to this count
+		NumDiffsForLogStatLocal++;
 
-		// Skip reporting the difference if we are over the limit
-		if (MaxDiffsToLog >= 0 && InOutDiffsLogged >= MaxDiffsToLog)
+		// Skip logging of the difference if we are over the limit
+		if (bCallstackSuppressLogging || (MaxDiffsToLog >= 0 && InOutDiffsLogged >= MaxDiffsToLog))
 		{
 			if (FirstUnreportedDiffIndex == -1)
 			{
@@ -802,7 +821,6 @@ void FAccumulator::CompareWithPreviousForSection(const FPackageData& SourcePacka
 			check(CallstackAtOffsetPtr && DifferenceCallstackDataPtr); // These were set up above.
 			const FCallstacks::FCallstackAtOffset& CallstackAtOffset = *CallstackAtOffsetPtr;
 			const FCallstacks::FCallstackData& DifferenceCallstackData = *DifferenceCallstackDataPtr;
-			const FString& DifferenceCallstackDataText = LastDifferenceCallstackDataText;
 
 			FString BeforePropertyVal;
 			FString AfterPropertyVal;
@@ -874,7 +892,7 @@ void FAccumulator::CompareWithPreviousForSection(const FPackageData& SourcePacka
 					TEXT("Difference occurs at index %lld within Serialize call at callstack:%s%s%s%s"),
 				*SectionFilename, LocalOffset, DestAbsoluteOffset, SourceByte, DestByte, NewLineToken,
 				DestAbsoluteOffset - CallstackAtOffset.Offset, NewLineToken,
-				*DifferenceCallstackDataText, *DiffValues, *DebugDataStackText
+				*LastDifferenceCallstackDataText, *DiffValues, *DebugDataStackText
 			));
 		}
 
@@ -904,11 +922,11 @@ void FAccumulator::CompareWithPreviousForSection(const FPackageData& SourcePacka
 		}
 	}
 
-	if (MaxDiffsToLog >= 0 && NumDiffsLocal > NumDiffsLoggedLocal)
+	if (MaxDiffsToLog >= 0 && NumDiffsForLogStatLocal > NumDiffsLoggedLocal)
 	{
 		MessageCallback(ELogVerbosity::Warning, FString::Printf(
 			TEXT("%s: %lld difference(s) not logged (first at offset: %lld)."),
-			*SectionFilename, NumDiffsLocal - NumDiffsLoggedLocal, FirstUnreportedDiffIndex));
+			*SectionFilename, NumDiffsForLogStatLocal - NumDiffsLoggedLocal, FirstUnreportedDiffIndex));
 	}
 }
 
@@ -1069,7 +1087,7 @@ void FAccumulator::GenerateDiffMapForSection(const FPackageData& SourcePackage,
 				if (DifferenceCallstackOffsetIndex >= 0 && DifferenceCallstackOffsetIndex != LastDifferenceCallstackOffsetIndex)
 				{
 					const FCallstacks::FCallstackAtOffset& CallstackAtOffset = Callstacks.GetCallstack(DifferenceCallstackOffsetIndex);
-					if (!CallstackAtOffset.bIgnore)
+					if (!CallstackAtOffset.bSuppressLogging)
 					{
 						FDiffInfo OffsetAndSize;
 						OffsetAndSize.Offset = CallstackAtOffset.Offset;
@@ -1093,7 +1111,7 @@ void FAccumulator::GenerateDiffMapForSection(const FPackageData& SourcePackage,
 			// Compare against the size without start offset as all callstack offsets are absolute (from the merged header + exports file)
 			if (CallstackAtOffset.Offset < DestPackage.Size)
 			{
-				if (!CallstackAtOffset.bIgnore)
+				if (!CallstackAtOffset.bSuppressLogging)
 				{
 					FDiffInfo OffsetAndSize;
 					OffsetAndSize.Offset = CallstackAtOffset.Offset;
@@ -1626,6 +1644,9 @@ void DumpPackageHeaderDiffs_ZenPackage(
 	const FMessageCallback& MessageCallback)
 {
 	// TODO: Fill in detailed diffing of Zen Package Summary
+	MessageCallback(ELogVerbosity::Warning, FString::Printf(
+		TEXT("%s: header is different (header diffing not yet implemented for -zenstore; cook with -skipzenstore to see header diffs)"),
+		*AssetFilename));
 }
 
 void DumpPackageHeaderDiffs(
diff --git a/Engine/Source/Editor/UnrealEd/Private/Cooker/DiffWriterArchive.h b/Engine/Source/Editor/UnrealEd/Private/Cooker/DiffWriterArchive.h
--- a/Engine/Source/Editor/UnrealEd/Private/Cooker/DiffWriterArchive.h
+++ b/Engine/Source/Editor/UnrealEd/Private/Cooker/DiffWriterArchive.h
@@ -104,8 +104,8 @@ public:
 		int64 Offset = -1;
 		/** Callstack CRC for the Serialize call */
 		uint32 Callstack = 0;
-		/** Collected inside of skip scope */
-		bool bIgnore = false;
+		/** Collected inside of a scope that indicates diff should be recorded but logging should be suppressed */
+		bool bSuppressLogging = false;
 	};
 
 	/** Struct to hold the actual Serialize call callstack and any associated data */
@@ -231,7 +231,7 @@ private:
 class FAccumulator : public FRefCountBase
 {
 public:
-	FAccumulator(UObject* InAsset, FName InPackageName, int32 InMaxDiffsToLog,
+	FAccumulator(UObject* InAsset, FName InPackageName, int32 InMaxDiffsToLog, bool bInIgnoreHeaderDiffs,
 		FMessageCallback&& InMessageCallback, EPackageHeaderFormat InPackageHeaderFormat);
 	virtual ~FAccumulator();
 
@@ -274,6 +274,7 @@ private:
 	EPackageHeaderFormat PackageHeaderFormat = EPackageHeaderFormat::PackageFileSummary;
 	bool bFirstSaveComplete = false;
 	bool bHasDifferences = false;
+	bool bIgnoreHeaderDiffs = false;
 
 	friend class FCallstacks;
 	friend class FDiffArchive;
diff --git a/Engine/Source/Editor/UnrealEd/Private/Tests/DiffWriterArchiveTests.cpp b/Engine/Source/Editor/UnrealEd/Private/Tests/DiffWriterArchiveTests.cpp
--- a/Engine/Source/Editor/UnrealEd/Private/Tests/DiffWriterArchiveTests.cpp
+++ b/Engine/Source/Editor/UnrealEd/Private/Tests/DiffWriterArchiveTests.cpp
@@ -48,9 +48,10 @@ bool FDiffWriterArchiveTestsCallstacks::RunTest(const FString& Parameters)
 	using namespace UE::DiffWriter::Tests;
 
 	const int32 MaxDiffsToLog = 1000;
+	bool bIgnoreHeaderDiffs = false;
 	FName PackageName(TEXT("PackageName"));
 	EPackageHeaderFormat HeaderFormat(EPackageHeaderFormat::PackageFileSummary);
-	TRefCountPtr<FAccumulator> Accumulator = new FAccumulator(nullptr, PackageName, MaxDiffsToLog,
+	TRefCountPtr<FAccumulator> Accumulator = new FAccumulator(nullptr, PackageName, MaxDiffsToLog, bIgnoreHeaderDiffs,
 		[](ELogVerbosity::Type Verbosity, FStringView Message) {}, HeaderFormat);
 	FBasicDiffState State = FBasicDiffState{ 1,2,3 };
 	FDiffArchiveForLinker Ar(*Accumulator);
@@ -62,13 +63,13 @@ bool FDiffWriterArchiveTestsCallstacks::RunTest(const FString& Parameters)
 	TestTrueExpr(Callstacks.GetEndOffset() == sizeof(FBasicDiffState));
 	TestTrueExpr(Callstacks.GetCallstack(0).Offset == 0);
 	TestTrueExpr(Callstacks.GetCallstack(0).Callstack == 0);
-	TestTrueExpr(Callstacks.GetCallstack(0).bIgnore == false);
+	TestTrueExpr(Callstacks.GetCallstack(0).bSuppressLogging == false);
 	TestTrueExpr(Callstacks.GetCallstack(1).Offset == 4);
 	TestTrueExpr(Callstacks.GetCallstack(1).Callstack == 0);
-	TestTrueExpr(Callstacks.GetCallstack(1).bIgnore == false);
+	TestTrueExpr(Callstacks.GetCallstack(1).bSuppressLogging == false);
 	TestTrueExpr(Callstacks.GetCallstack(2).Offset == 8);
 	TestTrueExpr(Callstacks.GetCallstack(2).Callstack == 0);
-	TestTrueExpr(Callstacks.GetCallstack(2).bIgnore == false);
+	TestTrueExpr(Callstacks.GetCallstack(2).bSuppressLogging == false);
 	return true;
 }
 
@@ -81,6 +82,7 @@ bool FDiffWriterArchiveTestsBasic::RunTest(const FString& Parameters)
 	using namespace UE::DiffWriter::Tests;
 
 	const int32 MaxDiffsToLog = 1000;
+	bool bIgnoreHeaderDiffs = false;
 	FName PackageName(TEXT("PackageName"));
 	FString Filename(TEXT("Filename"));
 	EPackageHeaderFormat HeaderFormat(EPackageHeaderFormat::PackageFileSummary);
@@ -88,7 +90,7 @@ bool FDiffWriterArchiveTestsBasic::RunTest(const FString& Parameters)
 	// SECTION("DiffMap - identical")
 	{
 		TRefCountPtr<FAccumulator> Accumulator = new FAccumulator(nullptr, PackageName, MaxDiffsToLog,
-			[](ELogVerbosity::Type Verbosity, FStringView Message) {}, HeaderFormat);
+			bIgnoreHeaderDiffs, [](ELogVerbosity::Type Verbosity, FStringView Message) {}, HeaderFormat);
 
 		FLargeMemoryWriter InitialState;
 		{
@@ -111,7 +113,7 @@ bool FDiffWriterArchiveTestsBasic::RunTest(const FString& Parameters)
 	// SECTION("DiffMap - mismatch")
 	{
 		TRefCountPtr<FAccumulator> Accumulator = new FAccumulator(nullptr, PackageName, MaxDiffsToLog,
-			[](ELogVerbosity::Type Verbosity, FStringView Message) {}, HeaderFormat);
+			bIgnoreHeaderDiffs, [](ELogVerbosity::Type Verbosity, FStringView Message) {}, HeaderFormat);
 
 		FLargeMemoryWriter InitialState;
 		{
@@ -134,7 +136,7 @@ bool FDiffWriterArchiveTestsBasic::RunTest(const FString& Parameters)
 	// SECTION("Compare - mismatch")
 	{
 		TRefCountPtr<FAccumulator> Accumulator = new FAccumulator(nullptr, PackageName, MaxDiffsToLog,
-			[](ELogVerbosity::Type Verbosity, FStringView Message) {}, HeaderFormat);
+			bIgnoreHeaderDiffs, [](ELogVerbosity::Type Verbosity, FStringView Message) {}, HeaderFormat);
 
 		FBasicDiffState InitialState = FBasicDiffState {1,2,3};
 		FLargeMemoryWriter InitialMemory;
