Update cook metadata to include enabled-by-default plugins when cooking dlc. This should clean up some horde logspam

diff --git a/Engine/Source/Editor/UnrealEd/Private/CookOnTheFlyServer.cpp b/Engine/Source/Editor/UnrealEd/Private/CookOnTheFlyServer.cpp
--- a/Engine/Source/Editor/UnrealEd/Private/CookOnTheFlyServer.cpp
+++ b/Engine/Source/Editor/UnrealEd/Private/CookOnTheFlyServer.cpp
@@ -9472,6 +9472,11 @@ void UCookOnTheFlyServer::WriteCookMetadata(const ITargetPlatform* InTargetPlatf
 
 			EnabledPlugins.RemoveAllSwap([&IncludedSet](TSharedRef<IPlugin>& EnabledPlugin)
 			{
+				// Never remove plugins that are enabled by default as they are on but not included in the dlc dependency tree.
+				if (EnabledPlugin->IsEnabledByDefault(true))
+				{
+					return false;
+				}
 				return IncludedSet.Contains(EnabledPlugin->GetName()) == false;
 			}, false /* bAllowShrinking */);
 
