#UE Minor optimization ULoadPackageCommandlet when using -loadlist

diff --git a/Engine/Source/Editor/UnrealEd/Private/Commandlets/PackageUtilities.cpp b/Engine/Source/Editor/UnrealEd/Private/Commandlets/PackageUtilities.cpp
--- a/Engine/Source/Editor/UnrealEd/Private/Commandlets/PackageUtilities.cpp
+++ b/Engine/Source/Editor/UnrealEd/Private/Commandlets/PackageUtilities.cpp
@@ -730,6 +730,7 @@ int32 ULoadPackageCommandlet::Main( const FString& Params )
 
 				Tokens.Empty(TempTokens.Num());
 				Tokens = TempTokens;
+				break;
 			}
 		}
 	}
