Fix ObjectTools::PrivatizeAssets loading external actors

diff --git a/Engine/Source/Editor/UnrealEd/Private/ObjectTools.cpp b/Engine/Source/Editor/UnrealEd/Private/ObjectTools.cpp
--- a/Engine/Source/Editor/UnrealEd/Private/ObjectTools.cpp
+++ b/Engine/Source/Editor/UnrealEd/Private/ObjectTools.cpp
@@ -2685,7 +2685,7 @@ namespace ObjectTools
 		TArray<UObject*> ObjectsToPrivatize;
 		for (const FAssetData& AssetToPrivatize : AssetsToPrivatize)
 		{
-			UObject* ObjectToPrivatize = AssetToPrivatize.GetAsset({ ULevel::LoadAllExternalObjectsTag });
+			UObject* ObjectToPrivatize = AssetToPrivatize.GetAsset();
 
 			if (ObjectToPrivatize)
 			{
