Remove ensure, as it will fire any time an animation blueprint is duplicated (because animation blueprints have a compiler generated set of sparse class data)

diff --git a/Engine/Source/Editor/UnrealEd/Private/Kismet2/BlueprintEditorUtils.cpp b/Engine/Source/Editor/UnrealEd/Private/Kismet2/BlueprintEditorUtils.cpp
--- a/Engine/Source/Editor/UnrealEd/Private/Kismet2/BlueprintEditorUtils.cpp
+++ b/Engine/Source/Editor/UnrealEd/Private/Kismet2/BlueprintEditorUtils.cpp
@@ -1728,8 +1728,13 @@ void FBlueprintEditorUtils::PostDuplicateBlueprint(UBlueprint* Blueprint, bool b
 			{
 				void* SparseDataInstance = Blueprint->GeneratedClass->GetOrCreateSparseClassData();
 
+				// Compile may have generated a new sparse class data, in which case we're just going to 
+				// use whatever the compiler generated - if we're reusing the source class's sparse
+				// class data then we can copy over the values immediately.. We could CPFUO here
+				// as well, but if the compiler generated the sparse data that could be undesirable - e.g.
+				// because the sparse data is caching information about the CDO or Class
 				const TObjectPtr<UScriptStruct> OldSparseData = OldBPGC->GetSparseClassDataStruct();
-				if (ensure(OldSparseData == SparseData))
+				if (OldSparseData == SparseData)
 				{
 					const void* OldSparseDataInstance = OldBPGC->GetSparseClassData(EGetSparseClassDataMethod::ReturnIfNull);
 					if (OldSparseDataInstance && ensure(OldSparseDataInstance != SparseDataInstance))
