[WP navmesh] Fix WP navmesh generation by setting GridSize for ANavigationDataChunkActor (needed since the addition of FPartitionActorDesc::GetEditorBounds())

diff --git a/Engine/Source/Editor/UnrealEd/Private/WorldPartition/WorldPartitionNavigationDataBuilder.cpp b/Engine/Source/Editor/UnrealEd/Private/WorldPartition/WorldPartitionNavigationDataBuilder.cpp
--- a/Engine/Source/Editor/UnrealEd/Private/WorldPartition/WorldPartitionNavigationDataBuilder.cpp
+++ b/Engine/Source/Editor/UnrealEd/Private/WorldPartition/WorldPartitionNavigationDataBuilder.cpp
@@ -404,9 +404,11 @@ bool UWorldPartitionNavigationDataBuilder::GenerateNavigationData(UWorldPartitio
 	int32 XMax = static_cast<int32>(GeneratingBounds.Max.X);
 	int32 YMax = static_cast<int32>(GeneratingBounds.Max.Y);
 
+	const uint32 GridSize = FMath::Max(NavigationDataActorClass->GetDefaultObject<APartitionActor>()->GetDefaultGridSize(World), 1u);
+
 	const FIntRect GeneratingBounds2D(XMin, YMin, XMax, YMax);
 	FActorPartitionGridHelper::ForEachIntersectingCell(NavigationDataActorClass, GeneratingBounds2D, World->PersistentLevel,
-		[&WorldPartition, &ActorCount, World, &ValidNavigationDataChunkActors, &NavDataBounds, this](const UActorPartitionSubsystem::FCellCoord& InCellCoord, const FIntRect& InCellBounds)->bool
+		[&WorldPartition, &ActorCount, World, &ValidNavigationDataChunkActors, &NavDataBounds, GridSize, this](const UActorPartitionSubsystem::FCellCoord& InCellCoord, const FIntRect& InCellBounds)->bool
 		{
 			TRACE_CPUPROFILER_EVENT_SCOPE(MakeNavigationDataChunkActorForGridCell);
 		
@@ -435,6 +437,8 @@ bool UWorldPartitionNavigationDataBuilder::GenerateNavigationData(UWorldPartitio
 			ANavigationDataChunkActor* DataChunkActor = World->SpawnActor<ANavigationDataChunkActor>(SpawnParams);
 			ActorCount++;
 
+			DataChunkActor->SetGridSize(GridSize);
+			
 			const FVector2D CellCenter = CellBounds.GetCenter();
 			DataChunkActor->SetActorLocation(FVector(CellCenter.X, CellCenter.Y, 0.f));
 
diff --git a/Engine/Source/Runtime/Engine/Private/WorldPartition/ActorPartition/PartitionActorDesc.cpp b/Engine/Source/Runtime/Engine/Private/WorldPartition/ActorPartition/PartitionActorDesc.cpp
--- a/Engine/Source/Runtime/Engine/Private/WorldPartition/ActorPartition/PartitionActorDesc.cpp
+++ b/Engine/Source/Runtime/Engine/Private/WorldPartition/ActorPartition/PartitionActorDesc.cpp
@@ -68,6 +68,7 @@ bool FPartitionActorDesc::Equals(const FWorldPartitionActorDesc* Other) const
 FBox FPartitionActorDesc::GetEditorBounds() const
 {
 	const UActorPartitionSubsystem::FCellCoord CellCoord(GridIndexX, GridIndexY, GridIndexZ, 0);
+	ensureMsgf(GridSize > 1, TEXT("Invalid GridSize. GridSize needs to be set."));
 	return UActorPartitionSubsystem::FCellCoord::GetCellBounds(CellCoord, GridSize);
 }
 
