Cooker: CookFirst: Turn off randomdifforder when running with -cookfirst, since it breaks cookfirst.
#rnx

diff --git a/Engine/Source/Editor/UnrealEd/Classes/CookOnTheSide/CookOnTheFlyServer.h b/Engine/Source/Editor/UnrealEd/Classes/CookOnTheSide/CookOnTheFlyServer.h
--- a/Engine/Source/Editor/UnrealEd/Classes/CookOnTheSide/CookOnTheFlyServer.h
+++ b/Engine/Source/Editor/UnrealEd/Classes/CookOnTheSide/CookOnTheFlyServer.h
@@ -1478,6 +1478,8 @@ private:
 	bool bRandomizeCookOrder = false;
 	/** True if commandline arguments specified that we suppress the cook of packages based on filter criteria. */
 	bool bCookFilter = false;
+	/** True if commandline arguments specify that packages on commandline should be cooked first. */
+	bool bCookFirst = false;
 	/**
 	 * Experimental feature to correctly invoke the BeginCacheForCookedPlatformData contracts for
 	 * objects created by PreSave
diff --git a/Engine/Source/Editor/UnrealEd/Private/CookOnTheFlyServer.cpp b/Engine/Source/Editor/UnrealEd/Private/CookOnTheFlyServer.cpp
--- a/Engine/Source/Editor/UnrealEd/Private/CookOnTheFlyServer.cpp
+++ b/Engine/Source/Editor/UnrealEd/Private/CookOnTheFlyServer.cpp
@@ -6778,8 +6778,9 @@ void UCookOnTheFlyServer::SetInitializeConfigSettings(UE::Cook::FInitializeConfi
 		}
 	}
 
-	bRandomizeCookOrder = FParse::Param(FCommandLine::Get(), TEXT("RANDOMPACKAGEORDER")) ||
-		(FParse::Param(FCommandLine::Get(), TEXT("DIFFONLY")) && !FParse::Param(FCommandLine::Get(), TEXT("DIFFNORANDCOOK")));
+	bCookFirst = FParse::Param(FCommandLine::Get(), TEXT("CookFirst"));
+	bRandomizeCookOrder = !bCookFirst && (FParse::Param(FCommandLine::Get(), TEXT("RANDOMPACKAGEORDER")) ||
+		(FParse::Param(FCommandLine::Get(), TEXT("DIFFONLY")) && !FParse::Param(FCommandLine::Get(), TEXT("DIFFNORANDCOOK"))));
 
 	ParseCookFilters();
 
@@ -11363,7 +11364,6 @@ void UCookOnTheFlyServer::GenerateInitialRequests(FBeginCookContext& BeginContex
 		GenerateLongPackageNames(FilesInPath, FilesInPathInstigators);
 	}
 	TSet<FName> CookFirstPackages;
-	bool bCookFirst = FParse::Param(FCommandLine::Get(), TEXT("CookFirst"));
 	if (bCookFirst)
 	{
 		for (const FString& CookMap : CookMaps)
