CookMetadata: DLC plugin filtering is unreliable and I was unable to find a safe way to do it, so we include all enabled plugins even when cooking dlc.

diff --git a/Engine/Source/Developer/CookMetadata/Public/CookMetadata.h b/Engine/Source/Developer/CookMetadata/Public/CookMetadata.h
--- a/Engine/Source/Developer/CookMetadata/Public/CookMetadata.h
+++ b/Engine/Source/Developer/CookMetadata/Public/CookMetadata.h
@@ -159,9 +159,7 @@ struct COOKMETADATA_API FCookMetadataPluginEntry
 
 struct COOKMETADATA_API FCookMetadataPluginHierarchy
 {
-	// The list of plugins that were enabled during the cook that generated the FCookMetadataState,
-	// pruned to remove plugins that aren't enabled on the cook platform. If the cook was DLC, this
-	// is further pruned to only the DLC plugin and its dependencies.
+	// The list of plugins that were enabled during the cook that generated the FCookMetadataState
 	TArray<FCookMetadataPluginEntry> PluginsEnabledAtCook;
 
 	// The list of plugin dependencies. FCookMetadataPluginEntry::DependencyIndexStart indexes into this
diff --git a/Engine/Source/Editor/UnrealEd/Private/CookOnTheFlyServer.cpp b/Engine/Source/Editor/UnrealEd/Private/CookOnTheFlyServer.cpp
--- a/Engine/Source/Editor/UnrealEd/Private/CookOnTheFlyServer.cpp
+++ b/Engine/Source/Editor/UnrealEd/Private/CookOnTheFlyServer.cpp
@@ -9413,10 +9413,10 @@ void UCookOnTheFlyServer::WriteCookMetadata(const ITargetPlatform* InTargetPlatf
 	FString PlatformNameString = InTargetPlatform->PlatformName();
 
 	//
-	// Write the plugin hierarchy for the plugins we cooked. For DLC this is limited to the DLC plugin, otherwise
-	// we write out all plugins enabled for the cook.
+	// Write the plugin hierarchy for the plugins enabled. Technically for a DLC cook we aren't cooking all plugins,
+	// however there's not a direct way to narrow the list (e.g. enabled by default + dlc plugins is too narrow), 
+	// so we just always write the entire set.
 	//
-
 	TArray<TSharedRef<IPlugin>> EnabledPlugins = IPluginManager::Get().GetEnabledPlugins();
 
 	// NOTE: We can't use IsEnabledForPlugin because it has an issue with the AllowTargets list where preventing a plugin on
@@ -9424,67 +9424,6 @@ void UCookOnTheFlyServer::WriteCookMetadata(const ITargetPlatform* InTargetPlatf
 	// on disk during stage existing but the plugin isn't in the cook manifest. I wasn't able to find a way to fix this
 	// with the current plugin system, so we include all enabled plugins.
 
-
-	// Filter to the DLC plugin + dependencies if we are a DLC cook.
-	if (IsCookingDLC())
-	{
-		TMap<FStringView, TSharedRef<IPlugin>> EnabledPluginMap;
-		IPlugin* DlcPlugin = nullptr;
-		for (TSharedRef<IPlugin>& Plugin : EnabledPlugins)
-		{
-			EnabledPluginMap.Add(Plugin->GetName(), Plugin);
-			if (Plugin->GetName() == CookByTheBookOptions->DlcName)
-			{
-				check(DlcPlugin == nullptr);
-				DlcPlugin = &Plugin.Get();
-			}
-		}
-
-		if (DlcPlugin == nullptr)
-		{
-			UE_LOG(LogCook, Warning, TEXT("Dlc plugin %s not found, as a result including all plugins in cook metadata file."), *CookByTheBookOptions->DlcName);
-		}
-		else
-		{
-			TSet<FStringView> IncludedSet;
-
-			TArray<IPlugin*> DependencyStack;
-			IncludedSet.Add(DlcPlugin->GetName());
-			DependencyStack.Add(DlcPlugin);
-
-			while (DependencyStack.Num())
-			{
-				IPlugin* CurrentPlugin = DependencyStack.Pop();
-
-				for (const FPluginReferenceDescriptor& ChildPlugin : CurrentPlugin->GetDescriptor().Plugins)
-				{
-					TSharedRef<IPlugin>* FoundPlugin = EnabledPluginMap.Find(ChildPlugin.Name);
-					if (FoundPlugin)
-					{
-						bool bAlreadyExists = false;
-						IncludedSet.Add(ChildPlugin.Name, &bAlreadyExists);
-						if (bAlreadyExists == false)
-						{
-							DependencyStack.Add(&FoundPlugin->Get());
-						}
-					}
-				}
-			}
-
-			EnabledPlugins.RemoveAllSwap([&IncludedSet](TSharedRef<IPlugin>& EnabledPlugin)
-			{
-				// Never remove plugins that are enabled by default as they are on but not included in the dlc dependency tree.
-				if (EnabledPlugin->IsEnabledByDefault(true))
-				{
-					return false;
-				}
-				return IncludedSet.Contains(EnabledPlugin->GetName()) == false;
-			}, false /* bAllowShrinking */);
-
-			UE_LOG(LogCook, Display, TEXT("Dlc cook narrowed cook plugin metadata list from %d to %d plugins"), EnabledPluginMap.Num(), EnabledPlugins.Num());
-		}
-	}
-
 	constexpr int32 AdditionalPseudoPlugins = 2; // /Engine and /Game.
 	if (IntFitsIn<uint16>(EnabledPlugins.Num() + AdditionalPseudoPlugins) == false)
 	{
