Fix FindNativeArchetype, which I broke in 27152424

diff --git a/Engine/Source/Editor/UnrealEd/Private/Kismet2/Kismet2.cpp b/Engine/Source/Editor/UnrealEd/Private/Kismet2/Kismet2.cpp
--- a/Engine/Source/Editor/UnrealEd/Private/Kismet2/Kismet2.cpp
+++ b/Engine/Source/Editor/UnrealEd/Private/Kismet2/Kismet2.cpp
@@ -973,7 +973,7 @@ static UObject* ConformComponentsUtils::FindNativeArchetype(const UObject* Nativ
 	UObject* NativeSubobject = StaticFindObjectFast(UObject::StaticClass(), const_cast<UObject*>(NativeCDO), Component->GetFName());
 	if (!NativeSubobject)
 	{
-		return nullptr;
+		return Component->GetClass()->ClassDefaultObject;
 	}
 	else if (!Component->IsA(NativeSubobject->GetClass()))
 	{
