Sequencer: FBX Export options for baking

There are now options for baking on FBX export. The options are:
- Bake None
- Bake Transforms
- Bake Channels
- Bake All

The option for Camera and Lights is separate from the option for Actors.

diff --git a/Engine/Source/Editor/UnrealEd/Classes/Exporters/FbxExportOption.h b/Engine/Source/Editor/UnrealEd/Classes/Exporters/FbxExportOption.h
--- a/Engine/Source/Editor/UnrealEd/Classes/Exporters/FbxExportOption.h
+++ b/Engine/Source/Editor/UnrealEd/Classes/Exporters/FbxExportOption.h
@@ -24,6 +24,18 @@ enum class EFbxExportCompatibility : uint8
 	FBX_2020,
 };
 
+// Bake options for animated properties of exported objects
+UENUM(BlueprintType)
+enum class EMovieSceneBakeType : uint8
+{
+	None = 0x000,
+
+	BakeChannels = 0x001,
+	BakeTransforms = 0x002,
+
+	BakeAll = BakeChannels | BakeTransforms,
+};
+
 UCLASS(config = EditorPerProjectUserSettings, MinimalAPI, BlueprintType)
 class UFbxExportOption : public UObject
 {
@@ -78,6 +90,14 @@ public:
 	UPROPERTY(EditAnywhere, BlueprintReadWrite, config, category = Animation)
 	uint32 bExportLocalTime : 1;
 
+	/** Bake settings for camera and light animation curves. Camera Scale not exported. */
+	UPROPERTY(EditAnywhere, BlueprintReadWrite, AdvancedDisplay, config, Category = Animation)
+	EMovieSceneBakeType BakeCameraAndLightAnimation;
+
+	/** Bake settings for exported non-camera, non-light object animation */
+	UPROPERTY(EditAnywhere, BlueprintReadWrite, AdvancedDisplay, config, Category = Animation)
+	EMovieSceneBakeType BakeActorAnimation;
+
 	/* Set all the FProperty to the CDO value */
 	void ResetToDefault();
 
diff --git a/Engine/Source/Editor/UnrealEd/Private/Fbx/FbxExportOption.cpp b/Engine/Source/Editor/UnrealEd/Private/Fbx/FbxExportOption.cpp
--- a/Engine/Source/Editor/UnrealEd/Private/Fbx/FbxExportOption.cpp
+++ b/Engine/Source/Editor/UnrealEd/Private/Fbx/FbxExportOption.cpp
@@ -28,6 +28,8 @@ UFbxExportOption::UFbxExportOption(const FObjectInitializer& ObjectInitializer)
 	VertexColor = true;
 	MapSkeletalMotionToRoot = false;
 	bExportLocalTime = true;
+	BakeCameraAndLightAnimation = EMovieSceneBakeType::BakeTransforms;
+	BakeActorAnimation = EMovieSceneBakeType::None;
 }
 
 void UFbxExportOption::ResetToDefault()
diff --git a/Engine/Source/Editor/UnrealEd/Private/Fbx/FbxMainExport.cpp b/Engine/Source/Editor/UnrealEd/Private/Fbx/FbxMainExport.cpp
--- a/Engine/Source/Editor/UnrealEd/Private/Fbx/FbxMainExport.cpp
+++ b/Engine/Source/Editor/UnrealEd/Private/Fbx/FbxMainExport.cpp
@@ -1877,6 +1877,17 @@ bool FFbxExporter::ExportLevelSequenceTracks(UMovieScene* MovieScene, IMovieScen
 	USceneComponent* BoundComponent = Cast<USceneComponent>(BoundObject);
 
 	const bool bIsCameraActor = BoundActor ? BoundActor->IsA(ACameraActor::StaticClass()) : BoundComponent ? BoundComponent->IsA(UCameraComponent::StaticClass()) : false;
+	const bool bIsLightActor = BoundActor ? BoundActor->IsA(ALight::StaticClass()) : BoundComponent ? BoundComponent->IsA(ULightComponent::StaticClass()) : false;
+
+	bool bBakeAllTransformChannels = false;
+	if (bIsCameraActor || bIsLightActor)
+	{
+		bBakeAllTransformChannels = (static_cast<uint8>(GetExportOptions()->BakeCameraAndLightAnimation) & static_cast<uint8>(EMovieSceneBakeType::BakeTransforms));
+	}
+	else
+	{
+		bBakeAllTransformChannels = (static_cast<uint8>(GetExportOptions()->BakeActorAnimation) & static_cast<uint8>(EMovieSceneBakeType::BakeTransforms));
+	}
 
 	// If this actor has attached actors, and the bound component isn't the root component, skip the 3d transform track because we don't want the 
 	// component transform to be assigned to the actor, which would subsequently apply to the attached children 
@@ -1892,7 +1903,7 @@ bool FFbxExporter::ExportLevelSequenceTracks(UMovieScene* MovieScene, IMovieScen
 	}
 
 	// If there's more than one transform track for this actor (ie. on the actor and on the root component) or if there's more than one section, evaluate baked
-	if (bIsCameraActor || TransformTracks.Num() > 1 || (TransformTracks.Num() != 0 && TransformTracks[0].Get()->GetAllSections().Num() > 1))
+	if (bBakeAllTransformChannels || TransformTracks.Num() > 1 || (TransformTracks.Num() != 0 && TransformTracks[0].Get()->GetAllSections().Num() > 1))
 	{
 		if (!bSkip3DTransformTrack)
 		{
@@ -1921,7 +1932,16 @@ bool FFbxExporter::ExportLevelSequenceTracks(UMovieScene* MovieScene, IMovieScen
 		}
 		else
 		{
-			ExportLevelSequenceTrackChannels(FbxActor, *Track, MovieScene->GetPlaybackRange(), RootToLocalTransform);
+			bool bBakeChannels = false;
+			if (bIsCameraActor || bIsLightActor)
+			{
+				bBakeChannels = (static_cast<uint8>(GetExportOptions()->BakeCameraAndLightAnimation) & static_cast<uint8>(EMovieSceneBakeType::BakeChannels));
+			}
+			else
+			{
+				bBakeChannels = (static_cast<uint8>(GetExportOptions()->BakeActorAnimation) & static_cast<uint8>(EMovieSceneBakeType::BakeChannels));
+			}
+			ExportLevelSequenceTrackChannels(FbxActor, *Track, MovieScene->GetPlaybackRange(), RootToLocalTransform, bBakeChannels);
 		}
 	}
 
@@ -2747,6 +2767,41 @@ void RichCurveInterpolationToFbxInterpolation(ERichCurveInterpMode InInterpolati
 	}
 }
 
+template<typename ChannelType>
+void FFbxExporter::ExportBezierChannelToFbxCurveBaked(FbxAnimCurve& InFbxCurve, const ChannelType& InChannel, FFrameRate TickResolution, const UMovieSceneTrack* Track, ERichCurveValueMode ValueMode, bool bNegative, const FMovieSceneSequenceTransform& RootToLocalTransform)
+{
+	using ChannelValueType = typename ChannelType::ChannelValueType;
+	using CurveValueType = typename ChannelType::CurveValueType;
+
+	const float NegateFactor = bNegative ? -1.f : 1.f;
+	const float kOneThird = 1.f / 3.f;
+
+	InFbxCurve.KeyModifyBegin();
+
+	FFrameRate DisplayRate = Track->GetTypedOuter<UMovieScene>()->GetDisplayRate();
+	TRange<FFrameNumber> PlaybackRange = Track->GetTypedOuter<UMovieScene>()->GetPlaybackRange();
+
+	int32 LocalStartFrame = FFrameRate::TransformTime(FFrameTime(UE::MovieScene::DiscreteInclusiveLower(PlaybackRange)), TickResolution, DisplayRate).RoundToFrame().Value;
+	int32 StartFrame = FFrameRate::TransformTime(FFrameTime(UE::MovieScene::DiscreteInclusiveLower(PlaybackRange) * RootToLocalTransform.InverseNoLooping()), TickResolution, DisplayRate).RoundToFrame().Value;
+	int32 AnimationLength = FFrameRate::TransformTime(FFrameTime(FFrameNumber(UE::MovieScene::DiscreteSize(PlaybackRange))), TickResolution, DisplayRate).RoundToFrame().Value;
+
+	for (int32 FrameCount = 0; FrameCount <= AnimationLength; ++FrameCount)
+	{
+		int32 LocalFrame = LocalStartFrame + FrameCount;
+
+		FFrameTime LocalTime = FFrameRate::TransformTime(FFrameTime(LocalFrame), DisplayRate, TickResolution);
+
+		float Value;
+		InChannel.Evaluate(LocalTime, Value);
+
+		FbxTime FbxTime;
+		FbxTime.SetSecondDouble(GetExportOptions()->bExportLocalTime ? DisplayRate.AsSeconds(LocalFrame) : DisplayRate.AsSeconds(StartFrame + FrameCount));
+
+		InFbxCurve.KeySet(InFbxCurve.KeyAdd(FbxTime), FbxTime, Value, FbxAnimCurveDef::eInterpolationLinear);
+	}
+	InFbxCurve.KeyModifyEnd();
+}
+
 template<typename ChannelType>
 void FFbxExporter::ExportBezierChannelToFbxCurve(FbxAnimCurve& InFbxCurve, const ChannelType& InChannel, FFrameRate TickResolution, ERichCurveValueMode ValueMode, bool bNegative, const FMovieSceneSequenceTransform& RootToLocalTransform)
 {
@@ -2896,7 +2951,9 @@ void FFbxExporter::ExportLevelSequence3DTransformTrack(FbxNode* FbxNode, IMovieS
 
 	const bool bIsCameraActor = BoundActor ? BoundActor->IsA(ACameraActor::StaticClass()) : BoundComponent ? BoundComponent->IsA(UCameraComponent::StaticClass()) : false;
 	const bool bIsLightActor = BoundActor ? BoundActor->IsA(ALight::StaticClass()) : BoundComponent ? BoundComponent->IsA(ULightComponent::StaticClass()) : false;
-	const bool bBakeRotations = bIsCameraActor || bIsLightActor;
+
+	// If bake rotations is set, and actor is a camera or light actor, only bake rotation channel
+	const bool bBakeRotations = (bIsCameraActor || bIsLightActor) && (static_cast<uint8>(GetExportOptions()->BakeCameraAndLightAnimation) & static_cast<uint8>(EMovieSceneBakeType::BakeTransforms));
 
 	if (!FbxNode)
 	{
@@ -3295,7 +3352,7 @@ void FFbxExporter::ExportLevelSequenceBaked3DTransformTrack(IAnimTrackAdapter& A
 }
 
 
-void FFbxExporter::ExportLevelSequenceTrackChannels( FbxNode* FbxNode, UMovieSceneTrack& Track, const TRange<FFrameNumber>& InPlaybackRange, const FMovieSceneSequenceTransform& RootToLocalTransform)
+void FFbxExporter::ExportLevelSequenceTrackChannels( FbxNode* FbxNode, UMovieSceneTrack& Track, const TRange<FFrameNumber>& InPlaybackRange, const FMovieSceneSequenceTransform& RootToLocalTransform, bool bBakeBezierCurves)
 {
 	// TODO: Support more than one section?
 	UMovieSceneSection* Section = Track.GetAllSections().Num() > 0 ? Track.GetAllSections()[0] : nullptr;
@@ -3436,14 +3493,28 @@ void FFbxExporter::ExportLevelSequenceTrackChannels( FbxNode* FbxNode, UMovieSce
 				CurveNode->SetChannelValue<double>(0U, DoubleChannel->GetDefault().Get(MAX_dbl));
 				CurveNode->ConnectToChannel(AnimCurve, 0U);
 
-				ExportChannelToFbxCurve(*AnimCurve, *DoubleChannel, TickResolution, IsFoV ? ERichCurveValueMode::Fov : ERichCurveValueMode::Default, false, RootToLocalTransform);
+				if (bBakeBezierCurves)
+				{
+					ExportBezierChannelToFbxCurveBaked(*AnimCurve, *DoubleChannel, TickResolution, &Track, IsFoV ? ERichCurveValueMode::Fov : ERichCurveValueMode::Default, false, RootToLocalTransform);
+				}
+ 				else
+ 				{
+					ExportChannelToFbxCurve(*AnimCurve, *DoubleChannel, TickResolution, IsFoV ? ERichCurveValueMode::Fov : ERichCurveValueMode::Default, false, RootToLocalTransform);
+				}
 			}
 			else if (FloatChannel)
 			{
 				CurveNode->SetChannelValue<double>(0U, FloatChannel->GetDefault().Get(MAX_flt));
 				CurveNode->ConnectToChannel(AnimCurve, 0U);
 
-				ExportChannelToFbxCurve(*AnimCurve, *FloatChannel, TickResolution, IsFoV ? ERichCurveValueMode::Fov : ERichCurveValueMode::Default, false, RootToLocalTransform);
+				if (bBakeBezierCurves)
+				{
+					ExportBezierChannelToFbxCurveBaked(*AnimCurve, *FloatChannel, TickResolution, &Track, IsFoV ? ERichCurveValueMode::Fov : ERichCurveValueMode::Default, false, RootToLocalTransform);
+				}
+				else
+				{
+					ExportChannelToFbxCurve(*AnimCurve, *FloatChannel, TickResolution, IsFoV ? ERichCurveValueMode::Fov : ERichCurveValueMode::Default, false, RootToLocalTransform);
+				}
 			}
 			else if (IntegerChannel)
 			{
diff --git a/Engine/Source/Editor/UnrealEd/Private/FbxExporter.h b/Engine/Source/Editor/UnrealEd/Private/FbxExporter.h
--- a/Engine/Source/Editor/UnrealEd/Private/FbxExporter.h
+++ b/Engine/Source/Editor/UnrealEd/Private/FbxExporter.h
@@ -437,10 +437,13 @@ private:
 	 */
 	UNREALED_API void ExportLevelSequenceBaked3DTransformTrack(IAnimTrackAdapter& AnimTrackAdapter, FbxNode* FbxActor, IMovieScenePlayer* MovieScenePlayer, FMovieSceneSequenceIDRef InSequenceID, TArray<TWeakObjectPtr<UMovieScene3DTransformTrack> > TransformTracks, UObject* BoundObject, const TRange<FFrameNumber>& InPlaybackRange, const FMovieSceneSequenceTransform& RootToLocalTransform);
 
-	/** 
+	/**
 	 * Exports a level sequence property track into the FBX animation stack. 
 	 */
-	UNREALED_API void ExportLevelSequenceTrackChannels( FbxNode* FbxActor, UMovieSceneTrack& Track, const TRange<FFrameNumber>& InPlaybackRange, const FMovieSceneSequenceTransform& RootToLocalTransform);
+	UNREALED_API void ExportLevelSequenceTrackChannels( FbxNode* FbxActor, UMovieSceneTrack& Track, const TRange<FFrameNumber>& InPlaybackRange, const FMovieSceneSequenceTransform& RootToLocalTransform, bool bBakeBezierCurves);
+
+	UE_DEPRECATED(5.4, "Use ExportLevelSequenceTrackChannels that indicates whether bezier channels should be baked")
+	UNREALED_API void ExportLevelSequenceTrackChannels(FbxNode* FbxActor, UMovieSceneTrack& Track, const TRange<FFrameNumber>& InPlaybackRange, const FMovieSceneSequenceTransform& RootToLocalTransform) { ExportLevelSequenceTrackChannels(FbxActor, Track, InPlaybackRange, RootToLocalTransform, false); }
 
 	/** Defines value export modes for the EportRichCurveToFbxCurve method. */
 	enum class ERichCurveValueMode
@@ -451,6 +454,10 @@ private:
 		Fov
 	};
 
+	/** Generic implementation of exporting a movie scene bezier curve channel to an fbx animation curve, baked per frame */
+	template<typename ChannelType>
+	void ExportBezierChannelToFbxCurveBaked(FbxAnimCurve& InFbxCurve, const ChannelType& InChannel, FFrameRate TickResolution, const UMovieSceneTrack* Track, ERichCurveValueMode ValueMode, bool bNegative, const FMovieSceneSequenceTransform& RootToLocalTransform);
+
 	/** Generic implementation of exporting a movie scene bezier curve channel to an fbx animation curve */
 	template<typename ChannelType>
 	void ExportBezierChannelToFbxCurve(FbxAnimCurve& InFbxCurve, const ChannelType& InChannel, FFrameRate TickResolution, ERichCurveValueMode ValueMode, bool bNegative, const FMovieSceneSequenceTransform& RootToLocalTransform);
