Avoid spamming cook rejection reason to log when rejecting external actors for the OnlyEditorOnly reason.

diff --git a/Engine/Source/Editor/UnrealEd/Private/CookOnTheFlyServer.cpp b/Engine/Source/Editor/UnrealEd/Private/CookOnTheFlyServer.cpp
--- a/Engine/Source/Editor/UnrealEd/Private/CookOnTheFlyServer.cpp
+++ b/Engine/Source/Editor/UnrealEd/Private/CookOnTheFlyServer.cpp
@@ -2643,7 +2643,7 @@ void UCookOnTheFlyServer::DemoteToIdle(UE::Cook::FPackageData& PackageData, UE::
 			TStringBuilder<256> PackageNameStr(InPlace, PackageData.GetPackageName());
 
 			// ExternalActors: Do not send a message for every NeverCook external Actor package; too much spam
-			if (Reason == ESuppressCookReason::NeverCook)
+			if ((Reason == ESuppressCookReason::NeverCook) || (Reason == ESuppressCookReason::OnlyEditorOnly))
 			{
 				bPrintDiagnostic &= UE::String::FindFirst(PackageNameStr.ToView(),
 					ULevel::GetExternalActorsFolderName(), ESearchCase::IgnoreCase) == INDEX_NONE;
