Adding tooltip for Monitor Editor Performance in the scalabiltiy settings.

#rnx

diff --git a/Engine/Source/Editor/UnrealEd/Private/SScalabilitySettings.cpp b/Engine/Source/Editor/UnrealEd/Private/SScalabilitySettings.cpp
--- a/Engine/Source/Editor/UnrealEd/Private/SScalabilitySettings.cpp
+++ b/Engine/Source/Editor/UnrealEd/Private/SScalabilitySettings.cpp
@@ -547,6 +547,16 @@ void SScalabilitySettings::Construct( const FArguments& InArgs )
 				[
 					SNew(STextBlock)
 					.Text(LOCTEXT("PerformanceWarningEnableDisableCheckbox", "Monitor Editor Performance?"))
+					.ToolTipText_Lambda([this]()
+						{
+							FProperty* Property = UEditorPerformanceSettings::StaticClass()->FindPropertyByName(GET_MEMBER_NAME_CHECKED(UEditorPerformanceSettings, bMonitorEditorPerformance));
+							if (Property)
+							{
+								return Property->GetToolTipText();
+							}
+							
+							return FText::GetEmpty();
+						})
 				]
 			]
 
