UE-190649 - BP compilation doesn't properly clear loading flags on trashed objects

Not a fix, but adds an ensure that validates the ComponentTemplates array of a Blueprint. The underlying cause of this issue is due to calling PostLoad on a Blueprint before we've resolved placeholder objects. In some cases, the Zen loader can cause this problem, such as during an unexpected LoadPackage call due to resolving a soft pointer.

diff --git a/Engine/Source/Editor/UnrealEd/Private/Kismet2/BlueprintEditorUtils.cpp b/Engine/Source/Editor/UnrealEd/Private/Kismet2/BlueprintEditorUtils.cpp
--- a/Engine/Source/Editor/UnrealEd/Private/Kismet2/BlueprintEditorUtils.cpp
+++ b/Engine/Source/Editor/UnrealEd/Private/Kismet2/BlueprintEditorUtils.cpp
@@ -7466,6 +7466,22 @@ void FBlueprintEditorUtils::UpdateComponentTemplates(UBlueprint* Blueprint)
 			ReferencedTemplates.Add(ActorComp);
 		}
 	}
+
+	for (UObject* ActorComp : Blueprint->ComponentTemplates)
+	{
+		if (ensure(ActorComp))
+		{
+			// The most common case of an invalid component is a rogue LinkerPlaceholderExport object, 
+			// which are outered to the transient package.
+			// These can occur if we update our components prior to resolving placeholder objects.
+			// Typically, these are the result of underlying loader issues.
+			
+			const bool bIsValidComponentEntry = !FBlueprintSupport::IsDeferredDependencyPlaceholder(ActorComp);
+
+			ensureAlwaysMsgf(bIsValidComponentEntry, TEXT("Invalid component found in Blueprint's component template! Blueprint: %s. Component: %s"), *Blueprint->GetName(), *ActorComp->GetName());
+		}
+	}
+
 	Blueprint->ComponentTemplates.Empty();
 	Blueprint->ComponentTemplates.Append(ReferencedTemplates);
 }
