World Partition - Builder Commandlet: When processing a collection, output missing map files as Log rather than Warning

diff --git a/Engine/Source/Editor/UnrealEd/Private/Commandlets/WorldPartitionBuilderCommandlet.cpp b/Engine/Source/Editor/UnrealEd/Private/Commandlets/WorldPartitionBuilderCommandlet.cpp
--- a/Engine/Source/Editor/UnrealEd/Private/Commandlets/WorldPartitionBuilderCommandlet.cpp
+++ b/Engine/Source/Editor/UnrealEd/Private/Commandlets/WorldPartitionBuilderCommandlet.cpp
@@ -173,7 +173,7 @@ TArray<FString> UWorldPartitionBuilderCommandlet::GatherMapsFromCollection(const
 		}
 		else
 		{
-			UE_LOG(LogWorldPartitionBuilderCommandlet, Warning, TEXT("%s was not found"), *PackageName);
+			UE_LOG(LogWorldPartitionBuilderCommandlet, Log, TEXT("%s was not found"), *PackageName);
 		}
 	}
 
