CookOnTheFly2: Wait for the cooker to block on the AssetRegistry's scan before copying the list of existing packages from the AssetRegistry.
#rnx

diff --git a/Engine/Source/Editor/UnrealEd/Private/CookOnTheFlyServer.cpp b/Engine/Source/Editor/UnrealEd/Private/CookOnTheFlyServer.cpp
--- a/Engine/Source/Editor/UnrealEd/Private/CookOnTheFlyServer.cpp
+++ b/Engine/Source/Editor/UnrealEd/Private/CookOnTheFlyServer.cpp
@@ -322,27 +322,34 @@ public:
 		return Cooker.SandboxFile->GetSandboxDirectory();
 	}
 
-	const ITargetPlatform* AddPlatform(const FName& PlatformName)
+	virtual const ITargetPlatform* AddPlatform(FName PlatformName, bool& bOutAlreadyInitialized) override
 	{
 		UE::Cook::FPlatformManager::FReadScopeLock PlatformScopeLock(Cooker.PlatformManager->ReadLockPlatforms());
 		const ITargetPlatform* TargetPlatform = AddPlatformInternal(PlatformName);
 		if (!TargetPlatform)
 		{
 			UE_LOG(LogCook, Warning, TEXT("Trying to add invalid platform '%s' on the fly"), *PlatformName.ToString());
+			bOutAlreadyInitialized = false;
 			return nullptr;
 		}
 
+		bOutAlreadyInitialized = Cooker.PlatformManager->HasSessionPlatform(TargetPlatform);
 		Cooker.PlatformManager->AddRefCookOnTheFlyPlatform(PlatformName, Cooker);
 
 		return TargetPlatform;
 	}
 
-	virtual void RemovePlatform(const FName& PlatformName) override
+	virtual void RemovePlatform(FName PlatformName) override
 	{
 		UE::Cook::FPlatformManager::FReadScopeLock PlatformScopeLock(Cooker.PlatformManager->ReadLockPlatforms());
 		Cooker.PlatformManager->ReleaseCookOnTheFlyPlatform(PlatformName);
 	}
 
+	virtual bool IsSchedulerThread() const override
+	{
+		return IsInGameThread();
+	}
+
 	virtual void GetUnsolicitedFiles(const FName& PlatformName, const FString& Filename, const bool bIsCookable, TArray<FString>& OutUnsolicitedFiles) override
 	{
 		UE::Cook::FPlatformManager::FReadScopeLock PlatformsScopeLock(Cooker.PlatformManager->ReadLockPlatforms());
@@ -620,7 +627,8 @@ bool UCookOnTheFlyServer::StartCookOnTheFly(FCookOnTheFlyStartupOptions InCookOn
 		{
 			if (Connection.GetTargetPlatform())
 			{
-				CookOnTheFlyServerInterface->AddPlatform(Connection.GetPlatformName());
+				bool bAlreadyInitialized = false;
+				CookOnTheFlyServerInterface->AddPlatform(Connection.GetPlatformName(), bAlreadyInitialized);
 			}
 		});
 
@@ -761,6 +769,12 @@ void UCookOnTheFlyServer::StartCookOnTheFlySessionFromGameThread(ITargetPlatform
 	// AddCookOnTheFlyPlatformFromGameThread can be called on cooker startup which occurs in UUnrealEdEngine::Init
 	// before all plugins are loaded.
 	BlockOnAssetRegistry();
+
+	if (CookOnTheFlyRequestManager)
+	{
+		FName PlatformName(*TargetPlatform->PlatformName());
+		CookOnTheFlyRequestManager->OnSessionStarted(PlatformName, bFirstCookInThisProcess);
+	}
 }
 
 void UCookOnTheFlyServer::OnTargetPlatformsInvalidated()
diff --git a/Engine/Source/Editor/UnrealEd/Private/Cooker/CookOnTheFlyServerInterface.h b/Engine/Source/Editor/UnrealEd/Private/Cooker/CookOnTheFlyServerInterface.h
--- a/Engine/Source/Editor/UnrealEd/Private/Cooker/CookOnTheFlyServerInterface.h
+++ b/Engine/Source/Editor/UnrealEd/Private/Cooker/CookOnTheFlyServerInterface.h
@@ -52,15 +52,17 @@ public:
 	/** Returns the cooker sandbox directory path. */
 	virtual FString GetSandboxDirectory() const = 0;
 	
-	/**
-	 * Remove platform from the cook-on-the-fly session.
-	 *
-	 * @param PlatformName The platform name.
-	 */
-	virtual void RemovePlatform(const FName& PlatformName) = 0;
+	/** Request cooker set up and be ready to cook PlatformName. */
+	virtual const ITargetPlatform* AddPlatform(FName PlatformName, bool& bOutAlreadyInitialized) = 0;
+
+	/** Remove platform from the cook-on-the-fly session. */
+	virtual void RemovePlatform(FName PlatformName) = 0;
+
+	/** Report whether the current thread is the cooker's scheduler thread. */
+	virtual bool IsSchedulerThread() const = 0;
 
 	/**
-	 * Returns all unsolicited files that has been produced has a result of a cook request.
+	 * Returns all unsolicited files that have been produced as a result of a cook request.
 	 *
 	 * @param PlatformName The platform name.
 	 * @param Filename The filename.
@@ -103,9 +105,12 @@ public:
 
 	virtual void Tick() = 0;
 
-	/** Shutdown the reques tmanager. */
+	/** Shutdown the request manager. */
 	virtual void Shutdown() = 0;
 
+	/** Called when cooker has completed setup for a platform and the AssetRegistry is available */
+	virtual void OnSessionStarted(FName PlatformName, bool bFirstSessionInThisProcess) {}
+
 	/** Called when a new package is generated */
 	virtual void OnPackageGenerated(const FName& PackageName) = 0;
 
diff --git a/Engine/Source/Editor/UnrealEd/Private/Cooker/CookPlatformManager.cpp b/Engine/Source/Editor/UnrealEd/Private/Cooker/CookPlatformManager.cpp
--- a/Engine/Source/Editor/UnrealEd/Private/Cooker/CookPlatformManager.cpp
+++ b/Engine/Source/Editor/UnrealEd/Private/Cooker/CookPlatformManager.cpp
@@ -305,12 +305,18 @@ void FPlatformManager::AddRefCookOnTheFlyPlatform(FName PlatformName, UCookOnThe
 	if (!HasSessionPlatform(PlatformData->TargetPlatform))
 	{
 		CookOnTheFlyServer.WorkerRequests->AddCookOnTheFlyCallback([PlatformName,
-			LocalCookOnTheFlyServer = &CookOnTheFlyServer]()
+			LocalCookOnTheFlyServer = &CookOnTheFlyServer, this]()
 			{
 				ITargetPlatform* TargetPlatform = GetTargetPlatformManager()->FindTargetPlatform(PlatformName.ToString());
 				if (TargetPlatform)
 				{
-					LocalCookOnTheFlyServer->StartCookOnTheFlySessionFromGameThread(TargetPlatform);
+					// We might get multiple AddRef calls that add this callback before the first one reaches
+					// StartCookOnTheFlySessionFromGameThread, so check again whether some earlier request
+					// has already added the sessionplatform
+					if (!HasSessionPlatform(TargetPlatform))
+					{
+						LocalCookOnTheFlyServer->StartCookOnTheFlySessionFromGameThread(TargetPlatform);
+					}
 				}
 			});
 	}
diff --git a/Engine/Source/Editor/UnrealEd/Private/Cooker/IoStoreCookOnTheFlyRequestManager.cpp b/Engine/Source/Editor/UnrealEd/Private/Cooker/IoStoreCookOnTheFlyRequestManager.cpp
--- a/Engine/Source/Editor/UnrealEd/Private/Cooker/IoStoreCookOnTheFlyRequestManager.cpp
+++ b/Engine/Source/Editor/UnrealEd/Private/Cooker/IoStoreCookOnTheFlyRequestManager.cpp
@@ -76,9 +76,6 @@ public:
 		ConnectionServer->OnRequest(ECookOnTheFlyMessage::RecookPackages).BindRaw(this, &FIoStoreCookOnTheFlyRequestManager::HandleClientRequest);
 
 		MessageEndpoint = FMessageEndpoint::Builder("FCookOnTheFly");
-		TArray<FAssetData> AllAssets;
-		AssetRegistry->GetAllAssets(AllAssets, true);
-		PackageRegistry.Add(AllAssets);
 	}
 
 private:
@@ -99,10 +96,10 @@ private:
 			FPackageStoreEntryResource Entry;
 		};
 
-		FPlatformContext(FName InPlatformName, IPackageStoreWriter* InPackageWriter, FPackageRegistry& InPackageRegistry)
+		FPlatformContext(FIoStoreCookOnTheFlyRequestManager& InManager, FName InPlatformName, IPackageStoreWriter* InPackageWriter)
 			: PlatformName(InPlatformName)
+			, Manager(InManager)
 			, PackageWriter(InPackageWriter)
-			, PackageRegistry(InPackageRegistry)
 		{
 		}
 
@@ -112,6 +109,18 @@ private:
 			{
 				FPlatformProcess::ReturnSynchEventToPool(PackageCookedEvent);
 			}
+			FScopeLock _(&GetLock());
+			if (bCookSessionRequested)
+			{
+				// The request can only be in progress if there is another thread calling RequestCookSession,
+				// and all threads should have completed before we are allowed to destroy. If we ever destroy
+				// when requests are still in progress, we will need to add a cancel mechanic and wait for the
+				// cancel here.
+				check(!bCookSessionRequestInProgress);
+				Manager.CookOnTheFlyServer.RemovePlatform(PlatformName);
+				bCookSessionRequested = false;
+				CookerReadyForRequests->Reset();
+			}
 		}
 
 		FCriticalSection& GetLock()
@@ -135,10 +144,65 @@ private:
 			}
 		}
 
+		void OnSessionStarted()
+		{
+			FScopeLock _(&GetLock());
+			if (bCookSessionRequested)
+			{
+				return;
+			}
+
+			if (bCookSessionRequestInProgress)
+			{
+				bCookSessionRequested = true;
+				bCookSessionRequestInProgress = false;
+			}
+			CookerReadyForRequests->Trigger();
+		}
+
+		void RequestCookSessionAndWait()
+		{
+			bool bCallOnSessionStartedManually = false;
+			{
+				FScopeLock _(&GetLock());
+				if (bCookSessionRequested)
+				{
+					return;
+				}
+				if (!bCookSessionRequestInProgress)
+				{
+					bCookSessionRequestInProgress = true;
+					bool bAlreadyInitialized = false;
+					Manager.CookOnTheFlyServer.AddPlatform(PlatformName, bAlreadyInitialized);
+					if (bAlreadyInitialized)
+					{
+						// The platform has already been initialized by an AddPlatform call from elsewhere
+						// and the cooker will not be sending an OnSessionStarted event. Trigger it now.
+						bCallOnSessionStartedManually = true;
+					}
+				}
+			}
+			if (bCallOnSessionStartedManually)
+			{
+				Manager.OnSessionStarted(PlatformName, false /* bFirstSessionInThisProcess */);
+			}
+
+			if (CookerReadyForRequests->Wait(0))
+			{
+				return;
+			}
+			// AddPlatform always queues the request rather than handling it immediately,
+			// and it handles it later in the tick on the cooker's scheduler thread. So if
+			// we try to wait for the event from the scheduler thread we will deadlock.
+			check(!Manager.CookOnTheFlyServer.IsSchedulerThread());
+			CookerReadyForRequests->Wait();
+		}
+
 		EPackageStoreEntryStatus RequestCook(UE::Cook::ICookOnTheFlyServer& InCookOnTheFlyServer, const FPackageId& PackageId, FPackageStoreEntryResource& OutEntry)
 		{
 			FPackage& Package = GetPackage(PackageId);
-			FName PackageName = PackageRegistry.Get(PackageId);
+			check(Manager.bCookerGlobalsInitialized); // Caller should have called RequestCookSessionAndWait
+			FName PackageName = Manager.PackageRegistry.Get(PackageId);
 			if (Package.Status == EPackageStatus::Cooked)
 			{
 				UE_LOG(LogCookOnTheFly, Verbose, TEXT("0x%llX was already cooked"), PackageId.ValueForDebugging());
@@ -285,10 +349,13 @@ private:
 		FCriticalSection CriticalSection;
 		FName PlatformName;
 		TMap<FPackageId, TUniquePtr<FPackage>> Packages;
+		FIoStoreCookOnTheFlyRequestManager& Manager;
 		IPackageStoreWriter* PackageWriter = nullptr;
-		FPackageRegistry& PackageRegistry;
 		int32 SingleThreadedClientsCount = 0;
 		FEvent* PackageCookedEvent = nullptr;
+		FEventRef CookerReadyForRequests;
+		bool bCookSessionRequested = false;
+		bool bCookSessionRequestInProgress = false;
 	};
 
 	virtual bool Initialize() override
@@ -306,6 +373,33 @@ private:
 		return true;
 	}
 
+	void InitializeCookSessionGlobals()
+	{
+		// This function is only called from Cooker's scheduler thread via OnSessionStarted, so we do not need
+		// to worry about synchronization within it. Threads reading the information written here must call 
+		// RequestCookSessionAndWait or assert that bCookerGlobalsInitialized is true.
+		if (bCookerGlobalsInitialized)
+		{
+			return;
+		}
+		ON_SCOPE_EXIT{ bCookerGlobalsInitialized = true; };
+
+		TArray<FAssetData> AllAssets;
+		IAssetRegistry* AssetRegistry = IAssetRegistry::Get();
+		AssetRegistry->GetAllAssets(AllAssets, true);
+		PackageRegistry.Add(AllAssets);
+	}
+
+	virtual void OnSessionStarted(FName PlatformName, bool bFirstSessionInThisProcess) override
+	{
+		InitializeCookSessionGlobals();
+		FPlatformContext* Context = FindContext(PlatformName);
+		if (Context)
+		{
+			Context->OnSessionStarted();
+		}
+	}
+
 	virtual void Shutdown() override
 	{
 	}
@@ -330,6 +424,12 @@ private:
 			PackagesToRecook.Empty();
 		}
 
+		// PackageRegistry is not initialized until InitializeCookSessionGlobals, and afterwards is supposed to be immutable
+		// so it can be accessed from multiple threads without a lock. We need to access it in this function, so we have to 
+		// assert it has been initialized. It must be initialized if we have any PackagesToRecook, because all the functions
+		// that add elements to PackagesToRecook call RequestCookSessionAndWait first
+		check(bCookerGlobalsInitialized);
+
 		TArray<FName, TInlineAllocator<128>> PackageNames;
 		PackageNames.Reserve(PackageIds.Num());
 
@@ -404,7 +504,8 @@ private:
 			FScopeLock _(&ContextsCriticalSection);
 			if (!PlatformContexts.Contains(PlatformName))
 			{
-				TUniquePtr<FPlatformContext>& Context = PlatformContexts.Add(PlatformName, MakeUnique<FPlatformContext>(PlatformName, PackageWriter, PackageRegistry));
+				TUniquePtr<FPlatformContext>& Context = PlatformContexts.Add(PlatformName,
+					MakeUnique<FPlatformContext>(*this, PlatformName, PackageWriter));
 
 				PackageWriter->GetEntries([&Context](TArrayView<const FPackageStoreEntryResource> Entries,
 					TArrayView<const IPackageStoreWriter::FOplogCookInfo> CookInfos)
@@ -507,8 +608,8 @@ private:
 		}
 		else
 		{
-			FScopeLock _(&ContextsCriticalSection);
 			FPlatformContext& Context = GetContext(PlatformName);
+			Context.RequestCookSessionAndWait();
 			FScopeLock __(&Context.GetLock());
 			FCompletedPackages CompletedPackages;
 			Context.GetCompletedPackages(CompletedPackages);
@@ -535,6 +636,8 @@ private:
 		TRACE_CPUPROFILER_EVENT_SCOPE(CookOnTheFly::HandleCookPackageRequest);
 
 		FPlatformContext& Context = GetContext(PlatformName);
+		Context.RequestCookSessionAndWait();
+
 		FCookPackageRequest CookRequest = Request.GetBodyAs<FCookPackageRequest>();
 		FCookPackageResponse CookResponse;
 
@@ -587,6 +690,9 @@ private:
 
 		TRACE_CPUPROFILER_EVENT_SCOPE(CookOnTheFly::HandleRecookPackagesRequest);
 
+		FPlatformContext& Context = GetContext(PlatformName);
+		Context.RequestCookSessionAndWait();
+
 		FRecookPackagesRequest RecookRequest = Request.GetBodyAs<FRecookPackagesRequest>();
 
 		UE_LOG(LogCookOnTheFly, Display, TEXT("Received recook request for %d packages"), RecookRequest.PackageIds.Num());
@@ -641,6 +747,7 @@ private:
 	void OnPackageStoreEntryCreated(const IPackageStoreWriter::FEntryCreatedEventArgs& EventArgs)
 	{
 		FPlatformContext& Context = GetContext(EventArgs.PlatformName);
+		Context.RequestCookSessionAndWait();
 
 		FScopeLock _(&Context.GetLock());
 		for (const FPackageId& ImportedPackageId : EventArgs.Entry.ImportedPackageIds)
@@ -737,7 +844,19 @@ private:
 		}
 	}
 	
-	FPlatformContext& GetContext(const FName& PlatformName)
+	FPlatformContext* FindContext(FName PlatformName)
+	{
+		FScopeLock _(&ContextsCriticalSection);
+		TUniquePtr<FPlatformContext>* Ctx = PlatformContexts.Find(PlatformName);
+		if (!Ctx)
+		{
+			return nullptr;
+		}
+		check(Ctx->IsValid());
+		return Ctx->Get();
+	}
+
+	FPlatformContext& GetContext(FName PlatformName)
 	{
 		FScopeLock _(&ContextsCriticalSection);
 		TUniquePtr<FPlatformContext>& Ctx = PlatformContexts.FindChecked(PlatformName);
@@ -770,6 +889,7 @@ private:
 	FPackageRegistry PackageRegistry;
 	FCriticalSection PackagesToRecookCritical;
 	TSet<FPackageId> PackagesToRecook;
+	bool bCookerGlobalsInitialized = false;
 };
 
 namespace UE { namespace Cook
