Move the Focus Selected command to the Actor context menu, rename it to Frame Selected, and make the View Options menu items appear in the Actor menu in the main menu bar as well as in the contextual menu.

diff --git a/Engine/Source/Editor/EditorStyle/Private/StarshipStyle.cpp b/Engine/Source/Editor/EditorStyle/Private/StarshipStyle.cpp
--- a/Engine/Source/Editor/EditorStyle/Private/StarshipStyle.cpp
+++ b/Engine/Source/Editor/EditorStyle/Private/StarshipStyle.cpp
@@ -431,6 +431,7 @@ void FStarshipEditorStyle::FStyle::SetupGeneralStyles()
 		Set("Icons.PlaceActors", new IMAGE_BRUSH_SVG("Starship/Common/PlaceActors", Icon16x16));
 		Set("Icons.ReplaceActor", new IMAGE_BRUSH_SVG("Starship/Common/ReplaceActors", Icon16x16));
 		Set("Icons.GroupActors", new IMAGE_BRUSH_SVG("Starship/Common/GroupActors", Icon16x16));
+		Set("Icons.FrameActor", new IMAGE_BRUSH_SVG("Starship/Common/FrameActor", Icon16x16));
 		Set("Icons.Transform", new IMAGE_BRUSH_SVG("Starship/Common/transform-local", Icon16x16));
 		Set("Icons.SetShowPivot", new IMAGE_BRUSH_SVG("Starship/Common/SetShowPivot", Icon16x16));
 		Set("Icons.Snap", new IMAGE_BRUSH_SVG("Starship/Common/Snap", Icon16x16));
diff --git a/Engine/Source/Editor/LevelEditor/Private/LevelEditorContextMenu.cpp b/Engine/Source/Editor/LevelEditor/Private/LevelEditorContextMenu.cpp
--- a/Engine/Source/Editor/LevelEditor/Private/LevelEditorContextMenu.cpp
+++ b/Engine/Source/Editor/LevelEditor/Private/LevelEditorContextMenu.cpp
@@ -388,13 +388,19 @@ void FLevelEditorContextMenu::RegisterActorContextMenu()
 		}
 
 
-		if (LevelEditorContext->ContextType == ELevelEditorMenuContext::Viewport || LevelEditorContext->ContextType == ELevelEditorMenuContext::SceneOutliner)
 		{
-			// Options that only appear in the viewport context menu or scene outliner (will affect the current viewport)
+			// Options that affect the current viewport.
 			// In most cases, you DO NOT want to extend this section; look at ActorUETools or ActorTypeTools below
 			FToolMenuSection& Section = InMenu->AddSection("ActorViewOptions", LOCTEXT("ViewOptionsHeading", "View Options"));
 			const FVector* ClickLocation = &GEditor->ClickLocation;
 
+			Section.AddMenuEntry(
+				FEditorViewportCommands::Get().FocusViewportToSelection,
+				TAttribute<FText>(),
+				TAttribute<FText>(),
+				FSlateIcon(FAppStyle::GetAppStyleSetName(), "Icons.FrameActor")
+			);
+
 			// This keys off the mouse position so can only appear in the viewport
 			if (LevelEditorContext->ContextType == ELevelEditorMenuContext::Viewport)
 			{
diff --git a/Engine/Source/Editor/LevelEditor/Private/LevelEditorMenu.cpp b/Engine/Source/Editor/LevelEditor/Private/LevelEditorMenu.cpp
--- a/Engine/Source/Editor/LevelEditor/Private/LevelEditorMenu.cpp
+++ b/Engine/Source/Editor/LevelEditor/Private/LevelEditorMenu.cpp
@@ -764,7 +764,6 @@ void FLevelEditorMenu::RegisterSelectMenu()
 				}
 				InSection.AddMenuEntry(FLevelEditorCommands::Get().SelectAllActorsOfSameClass, SelectAllActorsText);
 			}));
-		Section.AddMenuEntry(FEditorViewportCommands::Get().FocusViewportToSelection);
 	}
 
 	{
diff --git a/Engine/Source/Editor/UnrealEd/Private/EditorViewportCommands.cpp b/Engine/Source/Editor/UnrealEd/Private/EditorViewportCommands.cpp
--- a/Engine/Source/Editor/UnrealEd/Private/EditorViewportCommands.cpp
+++ b/Engine/Source/Editor/UnrealEd/Private/EditorViewportCommands.cpp
@@ -161,8 +161,8 @@ void FEditorViewportCommands::RegisterCommands()
 #endif
 	UI_COMMAND( CycleTransformGizmos, "Cycle Between Translate, Rotate, and Scale", "Cycles the transform gizmos between translate, rotate, and scale", EUserInterfaceActionType::Button, FInputChord(EKeys::SpaceBar) );
 	
-	UI_COMMAND( FocusAllViewportsToSelection, "Focus Selected Actors in All Viewports", "Moves the camera in front of the selected actors in all open viewports", EUserInterfaceActionType::Button, FInputChord(EKeys::F, EModifierKey::Shift) );
-	UI_COMMAND( FocusViewportToSelection, "Focus Selected", "Moves the camera in front of the selection", EUserInterfaceActionType::Button, FInputChord( EKeys::F ) );
+	UI_COMMAND( FocusAllViewportsToSelection, "Frame Selected Actors in All Viewports", "Centers all open viewports on the selected actors", EUserInterfaceActionType::Button, FInputChord(EKeys::F, EModifierKey::Shift) );
+	UI_COMMAND( FocusViewportToSelection, "Frame Selected", "Centers the viewport on the selected actors", EUserInterfaceActionType::Button, FInputChord( EKeys::F ) );
 
 	UI_COMMAND( LocationGridSnap, "Grid Snap", "Enables or disables snapping to the grid when dragging objects around", EUserInterfaceActionType::ToggleButton, FInputChord() );
 	UI_COMMAND( RotationGridSnap, "Rotation Snap", "Enables or disables snapping objects to a rotation grid", EUserInterfaceActionType::ToggleButton, FInputChord() );
