Add bindable viewport camera rotation UI commands

diff --git a/Engine/Source/Editor/UnrealEd/Private/EditorViewportClient.cpp b/Engine/Source/Editor/UnrealEd/Private/EditorViewportClient.cpp
--- a/Engine/Source/Editor/UnrealEd/Private/EditorViewportClient.cpp
+++ b/Engine/Source/Editor/UnrealEd/Private/EditorViewportClient.cpp
@@ -1784,6 +1784,11 @@ void FEditorViewportClient::UpdateCameraMovement( float DeltaTime )
 		bool bDownKeyState = false;
 		bool bZoomOutKeyState = false;
 		bool bZoomInKeyState = false;
+
+		bool bRotateUpKeyState = false;
+		bool bRotateDownKeyState = false;
+		bool bRotateLeftKeyState = false;
+		bool bRotateRightKeyState = false;
 		// Iterate through all key mappings to generate key state flags
 		for (uint32 i = 0; i < static_cast<uint8>(EMultipleKeyBindingIndex::NumChords); ++i)
 		{
@@ -1797,6 +1802,16 @@ void FEditorViewportClient::UpdateCameraMovement( float DeltaTime )
 			bDownKeyState |= Viewport->KeyState(FViewportNavigationCommands::Get().Down->GetActiveChord(ChordIndex)->Key);
 			bZoomOutKeyState |= Viewport->KeyState(FViewportNavigationCommands::Get().FovZoomOut->GetActiveChord(ChordIndex)->Key);
 			bZoomInKeyState |= Viewport->KeyState(FViewportNavigationCommands::Get().FovZoomIn->GetActiveChord(ChordIndex)->Key);
+
+			bRotateUpKeyState |= Viewport->KeyState(FViewportNavigationCommands::Get().RotateUp->GetActiveChord(ChordIndex)->Key);
+			bRotateDownKeyState |= Viewport->KeyState(FViewportNavigationCommands::Get().RotateDown->GetActiveChord(ChordIndex)->Key);
+			bRotateLeftKeyState |= Viewport->KeyState(FViewportNavigationCommands::Get().RotateLeft->GetActiveChord(ChordIndex)->Key);
+			bRotateRightKeyState |= Viewport->KeyState(FViewportNavigationCommands::Get().RotateRight->GetActiveChord(ChordIndex)->Key);
+		}
+
+		if (!CameraController->IsRotating())
+		{
+			CameraController->GetConfig().bForceRotationalPhysics = false;
 		}
 
 		// Forward/back
@@ -1853,6 +1868,30 @@ void FEditorViewportClient::UpdateCameraMovement( float DeltaTime )
 			CameraUserImpulseData->ZoomOutInImpulse -= 1.0f;
 		}
 
+		// Rotate up/down
+		if (bRemapWASDKeys && bRotateUpKeyState)
+		{
+			CameraUserImpulseData->RotatePitchImpulse += 1.0f * CameraController->GetConfig().RotationMultiplier;
+			CameraController->GetConfig().bForceRotationalPhysics = true;
+		}
+		if (bRemapWASDKeys && bRotateDownKeyState)
+		{
+			CameraUserImpulseData->RotatePitchImpulse -= 1.0f * CameraController->GetConfig().RotationMultiplier;
+			CameraController->GetConfig().bForceRotationalPhysics = true;
+		}
+
+		// Rotate left/right
+		if (bRemapWASDKeys && bRotateLeftKeyState)
+		{
+			CameraUserImpulseData->RotateYawImpulse -= 1.0f * CameraController->GetConfig().RotationMultiplier;
+			CameraController->GetConfig().bForceRotationalPhysics = true;
+		}
+		if (bRemapWASDKeys && bRotateRightKeyState)
+		{
+			CameraUserImpulseData->RotateYawImpulse += 1.0f * CameraController->GetConfig().RotationMultiplier;
+			CameraController->GetConfig().bForceRotationalPhysics = true;
+		}
+
 		// Record Stats
 		if ( CameraUserImpulseData->MoveForwardBackwardImpulse != 0 || CameraUserImpulseData->MoveRightLeftImpulse != 0 )
 		{
@@ -1867,11 +1906,6 @@ void FEditorViewportClient::UpdateCameraMovement( float DeltaTime )
 			FEditorViewportStats::Using(FEditorViewportStats::CAT_PERSPECTIVE_KEYBOARD_FOV_ZOOM);
 		}
 
-		if (!CameraController->IsRotating())
-		{
-			CameraController->GetConfig().bForceRotationalPhysics = false;
-		}
-
 		if( GetDefault<ULevelEditorViewportSettings>()->bLevelEditorJoystickControls )
 		{
 			//Now update for cached joystick info (relative movement first)
@@ -2594,7 +2628,11 @@ bool FEditorViewportClient::IsFlightCameraActive() const
 			|| Viewport->KeyState(FViewportNavigationCommands::Get().Up->GetActiveChord(ChordIndex)->Key)
 			|| Viewport->KeyState(FViewportNavigationCommands::Get().Down->GetActiveChord(ChordIndex)->Key)
 			|| Viewport->KeyState(FViewportNavigationCommands::Get().FovZoomIn->GetActiveChord(ChordIndex)->Key)
-			|| Viewport->KeyState(FViewportNavigationCommands::Get().FovZoomOut->GetActiveChord(ChordIndex)->Key));
+			|| Viewport->KeyState(FViewportNavigationCommands::Get().FovZoomOut->GetActiveChord(ChordIndex)->Key)
+			|| Viewport->KeyState(FViewportNavigationCommands::Get().RotateUp->GetActiveChord(ChordIndex)->Key)
+			|| Viewport->KeyState(FViewportNavigationCommands::Get().RotateDown->GetActiveChord(ChordIndex)->Key)
+			|| Viewport->KeyState(FViewportNavigationCommands::Get().RotateLeft->GetActiveChord(ChordIndex)->Key)
+			|| Viewport->KeyState(FViewportNavigationCommands::Get().RotateRight->GetActiveChord(ChordIndex)->Key));
 	}
 	const bool bIsUsingTrackpad = FSlateApplication::Get().IsUsingTrackpad();
 
@@ -6547,6 +6585,11 @@ void FViewportNavigationCommands::RegisterCommands()
 
 	UI_COMMAND(FovZoomIn, "FOV Zoom In", "Narrows the camers FOV", EUserInterfaceActionType::Button, FInputChord(EKeys::C));
 	UI_COMMAND(FovZoomOut, "FOV Zoom Out", "Widens the camera FOV", EUserInterfaceActionType::Button, FInputChord(EKeys::Z));
+
+	UI_COMMAND(RotateUp, "Rotate Up", "Rotates the camera Up", EUserInterfaceActionType::Button, FInputChord());
+	UI_COMMAND(RotateDown, "Rotate Down", "Rotates the camera Down", EUserInterfaceActionType::Button, FInputChord());
+	UI_COMMAND(RotateLeft, "Rotate Left", "Rotates the camera Left", EUserInterfaceActionType::Button, FInputChord());
+	UI_COMMAND(RotateRight, "Rotate Right", "Rotates the camera Right", EUserInterfaceActionType::Button, FInputChord());
 }
 
 #undef LOCTEXT_NAMESPACE
diff --git a/Engine/Source/Editor/UnrealEd/Public/EditorViewportClient.h b/Engine/Source/Editor/UnrealEd/Public/EditorViewportClient.h
--- a/Engine/Source/Editor/UnrealEd/Public/EditorViewportClient.h
+++ b/Engine/Source/Editor/UnrealEd/Public/EditorViewportClient.h
@@ -78,6 +78,11 @@ public:
 
 	TSharedPtr< FUICommandInfo > FovZoomIn;
 	TSharedPtr< FUICommandInfo > FovZoomOut;
+
+	TSharedPtr< FUICommandInfo > RotateUp;
+	TSharedPtr< FUICommandInfo > RotateDown;
+	TSharedPtr< FUICommandInfo > RotateLeft;
+	TSharedPtr< FUICommandInfo > RotateRight;
 };
 
 // FPrioritizedInputChord
