Avoid crashing when running a DiffOnly cook due to trying to serialize a property at an incorrect offset in the stream because of PAST length differences in the stream.

diff --git a/Engine/Source/Editor/UnrealEd/Private/Cooker/DiffWriterArchive.cpp b/Engine/Source/Editor/UnrealEd/Private/Cooker/DiffWriterArchive.cpp
--- a/Engine/Source/Editor/UnrealEd/Private/Cooker/DiffWriterArchive.cpp
+++ b/Engine/Source/Editor/UnrealEd/Private/Cooker/DiffWriterArchive.cpp
@@ -826,7 +826,14 @@ void FAccumulator::CompareWithPreviousForSection(const FPackageData& SourcePacka
 			FString AfterPropertyVal;
 			if (FProperty* SerProp = DifferenceCallstackData.SerializedProp)
 			{
-				if (SourceSize == DestSize && ShouldDumpPropertyValueState(SerProp))
+				// We don't attempt to serialize properties when we have already encountered at least one diff in the asset.
+				// That is because we don't handle length differences in the source and destination data caused by the previous
+				// diffs.  Those previous length differences could mean the current diff is at a position that is offset and
+				// invalid to serialize the property from.  That can result in things like serializing an array or string which has
+				// a negative number of elements, or other invalid situations that lead to an assert or crash.  If we must
+				// sesrialize these properties, we would have to ensure that we keep an appropriate offset separate for the source
+				// archive and the dest archive.
+				if ((SourceSize == DestSize) && (InOutDiffsLogged < 2) && ShouldDumpPropertyValueState(SerProp))
 				{
 					// Walk backwards until we find a callstack which wasn't from the given property
 					int64 OffsetX = DestAbsoluteOffset;
