Fix copying collapsed nodes in MaterialGraph

diff --git a/Engine/Source/Editor/UnrealEd/Classes/MaterialGraph/MaterialGraphNode_Composite.h b/Engine/Source/Editor/UnrealEd/Classes/MaterialGraph/MaterialGraphNode_Composite.h
--- a/Engine/Source/Editor/UnrealEd/Classes/MaterialGraph/MaterialGraphNode_Composite.h
+++ b/Engine/Source/Editor/UnrealEd/Classes/MaterialGraph/MaterialGraphNode_Composite.h
@@ -33,7 +33,8 @@ class UMaterialGraphNode_Composite : public UMaterialGraphNode
 	virtual void DestroyNode() override; 
 	virtual void PrepareForCopying() override;
 	virtual void PostPasteNode() override; 
-	virtual void OnRenameNode(const FString& NewName) override; 
+	virtual void OnRenameNode(const FString& NewName) override;
+	virtual void ReconstructNode() override;
 	virtual TArray<UEdGraph*> GetSubGraphs() const override { return TArray<UEdGraph*>( { BoundGraph } ); }
 	virtual TSharedPtr<SGraphNode> CreateVisualWidget() override;
 	//~ End UEdGraphNode Interface.
diff --git a/Engine/Source/Editor/UnrealEd/Private/MaterialGraphNode_Composite.cpp b/Engine/Source/Editor/UnrealEd/Private/MaterialGraphNode_Composite.cpp
--- a/Engine/Source/Editor/UnrealEd/Private/MaterialGraphNode_Composite.cpp
+++ b/Engine/Source/Editor/UnrealEd/Private/MaterialGraphNode_Composite.cpp
@@ -139,6 +139,11 @@ void UMaterialGraphNode_Composite::OnRenameNode(const FString& NewName)
 	MaterialDirtyDelegate.ExecuteIfBound();
 }
 
+void UMaterialGraphNode_Composite::ReconstructNode()
+{
+	FixupInputAndOutputPinBases();
+}
+
 TSharedPtr<SGraphNode> UMaterialGraphNode_Composite::CreateVisualWidget()
 {
 	return SNew(SGraphNodeMaterialComposite, this);
@@ -167,7 +172,7 @@ void UMaterialGraphNode_Composite::FixupInputAndOutputPinBases()
 			}
 		}
 
-		ReconstructNode();
+		UMaterialGraphNode_Base::ReconstructNode();
 	}
 }
 
