Fix duplicate LOCTEXT macros in AssetEditorSubsystem and AssetViewWidgets

diff --git a/Engine/Source/Editor/ContentBrowser/Private/AssetViewWidgets.cpp b/Engine/Source/Editor/ContentBrowser/Private/AssetViewWidgets.cpp
--- a/Engine/Source/Editor/ContentBrowser/Private/AssetViewWidgets.cpp
+++ b/Engine/Source/Editor/ContentBrowser/Private/AssetViewWidgets.cpp
@@ -831,7 +831,7 @@ TSharedRef<SWidget> SAssetViewItem::CreateToolTipWidget() const
 			{
 				if(AssetItem->GetItem().CanView())
 				{
-					PublicStateText = LOCTEXT("ReadOnlyAssetState", "View / Read Only");
+					PublicStateText = LOCTEXT("ViewReadOnlyAssetState", "View / Read Only");
 					PublicStateIcon = FAppStyle::GetBrush("AssetEditor.ReadOnlyOpenable");
 
 				}
diff --git a/Engine/Source/Editor/UnrealEd/Private/Subsystems/AssetEditorSubsystem.cpp b/Engine/Source/Editor/UnrealEd/Private/Subsystems/AssetEditorSubsystem.cpp
--- a/Engine/Source/Editor/UnrealEd/Private/Subsystems/AssetEditorSubsystem.cpp
+++ b/Engine/Source/Editor/UnrealEd/Private/Subsystems/AssetEditorSubsystem.cpp
@@ -1531,7 +1531,7 @@ bool UAssetEditorSubsystem::CanOpenEditorForAsset(UObject* Asset, const EAssetTy
 		{
 			if(OutErrorMsg)
 			{
-				*OutErrorMsg = LOCTEXT("AssetTypeDoesntSupportOpenMethod", "UClasses cannot be opened in Asset Editors!");
+				*OutErrorMsg = LOCTEXT("UClassesCantBeOpened", "UClasses cannot be opened in Asset Editors!");
 			}
 			return false;
 		}
@@ -1544,7 +1544,7 @@ bool UAssetEditorSubsystem::CanOpenEditorForAsset(UObject* Asset, const EAssetTy
 		{
 			if(OutErrorMsg)
 			{
-				*OutErrorMsg = LOCTEXT("AssetTypeDoesntSupportOpenMethod", "Unable to Edit Cooked asset");
+				*OutErrorMsg = LOCTEXT("AssetCantBeEdited", "Unable to Edit Cooked asset");
 			}
 			return false;
 		}
@@ -1560,7 +1560,7 @@ bool UAssetEditorSubsystem::CanOpenEditorForAsset(UObject* Asset, const EAssetTy
 				{
 					if(OutErrorMsg)
 					{
-						*OutErrorMsg = LOCTEXT("AssetTypeDoesntSupportOpenMethod", "This asset does not support being opened in read only mode.");
+						*OutErrorMsg = LOCTEXT("AssetDoesntSupportOpenMethod", "This asset does not support being opened in read only mode.");
 					}
 					return false;
 				}
