Word Partition - Builder Commandlet: Start by testing if the provided Tokens[0] is a collection, as this is a lot faster than doing a full scan for a short package name that in the end, doesn't exist

diff --git a/Engine/Source/Editor/UnrealEd/Private/Commandlets/WorldPartitionBuilderCommandlet.cpp b/Engine/Source/Editor/UnrealEd/Private/Commandlets/WorldPartitionBuilderCommandlet.cpp
--- a/Engine/Source/Editor/UnrealEd/Private/Commandlets/WorldPartitionBuilderCommandlet.cpp
+++ b/Engine/Source/Editor/UnrealEd/Private/Commandlets/WorldPartitionBuilderCommandlet.cpp
@@ -80,11 +80,7 @@ int32 UWorldPartitionBuilderCommandlet::Main(const FString& Params)
 
 	// Parse map name or maps collection
 	FString MapLongPackageName;
-	if (FPackageName::SearchForPackageOnDisk(Tokens[0], &MapLongPackageName))
-	{
-		MapPackagesNames = { MapLongPackageName };
-	}	
-	else if (CollectionManager.CollectionExists(FName(Tokens[0]), ECollectionShareType::CST_All))
+	if (CollectionManager.CollectionExists(FName(Tokens[0]), ECollectionShareType::CST_All))
 	{
 		MapPackagesNames = GatherMapsFromCollection(Tokens[0]);
 		if (MapPackagesNames.IsEmpty())
@@ -93,6 +89,10 @@ int32 UWorldPartitionBuilderCommandlet::Main(const FString& Params)
 		    return 0;
 	    }
 	}
+	else if (FPackageName::SearchForPackageOnDisk(Tokens[0], &MapLongPackageName))
+	{
+		MapPackagesNames = { MapLongPackageName };
+	}	
 	else
 	{
 		UE_LOG(LogWorldPartitionBuilderCommandlet, Error, TEXT("Missing world(s) as the first argument to the commandlet. Either supply the world name directly (WorldName or /Path/To/WorldName), or provide a collection name to have the builder operate on a set of maps."));
