WorldPartition - Builder Commandlet: Avoid querying the asset registry when validating maps found in a collection, as the initial scan is likely not completed at this point
* Fixes some collection maps not being processed by the builders

diff --git a/Engine/Source/Editor/UnrealEd/Private/Commandlets/WorldPartitionBuilderCommandlet.cpp b/Engine/Source/Editor/UnrealEd/Private/Commandlets/WorldPartitionBuilderCommandlet.cpp
--- a/Engine/Source/Editor/UnrealEd/Private/Commandlets/WorldPartitionBuilderCommandlet.cpp
+++ b/Engine/Source/Editor/UnrealEd/Private/Commandlets/WorldPartitionBuilderCommandlet.cpp
@@ -5,6 +5,7 @@
 #include "CoreMinimal.h"
 #include "EngineUtils.h"
 #include "EditorWorldUtils.h"
+#include "FileHelpers.h"
 #include "Logging/LogMacros.h"
 #include "Misc/CommandLine.h"
 #include "Misc/EngineVersion.h"
@@ -17,8 +18,6 @@
 
 #include "CollectionManagerModule.h"
 #include "ICollectionManager.h"
-#include "AssetRegistry/AssetData.h"
-#include "AssetRegistry/AssetRegistryModule.h"
 
 #include "ISourceControlModule.h"
 #include "ISourceControlProvider.h"
@@ -146,7 +145,6 @@ TArray<FString> UWorldPartitionBuilderCommandlet::GatherMapsFromCollection(const
 	TArray<FString> MapPackagesNames;
 
 	ICollectionManager& CollectionManager = FModuleManager::LoadModuleChecked<FCollectionManagerModule>("CollectionManager").Get();
-	IAssetRegistry& AssetRegistry = FModuleManager::LoadModuleChecked<FAssetRegistryModule>("AssetRegistry").Get();
 
 	TArray<FSoftObjectPath> AssetsPaths;
 	CollectionManager.GetAssetsInCollection(FName(CollectionName), ECollectionShareType::CST_All, AssetsPaths, ECollectionRecursionFlags::SelfAndChildren);
@@ -155,25 +153,15 @@ TArray<FString> UWorldPartitionBuilderCommandlet::GatherMapsFromCollection(const
 	for (const auto& AssetPath : AssetsPaths)
 	{
 		FString PackageName = AssetPath.GetLongPackageName();
-		UE_LOG(LogWorldPartitionBuilderCommandlet, Display, TEXT("* %s"), *PackageName);
 
-		const bool bIncludeOnlyOnDiskAssets = true;
-		FAssetData AssetData = AssetRegistry.GetAssetByObjectPath(AssetPath, bIncludeOnlyOnDiskAssets);
-
-		if (AssetData.IsValid())
+		if (FEditorFileUtils::IsMapPackageAsset(PackageName))
 		{
-			if (AssetData.AssetClassPath == UWorld::StaticClass()->GetClassPathName())
-			{
-				MapPackagesNames.Add(PackageName);
-			}
-			else
-			{
-				UE_LOG(LogWorldPartitionBuilderCommandlet, Warning, TEXT("%s is not a map asset"), *PackageName);
-			}
+			UE_LOG(LogWorldPartitionBuilderCommandlet, Display, TEXT("* %s"), *PackageName);
+			MapPackagesNames.Add(PackageName);
 		}
 		else
 		{
-			UE_LOG(LogWorldPartitionBuilderCommandlet, Log, TEXT("%s was not found"), *PackageName);
+			UE_LOG(LogWorldPartitionBuilderCommandlet, Log, TEXT("%s was not found or is not a map package"), *PackageName);
 		}
 	}
 
