Cooker: GeneratedPackages: During garbage collection, for generator and generated packages, fix the condition for whether we need to keep all CachedObjectsInOuter referenced to be HasTakenOverCachedCookedPlatformData. The previous condition of GetSaveState() > CallObjectsToMove was incorrect and should have been >=. HasTakenOverCachedCookedPlatformData is the same thing as comparing with State >= CallObjectsToMove, but is more authoritative if we ever change the state where we do the take over.
The incorrect comparison caused us not to keep the objects referenced, which allowed them to be GC'd, which breaks the CookPackageSplitter contract and also triggered an error in FGeneratorPackage::ResetSaveState when we detected that deletion and demoted the package out of save.
#rnx

diff --git a/Engine/Source/Editor/UnrealEd/Private/Cooker/CookPackageData.cpp b/Engine/Source/Editor/UnrealEd/Private/Cooker/CookPackageData.cpp
--- a/Engine/Source/Editor/UnrealEd/Private/Cooker/CookPackageData.cpp
+++ b/Engine/Source/Editor/UnrealEd/Private/Cooker/CookPackageData.cpp
@@ -1789,7 +1789,7 @@ void FGeneratorPackage::PreGarbageCollect(FCookGenerationInfo& Info, TArray<TObj
 			bOutShouldDemote = true;
 		}
 	}
-	if (Info.GetSaveState() > FCookGenerationInfo::ESaveState::CallObjectsToMove)
+	if (Info.HasTakenOverCachedCookedPlatformData())
 	{
 		if (GetCookPackageSplitterInstance()->UseInternalReferenceToAvoidGarbageCollect())
 		{
