Register Editor Modes after all modules have loaded

diff --git a/Engine/Source/Editor/UnrealEd/Private/Subsystems/AssetEditorSubsystem.cpp b/Engine/Source/Editor/UnrealEd/Private/Subsystems/AssetEditorSubsystem.cpp
--- a/Engine/Source/Editor/UnrealEd/Private/Subsystems/AssetEditorSubsystem.cpp
+++ b/Engine/Source/Editor/UnrealEd/Private/Subsystems/AssetEditorSubsystem.cpp
@@ -59,7 +59,7 @@ void UAssetEditorSubsystem::Initialize(FSubsystemCollectionBase& Collection)
 
 	GEditor->OnEditorClose().AddUObject(this, &UAssetEditorSubsystem::OnEditorClose);
 	FCoreDelegates::OnEnginePreExit.AddUObject(this, &UAssetEditorSubsystem::UnregisterEditorModes);
-	FCoreDelegates::OnPostEngineInit.AddUObject(this, &UAssetEditorSubsystem::RegisterEditorModes);
+	FCoreDelegates::OnAllModuleLoadingPhasesComplete.AddUObject(this, &UAssetEditorSubsystem::RegisterEditorModes);
 
 	if (FAssetRegistryModule* AssetRegistryModule = FModuleManager::GetModulePtr<FAssetRegistryModule>(TEXT("AssetRegistry")))
 	{
