Make sure the DLC assets are not ignored if -DlcRecook is specified, even though the source development asset registry says they are already cooked.

#rnx

diff --git a/Engine/Source/Editor/UnrealEd/Classes/CookOnTheSide/CookOnTheFlyServer.h b/Engine/Source/Editor/UnrealEd/Classes/CookOnTheSide/CookOnTheFlyServer.h
--- a/Engine/Source/Editor/UnrealEd/Classes/CookOnTheSide/CookOnTheFlyServer.h
+++ b/Engine/Source/Editor/UnrealEd/Classes/CookOnTheSide/CookOnTheFlyServer.h
@@ -1256,6 +1256,8 @@ private:
 
 	UNREALED_API FString GetContentDirectoryForDLC() const;
 
+	UNREALED_API FString GetMountedAssetPathForDLC() const;
+
 	FString GetMetadataDirectory() const;
 
 	/**
diff --git a/Engine/Source/Editor/UnrealEd/Private/CookOnTheFlyServer.cpp b/Engine/Source/Editor/UnrealEd/Private/CookOnTheFlyServer.cpp
--- a/Engine/Source/Editor/UnrealEd/Private/CookOnTheFlyServer.cpp
+++ b/Engine/Source/Editor/UnrealEd/Private/CookOnTheFlyServer.cpp
@@ -1209,6 +1209,17 @@ FString UCookOnTheFlyServer::GetBaseDirectoryForDLC() const
 	return FPaths::ProjectPluginsDir() / CookByTheBookOptions->DlcName;
 }
 
+FString UCookOnTheFlyServer::GetMountedAssetPathForDLC() const
+{
+	TSharedPtr<IPlugin> Plugin = IPluginManager::Get().FindPlugin(CookByTheBookOptions->DlcName);
+	if (Plugin.IsValid())
+	{
+		return Plugin->GetMountedAssetPath();
+	}
+
+	return FString::Printf(TEXT("/%s/"), *CookByTheBookOptions->DlcName);
+}
+
 FString UCookOnTheFlyServer::GetContentDirectoryForDLC() const
 {
 	return GetBaseDirectoryForDLC() / TEXT("Content");
@@ -11500,9 +11511,14 @@ void UCookOnTheFlyServer::RecordDLCPackagesFromBaseGame(FBeginCookContext& Begin
 
 		TArray<FName>& PlatformBasedPackages = CookByTheBookOptions->BasedOnReleaseCookedPackages.FindOrAdd(PlatformName);
 		PlatformBasedPackages.Reset(ActivePackageList.Num());
+		FString PluginPathToSkip(!!(CookOptions & ECookByTheBookOptions::DlcRecook) ? GetMountedAssetPathForDLC() : "");
 		for (UE::Cook::FConstructPackageData& PackageData : ActivePackageList)
 		{
-			PlatformBasedPackages.Add(PackageData.NormalizedFileName);
+			// If we are recooking a DLC, skip adding its assets to 'BasedOnReleaseCookedPackages' so that they don't get ignored in the final stages of the cook.
+			if ((PluginPathToSkip.Len() == 0) || !PackageData.PackageName.ToString().StartsWith(PluginPathToSkip))
+			{
+				PlatformBasedPackages.Add(PackageData.NormalizedFileName);
+			}
 		}
 		bFirstAddExistingPackageDatas = false;
 	}
