[Component Visualizer] Deprecate unused function for showing based on actor selection. This is explicitly the design of component visualizers, and is unnecessary.

diff --git a/Engine/Source/Editor/UnrealEd/Public/ComponentVisualizer.h b/Engine/Source/Editor/UnrealEd/Public/ComponentVisualizer.h
--- a/Engine/Source/Editor/UnrealEd/Public/ComponentVisualizer.h
+++ b/Engine/Source/Editor/UnrealEd/Public/ComponentVisualizer.h
@@ -155,6 +155,7 @@ public:
 	/** */
 	virtual void OnRegister() {}
 	/** Only show this visualizer if the actor is selected */
+	UE_DEPRECATED(5.4, "This function is unused and will be removed in a future version. Component visualizers are only shown for the active selection. Use bDebugDraw on the specific component, or the editor setting to control drawing of subcomponents for selected actors.")
 	virtual bool ShowWhenSelected() { return true; }
 	/** Show this visualizer if the component is directly is selected */
 	virtual bool ShouldShowForSelectedSubcomponents(const UActorComponent* Component) { return true; }
