Fix for Compilation Tests / Compile UnrealEditor StressTestUnity Win64
Issue 362201: Compile errors in ComponentTypeRegistry.cpp

diff --git a/Engine/Source/Editor/UnrealEd/Private/Commandlets/CookShadersCommandlet.cpp b/Engine/Source/Editor/UnrealEd/Private/Commandlets/CookShadersCommandlet.cpp
--- a/Engine/Source/Editor/UnrealEd/Private/Commandlets/CookShadersCommandlet.cpp
+++ b/Engine/Source/Editor/UnrealEd/Private/Commandlets/CookShadersCommandlet.cpp
@@ -19,7 +19,6 @@
 DEFINE_LOG_CATEGORY_STATIC(LogCookShadersCommandlet, Log, All);
 
 static const TCHAR* GlobalName = TEXT("Global");
-static const TCHAR* MaterialName = TEXT("Material");
 static const TCHAR* NiagaraName = TEXT("Niagara");
 
 // Examples
@@ -133,7 +132,7 @@ namespace CookShadersCommandlet {
 								}
 								else
 								{
-									Record.Type = MaterialName;
+									Record.Type = TEXT("Material");
 
 									// default is Num
 									Record.Quality = EMaterialQualityLevel::Num;
@@ -226,7 +225,7 @@ int32 UCookShadersCommandlet::Main(const FString& Params)
 	{
 		if (MaterialString == GlobalName)
 		{
-			// we don't have a way to specify global shader name, or compile oen specifically
+			// we don't have a way to specify global shader name, or compile one specifically
 			GlobalsToFind.Add(GlobalName);
 			MaterialString = TEXT("");
 		}
@@ -250,7 +249,7 @@ int32 UCookShadersCommandlet::Main(const FString& Params)
 		{
 			GlobalsToFind.Add(I.Name);
 		}
-		else if (I.Type == MaterialName)
+		else if (I.Type == TEXT("Material"))
 		{
 			TArray<FString> ShaderTypeNames;
 			ShaderTypeNames.Add(I.Shader);
