Cooker: GeneratedPackages: do not demote generator or generated packages when private objects are garbage collected. Demote  (and give an error) only  when public objects are garbage collected. Normal packages that are not a generator and are not owned by a generator keep the same behavior of demoting if any object, even a private one, is garbage collected.
#rnx

diff --git a/Engine/Source/Editor/UnrealEd/Classes/CookOnTheSide/CookOnTheFlyServer.h b/Engine/Source/Editor/UnrealEd/Classes/CookOnTheSide/CookOnTheFlyServer.h
--- a/Engine/Source/Editor/UnrealEd/Classes/CookOnTheSide/CookOnTheFlyServer.h
+++ b/Engine/Source/Editor/UnrealEd/Classes/CookOnTheSide/CookOnTheFlyServer.h
@@ -175,6 +175,7 @@ namespace UE::Cook
 	enum class ESuppressCookReason : uint8;
 	enum class ESendFlags : uint8;
 	struct FBeginCookConfigSettings;
+	struct FCachedObjectInOuter;
 	struct FConstructPackageData;
 	struct FCookByTheBookOptions;
 	struct FCookOnTheFlyOptions;
@@ -1155,7 +1156,7 @@ private:
 		bool bPrecaching);
 	/** Call BeginCacheForCookedPlatformData on all objects in a PackageData or a GeneratorPackage's current round. */
 	UNREALED_API UE::Cook::EPollStatus CallBeginCacheOnObjects(UE::Cook::FPackageData& PackageData, UPackage* Package,
-		TArray<FWeakObjectPtr>& Objects, int32& NextIndex, UE::Cook::FCookerTimer& Timer);
+		TArray<UE::Cook::FCachedObjectInOuter>& Objects, int32& NextIndex, UE::Cook::FCookerTimer& Timer);
 
 	/**
 	 * Frees all the memory used to call BeginCacheForCookedPlatformData on all the objects in PackageData.
diff --git a/Engine/Source/Editor/UnrealEd/Private/CookOnTheFlyServer.cpp b/Engine/Source/Editor/UnrealEd/Private/CookOnTheFlyServer.cpp
--- a/Engine/Source/Editor/UnrealEd/Private/CookOnTheFlyServer.cpp
+++ b/Engine/Source/Editor/UnrealEd/Private/CookOnTheFlyServer.cpp
@@ -3062,9 +3062,9 @@ UE::Cook::EPollStatus UCookOnTheFlyServer::ConditionalCreateGeneratorPackage(UE:
 
 	TArray<Private::FRegisteredCookPackageSplitter*> FoundRegisteredSplitters;
 
-	for (FWeakObjectPtr& WeakObj : PackageData.GetCachedObjectsInOuter())
+	for (FCachedObjectInOuter& CachedObjectInOuter : PackageData.GetCachedObjectsInOuter())
 	{
-		UObject* Obj = WeakObj.Get();
+		UObject* Obj = CachedObjectInOuter.Object.Get();
 		if (!Obj)
 		{
 			continue;
@@ -3702,7 +3702,7 @@ UE::Cook::EPollStatus UCookOnTheFlyServer::PrepareSaveInternal(UE::Cook::FPackag
 			CookedPlatformDataNextIndex = 0;
 		}
 
-		TArray<FWeakObjectPtr>& CachedObjectsInOuter = PackageData.GetCachedObjectsInOuter();
+		TArray<FCachedObjectInOuter>& CachedObjectsInOuter = PackageData.GetCachedObjectsInOuter();
 		EPollStatus Result = CallBeginCacheOnObjects(PackageData, Package, CachedObjectsInOuter,
 			CookedPlatformDataNextIndex, Timer);
 		if (Result != EPollStatus::Success)
@@ -3787,7 +3787,7 @@ UE::Cook::EPollStatus UCookOnTheFlyServer::PrepareSaveInternal(UE::Cook::FPackag
 }
 
 UE::Cook::EPollStatus UCookOnTheFlyServer::CallBeginCacheOnObjects(UE::Cook::FPackageData& PackageData,
-	UPackage* Package, TArray<FWeakObjectPtr>& Objects, int32& NextIndex, UE::Cook::FCookerTimer& Timer)
+	UPackage* Package, TArray<UE::Cook::FCachedObjectInOuter>& Objects, int32& NextIndex, UE::Cook::FCookerTimer& Timer)
 {
 	using namespace UE::Cook;
 
@@ -3796,11 +3796,11 @@ UE::Cook::EPollStatus UCookOnTheFlyServer::CallBeginCacheOnObjects(UE::Cook::FPa
 	TArray<const ITargetPlatform*, TInlineAllocator<ExpectedMaxNumPlatforms>> TargetPlatforms;
 	PackageData.GetCachedObjectsInOuterPlatforms(TargetPlatforms);
 
-	FWeakObjectPtr* ObjectsData = Objects.GetData();
+	FCachedObjectInOuter* ObjectsData = Objects.GetData();
 	int NumObjects = Objects.Num();
 	for (; NextIndex < NumObjects; ++NextIndex)
 	{
-		UObject* Obj = ObjectsData[NextIndex].Get();
+		UObject* Obj = ObjectsData[NextIndex].Object.Get();
 		if (!Obj)
 		{
 			// Objects can be marked as pending kill even without a garbage collect, and our weakptr.get will return
@@ -3809,7 +3809,7 @@ UE::Cook::EPollStatus UCookOnTheFlyServer::CallBeginCacheOnObjects(UE::Cook::FPa
 			// BeginCacheForCookedPlatformData and ClearAllCachedCookedPlatformData
 			// In case the weakptr is merely pendingkill, set it to null explicitly so we don't think that we've called
 			// BeginCacheForCookedPlatformData on it if it gets unmarked pendingkill later
-			ObjectsData[NextIndex] = nullptr;
+			ObjectsData[NextIndex].Object = nullptr;
 			continue;
 		}
 		FCachedCookedPlatformDataState& CCPDState = PackageDatas->GetCachedCookedPlatformDataObjects().FindOrAdd(Obj);
@@ -3897,9 +3897,9 @@ void UCookOnTheFlyServer::ReleaseCookedPlatformData(UE::Cook::FPackageData& Pack
 		// Since we have completed CookedPlatformData, we know we called BeginCacheForCookedPlatformData on all
 		// objects in the package, and none are pending
 		UE_SCOPED_HIERARCHICAL_COOKTIMER(ClearAllCachedCookedPlatformData);
-		for (FWeakObjectPtr& WeakPtr : PackageData.GetCachedObjectsInOuter())
+		for (FCachedObjectInOuter& CachedObjectInOuter : PackageData.GetCachedObjectsInOuter())
 		{
-			UObject* Object = WeakPtr.Get();
+			UObject* Object = CachedObjectInOuter.Object.Get();
 			if (Object)
 			{
 				FPendingCookedPlatformData::ClearCachedCookedPlatformData(Object, PackageData,
@@ -3932,10 +3932,10 @@ void UCookOnTheFlyServer::ReleaseCookedPlatformData(UE::Cook::FPackageData& Pack
 
 
 		// Iterate over all objects in the FPackageData up to GetCookedPlatformDataNextIndex
-		TArray<FWeakObjectPtr>& CachedObjects = PackageData.GetCachedObjectsInOuter();
+		TArray<FCachedObjectInOuter>& CachedObjects = PackageData.GetCachedObjectsInOuter();
 		for (int32 ObjectIndex = 0; ObjectIndex < PackageData.GetCookedPlatformDataNextIndex(); ++ObjectIndex)
 		{
-			UObject* Object = CachedObjects[ObjectIndex].Get();
+			UObject* Object = CachedObjects[ObjectIndex].Object.Get();
 			if (!Object)
 			{
 				continue;
@@ -5229,7 +5229,24 @@ void UCookOnTheFlyServer::PostGarbageCollect()
 
 	TSet<UObject*> SaveQueueObjectsThatStillExist;
 
-	// If any PackageDatas with ObjectPointers had any of their object pointers deleted out from under them, demote them back to request
+	// If garbage collection deleted a UPackage WHILE WE WERE SAVING IT, then we have problems.
+	check(!SavingPackageData || SavingPackageData->GetPackage() != nullptr);
+
+	// If there was a GarbageCollect after we already started calling BeginCacheCookedPlatformData, then we
+	// have a list of the WeakObjectPtr to all objects in the package (FPackageData::CachedObjejectsInOuter)
+	// and some of those objects may have been set to null. We declare a reference to prevent GC for the RF_Public
+	// objects in that list, but we do not declare that reference for private objects. The private objects may
+	// therefore have been deleted and set to null 
+	// Side note: because objects can be marked as pending kill at any time and we use FObjectWeakPtr.Get(),
+	// which returns null if pending kill, we need to skip nulls in the array at any point, not just after GC.
+	// 
+	// We do not want to prevent GC of private objects in case there is the expectation by some
+	// systems (blueprints, licensee code) that removing references to an object during PreCollectGarbage will cause
+	// it to be deleted by GC and be replaceable afterwards. We add any new private objects after the garbage collect
+	// and continue with the save. Public objects have a different contract; they are not replaceable across a
+	// GC because anything outside the package could be referring to them. So we keep them referenced. But GC may
+	// force delete them despite our reference, and the package is then in an unknown state. If that happens we
+	// demote the package back to request and start its load and save over.
 	TArray<FPackageData*> Demotes;
 	for (FPackageData* PackageData : PackageDatas->GetSaveQueue())
 	{
@@ -5242,9 +5259,9 @@ void UCookOnTheFlyServer::PostGarbageCollect()
 		else
 		{
 			// Mark that the objects for this package should be kept in CachedCookedPlatformData records
-			for (FWeakObjectPtr& WeakObjectPtr : PackageData->GetCachedObjectsInOuter())
+			for (FCachedObjectInOuter& CachedObjectInOuter : PackageData->GetCachedObjectsInOuter())
 			{
-				UObject* Object = WeakObjectPtr.Get();
+				UObject* Object = CachedObjectInOuter.Object.Get();
 				if (Object)
 				{
 					SaveQueueObjectsThatStillExist.Add(Object);
@@ -5276,12 +5293,6 @@ void UCookOnTheFlyServer::PostGarbageCollect()
 	// Remove objects that were deleted by garbage collection from our containers that track raw object pointers
 	PackageDatas->CachedCookedPlatformDataObjectsPostGarbageCollect(SaveQueueObjectsThatStillExist);
 
-	// If there was a GarbageCollect while we are saving a package, some of the WeakObjectPtr in SavingPackageData->CachedObjectPointers may have been deleted and set to null
-	// We need to handle nulls in that array at any point after calling SavePackage. We do not want to declare them as references and prevent their GC, in case there is 
-	// the expectation by some licensee code that removing references to an object will cause it to not be saved
-	// However, if garbage collection deleted the package WHILE WE WERE SAVING IT, then we have problems.
-	check(!SavingPackageData || SavingPackageData->GetPackage() != nullptr);
-
 	GCKeepObjects.Empty();
 	UPackage::SoftGCPackageToObjectList.Empty();
 	SoftGCPackageToObjectListBuffer.Empty();
@@ -12281,7 +12292,7 @@ EPackageWriterResult UCookOnTheFlyServer::SavePackageBeginCacheForCookedPlatform
 	FPackageData* PackageData = PackageDatas->FindPackageDataByPackageName(PackageName);
 	check(PackageData); // This callback is called from a packagesave we initiated, so it should exist
 
-	TArray<FWeakObjectPtr>& CachedObjectsInOuter = PackageData->GetCachedObjectsInOuter();
+	TArray<FCachedObjectInOuter>& CachedObjectsInOuter = PackageData->GetCachedObjectsInOuter();
 	int32& NextIndex = PackageData->GetCookedPlatformDataNextIndex();
 	TArray<UObject*> PendingObjects;
 	for (UObject* Object : SaveableObjects)
@@ -12292,11 +12303,11 @@ EPackageWriterResult UCookOnTheFlyServer::SavePackageBeginCacheForCookedPlatform
 			CCPDState.AddRefFrom(PackageData);
 
 			// NextIndex is usually at the end of CachedObjectsInOuter, but in case it is not, insert the new Object 
-			// at NextIndex so that we still record that we have not called BeginCace on objects after it. Then increment
+			// at NextIndex so that we still record that we have not called BeginCache on objects after it. Then increment
 			// NextIndex to indicate we have already called (down below) BeginCache on the added object, so that
 			// ReleaseCookedPlatformData knows that it needs to call Clear on it.
 			check(NextIndex >= 0);
-			CachedObjectsInOuter.Insert(Object, NextIndex);
+			CachedObjectsInOuter.Insert(FCachedObjectInOuter(Object), NextIndex);
 			++NextIndex;
 		}
 
diff --git a/Engine/Source/Editor/UnrealEd/Private/Cooker/CookPackageData.cpp b/Engine/Source/Editor/UnrealEd/Private/Cooker/CookPackageData.cpp
--- a/Engine/Source/Editor/UnrealEd/Private/Cooker/CookPackageData.cpp
+++ b/Engine/Source/Editor/UnrealEd/Private/Cooker/CookPackageData.cpp
@@ -1126,7 +1126,7 @@ void FPackageData::CheckPreloadEmpty()
 	check(!GetIsPreloaded());
 }
 
-TArray<FWeakObjectPtr>& FPackageData::GetCachedObjectsInOuter()
+TArray<FCachedObjectInOuter>& FPackageData::GetCachedObjectsInOuter()
 {
 	return CachedObjectsInOuter;
 }
@@ -1157,7 +1157,7 @@ void FPackageData::CreateObjectCache()
 		{
 			FWeakObjectPtr ObjectWeakPointer(Object);
 			check(ObjectWeakPointer.Get()); // GetObjectsWithOuter with Garbage filtered out should only return valid-for-weakptr objects
-			CachedObjectsInOuter.Emplace(MoveTemp(ObjectWeakPointer));
+			CachedObjectsInOuter.Emplace(ObjectWeakPointer);
 		}
 
 		for (TPair<const ITargetPlatform*, FPackagePlatformData>& Pair : PlatformDatas)
@@ -1224,9 +1224,9 @@ EPollStatus FPackageData::RefreshObjectCache(bool& bOutFoundNewObjects)
 
 	TArray<UObject*> OldObjects;
 	OldObjects.Reserve(CachedObjectsInOuter.Num());
-	for (FWeakObjectPtr& Object : CachedObjectsInOuter)
+	for (FCachedObjectInOuter& Object : CachedObjectsInOuter)
 	{
-		UObject* ObjectPtr = Object.Get();
+		UObject* ObjectPtr = Object.Object.Get();
 		if (ObjectPtr)
 		{
 			OldObjects.Add(ObjectPtr);
@@ -1381,19 +1381,28 @@ void FPackageData::UpdateSaveAfterGarbageCollect(bool& bOutDemote)
 		return;
 	}
 
-	if (GetPackage() == nullptr || !GetPackage()->IsFullyLoaded() ||
-		Algo::AnyOf(CachedObjectsInOuter, [](const FWeakObjectPtr& WeakPtr)
-			{
-				// TODO: Keep track of which objects were public, and only invalidate the save if the object
-				// that has been deleted or marked pending kill was public
-				// Until we make that change, we will unnecessarily invalidate and demote some packages after a
-				// garbage collect
-				return WeakPtr.Get() == nullptr;
-			}))
+	// Reexecute PrepareSave if we already completed it; we need to refresh our CachedObjectsInOuter list
+	// and call BeginCacheOnCookedPlatformData on any new objects.
+	SetCookedPlatformDataComplete(false);
+
+	if (GetPackage() == nullptr || !GetPackage()->IsFullyLoaded())
 	{
 		bOutDemote = true;
-		return;
 	}
+	else
+	{
+		for (FCachedObjectInOuter& CachedObjectInOuter : CachedObjectsInOuter)
+		{
+			if (CachedObjectInOuter.Object.Get() == nullptr)
+			{
+				// Deleting a public object puts the package in an invalid state; demote back to request
+				// and load/save it again
+				bool bPublicDeleted = !!(CachedObjectInOuter.ObjectFlags & RF_Public);;
+				bOutDemote |= bPublicDeleted;
+			}
+		}
+	}
+
 	if (GeneratorPackage)
 	{
 		GeneratorPackage->UpdateSaveAfterGarbageCollect(*this, bOutDemote);
@@ -1803,9 +1812,9 @@ void FGeneratorPackage::PreGarbageCollect(FCookGenerationInfo& Info, TArray<TObj
 				GCKeepPackageDatas.Add(Info.PackageData);
 			}
 			GCKeepPackages.Append(Info.KeepReferencedPackages);
-			for (FWeakObjectPtr& WeakObjectPtr : Info.PackageData->GetCachedObjectsInOuter())
+			for (FCachedObjectInOuter& CachedObjectInOuter : Info.PackageData->GetCachedObjectsInOuter())
 			{
-				UObject* Object = WeakObjectPtr.Get();
+				UObject* Object = CachedObjectInOuter.Object.Get();
 				if (Object)
 				{
 					GCKeepObjects.Add(Object);
@@ -1836,8 +1845,8 @@ void FGeneratorPackage::PostGarbageCollect()
 	else
 	{
 		// After the Generator Package is saved, we drop our references to it and it can be garbage collected
-		// If we have any packages left to populate, our splitter contract requires that it be garbage collected;
-		// we promise that the package is not partially GC'd during calls to TryPopulateGeneratedPackage
+		// If we have any packages left to populate, our splitter contract requires that it be garbage collected
+		// because we promise that the package is not partially GC'd during calls to TryPopulateGeneratedPackage
 		// The splitter can opt-out of this contract and keep it referenced itself if it desires.
 		UPackage* LocalOwnerPackage = FindObject<UPackage>(nullptr, *Owner.GetPackageName().ToString());
 		if (LocalOwnerPackage)
@@ -2030,39 +2039,29 @@ void FGeneratorPackage::UpdateSaveAfterGarbageCollect(const FPackageData& Packag
 		}
 	}
 
+	if (bInOutDemote && 
+		GetCookPackageSplitterInstance()->UseInternalReferenceToAvoidGarbageCollect() &&
+		Info->HasTakenOverCachedCookedPlatformData())
+	{
+		// No public objects should have been deleted; we are supposed to keep them referenced by keeping the package
+		// referenced in UCookOnTheFlyServer::PreGarbageCollect, and the package keeping its public objects referenced
+		// by UPackage::AddReferencedObjects. Since no public objects were deleted, our caller should not have
+		// set bInOutDemote=true.
+		// Allowing demotion after the splitter has started moving objects breaks our contract with the splitter
+		// and can cause a crash. So log this as an error.
+		UE_LOG(LogCook, Error, TEXT("A %s package had some of its UObjects deleted during garbage collection after it started generating. This will cause errors during save of the package.")
+			TEXT("\n\tSplitter=%s%s."),
+			Info->IsGenerator() ? TEXT("Generator") : TEXT("Generated"),
+			*GetSplitDataObjectName().ToString(),
+			Info->IsGenerator() ? TEXT("") : *FString::Printf(TEXT(", Generated=%s."), *Info->PackageData->GetPackageName().ToString()));
+	}
+
+	// Remove raw pointers from RootMovedObjects if they no longer exist in the weakpointers in CachedObjectsInOuter
 	TSet<UObject*> CachedObjectsInOuterSet;
-	TArray<FWeakObjectPtr>& CachedObjectsInOuter = Info->PackageData->GetCachedObjectsInOuter();
-	int32& NextIndexToBeginCache = Info->PackageData->GetCookedPlatformDataNextIndex();
-	for (int32 Index = 0; Index < CachedObjectsInOuter.Num(); ++Index)
+	for (FCachedObjectInOuter& CachedObjectInOuter : Info->PackageData->GetCachedObjectsInOuter())
 	{
-		FWeakObjectPtr& WeakObjectPtr = CachedObjectsInOuter[Index];
-		UObject* Object = WeakObjectPtr.Get();
-		if (!Object)
-		{
-			if (GetCookPackageSplitterInstance()->UseInternalReferenceToAvoidGarbageCollect())
-			{
-				// No objects should be allowed to be deleted; we are supposed to keep them referenced
-				// But allowing demotion will break things for sure.
-				// Log a cook error but remove the invalidated object.
-				UE_LOG(LogCook, Error, TEXT("PackageSplitter found an object returned from %s that was removed from memory during garbage collection. This will cause errors during save of the package.")
-					TEXT("\n\tSplitter=%s%s."),
-					Info->IsGenerator() ? TEXT("PopulateGeneratorPackage") : TEXT("PopulateGeneratedPackage"),
-					*GetSplitDataObjectName().ToString(),
-					Info->IsGenerator() ? TEXT("") : *FString::Printf(TEXT(", Generated=%s."), *Info->PackageData->GetPackageName().ToString()));
-
-				CachedObjectsInOuter.RemoveAt(Index);
-				if (NextIndexToBeginCache > Index)
-				{
-					--NextIndexToBeginCache;
-				}
-				--Index;
-			}
-			else
-			{
-				bInOutDemote = true;
-			}
-		}
-		else
+		UObject* Object = CachedObjectInOuter.Object.Get();
+		if (Object)
 		{
 			CachedObjectsInOuterSet.Add(Object);
 		}
@@ -2094,14 +2093,14 @@ void FCookGenerationInfo::SetSaveStateComplete(ESaveState CompletedState)
 }
 
 void FCookGenerationInfo::TakeOverCachedObjectsAndAddMoved(FGeneratorPackage& Generator,
-	TArray<FWeakObjectPtr>& CachedObjectsInOuter, TArray<UObject*>& MovedObjects)
+	TArray<FCachedObjectInOuter>& CachedObjectsInOuter, TArray<UObject*>& MovedObjects)
 {
 	RootMovedObjects.Reset();
 
 	TSet<UObject*> ObjectSet;
-	for (FWeakObjectPtr& ObjectInOuter : CachedObjectsInOuter)
+	for (FCachedObjectInOuter& ObjectInOuter : CachedObjectsInOuter)
 	{
-		UObject* Object = ObjectInOuter.Get();
+		UObject* Object = ObjectInOuter.Object.Get();
 		if (Object)
 		{
 			ObjectSet.Add(Object);
@@ -2125,7 +2124,7 @@ void FCookGenerationInfo::TakeOverCachedObjectsAndAddMoved(FGeneratorPackage& Ge
 		if (!bAlreadyExists)
 		{
 			RootMovedObjects.Add(Object);
-			CachedObjectsInOuter.Add(Object);
+			CachedObjectsInOuter.Emplace(Object);
 			GetObjectsWithOuter(Object, ChildrenOfMovedObjects, true /* bIncludeNestedObjects */, RF_NoFlags, EInternalObjectFlags::Garbage);
 		}
 	}
@@ -2137,7 +2136,7 @@ void FCookGenerationInfo::TakeOverCachedObjectsAndAddMoved(FGeneratorPackage& Ge
 		ObjectSet.Add(Object, &bAlreadyExists);
 		if (!bAlreadyExists)
 		{
-			CachedObjectsInOuter.Add(Object);
+			CachedObjectsInOuter.Emplace(Object);
 		}
 	}
 
@@ -2153,11 +2152,11 @@ EPollStatus FCookGenerationInfo::RefreshPackageObjects(FGeneratorPackage& Genera
 
 	TSet<UObject*> ObjectSet;
 	check(PackageData); // RefreshPackageObjects is only called when there is a PackageData
-	TArray<FWeakObjectPtr>& CachedObjectsInOuter = PackageData->GetCachedObjectsInOuter();
+	TArray<FCachedObjectInOuter>& CachedObjectsInOuter = PackageData->GetCachedObjectsInOuter();
 	ObjectSet.Reserve(CachedObjectsInOuter.Num());
-	for (FWeakObjectPtr& ExistingObject : CachedObjectsInOuter)
+	for (FCachedObjectInOuter& ExistingObject : CachedObjectsInOuter)
 	{
-		UObject* Object = ExistingObject.Get();
+		UObject* Object = ExistingObject.Object.Get();
 		if (Object)
 		{
 			bool bAlreadyExists;
@@ -2176,7 +2175,7 @@ EPollStatus FCookGenerationInfo::RefreshPackageObjects(FGeneratorPackage& Genera
 		ObjectSet.Add(Object, &bAlreadyExists);
 		if (!bAlreadyExists)
 		{
-			CachedObjectsInOuter.Add(Object);
+			CachedObjectsInOuter.Emplace(Object);
 			if (!FirstNewObject )
 			{
 				FirstNewObject  = Object;
diff --git a/Engine/Source/Editor/UnrealEd/Private/Cooker/CookPackageData.h b/Engine/Source/Editor/UnrealEd/Private/Cooker/CookPackageData.h
--- a/Engine/Source/Editor/UnrealEd/Private/Cooker/CookPackageData.h
+++ b/Engine/Source/Editor/UnrealEd/Private/Cooker/CookPackageData.h
@@ -89,6 +89,36 @@ struct FCachedCookedPlatformDataState
 	TMap<const ITargetPlatform*, ECachedCookedPlatformDataEvent> PlatformStates;
 };
 
+/**
+ * Objects we searched for in the Package being saved; we need to execute various operations on
+ * all of these objects, most notably BeginCacheForCookedPlatformData family of functions.
+ * We keep track of a WeakPtr to the object along with its flags in case it gets deleted and we
+ * need to decide how to respond to the deletion it based on what its flags were.
+ */
+struct FCachedObjectInOuter
+{
+	FWeakObjectPtr Object;
+	EObjectFlags ObjectFlags;
+
+	FCachedObjectInOuter(UObject* InObject = nullptr)
+	{
+		Object = InObject;
+		ObjectFlags = InObject ? InObject->GetFlags() : RF_NoFlags;
+	}
+	FCachedObjectInOuter(FWeakObjectPtr&& InObject)
+	{
+		Object = MoveTemp(InObject);
+		UObject* Ptr = Object.Get(true /* bEvenIfPendingKill */);
+		ObjectFlags = Ptr ? Ptr->GetFlags() : RF_NoFlags;
+	}
+	FCachedObjectInOuter(const FWeakObjectPtr& InObject)
+	{
+		Object = InObject;
+		UObject* Ptr = Object.Get(true /* bEvenIfPendingKill */);
+		ObjectFlags = Ptr ? Ptr->GetFlags() : RF_NoFlags;
+	}
+};
+
 /** Flags specifying the behavior of FPackageData::SendToState */
 enum class ESendFlags : uint8
 {
@@ -466,7 +496,7 @@ public:
 	 * The list of objects inside the package.  Only non-empty during saving; it is populated on demand by
 	 * TryCreateObjectCache and is cleared when leaving the save state.
 	 */
-	TArray<FWeakObjectPtr>& GetCachedObjectsInOuter();
+	TArray<FCachedObjectInOuter>& GetCachedObjectsInOuter();
 	template <typename ArrayType>
 	/** The list of platforms that were recorded as needscooking when CachedObjeObjectsInOuter was recorded. */
 	void GetCachedObjectsInOuterPlatforms(ArrayType& OutPlatforms) const;
@@ -691,7 +721,8 @@ private:
 	FGeneratorPackage* GeneratedOwner;
 	/** Data for each platform that has been interacted with by *this. */
 	TSortedMap<const ITargetPlatform*, FPackagePlatformData, TInlineAllocator<1>> PlatformDatas;
-	TArray<FWeakObjectPtr> CachedObjectsInOuter;
+
+	TArray<FCachedObjectInOuter> CachedObjectsInOuter;
 	FCompletionCallback CompletionCallback;
 	TUniquePtr<TMap<FPackageData*, EInstigator>> Unsolicited;
 	FName PackageName;
@@ -802,7 +833,7 @@ public:
 	 * and add them to this. Also add the list of NewObjects reported by the splitter that will be moved into the package.
 	 */
 	void TakeOverCachedObjectsAndAddMoved(FGeneratorPackage& Generator, 
-		TArray<FWeakObjectPtr>& CachedObjectsInOuter, TArray<UObject*>& MovedObjects);
+		TArray<FCachedObjectInOuter>& CachedObjectsInOuter, TArray<UObject*>& MovedObjects);
 	/**
 	 * Fetch all the objects currently in the package and add them the list of objects that need BeginCacheForCookedPlatformData.
 	 * Reports whether new objects were found. If DemotionState is not ESaveState::Last, will SetState back to DemotionState
