Delete Assets dialog with now output the actor label when deleting assets referenced by actors, instead of showing the external actor path.

diff --git a/Engine/Source/Editor/UnrealEd/Private/Dialogs/SDeleteAssetsDialog.cpp b/Engine/Source/Editor/UnrealEd/Private/Dialogs/SDeleteAssetsDialog.cpp
--- a/Engine/Source/Editor/UnrealEd/Private/Dialogs/SDeleteAssetsDialog.cpp
+++ b/Engine/Source/Editor/UnrealEd/Private/Dialogs/SDeleteAssetsDialog.cpp
@@ -1,6 +1,9 @@
 // Copyright Epic Games, Inc. All Rights Reserved.
 
 #include "Dialogs/SDeleteAssetsDialog.h"
+#include "AssetRegistry/AssetData.h"
+#include "AssetRegistry/AssetRegistryModule.h"
+#include "AssetRegistry/IAssetRegistry.h"
 #include "Framework/Commands/UIAction.h"
 #include "Framework/Commands/UICommandList.h"
 #include "Widgets/Notifications/SProgressBar.h"
@@ -813,9 +816,36 @@ FText SDeleteAssetsDialog::GetReferencingAssetsEmptyText() const
 {
 	FString DiskReferences = "There Are Some Non-Displayable References\n\n";
 
+	static FName NAME_ActorLabel(TEXT("ActorLabel"));
+	IAssetRegistry& AssetRegistry = FModuleManager::LoadModuleChecked<FAssetRegistryModule>(TEXT("AssetRegistry")).Get();
+
 	for ( const FName& DiskReference : DeleteModel->GetAssetReferences() )
 	{
-		DiskReferences += DiskReference.ToString() + "\n";
+		FString ReferenceToAppend;
+
+		FARFilter Filter;
+		Filter.PackageNames = { DiskReference };
+		
+		TArray<FAssetData> Assets;
+		if (AssetRegistry.GetAssets(Filter, Assets))
+		{
+			for (const FAssetData& Asset : Assets)
+			{
+				FString ActorLabel;
+				if (Asset.GetTagValue(NAME_ActorLabel, ActorLabel))
+				{
+					ReferenceToAppend = Asset.GetOptionalOuterPathName().ToString() + TEXT(".") + ActorLabel;
+					break;
+				}
+			}
+		}
+		
+		if (ReferenceToAppend.IsEmpty())
+		{
+			ReferenceToAppend = DiskReference.ToString();
+		}
+
+		DiskReferences += ReferenceToAppend + TEXT("\n");
 	}
 
 	return FText::FromString( DiskReferences );
