Cooker: When a generated package is demoted due to deleted objects, add feedback with the name of the object that was deleted.
#rnx

diff --git a/Engine/Source/Editor/UnrealEd/Private/Cooker/CookPackageData.cpp b/Engine/Source/Editor/UnrealEd/Private/Cooker/CookPackageData.cpp
--- a/Engine/Source/Editor/UnrealEd/Private/Cooker/CookPackageData.cpp
+++ b/Engine/Source/Editor/UnrealEd/Private/Cooker/CookPackageData.cpp
@@ -1131,6 +1131,11 @@ TArray<FCachedObjectInOuter>& FPackageData::GetCachedObjectsInOuter()
 	return CachedObjectsInOuter;
 }
 
+const TArray<FCachedObjectInOuter>& FPackageData::GetCachedObjectsInOuter() const
+{
+	return CachedObjectsInOuter;
+}
+
 void FPackageData::CheckObjectCacheEmpty() const
 {
 	check(CachedObjectsInOuter.Num() == 0);
@@ -2049,11 +2054,51 @@ void FGeneratorPackage::UpdateSaveAfterGarbageCollect(const FPackageData& Packag
 		// set bInOutDemote=true.
 		// Allowing demotion after the splitter has started moving objects breaks our contract with the splitter
 		// and can cause a crash. So log this as an error.
+		// For better feedback, look in our extra data to identify the name of the public UObject that was deleted.
+		FString DeletedObject;
+		if (!PackageData.GetPackage())
+		{
+			DeletedObject = FString::Printf(TEXT("UPackage %s"), *PackageData.GetPackageName().ToString());
+		}
+		else
+		{
+			TSet<UObject*> ExistingObjectsAfterSave;
+			for (const FCachedObjectInOuter& CachedObjectInOuter : PackageData.GetCachedObjectsInOuter())
+			{
+				UObject* Ptr = CachedObjectInOuter.Object.Get();
+				if (Ptr)
+				{
+					ExistingObjectsAfterSave.Add(Ptr);
+				}
+			}
+
+			for (const TPair<UObject*, FCachedObjectInOuterGeneratorInfo>& Pair : Info->CachedObjectsInOuterInfo)
+			{
+				if (Pair.Value.bPublic && !ExistingObjectsAfterSave.Contains(Pair.Key))
+				{
+					DeletedObject = Pair.Value.FullName;
+					break;
+				}
+			}
+			if (DeletedObject.IsEmpty())
+			{
+				if (!PackageData.GetPackage()->IsFullyLoaded())
+				{
+					DeletedObject = FString::Printf(TEXT("UPackage %s is no longer FullyLoaded"), *PackageData.GetPackageName().ToString());
+				}
+				else
+				{
+					DeletedObject = TEXT("<Unknown>");
+				}
+			}
+		}
 		UE_LOG(LogCook, Error, TEXT("A %s package had some of its UObjects deleted during garbage collection after it started generating. This will cause errors during save of the package.")
-			TEXT("\n\tSplitter=%s%s."),
+			TEXT("\n\tDeleted object: %s")
+			TEXT("\n\tSplitter=%s%s"),
 			Info->IsGenerator() ? TEXT("Generator") : TEXT("Generated"),
+			*DeletedObject,
 			*GetSplitDataObjectName().ToString(),
-			Info->IsGenerator() ? TEXT("") : *FString::Printf(TEXT(", Generated=%s."), *Info->PackageData->GetPackageName().ToString()));
+			Info->IsGenerator() ? TEXT(".") : *FString::Printf(TEXT(", Generated=%s."), *Info->PackageData->GetPackageName().ToString()));
 	}
 
 	// Remove raw pointers from RootMovedObjects if they no longer exist in the weakpointers in CachedObjectsInOuter
@@ -2066,9 +2111,10 @@ void FGeneratorPackage::UpdateSaveAfterGarbageCollect(const FPackageData& Packag
 			CachedObjectsInOuterSet.Add(Object);
 		}
 	}
-	for (TSet<UObject*>::TIterator Iter(Info->RootMovedObjects); Iter; ++Iter)
+	for (TMap<UObject*, FCachedObjectInOuterGeneratorInfo>::TIterator Iter(Info->CachedObjectsInOuterInfo);
+		Iter; ++Iter)
 	{
-		if (!CachedObjectsInOuterSet.Contains(*Iter))
+		if (!CachedObjectsInOuterSet.Contains(Iter->Key))
 		{
 			Iter.RemoveCurrent();
 		}
@@ -2092,18 +2138,33 @@ void FCookGenerationInfo::SetSaveStateComplete(ESaveState CompletedState)
 	}
 }
 
+void FCachedObjectInOuterGeneratorInfo::Initialize(UObject* Object)
+{
+	if (Object)
+	{
+		FullName = Object->GetFullName();
+		bPublic = Object->HasAnyFlags(RF_Public);
+	}
+	else
+	{
+		FullName.Empty();
+		bPublic = false;
+	}
+
+	bInitialized = true;
+}
+
 void FCookGenerationInfo::TakeOverCachedObjectsAndAddMoved(FGeneratorPackage& Generator,
 	TArray<FCachedObjectInOuter>& CachedObjectsInOuter, TArray<UObject*>& MovedObjects)
 {
-	RootMovedObjects.Reset();
+	CachedObjectsInOuterInfo.Reset();
 
-	TSet<UObject*> ObjectSet;
 	for (FCachedObjectInOuter& ObjectInOuter : CachedObjectsInOuter)
 	{
 		UObject* Object = ObjectInOuter.Object.Get();
 		if (Object)
 		{
-			ObjectSet.Add(Object);
+			CachedObjectsInOuterInfo.FindOrAdd(Object).Initialize(Object);
 		}
 	}
 
@@ -2119,11 +2180,12 @@ void FCookGenerationInfo::TakeOverCachedObjectsAndAddMoved(FGeneratorPackage& Ge
 				IsGenerator() ? TEXT("") : *FString::Printf(TEXT(", Package %s"), *PackageData->GetPackageName().ToString()));
 			continue;
 		}
-		bool bAlreadyExists;
-		ObjectSet.Add(Object, &bAlreadyExists);
-		if (!bAlreadyExists)
+		FCachedObjectInOuterGeneratorInfo& Info = CachedObjectsInOuterInfo.FindOrAdd(Object);
+		if (!Info.bInitialized)
 		{
-			RootMovedObjects.Add(Object);
+			Info.Initialize(Object);
+			Info.bMoved = true;
+			Info.bMovedRoot = true;
 			CachedObjectsInOuter.Emplace(Object);
 			GetObjectsWithOuter(Object, ChildrenOfMovedObjects, true /* bIncludeNestedObjects */, RF_NoFlags, EInternalObjectFlags::Garbage);
 		}
@@ -2132,10 +2194,11 @@ void FCookGenerationInfo::TakeOverCachedObjectsAndAddMoved(FGeneratorPackage& Ge
 	for (UObject* Object : ChildrenOfMovedObjects)
 	{
 		check(IsValid(Object));
-		bool bAlreadyExists;
-		ObjectSet.Add(Object, &bAlreadyExists);
-		if (!bAlreadyExists)
+		FCachedObjectInOuterGeneratorInfo& Info = CachedObjectsInOuterInfo.FindOrAdd(Object);
+		if (!Info.bInitialized)
 		{
+			Info.Initialize(Object);
+			Info.bMoved = true;
 			CachedObjectsInOuter.Emplace(Object);
 		}
 	}
@@ -2150,35 +2213,19 @@ EPollStatus FCookGenerationInfo::RefreshPackageObjects(FGeneratorPackage& Genera
 	TArray<UObject*> CurrentObjectsInOuter;
 	GetObjectsWithOuter(Package, CurrentObjectsInOuter, true /* bIncludeNestedObjects */, RF_NoFlags, EInternalObjectFlags::Garbage);
 
-	TSet<UObject*> ObjectSet;
 	check(PackageData); // RefreshPackageObjects is only called when there is a PackageData
 	TArray<FCachedObjectInOuter>& CachedObjectsInOuter = PackageData->GetCachedObjectsInOuter();
-	ObjectSet.Reserve(CachedObjectsInOuter.Num());
-	for (FCachedObjectInOuter& ExistingObject : CachedObjectsInOuter)
-	{
-		UObject* Object = ExistingObject.Object.Get();
-		if (Object)
-		{
-			bool bAlreadyExists;
-			ObjectSet.Add(Object, &bAlreadyExists);
-			check(!bAlreadyExists); // Objects in GetCachedObjectsInOuter are guaranteed unique and we haven't added any others yet
-			if (RootMovedObjects.Contains(Object))
-			{
-				GetObjectsWithOuter(Object, CurrentObjectsInOuter, true /* bIncludeNestedObjects */, RF_NoFlags, EInternalObjectFlags::Garbage);
-			}
-		}
-	}
 	UObject* FirstNewObject = nullptr;
 	for (UObject* Object : CurrentObjectsInOuter)
 	{
-		bool bAlreadyExists;
-		ObjectSet.Add(Object, &bAlreadyExists);
-		if (!bAlreadyExists)
+		FCachedObjectInOuterGeneratorInfo& Info = CachedObjectsInOuterInfo.FindOrAdd(Object);
+		if (!Info.bInitialized)
 		{
+			Info.Initialize(Object);
 			CachedObjectsInOuter.Emplace(Object);
-			if (!FirstNewObject )
+			if (!FirstNewObject)
 			{
-				FirstNewObject  = Object;
+				FirstNewObject = Object;
 			}
 		}
 	}
diff --git a/Engine/Source/Editor/UnrealEd/Private/Cooker/CookPackageData.h b/Engine/Source/Editor/UnrealEd/Private/Cooker/CookPackageData.h
--- a/Engine/Source/Editor/UnrealEd/Private/Cooker/CookPackageData.h
+++ b/Engine/Source/Editor/UnrealEd/Private/Cooker/CookPackageData.h
@@ -119,6 +119,29 @@ struct FCachedObjectInOuter
 	}
 };
 
+/**
+ * Extra information about CachedObjectsInOuter that a generator needs to know for diagnostics.
+ * The generator constructs an associated array (aka TMap) of these structures when it takes
+ * over the CachedObjectsInOuter.
+ */
+struct FCachedObjectInOuterGeneratorInfo
+{
+public:
+	/** Object->GetFullName() before it was deleted. */
+	FString FullName;
+	/** Has Initialize been called on *this. */
+	bool bInitialized = false;
+	/** Object->GetFlags() had RF_Public when the info was initialized. */
+	bool bPublic = false;
+	/** bMovedRoot is true, or this is a child object of such a moved object. */
+	bool bMoved = false;
+	/** Splitter informed us that object was moved into this package from another package. */
+	bool bMovedRoot = false;
+
+public:
+	void Initialize(UObject* Object);
+};
+
 /** Flags specifying the behavior of FPackageData::SendToState */
 enum class ESendFlags : uint8
 {
@@ -497,6 +520,7 @@ public:
 	 * TryCreateObjectCache and is cleared when leaving the save state.
 	 */
 	TArray<FCachedObjectInOuter>& GetCachedObjectsInOuter();
+	const TArray<FCachedObjectInOuter>& GetCachedObjectsInOuter() const;
 	template <typename ArrayType>
 	/** The list of platforms that were recorded as needscooking when CachedObjeObjectsInOuter was recorded. */
 	void GetCachedObjectsInOuterPlatforms(ArrayType& OutPlatforms) const;
@@ -868,7 +892,7 @@ public:
 	TArray<FAssetDependency> PackageDependencies;
 	FPackageData* PackageData = nullptr;
 	TArray<UPackage*> KeepReferencedPackages;
-	TSet<UObject*> RootMovedObjects;
+	TMap<UObject*, FCachedObjectInOuterGeneratorInfo> CachedObjectsInOuterInfo;
 private:
 	ESaveState GeneratorSaveState = ESaveState::StartGenerate;
 	bool bCreateAsMap : 1;
