[Editor Viewport] Add a cvar to disable invalidate on each SIE frame. Mouse clicks and other times a viewport would get invalidated should still occur.

r.Editor.Viewport.InvalidateEachSIEFrame (default to on to preserve behavior)

diff --git a/Engine/Source/Editor/UnrealEd/Private/EditorEngine.cpp b/Engine/Source/Editor/UnrealEd/Private/EditorEngine.cpp
--- a/Engine/Source/Editor/UnrealEd/Private/EditorEngine.cpp
+++ b/Engine/Source/Editor/UnrealEd/Private/EditorEngine.cpp
@@ -338,6 +338,11 @@ void DestroySelectionSets()
 
 } // namespace PrivateEditorSelection
 
+static FAutoConsoleVariable GInvalidateHitProxiesEachSIEFrameCVar(
+	TEXT("r.Editor.Viewport.InvalidateEachSIEFrame"),
+	1,
+	TEXT("Invalidate the viewport on each frame when SIE is running. Disabling this cvar (setting to 0) may improve performance, but impact the ability to click on objects that are moving in the viewport."));
+
 /**
 * A mapping of all startup packages to whether or not we have warned the user about editing them
 */
@@ -2381,7 +2386,7 @@ bool UEditorEngine::UpdateSingleViewportClient(FEditorViewportClient* InViewport
 				InViewportClient->GetWorld()->UpdateLevelStreaming();
 
 				// Also make sure hit proxies are refreshed for SIE viewports, as the user may be trying to grab an object or widget manipulator that's moving!
-				if( InViewportClient->IsRealtime() )
+				if( InViewportClient->IsRealtime() && (GInvalidateHitProxiesEachSIEFrameCVar->GetInt() != 0))
 				{
 					// @todo simulate: This may cause simulate performance to be worse in cases where you aren't needing to interact with gizmos.  Consider making this optional.
 					InViewportClient->RequestInvalidateHitProxy( InViewportClient->Viewport );
