Fix folder on duplicated actor when the folder contains spaces
#rnx

diff --git a/Engine/Source/Editor/UnrealEd/Private/EditorExporters.cpp b/Engine/Source/Editor/UnrealEd/Private/EditorExporters.cpp
--- a/Engine/Source/Editor/UnrealEd/Private/EditorExporters.cpp
+++ b/Engine/Source/Editor/UnrealEd/Private/EditorExporters.cpp
@@ -560,7 +560,7 @@ bool ULevelExporterT3D::ExportText( const FExportObjectInnerContext* Context, UO
 				FString SocketNameString = ( (ParentActor && SocketName != NAME_None) ? FString::Printf(TEXT(" SocketName=%s"), *SocketName.ToString() ) : TEXT(""));
 				FString GroupActor = (Actor->GroupActor ? FString::Printf(TEXT(" GroupActor=%s"), *Actor->GroupActor->GetName() ) : TEXT(""));
 				FString GroupFolder = (Actor->GroupActor ? FString::Printf(TEXT(" GroupFolder=%s"), *Actor->GroupActor->GetFolderPath().ToString() ) : TEXT(""));
-				FString ActorFolderPath = (Level->IsUsingActorFolders() ? FString::Printf(TEXT(" ActorFolderPath=%s"), *Actor->GetFolderPath().ToString()) : TEXT(""));
+				FString ActorFolderPath = (Level->IsUsingActorFolders() ? FString::Printf(TEXT(" ActorFolderPath=\"%s\""), *Actor->GetFolderPath().ToString()) : TEXT(""));
 				FString CopyPasteId = (Actor->CopyPasteId != INDEX_NONE) ? FString::Printf(TEXT(" CopyPasteId=%d"), Actor->CopyPasteId ) : TEXT("");
 				FString ContentBundleGuid = (Actor->GetContentBundleGuid().IsValid() ? FString::Printf(TEXT(" ActorContentBundleGuid=%s"), *Actor->GetContentBundleGuid().ToString()) : TEXT(""));
 				Ar.Logf( TEXT("%sBegin Actor Class=%s Name=%s Archetype=%s%s%s%s%s%s%s%s"), 
