Only show "Path Tracing" mode in viewport if ray tracing is enabled.

diff --git a/Engine/Source/Editor/UnrealEd/Private/SEditorViewportViewMenu.cpp b/Engine/Source/Editor/UnrealEd/Private/SEditorViewportViewMenu.cpp
--- a/Engine/Source/Editor/UnrealEd/Private/SEditorViewportViewMenu.cpp
+++ b/Engine/Source/Editor/UnrealEd/Private/SEditorViewportViewMenu.cpp
@@ -144,7 +144,7 @@ void SEditorViewportViewMenu::FillViewMenu(UToolMenu* Menu) const
 			}
 
 #if RHI_RAYTRACING
-			if (IsRayTracingAllowed())
+			if (IsRayTracingEnabled())
 			{
 				static auto PathTracingCvar = IConsoleManager::Get().FindTConsoleVariableDataInt(TEXT("r.PathTracing"));
 				if (PathTracingCvar && PathTracingCvar->GetValueOnAnyThread() != 0)
