DiffOnly header diffing for ZenStore
#rnx

diff --git a/Engine/Source/Editor/UnrealEd/Private/Cooker/DiffPackageWriter.cpp b/Engine/Source/Editor/UnrealEd/Private/Cooker/DiffPackageWriter.cpp
--- a/Engine/Source/Editor/UnrealEd/Private/Cooker/DiffPackageWriter.cpp
+++ b/Engine/Source/Editor/UnrealEd/Private/Cooker/DiffPackageWriter.cpp
@@ -30,6 +30,8 @@ DEFINE_LOG_CATEGORY_STATIC(LogDiff, Log, All);
 FDiffPackageWriter::FDiffPackageWriter(TUniquePtr<ICookedPackageWriter>&& InInner)
 	: Inner(MoveTemp(InInner))
 {
+	AccumulatorGlobals.Reset(new UE::DiffWriter::FAccumulatorGlobals(Inner.Get()));
+
 	GConfig->GetInt(TEXT("CookSettings"), TEXT("MaxDiffsToLog"), MaxDiffsToLog, GEditorIni);
 	// Command line override for MaxDiffsToLog
 	FParse::Value(FCommandLine::Get(), TEXT("MaxDiffstoLog="), MaxDiffsToLog);
@@ -180,6 +182,10 @@ void FDiffPackageWriter::WritePackageData(const FPackageInfo& Info, FLargeMemory
 	{
 		Accumulator.OnSecondSaveComplete(LocalInfo.HeaderSize);
 
+		// Avoid an assert when calling StaticFindObject during save to retrieve the actor's class.
+		// We are not writing the discovered objects into the saved package, so the call to StaticFindObject is legal.
+		TGuardValue<bool> GIsSavingPackageGuard(GIsSavingPackage, false);
+
 		TMap<FName, FArchiveDiffStats> PackageDiffStats;
 		const TCHAR* CutoffString = TEXT("UEditorEngine::Save()");
 		Accumulator.CompareWithPrevious(CutoffString, PackageDiffStats);
@@ -221,7 +227,7 @@ UE::DiffWriter::FAccumulator& FDiffPackageWriter::ConstructAccumulator(FName Pac
 	if (!Accumulator.IsValid())
 	{
 		check(!bHasStartedSecondSave); // Accumulator should already exist from CreateLinkerArchive in the first save
-		Accumulator = new UE::DiffWriter::FAccumulator(Asset, *PackageName.ToString(), MaxDiffsToLog,
+		Accumulator = new UE::DiffWriter::FAccumulator(*AccumulatorGlobals, Asset, *PackageName.ToString(), MaxDiffsToLog,
 			bIgnoreHeaderDiffs, GetDiffWriterMessageCallback(), Inner->GetCookCapabilities().HeaderFormat);
 	}
 	return *Accumulator;
diff --git a/Engine/Source/Editor/UnrealEd/Private/Cooker/DiffPackageWriter.h b/Engine/Source/Editor/UnrealEd/Private/Cooker/DiffPackageWriter.h
--- a/Engine/Source/Editor/UnrealEd/Private/Cooker/DiffPackageWriter.h
+++ b/Engine/Source/Editor/UnrealEd/Private/Cooker/DiffPackageWriter.h
@@ -6,6 +6,8 @@
 #include "Cooker/DiffWriterArchive.h"
 #include "Serialization/PackageWriter.h"
 
+namespace UE::DiffWriter { struct FAccumulatorGlobals; }
+
 /** A CookedPackageWriter that diffs output from the current cook with the file that was saved in the previous cook. */
 class FDiffPackageWriter : public ICookedPackageWriter
 {
@@ -131,6 +133,7 @@ protected:
 	TRefCountPtr<UE::DiffWriter::FAccumulator> Accumulators[2];
 	FBeginPackageInfo BeginInfo;
 	TUniquePtr<ICookedPackageWriter> Inner;
+	TUniquePtr<UE::DiffWriter::FAccumulatorGlobals> AccumulatorGlobals;
 	const TCHAR* Indent = nullptr;
 	const TCHAR* NewLine = nullptr;
 	FString DumpObjListParams;
diff --git a/Engine/Source/Editor/UnrealEd/Private/Cooker/DiffWriterArchive.cpp b/Engine/Source/Editor/UnrealEd/Private/Cooker/DiffWriterArchive.cpp
--- a/Engine/Source/Editor/UnrealEd/Private/Cooker/DiffWriterArchive.cpp
+++ b/Engine/Source/Editor/UnrealEd/Private/Cooker/DiffWriterArchive.cpp
@@ -3,6 +3,8 @@
 #include "Cooker/DiffWriterArchive.h"
 
 #include "Compression/CompressionUtil.h"
+#include "Cooker/DiffWriterLinkerLoadHeader.h"
+#include "Cooker/DiffWriterZenHeader.h"
 #include "HAL/FileManager.h"
 #include "HAL/PlatformFileManager.h"
 #include "HAL/PlatformStackWalk.h"
@@ -12,6 +14,9 @@
 #include "Misc/PackageName.h"
 #include "Misc/Paths.h"
 #include "Misc/ScopeExit.h"
+#include "PackageStoreOptimizer.h"
+#include "Serialization/AsyncLoading2.h"
+#include "Serialization/MemoryReader.h"
 #include "Serialization/StaticMemoryReader.h"
 #include "UObject/LinkerLoad.h"
 #include "UObject/ObjectMacros.h"
@@ -31,11 +36,8 @@ const TCHAR* const NewLineToken = TEXT("%DWA%\n");
 
 /** Logs any mismatching header data. */
 void DumpPackageHeaderDiffs(
-	const FPackageData& SourcePackage,
-	const FPackageData& DestPackage,
-	const FString& AssetFilename,
-	const int32 MaxDiffsToLog,
-	const EPackageHeaderFormat PackageHeaderFormat,
+	FAccumulatorGlobals& Globals, const FPackageData& SourcePackage, const FPackageData& DestPackage,
+	const FString& AssetFilename, const int32 MaxDiffsToLog, const EPackageHeaderFormat PackageHeaderFormat,
 	const FMessageCallback& MessageCallback);
 /** Returns a new linker for loading the specified package. */
 FLinkerLoad* CreateLinkerForPackage(
@@ -484,10 +486,38 @@ void FCallstacks::RecordSerialize(EOffsetFrame OffsetFrame, int64 CurrentOffset,
 	}
 }
 
-FAccumulator::FAccumulator(UObject* InAsset, FName InPackageName, int32 InMaxDiffsToLog, bool bInIgnoreHeaderDiffs,
-	FMessageCallback&& InMessageCallback, EPackageHeaderFormat InPackageHeaderFormat)
+FAccumulatorGlobals::FAccumulatorGlobals(ICookedPackageWriter* InnerPackageWriter)
+	: PackageWriter(InnerPackageWriter)
+{
+}
+
+void FAccumulatorGlobals::Initialize(EPackageHeaderFormat InFormat)
+{
+	if (bInitialized)
+	{
+		return;
+	}
+	bInitialized = true;
+	Format = InFormat;
+	switch (Format)
+	{
+	case EPackageHeaderFormat::PackageFileSummary:
+		break;
+	case EPackageHeaderFormat::ZenPackageSummary:
+		FPackageStoreOptimizer::FindScriptObjects(ScriptObjectsMap);
+		break;
+	default:
+		checkNoEntry();
+		break;
+	}
+}
+
+FAccumulator::FAccumulator(FAccumulatorGlobals& InGlobals, UObject* InAsset, FName InPackageName,
+	int32 InMaxDiffsToLog, bool bInIgnoreHeaderDiffs, FMessageCallback&& InMessageCallback,
+	EPackageHeaderFormat InPackageHeaderFormat)
 	: LinkerCallstacks()
 	, ExportsCallstacks()
+	, Globals(InGlobals)
 	, MessageCallback(MoveTemp(InMessageCallback))
 	, PackageName(InPackageName)
 	, Asset(InAsset)
@@ -942,6 +972,8 @@ void FAccumulator::CompareWithPrevious(const TCHAR* CallstackCutoffText, TMap<FN
 	// An FDiffArchiveForLinker should have been constructed by SavePackage and should still be in memory
 	check(LinkerArchive);
 
+	Globals.Initialize(PackageHeaderFormat);
+
 	const FName AssetClass = GetAssetClass();
 	OutStats.FindOrAdd(AssetClass).NewFileTotalSize = LinkerArchive->TotalSize();
 	if (PreviousPackageData.Size == 0)
@@ -983,8 +1015,20 @@ void FAccumulator::CompareWithPrevious(const TCHAR* CallstackCutoffText, TMap<FN
 
 	if (HeaderSize > 0 && OutStats.FindOrAdd(AssetClass).NumDiffs > 0)
 	{
-		DumpPackageHeaderDiffs(SourcePackage, DestPackage, Filename, MaxDiffsToLog,
-			PackageHeaderFormat, MessageCallback);
+		int32 NumHeaderDiffMessages = 0;
+		DumpPackageHeaderDiffs(Globals, SourcePackage, DestPackage, Filename, MaxDiffsToLog,
+			PackageHeaderFormat,
+			[&NumHeaderDiffMessages, this](ELogVerbosity::Type Verbosity, FStringView Message)
+			{
+				MessageCallback(Verbosity, Message);
+				++NumHeaderDiffMessages;
+			});
+		if (NumHeaderDiffMessages == 0)
+		{
+			MessageCallback(ELogVerbosity::Warning, FString::Printf(
+				TEXT("%s: headers are different, but DumpPackageHeaderDiffs does not yet implement describing the difference."),
+				*Filename));
+		}
 	}
 
 	FPackageData SourcePackageExports = SourcePackage;
@@ -1214,243 +1258,6 @@ FLinkerLoad* CreateLinkerForPackage(FUObjectSerializeContext* LoadContext, const
 	return Linker;
 }
 
-static FString GetTableKey(const FLinkerLoad* Linker, const FObjectExport& Export)
-{
-	FName ClassName = Export.ClassIndex.IsNull() ? FName(NAME_Class) : Linker->ImpExp(Export.ClassIndex).ObjectName;
-	return FString::Printf(TEXT("%s %s.%s"),
-		*ClassName.ToString(),
-		!Export.OuterIndex.IsNull() ? *Linker->ImpExp(Export.OuterIndex).ObjectName.ToString() : *FPackageName::GetShortName(Linker->LinkerRoot),
-		*Export.ObjectName.ToString());
-}
-
-static FString GetTableKey(const FLinkerLoad* Linker, const FObjectImport& Import)
-{
-	return FString::Printf(TEXT("%s %s.%s"),
-		*Import.ClassName.ToString(),
-		!Import.OuterIndex.IsNull() ? *Linker->ImpExp(Import.OuterIndex).ObjectName.ToString() : TEXT("NULL"),
-		*Import.ObjectName.ToString());
-}
-
-static inline FString GetTableKey(const FLinkerLoad* Linker, const FName& Name)
-{
-	return *Name.ToString();
-}
-
-static inline FString GetTableKey(const FLinkerLoad* Linker, FNameEntryId Id)
-{
-	return FName::GetEntry(Id)->GetPlainNameString();
-}
-
-static inline FString GetTableKeyForIndex(const FLinkerLoad* Linker, FPackageIndex Index)
-{
-	if (Index.IsNull())
-	{
-		return TEXT("NULL");
-	}
-	else if (Index.IsExport())
-	{
-		return GetTableKey(Linker, Linker->Exp(Index));
-	}
-	else
-	{
-		return GetTableKey(Linker, Linker->Imp(Index));
-	}
-}
-
-static bool ComparePackageIndices(FLinkerLoad* SourceLinker, FLinkerLoad* DestLinker, const FPackageIndex& SourceIndex,
-	const FPackageIndex& DestIndex, const FMessageCallback& MessageCallback);
-
-static bool CompareTableItem(FLinkerLoad* SourceLinker, FLinkerLoad* DestLinker,
-	const FName& SourceName, const FName& DestName, const FMessageCallback& MessageCallback)
-{
-	return SourceName == DestName;
-}
-
-static bool CompareTableItem(FLinkerLoad* SourceLinker, FLinkerLoad* DestLinker,
-	FNameEntryId SourceName, FNameEntryId DestName, const FMessageCallback& MessageCallback)
-{
-	return SourceName == DestName;
-}
-
-static FString ConvertItemToText(const FName& Name, FLinkerLoad* Linker)
-{
-	return Name.ToString();
-}
-
-static FString ConvertItemToText(FNameEntryId Id, FLinkerLoad* Linker)
-{
-	return FName::GetEntry(Id)->GetPlainNameString();
-}
-
-static bool CompareTableItem(FLinkerLoad* SourceLinker, FLinkerLoad* DestLinker, const FObjectImport& SourceImport,
-	const FObjectImport& DestImport, const FMessageCallback& MessageCallback)
-{
-	if (SourceImport.ObjectName != DestImport.ObjectName ||
-		SourceImport.ClassName != DestImport.ClassName ||
-		SourceImport.ClassPackage != DestImport.ClassPackage ||
-		!ComparePackageIndices(SourceLinker, DestLinker, SourceImport.OuterIndex,
-			DestImport.OuterIndex, MessageCallback))
-	{
-		return false;
-	}
-	else
-	{
-		return true;
-	}
-}
-
-static FString ConvertItemToText(const FObjectImport& Import, FLinkerLoad* Linker)
-{
-	return FString::Printf(
-		TEXT("%s ClassPackage: %s"),
-		*GetTableKey(Linker, Import),
-		*Import.ClassPackage.ToString()
-	);
-}
-
-static bool CompareTableItem(FLinkerLoad* SourceLinker, FLinkerLoad* DestLinker, const FObjectExport& SourceExport,
-	const FObjectExport& DestExport, const FMessageCallback& MessageCallback)
-{
-	if (SourceExport.ObjectName != DestExport.ObjectName ||
-		SourceExport.PackageFlags != DestExport.PackageFlags ||
-		SourceExport.ObjectFlags != DestExport.ObjectFlags ||
-		SourceExport.SerialSize != DestExport.SerialSize ||
-		SourceExport.bForcedExport != DestExport.bForcedExport ||
-		SourceExport.bNotForClient != DestExport.bNotForClient ||
-		SourceExport.bNotForServer != DestExport.bNotForServer ||
-		SourceExport.bNotAlwaysLoadedForEditorGame != DestExport.bNotAlwaysLoadedForEditorGame ||
-		SourceExport.bIsAsset != DestExport.bIsAsset ||
-		SourceExport.bIsInheritedInstance != DestExport.bIsInheritedInstance ||
-		SourceExport.bGeneratePublicHash != DestExport.bGeneratePublicHash ||
-		!ComparePackageIndices(SourceLinker, DestLinker, SourceExport.TemplateIndex,
-			DestExport.TemplateIndex, MessageCallback) ||
-		!ComparePackageIndices(SourceLinker, DestLinker, SourceExport.OuterIndex,
-			DestExport.OuterIndex, MessageCallback) ||
-		!ComparePackageIndices(SourceLinker, DestLinker, SourceExport.ClassIndex,
-			DestExport.ClassIndex, MessageCallback) ||
-		!ComparePackageIndices(SourceLinker, DestLinker, SourceExport.SuperIndex, 
-			DestExport.SuperIndex, MessageCallback))
-	{
-		return false;
-	}
-	else
-	{
-		return true;
-	}
-}
-
-static bool IsImportMapIdentical(FLinkerLoad* SourceLinker, FLinkerLoad* DestLinker,
-	const FMessageCallback& MessageCallback)
-{
-	bool bIdentical = (SourceLinker->ImportMap.Num() == DestLinker->ImportMap.Num());
-	if (bIdentical)
-	{
-		for (int32 ImportIndex = 0; ImportIndex < SourceLinker->ImportMap.Num(); ++ImportIndex)
-		{
-			if (!CompareTableItem(SourceLinker, DestLinker, SourceLinker->ImportMap[ImportIndex],
-				DestLinker->ImportMap[ImportIndex], MessageCallback))
-			{
-				bIdentical = false;
-				break;
-			}
-		}
-	}
-	return bIdentical;
-}
-
-static bool ComparePackageIndices(FLinkerLoad* SourceLinker, FLinkerLoad* DestLinker, const FPackageIndex& SourceIndex,
-	const FPackageIndex& DestIndex, const FMessageCallback& MessageCallback)
-{
-	if (SourceIndex.IsNull() && DestIndex.IsNull())
-	{
-		return true;
-	}
-
-	if (SourceIndex.IsExport() && DestIndex.IsExport())
-	{
-		int32 SourceArrayIndex = SourceIndex.ToExport();
-		int32 DestArrayIndex   = DestIndex  .ToExport();
-
-		if (!SourceLinker->ExportMap.IsValidIndex(SourceArrayIndex) || !DestLinker->ExportMap.IsValidIndex(DestArrayIndex))
-		{
-			MessageCallback(ELogVerbosity::Warning, FString::Printf(
-				TEXT("Invalid export indices found, source: %d (of %d), dest: %d (of %d)"),
-				SourceArrayIndex, SourceLinker->ExportMap.Num(), DestArrayIndex, DestLinker->ExportMap.Num()));
-			return false;
-		}
-
-		const FObjectExport& SourceOuterExport = SourceLinker->Exp(SourceIndex);
-		const FObjectExport& DestOuterExport   = DestLinker  ->Exp(DestIndex);
-
-		FString SourceOuterExportKey = GetTableKey(SourceLinker, SourceOuterExport);
-		FString DestOuterExportKey   = GetTableKey(DestLinker,   DestOuterExport);
-
-		return SourceOuterExportKey == DestOuterExportKey;
-	}
-
-	if (SourceIndex.IsImport() && DestIndex.IsImport())
-	{
-		int32 SourceArrayIndex = SourceIndex.ToImport();
-		int32 DestArrayIndex   = DestIndex  .ToImport();
-
-		if (!SourceLinker->ImportMap.IsValidIndex(SourceArrayIndex) || !DestLinker->ImportMap.IsValidIndex(DestArrayIndex))
-		{
-			MessageCallback(ELogVerbosity::Warning, FString::Printf(
-				TEXT("Invalid import indices found, source: %d (of %d), dest: %d (of %d)"),
-				SourceArrayIndex, SourceLinker->ImportMap.Num(), DestArrayIndex, DestLinker->ImportMap.Num()));
-			return false;
-		}
-
-		const FObjectImport& SourceOuterImport = SourceLinker->Imp(SourceIndex);
-		const FObjectImport& DestOuterImport   = DestLinker  ->Imp(DestIndex);
-
-		FString SourceOuterImportKey = GetTableKey(SourceLinker, SourceOuterImport);
-		FString DestOuterImportKey   = GetTableKey(DestLinker,   DestOuterImport);
-
-		return SourceOuterImportKey == DestOuterImportKey;
-	}
-
-	return false;
-}
-
-static FString ConvertItemToText(const FObjectExport& Export, FLinkerLoad* Linker)
-{
-	FName ClassName = Export.ClassIndex.IsNull() ? FName(NAME_Class) : Linker->ImpExp(Export.ClassIndex).ObjectName;
-	return FString::Printf(TEXT("%s Super: %s, Template: %s, Flags: %d, Size: %lld, PackageFlags: %d, ForcedExport: %d, NotForClient: %d, NotForServer: %d, NotAlwaysLoadedForEditorGame: %d, IsAsset: %d, IsInheritedInstance: %d, GeneratePublicHash: %d"),
-		*GetTableKey(Linker, Export),
-		*GetTableKeyForIndex(Linker, Export.SuperIndex),
-		*GetTableKeyForIndex(Linker, Export.TemplateIndex),
-		(int32)Export.ObjectFlags,
-		Export.SerialSize,
-		Export.PackageFlags,
-		Export.bForcedExport,
-		Export.bNotForClient,
-		Export.bNotForServer,
-		Export.bNotAlwaysLoadedForEditorGame,
-		Export.bIsAsset,
-		Export.bIsInheritedInstance,
-		Export.bGeneratePublicHash);
-}
-
-static bool IsExportMapIdentical(FLinkerLoad* SourceLinker, FLinkerLoad* DestLinker,
-	const FMessageCallback& MessageCallback)
-{
-	bool bIdentical = (SourceLinker->ExportMap.Num() == DestLinker->ExportMap.Num());
-	if (bIdentical)
-	{
-		for (int32 ExportIndex = 0; ExportIndex < SourceLinker->ExportMap.Num(); ++ExportIndex)
-		{
-			if (!CompareTableItem(SourceLinker, DestLinker, SourceLinker->ExportMap[ExportIndex],
-				DestLinker->ExportMap[ExportIndex], MessageCallback))
-			{
-				bIdentical = false;
-				break;
-			}
-		}
-	}
-	return bIdentical;
-}
-
 /** Structure that holds an item from the NameMap/ImportMap/ExportMap in a TSet for diffing */
 template <typename T>
 struct TTableItem
@@ -1481,16 +1288,15 @@ struct TTableItem
 };
 
 /** Dumps differences between Linker tables */
-template <typename T>
+template <typename T, typename ContextProvider>
 static void DumpTableDifferences(
-	FLinkerLoad* SourceLinker, 
-	FLinkerLoad* DestLinker, 
-	TArray<T>& SourceTable, 
-	TArray<T>& DestTable,
+	ContextProvider& SourceContext,
+	ContextProvider& DestContext,
+	TConstArrayView<T> SourceTable, 
+	TConstArrayView<T> DestTable,
 	const TCHAR* AssetFilename,
 	const TCHAR* ItemName,
-	const int32 MaxDiffsToLog,
-	const FMessageCallback& MessageCallback
+	const int32 MaxDiffsToLog
 )
 {
 	FString HumanReadableString;
@@ -1506,12 +1312,12 @@ static void DumpTableDifferences(
 	for (int32 Index = 0; Index < SourceTable.Num(); ++Index)
 	{
 		const T& Item = SourceTable[Index];
-		SourceSet.Add(TTableItem<T>(GetTableKey(SourceLinker, Item), &Item, Index));
+		SourceSet.Add(TTableItem<T>(SourceContext.GetTableKey(Item), &Item, Index));
 	}
 	for (int32 Index = 0; Index < DestTable.Num(); ++Index)
 	{
 		const T& Item = DestTable[Index];
-		DestSet.Add(TTableItem<T>(GetTableKey(DestLinker, Item), &Item, Index));
+		DestSet.Add(TTableItem<T>(DestContext.GetTableKey(Item), &Item, Index));
 	}
 
 	// Determine the list of items removed from the source package and added to the dest package
@@ -1523,8 +1329,8 @@ static void DumpTableDifferences(
 	{
 		if (const TTableItem<T>* ChangedDestItem = DestSet.Find(ChangedSourceItem))
 		{
-			if (!CompareTableItem(SourceLinker, DestLinker, *ChangedSourceItem.Item,
-				*ChangedDestItem->Item, MessageCallback))
+			if (!SourceContext.CompareTableItem(DestContext, *ChangedSourceItem.Item,
+				*ChangedDestItem->Item))
 			{
 				RemovedItems.Add(ChangedSourceItem);
 				AddedItems  .Add(*ChangedDestItem);
@@ -1540,13 +1346,13 @@ static void DumpTableDifferences(
 	for (const TTableItem<T>& RemovedItem : RemovedItems)
 	{
 		HumanReadableString += IndentToken;
-		HumanReadableString += FString::Printf(TEXT("-[%d] %s"), RemovedItem.Index, *ConvertItemToText(*RemovedItem.Item, SourceLinker));
+		HumanReadableString += FString::Printf(TEXT("-[%d] %s"), RemovedItem.Index, *SourceContext.ConvertItemToText(*RemovedItem.Item));
 		HumanReadableString += NewLineToken;
 	}
 	for (const TTableItem<T>& AddedItem : AddedItems)
 	{
 		HumanReadableString += IndentToken;
-		HumanReadableString += FString::Printf(TEXT("+[%d] %s"), AddedItem.Index, *ConvertItemToText(*AddedItem.Item, DestLinker));
+		HumanReadableString += FString::Printf(TEXT("+[%d] %s"), AddedItem.Index, *DestContext.ConvertItemToText(*AddedItem.Item));
 		HumanReadableString += NewLineToken;
 	}
 
@@ -1561,7 +1367,7 @@ static void DumpTableDifferences(
 		HumanReadableString += NewLineToken;
 	}
 
-	MessageCallback(ELogVerbosity::Warning, FString::Printf(
+	SourceContext.LogMessage(ELogVerbosity::Warning, FString::Printf(
 		TEXT("%s: %sMap is different (%d %ss in source package vs %d %ss in dest package):%s%s"),		
 		AssetFilename,
 		ItemName,
@@ -1574,6 +1380,7 @@ static void DumpTableDifferences(
 }
 
 void DumpPackageHeaderDiffs_LinkerLoad(
+	FAccumulatorGlobals& Globals,
 	const FPackageData& SourcePackage,
 	const FPackageData& DestPackage,
 	const FString& AssetFilename,
@@ -1614,22 +1421,25 @@ void DumpPackageHeaderDiffs_LinkerLoad(
 
 	if (SourceLinker && DestLinker)
 	{
+		FDiffWriterLinkerLoadHeader SourceContext(SourceLinker, MessageCallback);
+		FDiffWriterLinkerLoadHeader DestContext(DestLinker, MessageCallback);
+
 		if (SourceLinker->NameMap != DestLinker->NameMap)
 		{
-			DumpTableDifferences<FNameEntryId>(SourceLinker, DestLinker, SourceLinker->NameMap, DestLinker->NameMap,
-				*AssetFilename, TEXT("Name"), MaxDiffsToLog, MessageCallback);
+			DumpTableDifferences<FNameEntryId>(SourceContext, DestContext, SourceLinker->NameMap, DestLinker->NameMap,
+				*AssetFilename, TEXT("Name"), MaxDiffsToLog);
 		}
 
-		if (!IsImportMapIdentical(SourceLinker, DestLinker, MessageCallback))
+		if (!SourceContext.IsImportMapIdentical(DestContext))
 		{
-			DumpTableDifferences<FObjectImport>(SourceLinker, DestLinker, SourceLinker->ImportMap, DestLinker->ImportMap,
-				*AssetFilename, TEXT("Import"), MaxDiffsToLog, MessageCallback);
+			DumpTableDifferences<FObjectImport>(SourceContext, DestContext, SourceLinker->ImportMap, DestLinker->ImportMap,
+				*AssetFilename, TEXT("Import"), MaxDiffsToLog);
 		}
 
-		if (!IsExportMapIdentical(SourceLinker, DestLinker, MessageCallback))
+		if (!SourceContext.IsExportMapIdentical(DestContext))
 		{
-			DumpTableDifferences<FObjectExport>(SourceLinker, DestLinker, SourceLinker->ExportMap, DestLinker->ExportMap,
-				*AssetFilename, TEXT("Export"), MaxDiffsToLog, MessageCallback);
+			DumpTableDifferences<FObjectExport>(SourceContext, DestContext, SourceLinker->ExportMap, DestLinker->ExportMap,
+				*AssetFilename, TEXT("Export"), MaxDiffsToLog);
 		}
 	}
 
@@ -1644,19 +1454,71 @@ void DumpPackageHeaderDiffs_LinkerLoad(
 }
 
 void DumpPackageHeaderDiffs_ZenPackage(
+	FAccumulatorGlobals& Globals,
 	const FPackageData& SourcePackage,
 	const FPackageData& DestPackage,
 	const FString& AssetFilename,
 	const int32 MaxDiffsToLog,
 	const FMessageCallback& MessageCallback)
 {
-	// TODO: Fill in detailed diffing of Zen Package Summary
-	MessageCallback(ELogVerbosity::Warning, FString::Printf(
-		TEXT("%s: header is different (header diffing not yet implemented for -zenstore; cook with -skipzenstore to see header diffs)"),
-		*AssetFilename));
+	FDiffWriterZenHeader SourceHeader(Globals, MessageCallback, true /* bUseZenStore */, SourcePackage, AssetFilename,
+		TEXT("source"));
+	FDiffWriterZenHeader DestHeader(Globals, MessageCallback, false /* bUseZenStore */, DestPackage, AssetFilename,
+		TEXT("dest"));
+	if (!SourceHeader.IsValid() || !DestHeader.IsValid())
+	{
+		// Explanation was logged by TryConstructSummary
+		return;
+	}
+
+	TArray<FName> SourceNames;
+	TArray<FName> DestNames;
+	for (FDisplayNameEntryId Id : SourceHeader.GetPackageHeader().NameMap)
+	{
+		SourceNames.Add(Id.ToName(0));
+	}
+	for (FDisplayNameEntryId Id : DestHeader.GetPackageHeader().NameMap)
+	{
+		DestNames.Add(Id.ToName(0));
+	}
+	Algo::Sort(SourceNames, FNameFastLess());
+	Algo::Sort(DestNames, FNameFastLess());
+
+	if (SourceNames != DestNames)
+	{
+		DumpTableDifferences<FName>(SourceHeader, DestHeader, SourceNames, DestNames,
+			*AssetFilename, TEXT("Name"), MaxDiffsToLog);
+	}
+
+	if (!SourceHeader.IsImportMapIdentical(DestHeader))
+	{
+		DumpTableDifferences<FPackageObjectIndex>(SourceHeader, DestHeader, SourceHeader.GetPackageHeader().ImportMap,
+			DestHeader.GetPackageHeader().ImportMap, *AssetFilename, TEXT("Import"), MaxDiffsToLog);
+	}
+
+	if (!SourceHeader.IsExportMapIdentical(DestHeader))
+	{
+		int32 SourceNum = SourceHeader.GetPackageHeader().ExportMap.Num();
+		int32 DestNum = DestHeader.GetPackageHeader().ExportMap.Num();
+		TArray<FZenHeaderIndexIntoExportMap> SourceIndices;
+		SourceIndices.Reserve(SourceNum);
+		for (int N = 0; N < SourceNum; ++N)
+		{
+			SourceIndices.Add(FZenHeaderIndexIntoExportMap{ N });
+		}
+		TArray<FZenHeaderIndexIntoExportMap> DestIndices;
+		DestIndices.Reserve(DestNum);
+		for (int N = 0; N < DestNum; ++N)
+		{
+			DestIndices.Add(FZenHeaderIndexIntoExportMap{ N });
+		}
+		DumpTableDifferences<FZenHeaderIndexIntoExportMap>(SourceHeader, DestHeader, SourceIndices, DestIndices,
+			*AssetFilename, TEXT("Export"), MaxDiffsToLog);
+	}
 }
 
 void DumpPackageHeaderDiffs(
+	FAccumulatorGlobals& Globals,
 	const FPackageData& SourcePackage,
 	const FPackageData& DestPackage,
 	const FString& AssetFilename,
@@ -1667,10 +1529,10 @@ void DumpPackageHeaderDiffs(
 	switch (PackageHeaderFormat)
 	{
 	case EPackageHeaderFormat::PackageFileSummary:
-		DumpPackageHeaderDiffs_LinkerLoad(SourcePackage, DestPackage, AssetFilename, MaxDiffsToLog, MessageCallback);
+		DumpPackageHeaderDiffs_LinkerLoad(Globals, SourcePackage, DestPackage, AssetFilename, MaxDiffsToLog, MessageCallback);
 		break;
 	case EPackageHeaderFormat::ZenPackageSummary:
-		DumpPackageHeaderDiffs_ZenPackage(SourcePackage, DestPackage, AssetFilename, MaxDiffsToLog, MessageCallback);
+		DumpPackageHeaderDiffs_ZenPackage(Globals, SourcePackage, DestPackage, AssetFilename, MaxDiffsToLog, MessageCallback);
 		break;
 	default:
 		unimplemented();
diff --git a/Engine/Source/Editor/UnrealEd/Private/Cooker/DiffWriterArchive.h b/Engine/Source/Editor/UnrealEd/Private/Cooker/DiffWriterArchive.h
--- a/Engine/Source/Editor/UnrealEd/Private/Cooker/DiffWriterArchive.h
+++ b/Engine/Source/Editor/UnrealEd/Private/Cooker/DiffWriterArchive.h
@@ -8,6 +8,7 @@
 #include "Containers/UnrealString.h"
 #include "HAL/Platform.h"
 #include "Logging/LogVerbosity.h"
+#include "PackageStoreOptimizer.h"
 #include "Serialization/Archive.h"
 #include "Serialization/ArchiveProxy.h"
 #include "Serialization/ArchiveStackTrace.h"
@@ -219,6 +220,23 @@ private:
 	int64 EndOffset;
 };
 
+/** Global data (e.g. the FPackageId of every object in /Script) used during diffing */
+struct FAccumulatorGlobals
+{
+public:
+	// Zen variables
+	TMap<FPackageObjectIndex, FPackageStoreOptimizer::FScriptObjectData> ScriptObjectsMap;
+
+	// Shared variables
+	ICookedPackageWriter* PackageWriter = nullptr;
+	EPackageHeaderFormat Format = EPackageHeaderFormat::PackageFileSummary;
+	bool bInitialized = false;
+
+public:
+	FAccumulatorGlobals(ICookedPackageWriter* InnerPackageWriter = nullptr);
+	void Initialize(EPackageHeaderFormat Format);
+};
+
 /**
  * Collects the memory version of a saved package, compares it with an existing package on disk, and reports callstack
  * for the Serialize call at each offset where they differ.
@@ -231,8 +249,8 @@ private:
 class FAccumulator : public FRefCountBase
 {
 public:
-	FAccumulator(UObject* InAsset, FName InPackageName, int32 InMaxDiffsToLog, bool bInIgnoreHeaderDiffs,
-		FMessageCallback&& InMessageCallback, EPackageHeaderFormat InPackageHeaderFormat);
+	FAccumulator(FAccumulatorGlobals& InGlobals, UObject* InAsset, FName InPackageName, int32 InMaxDiffsToLog,
+		bool bInIgnoreHeaderDiffs, FMessageCallback&& InMessageCallback, EPackageHeaderFormat InPackageHeaderFormat);
 	virtual ~FAccumulator();
 
 	void OnFirstSaveComplete(FStringView InLooseFilePath, int64 InHeaderSize, int64 InPreTransformHeaderSize,
@@ -262,6 +280,7 @@ private:
 	ICookedPackageWriter::FPreviousCookedBytesData PreviousPackageData;
 	FDiffArchive* LinkerArchive = nullptr;
 	FDiffArchive* ExportsArchive = nullptr;
+	FAccumulatorGlobals& Globals;
 
 	FDiffMap DiffMap;
 	FMessageCallback MessageCallback;
diff --git a/Engine/Source/Editor/UnrealEd/Private/Cooker/DiffWriterLinkerLoadHeader.cpp b/Engine/Source/Editor/UnrealEd/Private/Cooker/DiffWriterLinkerLoadHeader.cpp
new file mode 100644
--- /dev/null
+++ b/Engine/Source/Editor/UnrealEd/Private/Cooker/DiffWriterLinkerLoadHeader.cpp
@@ -0,0 +1,257 @@
+// Copyright Epic Games, Inc. All Rights Reserved.
+
+#include "Cooker/DiffWriterLinkerLoadHeader.h"
+
+#include "UObject/LinkerLoad.h"
+
+namespace UE::DiffWriter
+{
+
+FDiffWriterLinkerLoadHeader::FDiffWriterLinkerLoadHeader(FLinkerLoad* InLinker, const FMessageCallback& InMessageCallback)
+	: Linker(InLinker)
+	, MessageCallback(InMessageCallback)
+{
+	Linker = InLinker;
+}
+FString FDiffWriterLinkerLoadHeader::GetTableKey(const FObjectExport& Export)
+{
+	FName ClassName = Export.ClassIndex.IsNull() ? FName(NAME_Class) : Linker->ImpExp(Export.ClassIndex).ObjectName;
+	return FString::Printf(TEXT("%s %s.%s"),
+		*ClassName.ToString(),
+		!Export.OuterIndex.IsNull() ? *Linker->ImpExp(Export.OuterIndex).ObjectName.ToString() : *FPackageName::GetShortName(Linker->LinkerRoot),
+		*Export.ObjectName.ToString());
+}
+FString FDiffWriterLinkerLoadHeader::GetTableKey(const FObjectImport& Import)
+{
+	return FString::Printf(TEXT("%s %s.%s"),
+		*Import.ClassName.ToString(),
+		!Import.OuterIndex.IsNull() ? *Linker->ImpExp(Import.OuterIndex).ObjectName.ToString() : TEXT("NULL"),
+		*Import.ObjectName.ToString());
+}
+FString FDiffWriterLinkerLoadHeader::GetTableKey(const FName& Name)
+{
+	return *Name.ToString();
+}
+
+FString FDiffWriterLinkerLoadHeader::GetTableKey(FNameEntryId Id)
+{
+	return FName::GetEntry(Id)->GetPlainNameString();
+}
+
+FString FDiffWriterLinkerLoadHeader::GetTableKeyForIndex(FPackageIndex Index)
+{
+	if (Index.IsNull())
+	{
+		return TEXT("NULL");
+	}
+	else if (Index.IsExport())
+	{
+		return GetTableKey(Linker->Exp(Index));
+	}
+	else
+	{
+		return GetTableKey(Linker->Imp(Index));
+	}
+}
+
+bool FDiffWriterLinkerLoadHeader::CompareTableItem(FDiffWriterLinkerLoadHeader& DestContext,
+	const FName& SourceName, const FName& DestName)
+{
+	return SourceName == DestName;
+}
+
+bool FDiffWriterLinkerLoadHeader::CompareTableItem(FDiffWriterLinkerLoadHeader& DestContext,
+	FNameEntryId SourceName, FNameEntryId DestName)
+{
+	return SourceName == DestName;
+}
+
+FString FDiffWriterLinkerLoadHeader::ConvertItemToText(const FName& Name)
+{
+	return Name.ToString();
+}
+
+FString FDiffWriterLinkerLoadHeader::ConvertItemToText(FNameEntryId Id)
+{
+	return FName::GetEntry(Id)->GetPlainNameString();
+}
+
+bool FDiffWriterLinkerLoadHeader::CompareTableItem(FDiffWriterLinkerLoadHeader& DestContext, const FObjectImport& SourceImport,
+	const FObjectImport& DestImport)
+{
+	if (SourceImport.ObjectName != DestImport.ObjectName ||
+		SourceImport.ClassName != DestImport.ClassName ||
+		SourceImport.ClassPackage != DestImport.ClassPackage ||
+		!ComparePackageIndices(DestContext, SourceImport.OuterIndex,
+			DestImport.OuterIndex))
+	{
+		return false;
+	}
+	else
+	{
+		return true;
+	}
+}
+
+FString FDiffWriterLinkerLoadHeader::ConvertItemToText(const FObjectImport& Import)
+{
+	return FString::Printf(
+		TEXT("%s ClassPackage: %s"),
+		*GetTableKey(Import),
+		*Import.ClassPackage.ToString()
+	);
+}
+
+bool FDiffWriterLinkerLoadHeader::CompareTableItem(FDiffWriterLinkerLoadHeader& DestContext, const FObjectExport& SourceExport,
+	const FObjectExport& DestExport)
+{
+	if (SourceExport.ObjectName != DestExport.ObjectName ||
+		SourceExport.PackageFlags != DestExport.PackageFlags ||
+		SourceExport.ObjectFlags != DestExport.ObjectFlags ||
+		SourceExport.SerialSize != DestExport.SerialSize ||
+		SourceExport.bForcedExport != DestExport.bForcedExport ||
+		SourceExport.bNotForClient != DestExport.bNotForClient ||
+		SourceExport.bNotForServer != DestExport.bNotForServer ||
+		SourceExport.bNotAlwaysLoadedForEditorGame != DestExport.bNotAlwaysLoadedForEditorGame ||
+		SourceExport.bIsAsset != DestExport.bIsAsset ||
+		SourceExport.bIsInheritedInstance != DestExport.bIsInheritedInstance ||
+		SourceExport.bGeneratePublicHash != DestExport.bGeneratePublicHash ||
+		!ComparePackageIndices(DestContext, SourceExport.TemplateIndex,
+			DestExport.TemplateIndex) ||
+		!ComparePackageIndices(DestContext, SourceExport.OuterIndex,
+			DestExport.OuterIndex) ||
+		!ComparePackageIndices(DestContext, SourceExport.ClassIndex,
+			DestExport.ClassIndex) ||
+		!ComparePackageIndices(DestContext, SourceExport.SuperIndex,
+			DestExport.SuperIndex))
+	{
+		return false;
+	}
+	else
+	{
+		return true;
+	}
+}
+
+bool FDiffWriterLinkerLoadHeader::IsImportMapIdentical(FDiffWriterLinkerLoadHeader& DestContext)
+{
+	FLinkerLoad* SourceLinker = this->Linker;
+	FLinkerLoad* DestLinker = DestContext.Linker;
+	bool bIdentical = (SourceLinker->ImportMap.Num() == DestLinker->ImportMap.Num());
+	if (bIdentical)
+	{
+		for (int32 ImportIndex = 0; ImportIndex < SourceLinker->ImportMap.Num(); ++ImportIndex)
+		{
+			if (!CompareTableItem(DestContext, SourceLinker->ImportMap[ImportIndex],
+				DestLinker->ImportMap[ImportIndex]))
+			{
+				bIdentical = false;
+				break;
+			}
+		}
+	}
+	return bIdentical;
+}
+
+bool FDiffWriterLinkerLoadHeader::ComparePackageIndices(FDiffWriterLinkerLoadHeader& DestContext, const FPackageIndex& SourceIndex,
+	const FPackageIndex& DestIndex)
+{
+	if (SourceIndex.IsNull() && DestIndex.IsNull())
+	{
+		return true;
+	}
+	FDiffWriterLinkerLoadHeader& SourceContext = *this;
+	FLinkerLoad* SourceLinker = SourceContext.Linker;
+	FLinkerLoad* DestLinker = DestContext.Linker;
+
+	if (SourceIndex.IsExport() && DestIndex.IsExport())
+	{
+		int32 SourceArrayIndex = SourceIndex.ToExport();
+		int32 DestArrayIndex = DestIndex.ToExport();
+
+		if (!SourceLinker->ExportMap.IsValidIndex(SourceArrayIndex) || !DestLinker->ExportMap.IsValidIndex(DestArrayIndex))
+		{
+			LogMessage(ELogVerbosity::Warning, FString::Printf(
+				TEXT("Invalid export indices found, source: %d (of %d), dest: %d (of %d)"),
+				SourceArrayIndex, SourceLinker->ExportMap.Num(), DestArrayIndex, DestLinker->ExportMap.Num()));
+			return false;
+		}
+
+		const FObjectExport& SourceOuterExport = SourceLinker->Exp(SourceIndex);
+		const FObjectExport& DestOuterExport = DestLinker->Exp(DestIndex);
+
+		FString SourceOuterExportKey = SourceContext.GetTableKey(SourceOuterExport);
+		FString DestOuterExportKey = DestContext.GetTableKey(DestOuterExport);
+
+		return SourceOuterExportKey == DestOuterExportKey;
+	}
+
+	if (SourceIndex.IsImport() && DestIndex.IsImport())
+	{
+		int32 SourceArrayIndex = SourceIndex.ToImport();
+		int32 DestArrayIndex = DestIndex.ToImport();
+
+		if (!SourceLinker->ImportMap.IsValidIndex(SourceArrayIndex) || !DestLinker->ImportMap.IsValidIndex(DestArrayIndex))
+		{
+			LogMessage(ELogVerbosity::Warning, FString::Printf(
+				TEXT("Invalid import indices found, source: %d (of %d), dest: %d (of %d)"),
+				SourceArrayIndex, SourceLinker->ImportMap.Num(), DestArrayIndex, DestLinker->ImportMap.Num()));
+			return false;
+		}
+
+		const FObjectImport& SourceOuterImport = SourceLinker->Imp(SourceIndex);
+		const FObjectImport& DestOuterImport = DestLinker->Imp(DestIndex);
+
+		FString SourceOuterImportKey = SourceContext.GetTableKey(SourceOuterImport);
+		FString DestOuterImportKey = DestContext.GetTableKey(DestOuterImport);
+
+		return SourceOuterImportKey == DestOuterImportKey;
+	}
+
+	return false;
+}
+
+FString FDiffWriterLinkerLoadHeader::ConvertItemToText(const FObjectExport& Export)
+{
+	return FString::Printf(TEXT("%s Super: %s, Template: %s, Flags: %d, Size: %lld, PackageFlags: %d, ForcedExport: %d, NotForClient: %d, NotForServer: %d, NotAlwaysLoadedForEditorGame: %d, IsAsset: %d, IsInheritedInstance: %d, GeneratePublicHash: %d"),
+		*GetTableKey(Export),
+		*GetTableKeyForIndex(Export.SuperIndex),
+		*GetTableKeyForIndex(Export.TemplateIndex),
+		(int32)Export.ObjectFlags,
+		Export.SerialSize,
+		Export.PackageFlags,
+		Export.bForcedExport,
+		Export.bNotForClient,
+		Export.bNotForServer,
+		Export.bNotAlwaysLoadedForEditorGame,
+		Export.bIsAsset,
+		Export.bIsInheritedInstance,
+		Export.bGeneratePublicHash);
+}
+
+bool FDiffWriterLinkerLoadHeader::IsExportMapIdentical(FDiffWriterLinkerLoadHeader& DestContext)
+{
+	FLinkerLoad* SourceLinker = this->Linker;
+	FLinkerLoad* DestLinker = DestContext.Linker;
+	bool bIdentical = (SourceLinker->ExportMap.Num() == DestLinker->ExportMap.Num());
+	if (bIdentical)
+	{
+		for (int32 ExportIndex = 0; ExportIndex < SourceLinker->ExportMap.Num(); ++ExportIndex)
+		{
+			if (!CompareTableItem(DestContext, SourceLinker->ExportMap[ExportIndex],
+				DestLinker->ExportMap[ExportIndex]))
+			{
+				bIdentical = false;
+				break;
+			}
+		}
+	}
+	return bIdentical;
+}
+
+void FDiffWriterLinkerLoadHeader::LogMessage(ELogVerbosity::Type Verbosity, FStringView Message)
+{
+	MessageCallback(Verbosity, Message);
+}
+
+}
\ No newline at end of file
diff --git a/Engine/Source/Editor/UnrealEd/Private/Cooker/DiffWriterLinkerLoadHeader.h b/Engine/Source/Editor/UnrealEd/Private/Cooker/DiffWriterLinkerLoadHeader.h
new file mode 100644
--- /dev/null
+++ b/Engine/Source/Editor/UnrealEd/Private/Cooker/DiffWriterLinkerLoadHeader.h
@@ -0,0 +1,47 @@
+// Copyright Epic Games, Inc. All Rights Reserved.
+
+#pragma once
+
+#include "DiffWriterArchive.h"
+
+namespace UE::DiffWriter
+{
+
+/**
+	* Context object for DumpTableDifferences when we are using EPackageHeaderFormat::PackageFileSummary and so the
+	* header information comes from FLinkerLoad.
+*/
+struct FDiffWriterLinkerLoadHeader
+{
+public:
+	FLinkerLoad* Linker = nullptr;
+	const FMessageCallback& MessageCallback;
+
+public:
+	// Constructor
+	FDiffWriterLinkerLoadHeader(FLinkerLoad* InLinker, const FMessageCallback& InMessageCallback);
+
+	// API for DumpTableDifferences
+	FString GetTableKey(const FObjectExport& Export);
+	FString GetTableKey(const FObjectImport& Import);
+	FString GetTableKey(const FName& Name);
+	FString GetTableKey(FNameEntryId Id);
+	FString GetTableKeyForIndex(FPackageIndex Index);
+	FString ConvertItemToText(const FObjectExport& Export);
+	FString ConvertItemToText(const FName& Name);
+	FString ConvertItemToText(FNameEntryId Id);
+	FString ConvertItemToText(const FObjectImport& Import);
+	bool CompareTableItem(FDiffWriterLinkerLoadHeader& DestContext, const FName& SourceName, const FName& DestName);
+	bool CompareTableItem(FDiffWriterLinkerLoadHeader& DestContext, FNameEntryId SourceName, FNameEntryId DestName);
+	bool CompareTableItem(FDiffWriterLinkerLoadHeader& DestContext, const FObjectImport& SourceImport,
+		const FObjectImport& DestImport);
+	bool CompareTableItem(FDiffWriterLinkerLoadHeader& DestContext, const FObjectExport& SourceExport,
+		const FObjectExport& DestExport);
+	bool ComparePackageIndices(FDiffWriterLinkerLoadHeader& DestContext, const FPackageIndex& SourceIndex,
+		const FPackageIndex& DestIndex);
+	bool IsImportMapIdentical(FDiffWriterLinkerLoadHeader& DestContext);
+	bool IsExportMapIdentical(FDiffWriterLinkerLoadHeader& DestContext);
+	void LogMessage(ELogVerbosity::Type Verbosity, FStringView Message);
+};
+
+}
\ No newline at end of file
diff --git a/Engine/Source/Editor/UnrealEd/Private/Cooker/DiffWriterZenHeader.cpp b/Engine/Source/Editor/UnrealEd/Private/Cooker/DiffWriterZenHeader.cpp
new file mode 100644
--- /dev/null
+++ b/Engine/Source/Editor/UnrealEd/Private/Cooker/DiffWriterZenHeader.cpp
@@ -0,0 +1,416 @@
+// Copyright Epic Games, Inc. All Rights Reserved.
+
+#include "Cooker/DiffWriterZenHeader.h"
+
+#include "UObject/Package.h"
+#include "UObject/UObjectGlobals.h"
+#include "UObject/UObjectHash.h"
+
+namespace UE::DiffWriter
+{
+
+FDiffWriterZenHeader::FDiffWriterZenHeader(FAccumulatorGlobals& InGlobals,
+	const FMessageCallback& InMessageCallback, bool bInUseZenStoreForImports,
+	const FPackageData& Package, const FString& AssetFilename, const TCHAR* WhichComparisonPackage)
+	: Globals(InGlobals)
+	, MessageCallback(InMessageCallback)
+	, bUseZenStoreForImports(bInUseZenStoreForImports)
+{
+	FString ErrorMessage;
+	PackageHeader = FZenPackageHeader::MakeView(FMemoryView(Package.Data, Package.Size), ErrorMessage);
+	if (!ErrorMessage.IsEmpty())
+	{
+		InMessageCallback(ELogVerbosity::Warning, FString::Printf(
+			TEXT("%s: header is different, but cannot display differences. Package data for %s is invalid: %s"),
+			*AssetFilename, WhichComparisonPackage, *ErrorMessage));
+		PackageHeader = FZenPackageHeader();
+		bValid = false;
+	}
+}
+
+bool FDiffWriterZenHeader::IsValid() const
+{
+	return bValid;
+}
+
+const FZenPackageHeader& FDiffWriterZenHeader::GetPackageHeader() const
+{
+	return PackageHeader;
+}
+
+FStringView FDiffWriterZenHeader::GetObjectIndexPathName(FPackageObjectIndex PackageObjectIndex)
+{
+	FString* PathName = &ObjectIndexPathNames.FindOrAdd(PackageObjectIndex);
+	if (!PathName->IsEmpty())
+	{
+		return *PathName;
+	}
+
+	if (PackageObjectIndex.IsExport())
+	{
+		uint32 ExportIndex = PackageObjectIndex.ToExport();
+		if (!PackageHeader.ExportMap.IsValidIndex(ExportIndex))
+		{
+			*PathName = TEXT("InvalidExport");
+			return *PathName;
+		}
+		const FExportMapEntry& Export = PackageHeader.ExportMap[ExportIndex];
+		if (Export.OuterIndex.IsNull())
+		{
+			PackageHeader.NameMap.GetName(Export.ObjectName).ToString(*PathName);
+		}
+		else
+		{
+			*PathName = TEXT("<Cycle>"); // Set to non-empty to prevent cycles
+			FStringView ParentText = GetObjectIndexPathName(Export.OuterIndex);
+			// PathName is no longer usable because we potentially modified ObjectIndexPathNames when we called GetObjectIndexPathName
+			PathName = ObjectIndexPathNames.Find(PackageObjectIndex);
+			check(PathName);
+			*PathName = ParentText;
+			PathName->AppendChar('/');
+			PackageHeader.NameMap.GetName(Export.ObjectName).AppendString(*PathName);
+		}
+	}
+	else if (PackageObjectIndex.IsScriptImport())
+	{
+		FPackageStoreOptimizer::FScriptObjectData* ObjectData = Globals.ScriptObjectsMap.Find(PackageObjectIndex);
+		if (ObjectData)
+		{
+			*PathName = ObjectData->FullName;
+		}
+		else
+		{
+			*PathName = TEXT("<UnknownScriptImport>");
+		}
+	}
+	else if (PackageObjectIndex.IsPackageImport())
+	{
+		FPackageImportReference Ref = PackageObjectIndex.ToPackageImportRef();
+		if (PackageHeader.ImportedPackageNames.IsValidIndex(Ref.GetImportedPackageIndex()))
+		{
+			FName PackageName = PackageHeader.ImportedPackageNames[Ref.GetImportedPackageIndex()];
+			FZenPackageExportsForDiff& PackageExports = ImportedPackageExports.FindOrAdd(PackageName);
+			PackageExports.Initialize(Globals, PackageName, bUseZenStoreForImports);
+			uint64 ExportHash = PackageHeader.ImportedPublicExportHashes[Ref.GetImportedPublicExportHashIndex()];
+			FStringView ExportName = PackageExports.GetExportPackageRelativePath(ExportHash);
+			if (ExportName.StartsWith('/'))
+			{
+				ExportName.RightChopInline(1);
+			}
+			*PathName = FString::Printf(TEXT("%s.%.*s"),
+				*WriteToString<256>(PackageName), ExportName.Len(), ExportName.GetData());
+		}
+		else
+		{
+			*PathName = TEXT("<UnknownPackageImport>");
+		}
+	}
+	else
+	{
+		check(PackageObjectIndex.IsNull());
+		*PathName = TEXT("NULL");
+	}
+	return *PathName;
+}
+
+FString FDiffWriterZenHeader::GetTableKey(FName Id)
+{
+	return Id.ToString();
+}
+
+bool FDiffWriterZenHeader::CompareTableItem(FDiffWriterZenHeader& DestContext,
+	FName SourceName, FName DestName)
+{
+	return SourceName == DestName;
+}
+
+FString FDiffWriterZenHeader::ConvertItemToText(FName Id)
+{
+	return Id.ToString();
+}
+
+bool FDiffWriterZenHeader::IsImportMapIdentical(FDiffWriterZenHeader& DestContext)
+{
+	int32 NumImports = PackageHeader.ImportMap.Num();
+	if (NumImports != DestContext.PackageHeader.ImportMap.Num())
+	{
+		return false;
+	}
+
+	for (int32 Index = 0; Index < NumImports; ++Index)
+	{
+		if (!CompareTableItem(DestContext, PackageHeader.ImportMap[Index],
+			DestContext.PackageHeader.ImportMap[Index]))
+		{
+			return false;
+		}
+	}
+	return true;
+}
+
+FString FDiffWriterZenHeader::GetTableKey(FPackageObjectIndex Id)
+{
+	// TODO: When Id is an import, we could calculate the key more efficiently by changing it to the PackageId + ExportHash instead
+	// of looking up the name of the ExportHash
+	return FString(GetObjectIndexPathName(Id));
+}
+
+bool FDiffWriterZenHeader::CompareTableItem(FDiffWriterZenHeader& DestContext, const FPackageObjectIndex& SourceIndex,
+	const FPackageObjectIndex& DestIndex)
+{
+	if (SourceIndex.IsImport())
+	{
+		if (!DestIndex.IsImport())
+		{
+			return false;
+		}
+		// TODO: we could compare more efficiently by comparing PackageId + ExportHash instead
+		// of looking up the name of the ExportHash
+		return GetObjectIndexPathName(SourceIndex) == DestContext.GetObjectIndexPathName(DestIndex);
+	}
+	else if (SourceIndex.IsExport())
+	{
+		if (!DestIndex.IsExport())
+		{
+			return false;
+		}
+		return GetObjectIndexPathName(SourceIndex) == DestContext.GetObjectIndexPathName(DestIndex);
+	}
+	else
+	{
+		check(SourceIndex.IsNull());
+		return DestIndex.IsNull();
+	}
+}
+
+FString FDiffWriterZenHeader::ConvertItemToText(FPackageObjectIndex Id)
+{
+	return FString(GetObjectIndexPathName(Id));
+}
+
+bool FDiffWriterZenHeader::IsExportMapIdentical(FDiffWriterZenHeader& DestContext)
+{
+	int32 NumExports = PackageHeader.ExportMap.Num();
+	if (NumExports != DestContext.PackageHeader.ExportMap.Num())
+	{
+		return false;
+	}
+
+	for (int32 ExportIndex = 0; ExportIndex < NumExports; ++ExportIndex)
+	{
+		FZenHeaderIndexIntoExportMap ExportIndexStruct{ ExportIndex };
+		if (!CompareTableItem(DestContext, ExportIndexStruct, ExportIndexStruct))
+		{
+			return false;
+		}
+	}
+	return true;
+}
+
+FString FDiffWriterZenHeader::GetTableKey(const FZenHeaderIndexIntoExportMap& Index)
+{
+	return FString(GetObjectIndexPathName(FPackageObjectIndex::FromExportIndex(Index.Index)));
+}
+
+bool FDiffWriterZenHeader::CompareTableItem(FDiffWriterZenHeader& DestContext,
+	const FZenHeaderIndexIntoExportMap& SourceExportIndex, const FZenHeaderIndexIntoExportMap& DestExportIndex)
+{
+	if (!PackageHeader.ExportMap.IsValidIndex(SourceExportIndex.Index) ||
+		!DestContext.PackageHeader.ExportMap.IsValidIndex(DestExportIndex.Index))
+	{
+		return false;
+	}
+	const FExportMapEntry& SourceExport = PackageHeader.ExportMap[SourceExportIndex.Index];
+	const FExportMapEntry& DestExport = DestContext.PackageHeader.ExportMap[DestExportIndex.Index];
+
+	// Ignore CookedSerialOffset; it is irrelevant to the comparison of the export
+	if (SourceExport.CookedSerialSize != DestExport.CookedSerialSize) return false;
+	if (PackageHeader.NameMap.GetName(SourceExport.ObjectName) !=
+		DestContext.PackageHeader.NameMap.GetName(DestExport.ObjectName)) return false;
+	if (!CompareTableItem(DestContext, SourceExport.OuterIndex, DestExport.OuterIndex)) return false;
+	if (!CompareTableItem(DestContext, SourceExport.ClassIndex, DestExport.ClassIndex)) return false;
+	if (!CompareTableItem(DestContext, SourceExport.SuperIndex, DestExport.SuperIndex)) return false;
+	if (!CompareTableItem(DestContext, SourceExport.TemplateIndex, DestExport.TemplateIndex)) return false;
+	if (SourceExport.PublicExportHash != DestExport.PublicExportHash) return false;
+	if (SourceExport.ObjectFlags != DestExport.ObjectFlags) return false;
+	if (SourceExport.FilterFlags != DestExport.FilterFlags) return false;
+
+	return true;
+}
+
+FString FDiffWriterZenHeader::ConvertItemToText(const FZenHeaderIndexIntoExportMap& Index)
+{
+	if (!PackageHeader.ExportMap.IsValidIndex(Index.Index))
+	{
+		return TEXT("<InvalidExport>");
+	}
+	const FExportMapEntry& Export = PackageHeader.ExportMap[Index.Index];
+	FString ClassName = FString(GetObjectIndexPathName(Export.ClassIndex));
+	FString PathName = FString(GetObjectIndexPathName(FPackageObjectIndex::FromExportIndex(Index.Index)));
+	FString SuperName = FString(GetObjectIndexPathName(Export.SuperIndex));
+	FString TemplateName = FString(GetObjectIndexPathName(Export.TemplateIndex));
+
+	return FString::Printf(TEXT("%s %s Super: %s, Template: %s, Flags: %d, Size: %" INT64_FMT ", FilterFlags: %d"),
+		*ClassName,
+		*PathName,
+		*SuperName,
+		*TemplateName,
+		(int32)Export.ObjectFlags,
+		Export.CookedSerialSize,
+		(int32)Export.FilterFlags);
+}
+
+void FDiffWriterZenHeader::LogMessage(ELogVerbosity::Type Verbosity, FStringView Message)
+{
+	MessageCallback(Verbosity, Message);
+}
+
+void FZenPackageExportsForDiff::Initialize(FAccumulatorGlobals& Globals, FName InPackageName, bool bUseZenStore)
+{
+	if (!PackageName.IsNone())
+	{
+		return;
+	}
+	PackageName = InPackageName;
+	ExportPaths.Reset();
+	if (bUseZenStore)
+	{
+		InitializeFromZenStore(Globals);
+	}
+	else
+	{
+		InitializeFromMemory(Globals);
+	}
+}
+
+void FZenPackageExportsForDiff::InitializeFromZenStore(FAccumulatorGlobals& Globals)
+{
+	// Silently ignore on errors; GetExportName will report <UnknownExport.%u>
+	if (!Globals.PackageWriter)
+	{
+		return;
+	}
+
+	FPackageId PackageId = FPackageId::FromName(PackageName);
+	IPackageWriter::FPackageInfo PackageInfo;
+	uint16 MultiOutputIndex = 0; // We're reading exports from the non-optional version of the package
+	PackageInfo.ChunkId = CreateIoChunkId(PackageId.Value(), MultiOutputIndex, EIoChunkType::ExportBundleData);
+	ICookedPackageWriter::FPreviousCookedBytesData PackageBytes;
+	if (!Globals.PackageWriter->GetPreviousCookedBytes(PackageInfo, PackageBytes))
+	{
+		return;
+	}
+	FString Error;
+	FZenPackageHeader PackageHeader = FZenPackageHeader::MakeView(
+		FMemoryView(PackageBytes.Data.Get(), PackageBytes.Size), Error);
+	if (!Error.IsEmpty())
+	{
+		return;
+	}
+
+	TArray<FString> LocalPackageRelativePaths;
+	FString EmptyPath;
+	FString UnknownImportPath(TEXT("<UnknownImportInOtherPackage>"));
+	FString InvalidExportIndex(TEXT("<InvalidExportIndex>"));
+	FString InvalidCycleText(TEXT("<Cycle>"));
+	TFunction<const FString& (FPackageObjectIndex)> FindOrConstructPackageRelativePath;
+	FindOrConstructPackageRelativePath =
+		[&FindOrConstructPackageRelativePath /** Recursive calls */, &LocalPackageRelativePaths, &PackageHeader, &EmptyPath,
+		&UnknownImportPath, &InvalidExportIndex, &InvalidCycleText]
+	(FPackageObjectIndex Index) -> const FString&
+	{
+		if (Index.IsNull())
+		{
+			return EmptyPath;
+		}
+		if (Index.IsExport())
+		{
+			uint32 ExportIndex = Index.ToExport();
+			if (!LocalPackageRelativePaths.IsValidIndex(ExportIndex))
+			{
+				return InvalidExportIndex;
+			}
+			FString& PackageRelativePath = LocalPackageRelativePaths[ExportIndex];
+			if (!PackageRelativePath.IsEmpty())
+			{
+				return PackageRelativePath;
+			}
+			// Avoid cycles by setting ObjectName before the recursive call to FindOrConstructPackageRelativePath
+			PackageRelativePath = InvalidCycleText;
+
+			check(PackageHeader.ExportMap.IsValidIndex(ExportIndex)); // LocalExportNames is the same length and was checked above
+			const FExportMapEntry& Export = PackageHeader.ExportMap[ExportIndex];
+			FName LeafName;
+			if (!PackageHeader.NameMap.TryGetName(Export.ObjectName, LeafName) || LeafName.IsNone())
+			{
+				PackageRelativePath = InvalidExportIndex;
+				return PackageRelativePath;
+			}
+			const FString& ParentPackageRelativePath = FindOrConstructPackageRelativePath(Export.OuterIndex);
+
+			// Every path name is ParentPath (even if empty) + / + LeafName. PackageRelativePaths when parent is empty start with / 
+			PackageRelativePath = ParentPackageRelativePath;
+			PackageRelativePath.Append(TEXT("/"));
+			LeafName.AppendString(PackageRelativePath);
+
+			// Down-case the path names to match the behavior for the path names constructed from UObjects in memory
+			// by FPackageStoreOptimizer::AppendPathForPublicExportHash
+			PackageRelativePath.ToLowerInline();
+
+			return PackageRelativePath;
+		}
+		else
+		{
+			return UnknownImportPath;
+		}
+	};
+
+	uint32 NumExports = static_cast<uint32>(PackageHeader.ExportMap.Num());
+	LocalPackageRelativePaths.SetNum(NumExports);
+	for (uint32 ExportIndex = 0; ExportIndex < NumExports; ++ExportIndex)
+	{
+		FindOrConstructPackageRelativePath(FPackageObjectIndex::FromExportIndex(ExportIndex));
+	}
+
+	ExportPaths.Reserve(NumExports);
+	for (uint32 ExportIndex = 0; ExportIndex < NumExports; ++ExportIndex)
+	{
+		const FExportMapEntry& Export = PackageHeader.ExportMap[ExportIndex];
+		ExportPaths.Add(Export.PublicExportHash, LocalPackageRelativePaths[ExportIndex]);
+	}
+}
+
+void FZenPackageExportsForDiff::InitializeFromMemory(FAccumulatorGlobals& Globals)
+{
+	// Silently ignore on errors; GetExportName will report <UnknownExport.%u>
+	UPackage* Package = FindPackage(nullptr, *WriteToString<256>(PackageName));
+	if (!Package)
+	{
+		return;
+	}
+
+	// We do not add the UPackage object itself, because UPackage object is not stored in the list
+	// of exports in the disk version of the package
+	ForEachObjectWithOuter(Package, [this](UObject* Object)
+		{
+			TStringBuilder<256> PackageRelativePath;
+			FPackageStoreOptimizer::AppendPathForPublicExportHash(Object, PackageRelativePath);
+			uint64 PublicExportHash;
+			if (FPackageStoreOptimizer::TryGetPublicExportHash(PackageRelativePath, PublicExportHash))
+			{
+				ExportPaths.Add(PublicExportHash, FString(PackageRelativePath));
+			}
+		}, false /* bIncludeNestedObjects */);
+}
+
+FStringView FZenPackageExportsForDiff::GetExportPackageRelativePath(uint64 PublicExportHash)
+{
+	FString& ExportPath = ExportPaths.FindOrAdd(PublicExportHash);
+	if (ExportPath.IsEmpty())
+	{
+		ExportPath = FString::Printf(TEXT("<UnknownExport.%u>"), PublicExportHash);
+	}
+	return ExportPath;
+}
+
+}
\ No newline at end of file
diff --git a/Engine/Source/Editor/UnrealEd/Private/Cooker/DiffWriterZenHeader.h b/Engine/Source/Editor/UnrealEd/Private/Cooker/DiffWriterZenHeader.h
new file mode 100644
--- /dev/null
+++ b/Engine/Source/Editor/UnrealEd/Private/Cooker/DiffWriterZenHeader.h
@@ -0,0 +1,84 @@
+// Copyright Epic Games, Inc. All Rights Reserved.
+
+#pragma once
+
+#include "DiffWriterArchive.h"
+#include "Serialization/ZenPackageHeader.h"
+
+namespace UE::DiffWriter
+{
+
+/**
+ * Helper object for FDiffWriterZenHeader: Stores the map from PublicExportHash to package-relative
+ * objectpath for all of the objects in a package, used to convert a ZenPackageHeader's import's
+ * FPackageObjectIndex into the full path to the imported object
+ */
+class FZenPackageExportsForDiff
+{
+public:
+	void Initialize(FAccumulatorGlobals& Globals, FName InPackageName, bool bUseZenStore);
+	FStringView GetExportPackageRelativePath(uint64 PublicExportHash);
+
+private:
+	void InitializeFromZenStore(FAccumulatorGlobals& Globals);
+	void InitializeFromMemory(FAccumulatorGlobals& Globals);
+
+private:
+	FName PackageName = NAME_None;
+	TMap<uint64, FString> ExportPaths;
+};
+
+/** Adapter type for FDiffWriterZenHeader to represent its exports in DumpTableDifferences. */
+struct FZenHeaderIndexIntoExportMap
+{
+	/** Index of the given export in FZenPackageHeader::ExportMap. */
+	int32 Index;
+};
+
+/**
+ * Context object for DumpTableDifferences when we are using EPackageHeaderFormat::ZenPackageSummary and so the
+ * header information comes from ZenStore's optimized package format parsed into FZenPackageHeader.
+*/
+class FDiffWriterZenHeader
+{
+public:
+	// Constructor and simple data readers
+	FDiffWriterZenHeader(FAccumulatorGlobals& InGlobals, const FMessageCallback& InMessageCallback,
+		bool bInUseZenStoreForImports, const FPackageData& Package, const FString& AssetFilename,
+		const TCHAR* WhichComparisonPackage);
+	bool IsValid() const;
+	const FZenPackageHeader& GetPackageHeader() const;
+
+	// Optional data loaded on demand
+	FStringView GetObjectIndexPathName(FPackageObjectIndex PackageObjectIndex);
+
+	// API for DumpTableDifferences
+	FString GetTableKey(FName Id);
+	bool CompareTableItem(FDiffWriterZenHeader& DestContext, FName SourceName, FName DestName);
+	FString ConvertItemToText(FName Id);
+
+	bool IsImportMapIdentical(FDiffWriterZenHeader& DestContext);
+	FString GetTableKey(FPackageObjectIndex Id);
+	bool CompareTableItem(FDiffWriterZenHeader& DestContext, const FPackageObjectIndex& SourceIndex,
+		const FPackageObjectIndex& DestIndex);
+	FString ConvertItemToText(FPackageObjectIndex Id);
+
+	bool IsExportMapIdentical(FDiffWriterZenHeader& DestContext);
+	FString GetTableKey(const FZenHeaderIndexIntoExportMap& Id);
+	bool CompareTableItem(FDiffWriterZenHeader& DestContext, const FZenHeaderIndexIntoExportMap& SourceExport,
+		const FZenHeaderIndexIntoExportMap& DestExport);
+	FString ConvertItemToText(const FZenHeaderIndexIntoExportMap& Id);
+
+	void LogMessage(ELogVerbosity::Type Verbosity, FStringView Message);
+
+private:
+	FZenPackageHeader PackageHeader;
+	TMap<FName, FZenPackageExportsForDiff> ImportedPackageExports;
+	TMap<FPackageObjectIndex, FString> ObjectIndexPathNames;
+	FAccumulatorGlobals& Globals;
+	const FMessageCallback& MessageCallback;
+	bool bUseZenStoreForImports = false;
+	bool bValid = true;
+};
+
+}
\ No newline at end of file
diff --git a/Engine/Source/Editor/UnrealEd/Private/Tests/DiffWriterArchiveTests.cpp b/Engine/Source/Editor/UnrealEd/Private/Tests/DiffWriterArchiveTests.cpp
--- a/Engine/Source/Editor/UnrealEd/Private/Tests/DiffWriterArchiveTests.cpp
+++ b/Engine/Source/Editor/UnrealEd/Private/Tests/DiffWriterArchiveTests.cpp
@@ -51,8 +51,10 @@ bool FDiffWriterArchiveTestsCallstacks::RunTest(const FString& Parameters)
 	bool bIgnoreHeaderDiffs = false;
 	FName PackageName(TEXT("PackageName"));
 	EPackageHeaderFormat HeaderFormat(EPackageHeaderFormat::PackageFileSummary);
-	TRefCountPtr<FAccumulator> Accumulator = new FAccumulator(nullptr, PackageName, MaxDiffsToLog, bIgnoreHeaderDiffs,
-		[](ELogVerbosity::Type Verbosity, FStringView Message) {}, HeaderFormat);
+	UE::DiffWriter::FAccumulatorGlobals AccumulatorGlobals;
+
+	TRefCountPtr<FAccumulator> Accumulator = new FAccumulator(AccumulatorGlobals, nullptr, PackageName, MaxDiffsToLog,
+		bIgnoreHeaderDiffs, [](ELogVerbosity::Type Verbosity, FStringView Message) {}, HeaderFormat);
 	FBasicDiffState State = FBasicDiffState{ 1,2,3 };
 	FDiffArchiveForLinker Ar(*Accumulator);
 
@@ -86,11 +88,13 @@ bool FDiffWriterArchiveTestsBasic::RunTest(const FString& Parameters)
 	FName PackageName(TEXT("PackageName"));
 	FString Filename(TEXT("Filename"));
 	EPackageHeaderFormat HeaderFormat(EPackageHeaderFormat::PackageFileSummary);
+	UE::DiffWriter::FAccumulatorGlobals AccumulatorGlobals;
 
 	// SECTION("DiffMap - identical")
 	{
-		TRefCountPtr<FAccumulator> Accumulator = new FAccumulator(nullptr, PackageName, MaxDiffsToLog,
-			bIgnoreHeaderDiffs, [](ELogVerbosity::Type Verbosity, FStringView Message) {}, HeaderFormat);
+		TRefCountPtr<FAccumulator> Accumulator = new FAccumulator(AccumulatorGlobals, nullptr, PackageName,
+			MaxDiffsToLog, bIgnoreHeaderDiffs,
+			[](ELogVerbosity::Type Verbosity, FStringView Message) {}, HeaderFormat);
 
 		FLargeMemoryWriter InitialState;
 		{
@@ -112,8 +116,9 @@ bool FDiffWriterArchiveTestsBasic::RunTest(const FString& Parameters)
 
 	// SECTION("DiffMap - mismatch")
 	{
-		TRefCountPtr<FAccumulator> Accumulator = new FAccumulator(nullptr, PackageName, MaxDiffsToLog,
-			bIgnoreHeaderDiffs, [](ELogVerbosity::Type Verbosity, FStringView Message) {}, HeaderFormat);
+		TRefCountPtr<FAccumulator> Accumulator = new FAccumulator(AccumulatorGlobals, nullptr, PackageName,
+			MaxDiffsToLog, bIgnoreHeaderDiffs,
+			[](ELogVerbosity::Type Verbosity, FStringView Message) {}, HeaderFormat);
 
 		FLargeMemoryWriter InitialState;
 		{
@@ -135,8 +140,9 @@ bool FDiffWriterArchiveTestsBasic::RunTest(const FString& Parameters)
 
 	// SECTION("Compare - mismatch")
 	{
-		TRefCountPtr<FAccumulator> Accumulator = new FAccumulator(nullptr, PackageName, MaxDiffsToLog,
-			bIgnoreHeaderDiffs, [](ELogVerbosity::Type Verbosity, FStringView Message) {}, HeaderFormat);
+		TRefCountPtr<FAccumulator> Accumulator = new FAccumulator(AccumulatorGlobals, nullptr, PackageName,
+			MaxDiffsToLog, bIgnoreHeaderDiffs,
+			[](ELogVerbosity::Type Verbosity, FStringView Message) {}, HeaderFormat);
 
 		FBasicDiffState InitialState = FBasicDiffState {1,2,3};
 		FLargeMemoryWriter InitialMemory;
