Cooker: CookPackageSplitter: Add assertion guarantees that ICookPackageSplitter is not accessed after Teardown is called.
#rnx

diff --git a/Engine/Source/Editor/UnrealEd/Private/Cooker/CookPackageData.cpp b/Engine/Source/Editor/UnrealEd/Private/Cooker/CookPackageData.cpp
--- a/Engine/Source/Editor/UnrealEd/Private/Cooker/CookPackageData.cpp
+++ b/Engine/Source/Editor/UnrealEd/Private/Cooker/CookPackageData.cpp
@@ -1544,9 +1544,11 @@ void FGeneratorPackage::InitializeSave(const UObject* InSplitDataObject,
 			delete InCookPackageSplitterInstance;
 		}
 
+		bInitialized = true;
 		FName InSplitDataObjectName = *InSplitDataObject->GetFullName();
 		check(SplitDataObjectName.IsNone() || SplitDataObjectName == InSplitDataObjectName);
 		SplitDataObjectName = InSplitDataObjectName;
+		bUseInternalReferenceToAvoidGarbageCollect = CookPackageSplitterInstance->UseInternalReferenceToAvoidGarbageCollect();
 		SetOwnerPackage(GetOwner().GetPackage());
 	}
 }
@@ -1568,6 +1570,7 @@ void FGeneratorPackage::ConditionalNotifyCompletion(ICookPackageSplitter::ETeard
 	{
 		bNotifiedCompletion = true;
 		CookPackageSplitterInstance->Teardown(Status);
+		CookPackageSplitterInstance.Reset();
 	}
 }
 
@@ -1598,7 +1601,7 @@ bool FGeneratorPackage::TryGenerateList(UObject* OwnerObject, FPackageDatas& Pac
 	{
 		UCookOnTheFlyServer::FScopedActivePackage ScopedActivePackage(COTFS, OwnerPackageData.GetPackageName(),
 			PackageAccessTrackingOps::NAME_CookerBuildObject);
-		GeneratorDatas = CookPackageSplitterInstance->GetGenerateList(LocalOwnerPackage, OwnerObject);
+		GeneratorDatas = GetCookPackageSplitterInstance()->GetGenerateList(LocalOwnerPackage, OwnerObject);
 	}
 	PackagesToGenerate.Reset(GeneratorDatas.Num());
 	TArray<const ITargetPlatform*, TInlineAllocator<1>> PlatformsToCook;
@@ -1762,6 +1765,12 @@ const FCookGenerationInfo* FGeneratorPackage::FindInfo(const FPackageData& Packa
 	return const_cast<FGeneratorPackage*>(this)->FindInfo(PackageData);
 }
 
+ICookPackageSplitter* FGeneratorPackage::GetCookPackageSplitterInstance() const
+{
+	checkf(!bNotifiedCompletion, TEXT("It is illegal for the cooker to try to access the CookPackageSplitterInstance after calling Teardown on it."));
+	return CookPackageSplitterInstance.Get();
+}
+
 UObject* FGeneratorPackage::FindSplitDataObject() const
 {
 	check(IsInitialized());
@@ -1789,7 +1798,7 @@ void FGeneratorPackage::PreGarbageCollect(FCookGenerationInfo& Info, TArray<TObj
 	check(Info.PackageData); // Caller validates this is non-null
 	if (Info.GetSaveState() > FCookGenerationInfo::ESaveState::CallPopulate)
 	{
-		if (GetCookPackageSplitterInstance()->UseInternalReferenceToAvoidGarbageCollect())
+		if (IsUseInternalReferenceToAvoidGarbageCollect())
 		{
 			UPackage* Package = Info.PackageData->GetPackage();
 			if (Package)
@@ -1805,7 +1814,7 @@ void FGeneratorPackage::PreGarbageCollect(FCookGenerationInfo& Info, TArray<TObj
 	}
 	if (Info.HasTakenOverCachedCookedPlatformData())
 	{
-		if (GetCookPackageSplitterInstance()->UseInternalReferenceToAvoidGarbageCollect())
+		if (IsUseInternalReferenceToAvoidGarbageCollect())
 		{
 			// For the UseInternalReferenceToAvoidGarbageCollect case, part of the CookPackageSplitter contract is that
 			// the Cooker will keep referenced the package and all objects returned from GetObjectsToMove* functions
@@ -1858,7 +1867,7 @@ void FGeneratorPackage::PostGarbageCollect()
 		{
 			if (RemainingToPopulate > 0 &&
 				!Owner.IsKeepReferencedDuringGC() &&
-				!CookPackageSplitterInstance->UseInternalReferenceToAvoidGarbageCollect())
+				!IsUseInternalReferenceToAvoidGarbageCollect())
 			{
 				UE_LOG(LogCook, Error, TEXT("PackageSplitter found the Generator package still in memory after it should have been deleted by GC.")
 					TEXT("\n\tThis is unexpected since garbage has been collected and the package should have been unreferenced so it should have been collected, and will break population of Generated packages.")
@@ -2005,7 +2014,7 @@ void FGeneratorPackage::ResetSaveState(FCookGenerationInfo& Info, UPackage* Pack
 	if (Info.HasTakenOverCachedCookedPlatformData())
 	{
 		if (Info.PackageData && Info.PackageData->GetCachedObjectsInOuter().Num() != 0 &&
-			GetCookPackageSplitterInstance()->UseInternalReferenceToAvoidGarbageCollect() &&
+			IsUseInternalReferenceToAvoidGarbageCollect() &&
 			(ReleaseSaveReason == EReleaseSaveReason::Demoted || ReleaseSaveReason == EReleaseSaveReason::RecreateObjectCache))
 		{
 			UE_LOG(LogCook, Error, TEXT("CookPackageSplitter failure: We are demoting a %s package from save and removing our references that keep its objects loaded.\n")
@@ -2044,9 +2053,7 @@ void FGeneratorPackage::UpdateSaveAfterGarbageCollect(const FPackageData& Packag
 		}
 	}
 
-	if (bInOutDemote && 
-		GetCookPackageSplitterInstance()->UseInternalReferenceToAvoidGarbageCollect() &&
-		Info->HasTakenOverCachedCookedPlatformData())
+	if (bInOutDemote && IsUseInternalReferenceToAvoidGarbageCollect() && Info->HasTakenOverCachedCookedPlatformData())
 	{
 		// No public objects should have been deleted; we are supposed to keep them referenced by keeping the package
 		// referenced in UCookOnTheFlyServer::PreGarbageCollect, and the package keeping its public objects referenced
diff --git a/Engine/Source/Editor/UnrealEd/Private/Cooker/CookPackageData.h b/Engine/Source/Editor/UnrealEd/Private/Cooker/CookPackageData.h
--- a/Engine/Source/Editor/UnrealEd/Private/Cooker/CookPackageData.h
+++ b/Engine/Source/Editor/UnrealEd/Private/Cooker/CookPackageData.h
@@ -915,7 +915,7 @@ public:
 		ICookPackageSplitter* InCookPackageSplitterInstance);
 	~FGeneratorPackage();
 	void InitializeSave(const UObject* InSplitDataObject, ICookPackageSplitter* InCookPackageSplitterInstance);
-	bool IsInitialized() const { return CookPackageSplitterInstance.IsValid(); }
+	bool IsInitialized() const { return bInitialized; }
 	/** Clear references to owned generated packages, and mark those packages as orphaned */
 	void ClearGeneratedPackages();
 
@@ -935,9 +935,12 @@ public:
 	const UE::Cook::FCookGenerationInfo* FindInfo(const FPackageData& PackageData) const;
 
 	/** Return CookPackageSplitter. */
-	ICookPackageSplitter* GetCookPackageSplitterInstance() const { return CookPackageSplitterInstance.Get(); }
+	ICookPackageSplitter* GetCookPackageSplitterInstance() const;
 	/** Return the SplitDataObject's FullObjectPath. */
 	const FName GetSplitDataObjectName() const { check(IsInitialized()); return SplitDataObjectName; }
+	/** Return the Splitter's value for virtual bool UseInternalReferenceToAvoidGarbageCollect() */
+	bool IsUseInternalReferenceToAvoidGarbageCollect() const { return bUseInternalReferenceToAvoidGarbageCollect; }
+
 	/**
 	 * Find again the split object from its name, or return null if no longer in memory.
 	 * It may have been GC'd and reloaded since the last time we used it.
@@ -988,7 +991,9 @@ private:
 	int32 NextPopulateIndex = 0;
 	int32 RemainingToPopulate = 0;
 
+	bool bInitialized = false;
 	bool bNotifiedCompletion = false;
+	bool bUseInternalReferenceToAvoidGarbageCollect = false;
 };
 
 
