Cooker: Suppress the warning for cooker blocked by busy materials unless the shadercompiler is idle for two reports in a row.
#rnx

diff --git a/Engine/Source/Editor/UnrealEd/Classes/CookOnTheSide/CookOnTheFlyServer.h b/Engine/Source/Editor/UnrealEd/Classes/CookOnTheSide/CookOnTheFlyServer.h
--- a/Engine/Source/Editor/UnrealEd/Classes/CookOnTheSide/CookOnTheFlyServer.h
+++ b/Engine/Source/Editor/UnrealEd/Classes/CookOnTheSide/CookOnTheFlyServer.h
@@ -1447,6 +1447,8 @@ private:
 	bool bLoadBusy = false;
 	/** True when PumpSaves has detected it is blocked on async work and CookOnTheFlyServer should do work elsewhere. */
 	bool bSaveBusy = false;
+	/** We need to track whether the compiler has been inactive for a long time before issuing a warning about it. */
+	bool bShaderCompilerWasActiveeOnPreviousBusyReport = true;
 	/**
 	 * If preloading is enabled, we call TryPreload until it returns true before sending the package to LoadReady,
 	 * otherwise we skip TryPreload and it goes immediately.
diff --git a/Engine/Source/Editor/UnrealEd/Private/CookOnTheFlyServer.cpp b/Engine/Source/Editor/UnrealEd/Private/CookOnTheFlyServer.cpp
--- a/Engine/Source/Editor/UnrealEd/Private/CookOnTheFlyServer.cpp
+++ b/Engine/Source/Editor/UnrealEd/Private/CookOnTheFlyServer.cpp
@@ -1548,6 +1548,9 @@ void UCookOnTheFlyServer::SetSaveBusy(bool bInBusy)
 			SaveBusyStartTimeSeconds = MAX_flt;
 			SaveBusyRetryTimeSeconds = MAX_flt;
 			SaveBusyWarnTimeSeconds = MAX_flt;
+			// Whenever we set Save back to non-busy, reset the counter for how many busy reports with an
+			// idle shadercompiler we need before we issue a warning
+			bShaderCompilerWasActiveeOnPreviousBusyReport = true;
 		}
 	}
 	else if (bSaveBusy)
@@ -1556,6 +1559,16 @@ void UCookOnTheFlyServer::SetSaveBusy(bool bInBusy)
 		SaveBusyRetryTimeSeconds = CurrentTime + CookProgressRetryBusyPeriodSeconds;
 		if (CurrentTime >= SaveBusyWarnTimeSeconds)
 		{
+			// Compiler users - classes using the shader compiler - can take multiple minutes to be compiled due to long
+			// queues and long compile times, so we do not issue a warning when they are the only objects holding us up,
+			// so long as the shadercompiler reports it is working on them.
+			bool bShaderCompilerIsActive = GShaderCompilingManager->IsCompiling();
+			bool bBusyCompilationUsersAreExpected = bShaderCompilerIsActive ||
+				// Even if the ShaderCompilerManager is not currently compiling, it might shortly begin or have recently finished.
+				// Issue a warning for blocked compiler users only if there are two reports in a row where the compiler is not active
+				bShaderCompilerWasActiveeOnPreviousBusyReport;
+			bShaderCompilerWasActiveeOnPreviousBusyReport = bShaderCompilerIsActive;
+
 			// Issue a status update. For each UObject we're still waiting on, check whether the long duration is expected using type-specific checks
 			// Make the status update a warning if the long duration is not reported as expected.
 			TArray<UObject*> NonExpectedObjects;
@@ -1566,7 +1579,7 @@ void UCookOnTheFlyServer::SetSaveBusy(bool bInBusy)
 			TArray<UClass*> CompilationUsers({ UMaterialInterface::StaticClass(), FindObject<UClass>(nullptr, TEXT("/Script/Niagara.NiagaraScript")) });
 
 			PackageDatas->ForEachPendingCookedPlatformData(
-			[&CompilationUsers, &ExpectedObjects, &ExpectedPackages, &NonExpectedObjects, &NonExpectedPackages]
+			[&CompilationUsers, &ExpectedObjects, &ExpectedPackages, &NonExpectedObjects, &NonExpectedPackages, bBusyCompilationUsersAreExpected]
 			(const FPendingCookedPlatformData& Data)
 			{
 				UObject* Object = Data.Object.Get();
@@ -1583,7 +1596,7 @@ void UCookOnTheFlyServer::SetSaveBusy(bool bInBusy)
 						break;
 					}
 				}
-				if (bCompilationUser && GShaderCompilingManager->IsCompiling())
+				if (bCompilationUser && bBusyCompilationUsersAreExpected)
 				{
 					ExpectedObjects.Add(Object);
 					ExpectedPackages.Add(Object->GetPackage());
@@ -5310,6 +5323,10 @@ void UCookOnTheFlyServer::PostGarbageCollect()
 	});
 
 	CookedPackageCountSinceLastGC = 0;
+
+	// Whenever we collect garbage, reset the counter for how many busy reports with an
+	// idle shadercompiler we need before we issue a warning
+	bShaderCompilerWasActiveeOnPreviousBusyReport = true;
 }
 
 void UCookOnTheFlyServer::EvaluateGarbageCollectionResults(bool bWasDueToOOM, bool bWasPartialGC, uint32 ResultFlags,
