Cooker: Fix compile error when ENABLE_COOK_STATS is turned off.
#rnx

diff --git a/Engine/Source/Editor/UnrealEd/Private/Cooker/CookDirector.cpp b/Engine/Source/Editor/UnrealEd/Private/Cooker/CookDirector.cpp
--- a/Engine/Source/Editor/UnrealEd/Private/Cooker/CookDirector.cpp
+++ b/Engine/Source/Editor/UnrealEd/Private/Cooker/CookDirector.cpp
@@ -136,7 +136,9 @@ FCookDirector::FCookDirector(UCookOnTheFlyServer& InCOTFS, int32 CookProcessCoun
 
 	LastTickTimeSeconds = FPlatformTime::Seconds();
 
+#if ENABLE_COOK_STATS
 	FCookStatsManager::CookStatsCallbacks.AddRaw(this, &FCookDirector::LogCookStats);
+#endif
 }
 
 bool FCookDirector::IsMultiprocessAvailable() const
@@ -203,7 +205,9 @@ void FCookDirector::ParseConfig(int32 CookProcessCount, bool& bOutValid)
 FCookDirector::~FCookDirector()
 {
 	StopCommunicationThread();
+#if ENABLE_COOK_STATS
 	FCookStatsManager::CookStatsCallbacks.RemoveAll(this);
+#endif
 
 	TSet<FPackageData*> AbortedAssignments;
 	for (TPair<int32, TRefCountPtr<FCookWorkerServer>>& Pair : RemoteWorkers)
@@ -1405,6 +1409,7 @@ FCookDirector::FLaunchInfo FCookDirector::GetLaunchInfo(FWorkerId WorkerId, int3
 	return Info;
 }
 
+#if ENABLE_COOK_STATS
 void FCookDirector::LogCookStats(FCookStatsManager::AddStatFuncRef AddStat)
 {
 	auto IdleTimeToString = [](float IdleTime)
@@ -1421,6 +1426,7 @@ void FCookDirector::LogCookStats(FCookStatsManager::AddStatFuncRef AddStat)
 	}
 	AddStat(TEXT("CookDirector"), Stats);
 }
+#endif
 
 void FCookDirector::HandleRetractionMessage(FMPCollectorServerMessageContext& Context, bool bReadSuccessful,
 	FRetractionResultsMessage&& Message)
diff --git a/Engine/Source/Editor/UnrealEd/Private/Cooker/CookDirector.h b/Engine/Source/Editor/UnrealEd/Private/Cooker/CookDirector.h
--- a/Engine/Source/Editor/UnrealEd/Private/Cooker/CookDirector.h
+++ b/Engine/Source/Editor/UnrealEd/Private/Cooker/CookDirector.h
@@ -199,7 +199,9 @@ private:
 	 */
 	void SetWorkersStalled(bool bInWorkersStalled);
 	/** Callback for CookStats system to log our stats. */
+#if ENABLE_COOK_STATS
 	void LogCookStats(FCookStatsManager::AddStatFuncRef AddStat);
+#endif
 	void TickRetractionFromSchedulerThread(bool bAnyIdle, int32 BusiestNumAssignments);
 	void HandleRetractionMessage(FMPCollectorServerMessageContext& Context, bool bReadSuccessful,
 		FRetractionResultsMessage&& Message);
