Create IAssetTools::IsAssetVisible and hook up material function menu to it

diff --git a/Engine/Source/Developer/AssetTools/Private/AssetTools.cpp b/Engine/Source/Developer/AssetTools/Private/AssetTools.cpp
--- a/Engine/Source/Developer/AssetTools/Private/AssetTools.cpp
+++ b/Engine/Source/Developer/AssetTools/Private/AssetTools.cpp
@@ -5933,6 +5933,42 @@ TSharedRef<FPathPermissionList>& UAssetToolsImpl::GetWritableFolderPermissionLis
 	return WritableFolderPermissionList;
 }
 
+bool UAssetToolsImpl::IsAssetVisible(const FAssetData& AssetData, bool bCheckAliases) const
+{
+	// If the class is not visible, the asset is never visible regardless of the path
+	if (!GetAssetClassPathPermissionList(EAssetClassAction::ViewAsset)->PassesFilter(AssetData.AssetClassPath.ToString()))
+	{
+		return false;
+	}
+	
+	// If this asset's package is in a visible folder, then it is visible
+	if (FolderPermissionList->PassesStartsWithFilter(AssetData.PackagePath))
+	{
+		return true;
+	}
+
+	// Otherwise, check if any of the asset's aliases are in a visible folder
+	if (bCheckAliases)
+	{
+		const TArray<FString> Aliases = IContentBrowserSingleton::Get().GetAliasesForPath(AssetData.ToSoftObjectPath());
+		for (const FString& Alias : Aliases)
+		{
+			FStringView AliasPackagePath(Alias);
+			int32 LastSlashIndex = INDEX_NONE;
+			if (AliasPackagePath.FindLastChar(TEXT('/'), LastSlashIndex))
+			{
+				AliasPackagePath.LeftInline(LastSlashIndex);
+				if (FolderPermissionList->PassesStartsWithFilter(AliasPackagePath))
+				{
+					return true;
+				}
+			}
+		}
+	}
+
+	return false;
+}
+
 bool UAssetToolsImpl::AllPassWritableFolderFilter(const TArray<FString>& InPaths) const
 {
 	if (WritableFolderPermissionList->HasFiltering())
diff --git a/Engine/Source/Developer/AssetTools/Private/AssetTools.h b/Engine/Source/Developer/AssetTools/Private/AssetTools.h
--- a/Engine/Source/Developer/AssetTools/Private/AssetTools.h
+++ b/Engine/Source/Developer/AssetTools/Private/AssetTools.h
@@ -139,6 +139,7 @@ public:
 	virtual TSet<EBlueprintType>& GetAllowedBlueprintTypes() override;
 	virtual TSharedRef<FPathPermissionList>& GetFolderPermissionList() override;
 	virtual TSharedRef<FPathPermissionList>& GetWritableFolderPermissionList() override;
+	virtual bool IsAssetVisible(const FAssetData& AssetData, bool bCheckAliases = true) const override;
 	virtual bool AllPassWritableFolderFilter(const TArray<FString>& InPaths) const override;
 	virtual void NotifyBlockedByWritableFolderFilter() const;
 	virtual bool IsNameAllowed(const FString& Name, FText* OutErrorMessage = nullptr) const override;
diff --git a/Engine/Source/Developer/AssetTools/Public/IAssetTools.h b/Engine/Source/Developer/AssetTools/Public/IAssetTools.h
--- a/Engine/Source/Developer/AssetTools/Public/IAssetTools.h
+++ b/Engine/Source/Developer/AssetTools/Public/IAssetTools.h
@@ -598,6 +598,9 @@ public:
 	/** Get writable folder permission list for content browser and other systems */
 	virtual TSharedRef<FPathPermissionList>& GetWritableFolderPermissionList() = 0;
 
+	/** Determines whether an asset has a viewable asset class and its folder is allowed. If not, it can also check if it's aliased to an allowed location. */
+	virtual bool IsAssetVisible(const FAssetData& AssetData, bool bCheckAliases = true) const = 0;
+
 	/** Returns true if all in list pass writable folder filter */
 	virtual bool AllPassWritableFolderFilter(const TArray<FString>& InPaths) const = 0;
 
diff --git a/Engine/Source/Editor/ContentBrowser/Private/ContentBrowserSingleton.cpp b/Engine/Source/Editor/ContentBrowser/Private/ContentBrowserSingleton.cpp
--- a/Engine/Source/Editor/ContentBrowser/Private/ContentBrowserSingleton.cpp
+++ b/Engine/Source/Editor/ContentBrowser/Private/ContentBrowserSingleton.cpp
@@ -1317,4 +1317,17 @@ FContentBrowserItemPath FContentBrowserSingleton::GetInitialPathToSaveAsset(cons
 	return FContentBrowserItemPath(AssetPath, EContentBrowserPathType::Internal);
 }
 
+TArray<FString> FContentBrowserSingleton::GetAliasesForPath(const FSoftObjectPath& InPath) const
+{
+	TArray<FString> OutPaths;
+
+	const TArray<FContentBrowserItemPath> Items = IContentBrowserDataModule::Get().GetSubsystem()->GetAliasesForPath(InPath);
+	for (const FContentBrowserItemPath& Item : Items)
+	{
+		OutPaths.Add(Item.GetInternalPathString());
+	}
+
+	return OutPaths;
+}
+
 #undef LOCTEXT_NAMESPACE
diff --git a/Engine/Source/Editor/ContentBrowser/Private/ContentBrowserSingleton.h b/Engine/Source/Editor/ContentBrowser/Private/ContentBrowserSingleton.h
--- a/Engine/Source/Editor/ContentBrowser/Private/ContentBrowserSingleton.h
+++ b/Engine/Source/Editor/ContentBrowser/Private/ContentBrowserSingleton.h
@@ -119,6 +119,7 @@ PRAGMA_ENABLE_DEPRECATION_WARNINGS
 	virtual void UnregisterIsFolderShowPrivateContentToggleableDelegate() override;
 	virtual FDelegateHandle RegisterOnFavoritesChangedHandler(FSimpleDelegate OnFavoritesChanged) override;
 	virtual void UnregisterOnFavoritesChangedDelegate(FDelegateHandle Handle) override;
+	virtual TArray<FString> GetAliasesForPath(const FSoftObjectPath& InPath) const override;
 
 	/** Broadcast that the favorites have changed. */
 	void BroadcastFavoritesChanged() const;
diff --git a/Engine/Source/Editor/ContentBrowser/Public/IContentBrowserSingleton.h b/Engine/Source/Editor/ContentBrowser/Public/IContentBrowserSingleton.h
--- a/Engine/Source/Editor/ContentBrowser/Public/IContentBrowserSingleton.h
+++ b/Engine/Source/Editor/ContentBrowser/Public/IContentBrowserSingleton.h
@@ -720,4 +720,12 @@ public:
 
 	/** Unregister a previously-registered handler for when Favorites changes. */
 	virtual void UnregisterOnFavoritesChangedDelegate(FDelegateHandle Handle) = 0;
+
+	/**
+	 * Get a list of other paths that the data source may be using to represent a specific path
+	 *
+	 * @param The internal path (or object path) of an asset to get aliases for
+	 * @return All alternative paths that represent the input path (not including the input path itself)
+	 */
+	virtual TArray<FString> GetAliasesForPath(const FSoftObjectPath& InPath) const = 0;
 };
diff --git a/Engine/Source/Editor/UnrealEd/Private/MaterialGraphSchema.cpp b/Engine/Source/Editor/UnrealEd/Private/MaterialGraphSchema.cpp
--- a/Engine/Source/Editor/UnrealEd/Private/MaterialGraphSchema.cpp
+++ b/Engine/Source/Editor/UnrealEd/Private/MaterialGraphSchema.cpp
@@ -41,16 +41,13 @@
 #include "GraphEditorActions.h"
 #include "GraphEditorSettings.h"
 #include "AssetRegistry/AssetRegistryModule.h"
+#include "IAssetTools.h"
 #include "MaterialEditorActions.h"
 #include "MaterialGraphNode_Knot.h"
 #include "RenderUtils.h"
 
 #define LOCTEXT_NAMESPACE "MaterialGraphSchema"
 
-static bool GDisableMaterialFunctionContextMenuActions = false;
-static FAutoConsoleVariableRef DisableMaterialFunctionContextMenuActionsCVar(TEXT("r.Material.DisableMaterialFunctionContextMenuActions"),
-	GDisableMaterialFunctionContextMenuActions, TEXT("Prevents MaterialFunction creation actions from being added to MaterialEditor context menus (False: disabled, True: enabled"));
-
 int32 UMaterialGraphSchema::CurrentCacheRefreshID = 0;
 
 ////////////////////////////////////////
@@ -926,13 +923,6 @@ TSharedPtr<FEdGraphSchemaAction> UMaterialGraphSchema::GetCreateCommentAction()
 
 void UMaterialGraphSchema::GetMaterialFunctionActions(FGraphActionMenuBuilder& ActionMenuBuilder) const
 {
-	// Note: this is a temporary security solution. We need to improve on this by allowing functions that in included folders, 
-	// also considering that some folders may be aliased to different paths.
-	if (GDisableMaterialFunctionContextMenuActions)
-	{
-		return;
-	}
-
 	// Get type of dragged pin
 	uint32 FromPinType = 0;
 	if (ActionMenuBuilder.FromPin)
@@ -942,6 +932,7 @@ void UMaterialGraphSchema::GetMaterialFunctionActions(FGraphActionMenuBuilder& A
 
 	// Load the asset registry module
 	FAssetRegistryModule& AssetRegistryModule = FModuleManager::LoadModuleChecked<FAssetRegistryModule>(TEXT("AssetRegistry"));
+	IAssetTools& AssetTools = IAssetTools::Get();
 
 	// Collect a full list of assets with the specified class
 	TArray<FAssetData> AssetDataList;
@@ -962,6 +953,11 @@ void UMaterialGraphSchema::GetMaterialFunctionActions(FGraphActionMenuBuilder& A
 				}
 			}
 
+			if (!AssetTools.IsAssetVisible(AssetData))
+			{
+				continue;
+			}
+
 			if (!ActionMenuBuilder.FromPin || HasCompatibleConnection(AssetData, FromPinType, ActionMenuBuilder.FromPin->Direction))
 			{
 				// Gather the relevant information from the asset data
