Save script serialization offsets into packages that use versioned property serialization.
Offset is saved relative to export serial data.
Add new object version entry SCRIPT_SERIALIZATION_OFFSET.

diff --git a/Engine/Source/Developer/IoStoreUtilities/Private/PackageStoreOptimizer.cpp b/Engine/Source/Developer/IoStoreUtilities/Private/PackageStoreOptimizer.cpp
--- a/Engine/Source/Developer/IoStoreUtilities/Private/PackageStoreOptimizer.cpp
+++ b/Engine/Source/Developer/IoStoreUtilities/Private/PackageStoreOptimizer.cpp
@@ -138,6 +138,7 @@ FPackageStoreOptimizer::FCookedHeaderData FPackageStoreOptimizer::LoadCookedHead
 		Ar << Summary;
 	}
 
+	Ar.SetUseUnversionedPropertySerialization((CookedHeaderData.Summary.GetPackageFlags() & EPackageFlags::PKG_UnversionedProperties) != 0);
 	Ar.SetFilterEditorOnly((CookedHeaderData.Summary.GetPackageFlags() & EPackageFlags::PKG_FilterEditorOnly) != 0);
 
 	if (Summary.NameCount > 0)
diff --git a/Engine/Source/Editor/UnrealEd/Private/Commandlets/DiffPackagesCommandlet.cpp b/Engine/Source/Editor/UnrealEd/Private/Commandlets/DiffPackagesCommandlet.cpp
--- a/Engine/Source/Editor/UnrealEd/Private/Commandlets/DiffPackagesCommandlet.cpp
+++ b/Engine/Source/Editor/UnrealEd/Private/Commandlets/DiffPackagesCommandlet.cpp
@@ -104,38 +104,6 @@ struct FObjectComparison
 	EObjectDiff OverallDiffType;
 };
 
-/**
- * Constructor
- *
- * Populates the PropertyData and PropertyText members if associated with a valid UObject
- */
-FNativePropertyData::FNativePropertyData( UObject* InObject )
-: Object(InObject)
-{
-	SetObject(InObject);
-}
-
-namespace DiffPackagesCommandletImpl
-{
-	void LoadNativePropertyData(UObject* Object, TArray<uint8>& out_NativePropertyData);
-}
-
-/**
- * Changes the UObject associated with this native property data container and re-initializes the
- * PropertyData and PropertyText members
- */
-void FNativePropertyData::SetObject( UObject* NewObject )
-{
-	Object = NewObject;
-	PropertyData.Empty();
-	PropertyText.Empty();
-
-	if ( Object != NULL )
-	{
-		DiffPackagesCommandletImpl::LoadNativePropertyData(Object, PropertyData);
-		Object->GetNativePropertyValues(PropertyText, PPF_SimpleObjectText);
-	}
-}
 
 /**
  * Constructor
@@ -198,50 +166,4 @@ UObject* FObjectGraph::GetRootObject() const
 	return Objects[0].Object;
 }
 
-namespace DiffPackagesCommandletImpl
-{
-
-void LoadNativePropertyData( UObject* Object, TArray<uint8>& out_NativePropertyData )
-{
-	// first, validate our input parameters
-	check(Object);
-
-	auto ObjectLinker = Object->GetLinker();
-	check(ObjectLinker);
-
-	int32 ObjectLinkerIndex = Object->GetLinkerIndex();
-	check(ObjectLinker->ExportMap.IsValidIndex(ObjectLinkerIndex));
-
-	// now begin the process of loading the data for this object's natively serialized properties into the memory archive
-	out_NativePropertyData.Empty();
-	
-	FObjectExport& ObjectExport = ObjectLinker->ExportMap[ObjectLinkerIndex];
-	const int64 ScriptStartPos = ObjectExport.ScriptSerializationStartOffset;
-	const int64 ScriptEndPos = ObjectExport.ScriptSerializationEndOffset;
-
-	const int64 NativeStartPos = ScriptEndPos;
-	const int64 NativeEndPos = ObjectExport.SerialOffset + ObjectExport.SerialSize;
-
-	const int64 NativePropertySerialSize = NativeEndPos - NativeStartPos;
-	if ( NativePropertySerialSize > 0 )
-	{
-		checkSlow(NativeStartPos>=ObjectExport.SerialOffset);
-		checkSlow(NativeStartPos<NativeEndPos);
-		// but this might not be the case - need to make sure we catch any native data that is serialized before the property data
-		const int64 SavedPos = ((FArchive*)ObjectLinker)->Tell();
-
-		((FArchive*)ObjectLinker)->Seek(NativeStartPos);
-		((FArchive*)ObjectLinker)->Precache(NativeStartPos, NativePropertySerialSize);
-
-		// allocate enough space to contain the data we're about to read from disk
-		out_NativePropertyData.AddZeroed(IntCastChecked<int32, int64>(NativePropertySerialSize));
-		((FArchive*)ObjectLinker)->Serialize(out_NativePropertyData.GetData(), NativePropertySerialSize);
-
-		// return the linker to its previous position
-		((FArchive*)ObjectLinker)->Seek(SavedPos);
-	}
-}
-
-}
-
 PRAGMA_ENABLE_DEPRECATION_WARNINGS;
diff --git a/Engine/Source/Editor/UnrealEd/Public/Commandlets/EditorCommandlets.h b/Engine/Source/Editor/UnrealEd/Public/Commandlets/EditorCommandlets.h
--- a/Engine/Source/Editor/UnrealEd/Public/Commandlets/EditorCommandlets.h
+++ b/Engine/Source/Editor/UnrealEd/Public/Commandlets/EditorCommandlets.h
@@ -233,44 +233,3 @@ struct UE_DEPRECATED(5.3, "DiffPackagesCommandlet has been deleted.") FObjectGra
 };
 PRAGMA_ENABLE_DEPRECATION_WARNINGS
 
-/**
- * Contains the natively serialized property data for a single UObject.
- */
-PRAGMA_DISABLE_DEPRECATION_WARNINGS
-struct UE_DEPRECATED(5.3, "DiffPackagesCommandlet has been deleted.") FNativePropertyData
-{
-	/** the object that this property data is for */
-	UObject*				Object;
-
-	/** the raw bytes corresponding to this object's natively serialized property data */
-	TArray<uint8>			PropertyData;
-
-	/** the property names and textual representations of this object's natively serialized data */
-	TMap<FString,FString>	PropertyText;
-
-	/** Constructor */
-	FNativePropertyData( UObject* InObject );
-
-	/**
-	 * Changes the UObject associated with this native property data container and re-initializes the
-	 * PropertyData and PropertyText members
-	 */
-	void SetObject( UObject* NewObject );
-
-	/** Comparison operators */
-	inline bool operator==( const FNativePropertyData& Other ) const
-	{
-		return ((Object == NULL) == (Other.Object == NULL)) && PropertyData == Other.PropertyData && LegacyCompareEqual(PropertyText,Other.PropertyText);
-	}
-	inline bool operator!=( const FNativePropertyData& Other ) const
-	{
-		return ((Object == NULL) != (Other.Object == NULL)) || PropertyData != Other.PropertyData || LegacyCompareNotEqual(PropertyText,Other.PropertyText);
-	}
-
-	/** bool operator */
-	inline explicit operator bool() const
-	{
-		return PropertyData.Num() || PropertyText.Num();
-	}
-};
-PRAGMA_ENABLE_DEPRECATION_WARNINGS
diff --git a/Engine/Source/Runtime/Core/Public/UObject/ObjectVersion.h b/Engine/Source/Runtime/Core/Public/UObject/ObjectVersion.h
--- a/Engine/Source/Runtime/Core/Public/UObject/ObjectVersion.h
+++ b/Engine/Source/Runtime/Core/Public/UObject/ObjectVersion.h
@@ -72,6 +72,9 @@ enum class EUnrealEngineObjectUE5Version : uint32
 
 	// Added bulk/data resource table
 	DATA_RESOURCES,
+	
+	// Added script property serialization offset to export table entries for saved, versioned packages
+	SCRIPT_SERIALIZATION_OFFSET,
 
 	// -----<new versions can be added before this line>-------------------------------------------------
 	// - this needs to be the last line (see note below)
diff --git a/Engine/Source/Runtime/CoreUObject/Private/UObject/Class.cpp b/Engine/Source/Runtime/CoreUObject/Private/UObject/Class.cpp
--- a/Engine/Source/Runtime/CoreUObject/Private/UObject/Class.cpp
+++ b/Engine/Source/Runtime/CoreUObject/Private/UObject/Class.cpp
@@ -5368,6 +5368,7 @@ void UClass::SerializeDefaultObject(UObject* Object, FStructuredArchive::FSlot S
 	// tell the archive that it's allowed to load data for transient properties
 	FArchive& UnderlyingArchive = Slot.GetUnderlyingArchive();
 
+	UnderlyingArchive.MarkScriptSerializationStart(Object);
 	UnderlyingArchive.StartSerializingDefaults();
 
 	if( ((UnderlyingArchive.IsLoading() || UnderlyingArchive.IsSaving()) && !UnderlyingArchive.WantBinaryPropertySerialization()) )
@@ -5386,6 +5387,7 @@ void UClass::SerializeDefaultObject(UObject* Object, FStructuredArchive::FSlot S
 		SerializeBin(Slot, Object);
 	}
 	UnderlyingArchive.StopSerializingDefaults();
+	UnderlyingArchive.MarkScriptSerializationEnd(Object);
 }
 
 void UClass::SerializeSparseClassData(FStructuredArchive::FSlot Slot)
diff --git a/Engine/Source/Runtime/CoreUObject/Private/UObject/LinkerLoad.cpp b/Engine/Source/Runtime/CoreUObject/Private/UObject/LinkerLoad.cpp
--- a/Engine/Source/Runtime/CoreUObject/Private/UObject/LinkerLoad.cpp
+++ b/Engine/Source/Runtime/CoreUObject/Private/UObject/LinkerLoad.cpp
@@ -6007,7 +6007,14 @@ void FLinkerLoad::MarkScriptSerializationStart( const UObject* Obj )
 		if (ExportMap.IsValidIndex(Index))
 		{
 			FObjectExport& Export = ExportMap[Index];
-			Export.ScriptSerializationStartOffset = Tell();
+			if (!UseUnversionedPropertySerialization() && UEVer() >= EUnrealEngineObjectUE5Version::SCRIPT_SERIALIZATION_OFFSET)
+			{
+				checkf(Export.ScriptSerializationStartOffset == (Tell() - Export.SerialOffset), TEXT("Serialized script property start offset does not match offset during deserialization"));
+			}
+			else
+			{
+				Export.ScriptSerializationStartOffset = Tell() - Export.SerialOffset;
+			}
 		}
 	}
 }
@@ -6023,7 +6030,14 @@ void FLinkerLoad::MarkScriptSerializationEnd( const UObject* Obj )
 		if (ExportMap.IsValidIndex(Index))
 		{
 			FObjectExport& Export = ExportMap[Index];
-			Export.ScriptSerializationEndOffset = Tell();
+			if (!UseUnversionedPropertySerialization() && UEVer() >= EUnrealEngineObjectUE5Version::SCRIPT_SERIALIZATION_OFFSET)
+			{
+				checkf(Export.ScriptSerializationEndOffset == (Tell() - Export.SerialOffset), TEXT("Serialized script property end offset does not match offset during deserialization"));
+			}
+			else
+			{
+				Export.ScriptSerializationEndOffset = Tell() - Export.SerialOffset;
+			}
 		}
 	}
 }
diff --git a/Engine/Source/Runtime/CoreUObject/Private/UObject/LinkerSave.cpp b/Engine/Source/Runtime/CoreUObject/Private/UObject/LinkerSave.cpp
--- a/Engine/Source/Runtime/CoreUObject/Private/UObject/LinkerSave.cpp
+++ b/Engine/Source/Runtime/CoreUObject/Private/UObject/LinkerSave.cpp
@@ -293,6 +293,20 @@ FPackageIndex FLinkerSave::MapObject(TObjectPtr<const UObject> Object) const
 	return FPackageIndex();
 }
 
+void FLinkerSave::MarkScriptSerializationStart(const UObject* Obj) 
+{
+	check(Obj == CurrentlySavingExportObject);
+	FObjectExport& Export = ExportMap[CurrentlySavingExport.ToExport()];
+	Export.ScriptSerializationStartOffset = Tell();
+}
+
+void FLinkerSave::MarkScriptSerializationEnd(const UObject* Obj) 
+{
+	check(Obj == CurrentlySavingExportObject);
+	FObjectExport& Export = ExportMap[CurrentlySavingExport.ToExport()];
+	Export.ScriptSerializationEndOffset = Tell();
+}
+
 void FLinkerSave::Seek( int64 InPos )
 {
 	Saver->Seek( InPos );
diff --git a/Engine/Source/Runtime/CoreUObject/Private/UObject/ObjectResource.cpp b/Engine/Source/Runtime/CoreUObject/Private/UObject/ObjectResource.cpp
--- a/Engine/Source/Runtime/CoreUObject/Private/UObject/ObjectResource.cpp
+++ b/Engine/Source/Runtime/CoreUObject/Private/UObject/ObjectResource.cpp
@@ -204,6 +204,17 @@ void operator<<(FStructuredArchive::FSlot Slot, FObjectExport& E)
 		Record << SA_VALUE(TEXT("SerializationBeforeCreateDependencies"), E.SerializationBeforeCreateDependencies);
 		Record << SA_VALUE(TEXT("CreateBeforeCreateDependencies"), E.CreateBeforeCreateDependencies);
 	}	
+	
+	if (!BaseArchive.UseUnversionedPropertySerialization() && BaseArchive.UEVer() >= EUnrealEngineObjectUE5Version::SCRIPT_SERIALIZATION_OFFSET)
+	{
+		Record << SA_VALUE(TEXT("ScriptSerializationStartOffset"), E.ScriptSerializationStartOffset);
+		Record << SA_VALUE(TEXT("ScriptSerializationEndOffset"), E.ScriptSerializationEndOffset);
+	}
+	else if (BaseArchive.IsLoading())
+	{
+		E.ScriptSerializationStartOffset = 0;
+		E.ScriptSerializationEndOffset = 0;
+	}
 
 	#undef SERIALIZE_BIT_TO_RECORD
 }
diff --git a/Engine/Source/Runtime/CoreUObject/Private/UObject/SavePackage2.cpp b/Engine/Source/Runtime/CoreUObject/Private/UObject/SavePackage2.cpp
--- a/Engine/Source/Runtime/CoreUObject/Private/UObject/SavePackage2.cpp
+++ b/Engine/Source/Runtime/CoreUObject/Private/UObject/SavePackage2.cpp
@@ -2034,6 +2034,18 @@ ESavePackageResult WritePackageTextHeader(FStructuredArchive::FRecord& Structure
 		Linker.CurrentlySavingExport = FPackageIndex();
 		Linker.CurrentlySavingExportObject = nullptr;
 		Export.SerialSize = ExportsArchive.Tell() - Export.SerialOffset;
+
+		if (Export.ScriptSerializationEndOffset - Export.ScriptSerializationStartOffset > 0)
+		{
+			// Offset is already relative to export offset because of FExportProxyArchive 
+			check(Export.ScriptSerializationStartOffset >= 0);
+			check(Export.ScriptSerializationEndOffset <= Export.SerialSize);
+		}
+		else
+		{
+			check(Export.ScriptSerializationEndOffset == 0);
+			check(Export.ScriptSerializationStartOffset == 0);
+		}
 	}
 
 	return Linker.IsError() ? ESavePackageResult::Error : ReturnSuccessOrCancel();
@@ -2122,6 +2134,19 @@ ESavePackageResult WriteExports(FStructuredArchive::FRecord& StructuredArchiveRo
 			Linker->CurrentlySavingExport = FPackageIndex();
 			Linker->CurrentlySavingExportObject = nullptr;
 			Export.SerialSize = Linker->Tell() - Export.SerialOffset;
+
+			if (Export.ScriptSerializationEndOffset - Export.ScriptSerializationStartOffset > 0)
+			{
+				Export.ScriptSerializationStartOffset -= Export.SerialOffset;
+				Export.ScriptSerializationEndOffset -= Export.SerialOffset;
+				check(Export.ScriptSerializationStartOffset >= 0);
+				check(Export.ScriptSerializationEndOffset <= Export.SerialSize);
+			}
+			else
+			{
+				check(Export.ScriptSerializationEndOffset == 0);
+				check(Export.ScriptSerializationStartOffset == 0);
+			}
 		}
 	}
 	// if an error occurred on the linker while serializing exports, return an error
diff --git a/Engine/Source/Runtime/CoreUObject/Public/UObject/LinkerSave.h b/Engine/Source/Runtime/CoreUObject/Public/UObject/LinkerSave.h
--- a/Engine/Source/Runtime/CoreUObject/Public/UObject/LinkerSave.h
+++ b/Engine/Source/Runtime/CoreUObject/Public/UObject/LinkerSave.h
@@ -155,6 +155,16 @@ public:
 	/** Returns the appropriate package index for the source object, or default value if not found in ObjectIndicesMap */
 	FPackageIndex MapObject(TObjectPtr<const UObject> Object) const;
 
+	/**
+	 * Called when an object begins serializing property data using script serialization.
+	 */
+	virtual void MarkScriptSerializationStart(const UObject* Obj) override;
+
+	/**
+	 * Called when an object stops serializing property data using script serialization.
+	 */
+	virtual void MarkScriptSerializationEnd(const UObject* Obj) override;
+
 	// FArchive interface.
 	using FArchiveUObject::operator<<; // For visibility of the overloads we don't override
 	FArchive& operator<<( FName& InName );
diff --git a/Engine/Source/Runtime/CoreUObject/Public/UObject/ObjectResource.h b/Engine/Source/Runtime/CoreUObject/Public/UObject/ObjectResource.h
--- a/Engine/Source/Runtime/CoreUObject/Public/UObject/ObjectResource.h
+++ b/Engine/Source/Runtime/CoreUObject/Public/UObject/ObjectResource.h
@@ -266,16 +266,18 @@ struct FObjectExport : public FObjectResource
 	int64         	SerialOffset;
 
 	/**
-	 * The location (into the FLinker's underlying file reader archive) of the beginning of the
-	 * portion of this export's data that is serialized using script serialization.
-	 * Transient
+	 * The location (relative to SerialOffset) of the beginning of the portion of this export's data that is
+	 * serialized using tagged property serialization.
+	 * Serialized into versioned packages as of EUnrealEngineObjectUE5Version::SCRIPT_SERIALIZATION_OFFSET
+	 * Otherwise transient
 	 */
 	int64				ScriptSerializationStartOffset;
 
 	/**
-	 * The location (into the FLinker's underlying file reader archive) of the end of the
-	 * portion of this export's data that is serialized using script serialization.
-	 * Transient
+	 * The location (relative to SerialOffset) of the end of the portion of this export's data that is 
+	 * serialized using tagged property serialization.
+	 * Serialized into versioned packages as of EUnrealEngineObjectUE5Version::SCRIPT_SERIALIZATION_OFFSET
+	 * Otherwise transient
 	 */
 	int64				ScriptSerializationEndOffset;
 
