Editor: Add an environment variable "UE_EditorUIPid" to indicate that the current process or a parent process is the editor UI.
Cooker: When launched from the editor, use a different setting for cookprocesscount, so licensees can reduce amount of cookprocesses (and therefore the amount of memory) used when the editor is running.
#rnx

diff --git a/Engine/Config/BaseEditor.ini b/Engine/Config/BaseEditor.ini
--- a/Engine/Config/BaseEditor.ini
+++ b/Engine/Config/BaseEditor.ini
@@ -372,6 +372,10 @@ IterativeUseClassFilters=false
 ; CookProcessCount sets whether we cook singleprocess or multiprocess and how many processes there should be.
 ; CookProcessCount=(1 or less) is singleprocess. CookProcessCount=(N>1) is 1 director and N-1 cookworkers.
 CookProcessCount=1
+; CookProcessCountFromEditor sets the value for CookProcessCount when the cooker was launched by an editor UI process
+; A negative value is ignored and CookProcessCount is used instead.
+CookProcessCountFromEditor=-1
+
 
 +ExcludedIniParams=Game:/Script/BlueprintContext.McpContext
 +ExcludedIniParams=Editor:CookSettings
diff --git a/Engine/Source/Editor/UnrealEd/Private/CookOnTheFlyServer.cpp b/Engine/Source/Editor/UnrealEd/Private/CookOnTheFlyServer.cpp
--- a/Engine/Source/Editor/UnrealEd/Private/CookOnTheFlyServer.cpp
+++ b/Engine/Source/Editor/UnrealEd/Private/CookOnTheFlyServer.cpp
@@ -49,6 +49,7 @@
 #include "Dom/JsonObject.h"
 #include "Editor.h"
 #include "Editor/UnrealEdEngine.h"
+#include "EditorCommandLineUtils.h"
 #include "EditorDomain/EditorDomain.h"
 #include "EditorDomain/EditorDomainUtils.h"
 #include "Engine/AssetManager.h"
@@ -6499,9 +6500,17 @@ void UCookOnTheFlyServer::Initialize( ECookMode::Type DesiredCookMode, ECookInit
 	DirectorCookMode = WorkerRequests->GetDirectorCookMode(*this);
 	if (IsCookByTheBookMode() && !IsCookingInEditor())
 	{
-		int32 CookProcessCount=1;
-		GConfig->GetInt(TEXT("CookSettings"), TEXT("CookProcessCount"), CookProcessCount, GEditorIni);
+		bool bLaunchedByEditor = !FPlatformMisc::GetEnvironmentVariable(GEditorUIPidVariable).IsEmpty();
+		int32 CookProcessCount=-1;
 		FParse::Value(FCommandLine::Get(), TEXT("-CookProcessCount="), CookProcessCount);
+		if (CookProcessCount < 0 && bLaunchedByEditor)
+		{
+			GConfig->GetInt(TEXT("CookSettings"), TEXT("CookProcessCountFromEditor"), CookProcessCount, GEditorIni);
+		}
+		if (CookProcessCount < 0)
+		{
+			GConfig->GetInt(TEXT("CookSettings"), TEXT("CookProcessCount"), CookProcessCount, GEditorIni);
+		}
 		CookProcessCount = FMath::Max(1, CookProcessCount);
 		if (CookProcessCount > 1)
 		{
diff --git a/Engine/Source/Editor/UnrealEd/Private/UnrealEdEngine.cpp b/Engine/Source/Editor/UnrealEd/Private/UnrealEdEngine.cpp
--- a/Engine/Source/Editor/UnrealEd/Private/UnrealEdEngine.cpp
+++ b/Engine/Source/Editor/UnrealEd/Private/UnrealEdEngine.cpp
@@ -31,6 +31,7 @@
 #include "Engine/Selection.h"
 #include "Editor.h"
 #include "LevelEditorViewport.h"
+#include "EditorCommandLineUtils.h"
 #include "EditorModeRegistry.h"
 #include "EditorModeManager.h"
 #include "EditorModes.h"
@@ -217,6 +218,17 @@ void UUnrealEdEngine::Init(IEngineLoop* InEngineLoop)
 
 	}
 
+	// Set the UE_EditorUIPid variable; must be set before constructing the UCookOnTheFlyServer
+	if (!IsRunningCommandlet())
+	{
+		FString ParentPid = FPlatformMisc::GetEnvironmentVariable(GEditorUIPidVariable);
+		if (ParentPid.IsEmpty())
+		{
+			FPlatformMisc::SetEnvironmentVar(GEditorUIPidVariable,
+				*LexToString(FPlatformProcess::GetCurrentProcessId()));
+		}
+	}
+
 	if (!IsRunningCommandlet())
 	{
 		UEditorExperimentalSettings const* ExperimentalSettings = GetDefault<UEditorExperimentalSettings>();
diff --git a/Engine/Source/Editor/UnrealEd/Public/EditorCommandLineUtils.h b/Engine/Source/Editor/UnrealEd/Public/EditorCommandLineUtils.h
--- a/Engine/Source/Editor/UnrealEd/Public/EditorCommandLineUtils.h
+++ b/Engine/Source/Editor/UnrealEd/Public/EditorCommandLineUtils.h
@@ -2,7 +2,12 @@
 
 #pragma once
 
+#if UE_ENABLE_INCLUDE_ORDER_DEPRECATED_IN_5_2
 #include "CoreMinimal.h"
+#endif
+
+#include "Containers/StringView.h"
+#include "Containers/UnrealString.h"
 
 /** 
  * 
@@ -32,3 +37,5 @@ struct FEditorCommandLineUtils
 	 */
 	UNREALED_API static bool ParseGameProjectPath(const TCHAR* EditorCmdLine, FString& ProjPathOut, FString& GameNameOut);
 };
+
+constexpr const TCHAR* GEditorUIPidVariable = TEXT("UE_EditorUIPid");
