Change Material root node pin order for FrontMaterial and MaterialAttribute so that FrontMaterial is always the last one (more coherent UX behavior).

diff --git a/Engine/Source/Editor/UnrealEd/Private/MaterialGraph.cpp b/Engine/Source/Editor/UnrealEd/Private/MaterialGraph.cpp
--- a/Engine/Source/Editor/UnrealEd/Private/MaterialGraph.cpp
+++ b/Engine/Source/Editor/UnrealEd/Private/MaterialGraph.cpp
@@ -264,12 +264,14 @@ void UMaterialGraph::RebuildGraphInternal(const TMap<UMaterialExpression*, TArra
 		MaterialInputs.Add(FMaterialInputInfo(FMaterialAttributeDefinitionMap::GetDisplayNameForMaterial(MP_PixelDepthOffset, Material), MP_PixelDepthOffset, LOCTEXT("PixelDepthOffsetToolTip", "Pixel Depth Offset")));
 		MaterialInputs.Add(FMaterialInputInfo(FMaterialAttributeDefinitionMap::GetDisplayNameForMaterial(MP_ShadingModel, Material), MP_ShadingModel, LOCTEXT("ShadingModelToolTip", "Selects which shading model should be used per pixel")));
 		MaterialInputs.Add(FMaterialInputInfo(FMaterialAttributeDefinitionMap::GetDisplayNameForMaterial(MP_SurfaceThickness, Material), MP_SurfaceThickness, LOCTEXT("SurfaceThicknessToolTip", "Defines the surface's thickness when IsThinSurface is enabled")));
-		MaterialInputs.Add(FMaterialInputInfo(FMaterialAttributeDefinitionMap::GetDisplayNameForMaterial(MP_FrontMaterial, Material), MP_FrontMaterial, LOCTEXT("FrontMaterialToolTip", "Specify the front facing material")));
 		MaterialInputs.Add(FMaterialInputInfo(FMaterialAttributeDefinitionMap::GetDisplayNameForMaterial(MP_Displacement, Material), MP_Displacement, LOCTEXT("DisplacementToolTip", "Specifies scalar vertex displacement."))); 
 
 		//^^^ New material properties go above here. ^^^^
 		MaterialInputs.Add(FMaterialInputInfo(LOCTEXT("MaterialAttributes", "Material Attributes"), MP_MaterialAttributes, LOCTEXT("MaterialAttributesToolTip", "Material Attributes")));
 
+		// FrontMaterial is always added the latest for UX purpose
+		MaterialInputs.Add(FMaterialInputInfo(FMaterialAttributeDefinitionMap::GetDisplayNameForMaterial(MP_FrontMaterial, Material), MP_FrontMaterial, LOCTEXT("FrontMaterialToolTip", "Specify the front facing material")));
+
 		// Add Root Node
 		{
 			FGraphNodeCreator<UMaterialGraphNode_Root> NodeCreator(*this);
