Make sure physics asset cannot be imported

#rnx

diff --git a/Engine/Source/Editor/UnrealEd/Private/Fbx/FbxFactory.cpp b/Engine/Source/Editor/UnrealEd/Private/Fbx/FbxFactory.cpp
--- a/Engine/Source/Editor/UnrealEd/Private/Fbx/FbxFactory.cpp
+++ b/Engine/Source/Editor/UnrealEd/Private/Fbx/FbxFactory.cpp
@@ -36,6 +36,7 @@
 #include "Misc/Paths.h"
 #include "Misc/ScopedSlowTask.h"
 #include "ObjectTools.h"
+#include "PhysicsEngine/PhysicsAsset.h"
 #include "StaticMeshResources.h"
 #include "UObject/UnrealType.h"
 
@@ -403,6 +404,10 @@ UObject* UFbxFactory::FactoryCreateFile
 				ImportOptions->bUpdateSkeletonReferencePose = false;
 			}
 		}
+		if (!FbxImporter->CanCreateClass(UPhysicsAsset::StaticClass()))
+		{
+			ImportOptions->bCreatePhysicsAsset = false;
+		}
 		
 		if ( !FbxImporter->ImportFromFile( *UFactory::CurrentFilename, Type, true ) )
 		{
diff --git a/Engine/Source/Editor/UnrealEd/Private/Fbx/FbxMainImport.cpp b/Engine/Source/Editor/UnrealEd/Private/Fbx/FbxMainImport.cpp
--- a/Engine/Source/Editor/UnrealEd/Private/Fbx/FbxMainImport.cpp
+++ b/Engine/Source/Editor/UnrealEd/Private/Fbx/FbxMainImport.cpp
@@ -129,11 +129,6 @@ FBXImportOptions* GetImportOptions( UnFbx::FFbxImporter* FbxImporter, UFbxImport
 			ImportUI->PhysicsAsset = NULL;
 		}
 
-		if (!FbxImporter->CanCreateClass(UPhysicsAsset::StaticClass()))
-		{
-			ImportUI->bCreatePhysicsAsset = false;
-		}
-
 		if(bForceImportType)
 		{
 			ImportUI->MeshTypeToImport = ImportType;
@@ -281,6 +276,11 @@ FBXImportOptions* GetImportOptions( UnFbx::FFbxImporter* FbxImporter, UFbxImport
 			}
 		}
 
+		if (!FbxImporter->CanCreateClass(UPhysicsAsset::StaticClass()))
+		{
+			ImportUI->bCreatePhysicsAsset = false;
+		}
+
 		UFbxImportUI::SaveOptions(ImportUI);
 
 		if( ImportUI->StaticMeshImportData )
@@ -325,10 +325,22 @@ FBXImportOptions* GetImportOptions( UnFbx::FFbxImporter* FbxImporter, UFbxImport
 		//Clean up the options
 		UnFbx::FBXImportOptions::ResetOptions(ImportOptions);
 		ApplyImportUIToImportOptions(ImportUI, *ImportOptions);
+		
+		if (!FbxImporter->CanCreateClass(UPhysicsAsset::StaticClass()))
+		{
+			ImportOptions->bCreatePhysicsAsset = false;
+		}
+		
 		return ImportOptions;
 	}
 	else
 	{
+		
+		if (!FbxImporter->CanCreateClass(UPhysicsAsset::StaticClass()))
+		{
+			FbxImporter->GetImportOptions()->bCreatePhysicsAsset = false;
+		}
+
 		return FbxImporter->GetImportOptions();
 	}
 
