Fix bug with Select All Actor of Class
The previous code could result in ambiguous search results and fire an ensure if two assets of the same name existed in different folders and were instantiated in the same scene.

To test\repro:
- Create two Actor BPs of the same name in different folders
    - /Game/FolderA/TestBP
    - /Game/FolderB/TestBP
- Add instances of both BPs in a level
- Select any instance
- Select menu \ Select All TestBP(s)

diff --git a/Engine/Source/Editor/UnrealEd/Private/EditorEngine.cpp b/Engine/Source/Editor/UnrealEd/Private/EditorEngine.cpp
--- a/Engine/Source/Editor/UnrealEd/Private/EditorEngine.cpp
+++ b/Engine/Source/Editor/UnrealEd/Private/EditorEngine.cpp
@@ -3364,7 +3364,7 @@ void UEditorEngine::SelectAllActorsWithClass( bool bArchetype )
 		UWorld* CurrentEditorWorld = GetEditorWorldContext().World();
 		for (UClass* Class : SelectedClasses)
 		{
-			Exec(CurrentEditorWorld, *FString::Printf(TEXT("ACTOR SELECT OFCLASS CLASS=%s"), *Class->GetName()));
+			Exec(CurrentEditorWorld, *FString::Printf(TEXT("ACTOR SELECT OFCLASS CLASS=%s"), *Class->GetPathName()));
 		}
 	}
 	else
