Fix Editor Context crash: Edge case with non-wp world and no outliner window

- make sure World folder exists (create it if it doesn't) so that Push/Pop are symetric

#rnx

diff --git a/Engine/Source/Editor/UnrealEd/Private/EditorActorFolders.cpp b/Engine/Source/Editor/UnrealEd/Private/EditorActorFolders.cpp
--- a/Engine/Source/Editor/UnrealEd/Private/EditorActorFolders.cpp
+++ b/Engine/Source/Editor/UnrealEd/Private/EditorActorFolders.cpp
@@ -600,10 +600,7 @@ void FActorFolders::OnExecuteActorEditorContextAction(UWorld* InWorld, const EAc
 			SetActorEditorContextFolder(*InWorld, FFolder::GetWorldRootFolder(InWorld));
 			break;
 		case EActorEditorContextAction::PushContext:
-			if (UWorldFolders** Folders = (UWorldFolders**)WorldFolders.Find(InWorld))
-			{
-				(*Folders)->PushActorEditorContext();
-			}
+			GetOrCreateWorldFolders(*InWorld).PushActorEditorContext();
 			break;
 		case EActorEditorContextAction::PopContext:
 			if (UWorldFolders** Folders = (UWorldFolders**)WorldFolders.Find(InWorld))
