- Removed calls to world delegates PreLevelRemovedFromWorld/LevelRemovedFromWorld from UWorld::RemoveLevel and manually call them in UEditorLevelUtils::PrivateRemoveLevelFromWorld as it emulates UWorld::RemoveFromWorld.
- The goal of this change is for UWorld::RemoveLevel to be symmetric with UWorld::AddLevel (which doesn't call world delegates).
- LevelAddedToWorld/PreLevelRemovedFromWorld/LevelRemovedFromWorld are meant to be called by the level streaming (UWorld::AddToWorld and UWorld::RemoveFromWorld).

diff --git a/Engine/Source/Editor/UnrealEd/Private/EditorLevelUtils.cpp b/Engine/Source/Editor/UnrealEd/Private/EditorLevelUtils.cpp
--- a/Engine/Source/Editor/UnrealEd/Private/EditorLevelUtils.cpp
+++ b/Engine/Source/Editor/UnrealEd/Private/EditorLevelUtils.cpp
@@ -976,7 +976,16 @@ void UEditorLevelUtils::PrivateRemoveLevelFromWorld(ULevel* InLevel)
 
 	IStreamingManager::Get().RemoveLevel(InLevel);
 	UWorld* World = InLevel->OwningWorld;
-	World->RemoveLevel(InLevel);
+	if (World->ContainsLevel(InLevel))
+	{
+		// Manually call level removal world delegates PreLevelRemovedFromWorld/LevelRemovedFromWorld to simulate what UWorld::RemoveFromWorld does.
+		FWorldDelegates::PreLevelRemovedFromWorld.Broadcast(InLevel, World);
+		if (World->RemoveLevel(InLevel))
+		{
+			FWorldDelegates::LevelRemovedFromWorld.Broadcast(InLevel, World);
+		}
+	}
+
 	if (InLevel->bIsLightingScenario)
 	{
 		World->PropagateLightingScenarioChange();
diff --git a/Engine/Source/Runtime/Engine/Private/World.cpp b/Engine/Source/Runtime/Engine/Private/World.cpp
--- a/Engine/Source/Runtime/Engine/Private/World.cpp
+++ b/Engine/Source/Runtime/Engine/Private/World.cpp
@@ -8072,7 +8072,6 @@ bool UWorld::RemoveLevel( ULevel* InLevel )
 	bool bRemovedLevel = false;
 	if(ContainsLevel( InLevel ) == true )
 	{
-		FWorldDelegates::PreLevelRemovedFromWorld.Broadcast(InLevel, this);
 		bRemovedLevel = true;
 		
 #if WITH_EDITOR
@@ -8082,7 +8081,6 @@ bool UWorld::RemoveLevel( ULevel* InLevel )
 		}
 #endif //WITH_EDITOR
 		Levels.Remove( InLevel );
-		FWorldDelegates::LevelRemovedFromWorld.Broadcast(InLevel, this);
 		BroadcastLevelsChanged();
 	}
 	return bRemovedLevel;
