Fixed double clicking on an actor in the outliner not focusing the viewport on it during piloting.
- We can either press F or double click in the outliner to focus on an actor. Pressing F goes through a code path that considers the active viewport only, and allow focusing when piloting. Double clicking in the outliner sends the focus request to all viewports and wouldn't allow it if we were piloting an actor.
- The solution is to remove the condition in the multi-viewport code path to match the behavior of the single viewport code path. The condition was added a while ago to fix a feedback issue between the actor and the camera and isn't needed anymore.

diff --git a/Engine/Source/Editor/UnrealEd/Private/EditorServer.cpp b/Engine/Source/Editor/UnrealEd/Private/EditorServer.cpp
--- a/Engine/Source/Editor/UnrealEd/Private/EditorServer.cpp
+++ b/Engine/Source/Editor/UnrealEd/Private/EditorServer.cpp
@@ -4931,9 +4931,7 @@ void UEditorEngine::MoveViewportCamerasToBox(const FBox& BoundingBox, bool bActi
 			// Update all viewports.
 			for (FLevelEditorViewportClient* LinkedViewportClient : GetLevelViewportClients())
 			{
-				//Dont move camera attach to an actor
-				if (!LinkedViewportClient->IsAnyActorLocked())
-					LinkedViewportClient->FocusViewportOnBox(BoundingBox);
+				LinkedViewportClient->FocusViewportOnBox(BoundingBox);
 			}
 		}
 
