UWorldFolders::LoadState and UWorldFolders::SaveState now early out for temp, instanced and game worlds

diff --git a/Engine/Source/Editor/UnrealEd/Private/WorldFolders.cpp b/Engine/Source/Editor/UnrealEd/Private/WorldFolders.cpp
--- a/Engine/Source/Editor/UnrealEd/Private/WorldFolders.cpp
+++ b/Engine/Source/Editor/UnrealEd/Private/WorldFolders.cpp
@@ -280,6 +280,10 @@ void UWorldFolders::Serialize(FArchive& Ar)
 
 FString UWorldFolders::GetWorldStateFilename() const
 {
+	if (World->IsGameWorld() || World->IsInstanced() || FPackageName::IsTempPackage(World->GetPackage()->GetName()))
+	{
+		return FString();
+	}
 	UPackage* Package = World->GetOutermost();
 	const FString PathName = Package->GetPathName();
 	const uint32 PathNameCrc = FCrc::MemCrc32(*PathName, sizeof(TCHAR) * PathName.Len());
@@ -288,20 +292,24 @@ FString UWorldFolders::GetWorldStateFilename() const
 
 void UWorldFolders::LoadState()
 {
-	FFolder WorldDefaultFolder = FFolder::GetWorldRootFolder(World.Get());
-	check(WorldDefaultFolder.IsRootObjectValid());
-	const FFolder::FRootObject WorldRootObject = WorldDefaultFolder.GetRootObject();
+	const FString Filename = GetWorldStateFilename();
+	if (Filename.IsEmpty())
+	{
+		return;
+	}
 
 	// Attempt to load the folder properties from user's saved world state directory and apply them.
-	const auto Filename = GetWorldStateFilename();
 	TUniquePtr<FArchive> Ar(IFileManager::Get().CreateFileReader(*Filename));
 	if (Ar)
 	{
 		TSharedPtr<FJsonObject> RootObject = MakeShareable(new FJsonObject);
-
 		auto Reader = TJsonReaderFactory<TCHAR>::Create(Ar.Get());
 		if (FJsonSerializer::Deserialize(Reader, RootObject))
 		{
+			FFolder WorldDefaultFolder = FFolder::GetWorldRootFolder(World.Get());
+			check(WorldDefaultFolder.IsRootObjectValid());
+			const FFolder::FRootObject WorldRootObject = WorldDefaultFolder.GetRootObject();
+
 			const TSharedPtr<FJsonObject>& JsonFolders = RootObject->GetObjectField(TEXT("Folders"));
 			for (const auto& KeyValue : JsonFolders->Values)
 			{
@@ -323,7 +331,12 @@ void UWorldFolders::LoadState()
 
 void UWorldFolders::SaveState()
 {
-	const auto Filename = GetWorldStateFilename();
+	const FString Filename = GetWorldStateFilename();
+	if (Filename.IsEmpty())
+	{
+		return;
+	}
+
 	TUniquePtr<FArchive> Ar(IFileManager::Get().CreateFileWriter(*Filename));
 	if (Ar)
 	{
