Skip redirector packages when gathering asset localization

This avoids a package localization ID conflict warning between the redirector and the moved asset

#rnx

diff --git a/Engine/Source/Editor/UnrealEd/Private/Commandlets/GatherTextFromAssetsCommandlet.cpp b/Engine/Source/Editor/UnrealEd/Private/Commandlets/GatherTextFromAssetsCommandlet.cpp
--- a/Engine/Source/Editor/UnrealEd/Private/Commandlets/GatherTextFromAssetsCommandlet.cpp
+++ b/Engine/Source/Editor/UnrealEd/Private/Commandlets/GatherTextFromAssetsCommandlet.cpp
@@ -672,6 +672,12 @@ int32 UGatherTextFromAssetsCommandlet::Main(const FString& Params)
 	const double FilteringAssetsByIncludeExcludePathsStartTime = FPlatformTime::Seconds();
 	AssetDataArray.RemoveAll([&](const FAssetData& PartiallyFilteredAssetData) -> bool
 	{
+		if (PartiallyFilteredAssetData.IsRedirector())
+		{
+			// Redirectors never have localization
+			return true;
+		}
+
 		FString PackageFilePathWithoutExtension;
 		if (!FPackageName::TryConvertLongPackageNameToFilename(PartiallyFilteredAssetData.PackageName.ToString(), PackageFilePathWithoutExtension))
 		{
