Expose RemoveLevelFromWorld as a blueprint function.

diff --git a/Engine/Source/Editor/UnrealEd/Public/EditorLevelUtils.h b/Engine/Source/Editor/UnrealEd/Public/EditorLevelUtils.h
--- a/Engine/Source/Editor/UnrealEd/Public/EditorLevelUtils.h
+++ b/Engine/Source/Editor/UnrealEd/Public/EditorLevelUtils.h
@@ -195,6 +195,21 @@ public:
 		return AddLevelToWorld(World, *LevelPackageName, LevelStreamingClass, FTransform::Identity);
 	}
 
+	/**
+	 * Removes given level from the world. Note, this will only work for sub-levels in the main level.
+	 *
+	 * @param	InLevel				    Level asset to remove from the world.
+	 * @param	bClearSelection			If true, it will clear the editor selection.
+	 * @param	bResetTransactionBuffer	If true, it will reset the transaction buffer (i.e. clear undo history)
+	 *
+	 * @return							True if the level was successfully removed.
+	 */
+	UFUNCTION(BlueprintCallable, Category = "Editor Scripting | Level Creation", meta=(DisplayName="Remove Level From World", ScriptName="RemoveLevelFromWorld"))
+	static bool K2_RemoveLevelFromWorld(ULevel* InLevel, bool bClearSelection = true, bool bResetTransactionBuffer = true)
+	{
+		return RemoveLevelFromWorld(InLevel, bClearSelection, bResetTransactionBuffer);
+	}
+
 	/**
 	 * Adds the named level package to the world at the given position.  Does nothing if the level already exists in the world.
 	 *
