Prevent new process launches from editor when running NullApplication

There's currently no way to make the child window visible so just don't try when using NullApplication

diff --git a/Engine/Source/Editor/UnrealEd/Private/Kismet2/DebuggerCommands.cpp b/Engine/Source/Editor/UnrealEd/Private/Kismet2/DebuggerCommands.cpp
--- a/Engine/Source/Editor/UnrealEd/Private/Kismet2/DebuggerCommands.cpp
+++ b/Engine/Source/Editor/UnrealEd/Private/Kismet2/DebuggerCommands.cpp
@@ -6,6 +6,7 @@
 #include "Misc/MessageDialog.h"
 #include "Misc/App.h"
 #include "Modules/ModuleManager.h"
+#include "Null/NullPlatformApplicationMisc.h"
 #include "Framework/Application/SlateApplication.h"
 #include "Widgets/Text/STextBlock.h"
 #include "Framework/MultiBox/MultiBoxExtender.h"
@@ -1680,7 +1681,16 @@ void FInternalPlayWorldCommandCallbacks::PlayInNewProcess_Clicked(EPlayModeType
 
 bool FInternalPlayWorldCommandCallbacks::PlayInNewProcess_CanExecute()
 {
-	return !IsStoppedAtBreakpoint_InEngineMode();
+	if (IsStoppedAtBreakpoint_InEngineMode())
+	{
+		return false;
+	}
+	else if (FNullPlatformApplicationMisc::IsUsingNullApplication())
+	{
+		return false;
+	}
+
+	return true;
 }
 
 
