- Greatly optimized FActorFolders in PIE (now has no impact of PIE performance)
- Outliner automatically hides empty folders in PIE

diff --git a/Engine/Source/Editor/SceneOutliner/Private/ActorFolderTreeItem.cpp b/Engine/Source/Editor/SceneOutliner/Private/ActorFolderTreeItem.cpp
--- a/Engine/Source/Editor/SceneOutliner/Private/ActorFolderTreeItem.cpp
+++ b/Engine/Source/Editor/SceneOutliner/Private/ActorFolderTreeItem.cpp
@@ -260,7 +260,7 @@ FActorFolderTreeItem::FActorFolderTreeItem(const FFolder& InFolder, const TWeakO
 
 	if (World.IsValid())
 	{
-		Flags.bIsExpanded = FActorFolders::Get().IsFolderExpanded(*World, GetFolder());
+		Flags.bIsExpanded = !World->IsGameWorld() ? FActorFolders::Get().IsFolderExpanded(*World, GetFolder()) : true;
 	}
 }
 
@@ -376,6 +376,17 @@ bool FActorFolderTreeItem::ShouldShowVisibilityState() const
 	return World.IsValid() && Level && (Level->IsPersistentLevel() || !Level->IsInstancedLevel());
 }
 
+bool FActorFolderTreeItem::ShouldRemoveOnceLastChildRemoved() const
+{
+	if (FFolderTreeItem::ShouldRemoveOnceLastChildRemoved())
+	{
+		return true;
+	}
+	// Don't show empty folders in game world
+	const bool bIsGameWorld = World.IsValid() && World->IsGameWorld();
+	return bIsGameWorld;
+}
+
 FFolder FActorFolderTreeItem::GetFolder() const
 {
 	// Use Folder resolved by ActorFolder when valid (a null owning world would will fail resolving, see UActorFolder::GetFolder())
diff --git a/Engine/Source/Editor/SceneOutliner/Private/ActorHierarchy.cpp b/Engine/Source/Editor/SceneOutliner/Private/ActorHierarchy.cpp
--- a/Engine/Source/Editor/SceneOutliner/Private/ActorHierarchy.cpp
+++ b/Engine/Source/Editor/SceneOutliner/Private/ActorHierarchy.cpp
@@ -387,6 +387,12 @@ bool FActorHierarchy::IsShowingUnloadedActors() const
 
 void FActorHierarchy::CreateFolderItems(UWorld* InWorld, TArray<FSceneOutlinerTreeItemPtr>& OutItems) const
 {
+	// In game world, create folder through FindOrCreateParentItem
+	if (InWorld->IsGameWorld())
+	{
+		return;
+	}
+
 	if (Mode->ShouldShowFolders() && bShowingEmptyFolders)
 	{
 		// Add any folders which might match the current search terms
@@ -540,6 +546,12 @@ void FActorHierarchy::CreateChildren(const FSceneOutlinerTreeItemPtr& Item, TArr
 {
 	auto CreateChildrenFolders = [this](UWorld* InWorld, const FFolder& InParentFolder, const FFolder::FRootObject& InFolderRootObject, TArray<FSceneOutlinerTreeItemPtr>& OutChildren)
 	{
+		// In game world, create folder through FindOrCreateParentItem
+		if (InWorld->IsGameWorld())
+		{
+			return;
+		}
+
 		FActorFolders::Get().ForEachFolderWithRootObject(*InWorld, InFolderRootObject, [this, InWorld, &InParentFolder, &OutChildren](const FFolder& Folder)
 		{
 			if (Folder.IsChildOf(InParentFolder))
diff --git a/Engine/Source/Editor/SceneOutliner/Private/SSceneOutliner.cpp b/Engine/Source/Editor/SceneOutliner/Private/SSceneOutliner.cpp
--- a/Engine/Source/Editor/SceneOutliner/Private/SSceneOutliner.cpp
+++ b/Engine/Source/Editor/SceneOutliner/Private/SSceneOutliner.cpp
@@ -865,10 +865,13 @@ void SSceneOutliner::RepopulateEntireTree()
 
 void SSceneOutliner::OnChildRemovedFromParent(ISceneOutlinerTreeItem& Parent)
 {
-	if (Parent.Flags.bIsFilteredOut && !Parent.GetChildren().Num())
+	if (!Parent.GetChildren().Num())
 	{
-		// The parent no longer has any children that match the current search terms. Remove it.
-		RemoveItemFromTree(Parent.AsShared());
+		if (Parent.ShouldRemoveOnceLastChildRemoved())
+		{
+			// The parent no longer has any children that match the current search terms. Remove it.
+			RemoveItemFromTree(Parent.AsShared());
+		}
 	}
 }
 
diff --git a/Engine/Source/Editor/SceneOutliner/Public/ActorFolderTreeItem.h b/Engine/Source/Editor/SceneOutliner/Public/ActorFolderTreeItem.h
--- a/Engine/Source/Editor/SceneOutliner/Public/ActorFolderTreeItem.h
+++ b/Engine/Source/Editor/SceneOutliner/Public/ActorFolderTreeItem.h
@@ -24,6 +24,7 @@ public:
 	virtual TSharedRef<SWidget> GenerateLabelWidget(ISceneOutliner& Outliner, const STableRow<FSceneOutlinerTreeItemPtr>& InRow) override;
 	virtual bool ShouldShowPinnedState() const override { return false; }
 	virtual bool ShouldShowVisibilityState() const override;
+	virtual bool ShouldRemoveOnceLastChildRemoved() const override;
 	virtual FFolder GetFolder() const override;
 	/* End FFolderTreeItem Implementation */
 		
diff --git a/Engine/Source/Editor/SceneOutliner/Public/ISceneOutlinerTreeItem.h b/Engine/Source/Editor/SceneOutliner/Public/ISceneOutlinerTreeItem.h
--- a/Engine/Source/Editor/SceneOutliner/Public/ISceneOutlinerTreeItem.h
+++ b/Engine/Source/Editor/SceneOutliner/Public/ISceneOutlinerTreeItem.h
@@ -178,4 +178,7 @@ public:
 
 	/** Called when this item's label has changed */
 	virtual void OnLabelChanged() {}
+
+	/** Whether the item should be removed when the last child has been removed. */
+	virtual bool ShouldRemoveOnceLastChildRemoved() const { check(GetChildren().IsEmpty()); return Flags.bIsFilteredOut; }
 };
diff --git a/Engine/Source/Editor/UnrealEd/Private/EditorActorFolders.cpp b/Engine/Source/Editor/UnrealEd/Private/EditorActorFolders.cpp
--- a/Engine/Source/Editor/UnrealEd/Private/EditorActorFolders.cpp
+++ b/Engine/Source/Editor/UnrealEd/Private/EditorActorFolders.cpp
@@ -134,46 +134,34 @@ void FActorFolders::BroadcastOnActorFolderMoved(UWorld& InWorld, const FFolder&
 	BroadcastOnActorEditorContextClientChanged(InWorld);
 }
 
-void FActorFolders::OnLevelActorListChanged()
+static bool IsRunningGameOrPIE()
 {
-	bAnyLevelsChanged = true;
+	return (GIsEditor && GEditor && GEditor->GetPIEWorldContext()) || IsRunningGame();
 }
 
-void FActorFolders::OnAllLevelsChanged()
+void FActorFolders::OnLevelActorListChanged()
 {
-	if (!bAnyLevelsChanged)
+	if (!IsRunningGameOrPIE())
 	{
-		return;
+		bAnyLevelsChanged = true;
 	}
-	QUICK_SCOPE_CYCLE_COUNTER(FActorFolders_OnAllLevelsChanged);
-	bAnyLevelsChanged = false;
-	Housekeeping();
+}
 
-	check(GEngine);
+void FActorFolders::OnAllLevelsChanged()
+{
+	TRACE_CPUPROFILER_EVENT_SCOPE(FActorFolders::OnAllLevelsChanged);
 
-	UWorld* World = nullptr;
-	for (const FWorldContext& Context : GEngine->GetWorldContexts())
+	if (bAnyLevelsChanged && !IsRunningGameOrPIE())
 	{
-		UWorld* ThisWorld = Context.World();
-		if (!ThisWorld)
-		{
-			continue;
-		}
-		else if (Context.WorldType == EWorldType::PIE)
-		{
-			World = ThisWorld;
-			break;
-		}
-		else if (Context.WorldType == EWorldType::Editor)
+		QUICK_SCOPE_CYCLE_COUNTER(FActorFolders_OnAllLevelsChanged);
+		bAnyLevelsChanged = false;
+		Housekeeping();
+
+		if (UWorld* World = GEditor ? GEditor->GetEditorWorldContext().World() : nullptr)
 		{
-			World = ThisWorld;
+			RebuildFolderListForWorld(*World);
 		}
 	}
-
-	if (World)
-	{
-		RebuildFolderListForWorld(*World);
-	}
 }
 
 void FActorFolders::OnMapChange(uint32 MapChangeFlags)
@@ -226,6 +214,8 @@ void FActorFolders::OnActorFolderChanged(const AActor* InActor, FName OldPath)
 
 void FActorFolders::RebuildFolderListForWorld(UWorld& InWorld)
 {
+	TRACE_CPUPROFILER_EVENT_SCOPE(FActorFolders::RebuildFolderListForWorld);
+
 	if (auto* Folders = WorldFolders.Find(&InWorld))
 	{
 		// For world folders, we don't empty the existing folders so that we keep empty ones.
diff --git a/Engine/Source/Editor/UnrealEd/Private/WorldFolders.cpp b/Engine/Source/Editor/UnrealEd/Private/WorldFolders.cpp
--- a/Engine/Source/Editor/UnrealEd/Private/WorldFolders.cpp
+++ b/Engine/Source/Editor/UnrealEd/Private/WorldFolders.cpp
@@ -19,6 +19,8 @@ DEFINE_LOG_CATEGORY(LogWorldFolders);
 
 void UWorldFolders::Initialize(UWorld* InWorld)
 {
+	TRACE_CPUPROFILER_EVENT_SCOPE(UWorldFolders::Initialize);
+
 	check(!World.IsValid());
 	check(IsValidChecked(InWorld));
 	
@@ -35,6 +37,13 @@ void UWorldFolders::Initialize(UWorld* InWorld)
 
 void UWorldFolders::RebuildList()
 {
+	TRACE_CPUPROFILER_EVENT_SCOPE(UWorldFolders::RebuildList);
+
+	if (GetWorld()->IsGameWorld())
+	{
+		return;
+	}
+
 	Modify();
 	
 	// Clear folders with a Root Object.
