Cooker: ZenStore -diffonly bugfix: need to gather nested objects as well as toplevel objects when listing the exports that would be present from an in-memory package.
#rnx

diff --git a/Engine/Source/Editor/UnrealEd/Private/Cooker/DiffWriterZenHeader.cpp b/Engine/Source/Editor/UnrealEd/Private/Cooker/DiffWriterZenHeader.cpp
--- a/Engine/Source/Editor/UnrealEd/Private/Cooker/DiffWriterZenHeader.cpp
+++ b/Engine/Source/Editor/UnrealEd/Private/Cooker/DiffWriterZenHeader.cpp
@@ -400,7 +400,7 @@ void FZenPackageExportsForDiff::InitializeFromMemory(FAccumulatorGlobals& Global
 			{
 				ExportPaths.Add(PublicExportHash, FString(PackageRelativePath));
 			}
-		}, false /* bIncludeNestedObjects */);
+		}, true /* bIncludeNestedObjects */);
 }
 
 FStringView FZenPackageExportsForDiff::GetExportPackageRelativePath(uint64 PublicExportHash)
