Editor: add viewport client to the view modifier parameters so that subscribers know what viewport is being modified

diff --git a/Engine/Source/Editor/UnrealEd/Private/EditorViewportClient.cpp b/Engine/Source/Editor/UnrealEd/Private/EditorViewportClient.cpp
--- a/Engine/Source/Editor/UnrealEd/Private/EditorViewportClient.cpp
+++ b/Engine/Source/Editor/UnrealEd/Private/EditorViewportClient.cpp
@@ -895,6 +895,7 @@ FSceneView* FEditorViewportClient::CalcSceneView(FSceneViewFamily* ViewFamily, c
 	// Apply view modifiers.
 	FEditorViewportViewModifierParams ViewModifierParams;
 	{
+		ViewModifierParams.ViewportClient = this;
 		ViewModifierParams.ViewInfo.Location = ViewTransform.GetLocation();
 		ViewModifierParams.ViewInfo.Rotation = ViewTransform.GetRotation();
 
diff --git a/Engine/Source/Editor/UnrealEd/Public/EditorViewportClient.h b/Engine/Source/Editor/UnrealEd/Public/EditorViewportClient.h
--- a/Engine/Source/Editor/UnrealEd/Public/EditorViewportClient.h
+++ b/Engine/Source/Editor/UnrealEd/Public/EditorViewportClient.h
@@ -297,6 +297,8 @@ private:
 /** Parameter struct for editor viewport view modifiers */
 struct FEditorViewportViewModifierParams
 {
+	FEditorViewportClient* ViewportClient = nullptr;
+
 	FMinimalViewInfo ViewInfo;
 
 	void AddPostProcessBlend(const FPostProcessSettings& Settings, float Weight)
