Cooker: RequestCluster in non-incremental cook: Remove the use of an AsyncTask for reporting the empty cache results. It was needed only for shake down of the async threading flow, which has completed, and the use of an AsyncTask causes a performance degredation when worker threads are saturated.
#rnx

diff --git a/Engine/Source/Editor/UnrealEd/Private/Cooker/CookRequestCluster.cpp b/Engine/Source/Editor/UnrealEd/Private/Cooker/CookRequestCluster.cpp
--- a/Engine/Source/Editor/UnrealEd/Private/Cooker/CookRequestCluster.cpp
+++ b/Engine/Source/Editor/UnrealEd/Private/Cooker/CookRequestCluster.cpp
@@ -1463,23 +1463,19 @@ void FRequestCluster::FQueryVertexBatch::Send()
 		}
 		else
 		{
-			// When we do not need to asynchronously fetch, we record empty cache results from an AsyncTask.
-			// Using an AsyncTask keeps the threading flow similar to the FetchCookAttachments case
-			AsyncTask(ENamedThreads::AnyThread,
-				[this, PlatformIndex]()
+			// When we do not need to asynchronously fetch, we record empty cache results to keep the edgefetch
+			// flow similar to the FetchCookAttachments case
+
+			// Don't use a ranged-for, as we are not allowed to access this or this->PackageNames after the
+			// last index, and ranged-for != at the end of the final loop iteration can read from PackageNames
+			int32 NumPackageNames = PlatformData.PackageNames.Num();
+			FName* PackageNamesData = PlatformData.PackageNames.GetData();
+			for (int32 PackageNameIndex = 0; PackageNameIndex < NumPackageNames; ++PackageNameIndex)
 			{
-				FPlatformData& PlatformData = PlatformDatas[PlatformIndex];
-				// Don't use a ranged-for, as we are not allowed to access this or this->PackageNames after the
-				// last index, and ranged-for != at the end of the final loop iteration can read from PackageNames
-				int32 NumPackageNames = PlatformData.PackageNames.Num();
-				FName* PackageNamesData = PlatformData.PackageNames.GetData();
-				for (int32 PackageNameIndex = 0; PackageNameIndex < NumPackageNames; ++PackageNameIndex)
-				{
-					FName PackageName = PackageNamesData[PackageNameIndex];
-					UE::TargetDomain::FCookAttachments Attachments;
-					RecordCacheResults(PackageName, PlatformIndex, MoveTemp(Attachments));
-				}
-			});
+				FName PackageName = PackageNamesData[PackageNameIndex];
+				UE::TargetDomain::FCookAttachments Attachments;
+				RecordCacheResults(PackageName, PlatformIndex, MoveTemp(Attachments));
+			}
 		}
 	}
 }
