Add new message log token type for FAssetData.
Implement click callback to handle assets or loaded/unloaded actors.
Use asset data token in content validation logs.
Get full asset data (including actor metadata) for actor validation context menu entry.

diff --git a/Engine/Plugins/Editor/DataValidation/Source/DataValidation/Private/DataValidationModule.cpp b/Engine/Plugins/Editor/DataValidation/Source/DataValidation/Private/DataValidationModule.cpp
--- a/Engine/Plugins/Editor/DataValidation/Source/DataValidation/Private/DataValidationModule.cpp
+++ b/Engine/Plugins/Editor/DataValidation/Source/DataValidation/Private/DataValidationModule.cpp
@@ -113,6 +113,12 @@ void FDataValidationModule::RegisterMenus()
 					}
 					return true;
 				});
+				FAssetRegistryModule& AssetRegistryModule = FModuleManager::GetModuleChecked<FAssetRegistryModule>("AssetRegistry");
+				IAssetRegistry& AssetRegistry = AssetRegistryModule.Get();
+				for (FAssetData& AssetData : SelectedActorAssets)
+				{
+					AssetData = AssetRegistry.GetAssetByObjectPath(AssetData.GetSoftObjectPath(), true);
+				}
 				ValidateAssets(SelectedActorAssets, false, EDataValidationUsecase::Manual);
 			}
 		});
diff --git a/Engine/Source/Developer/MessageLog/Private/UserInterface/SMessageLogMessageListRow.cpp b/Engine/Source/Developer/MessageLog/Private/UserInterface/SMessageLogMessageListRow.cpp
--- a/Engine/Source/Developer/MessageLog/Private/UserInterface/SMessageLogMessageListRow.cpp
+++ b/Engine/Source/Developer/MessageLog/Private/UserInterface/SMessageLogMessageListRow.cpp
@@ -11,6 +11,7 @@
 #include "Widgets/Input/SHyperlink.h"
 #include "Internationalization/Regex.h"
 #include "Styling/StyleColors.h"
+#include "AssetRegistry/AssetDataToken.h"
 
 #define LOCTEXT_NAMESPACE "SMessageLogMessageListRow"
 
@@ -167,6 +168,16 @@ void SMessageLogMessageListRow::CreateMessage(const TSharedRef<SHorizontalBox>&
 	}
 		break;
 
+	case EMessageToken::AssetData:
+	{
+		const TSharedRef<FAssetDataToken> AssetDataToken = StaticCastSharedRef<FAssetDataToken>(InMessageToken);
+
+		IconBrushName = FName("Icons.Search");
+		RowContent = CreateHyperlink(InMessageToken, FAssetDataToken::DefaultOnGetAssetDisplayName().IsBound()
+			? FAssetDataToken::DefaultOnGetAssetDisplayName().Execute(AssetDataToken->GetAssetData(), true)
+			: InMessageToken->ToText());
+	}
+		break;
 	case EMessageToken::Object:
 	{
 		const TSharedRef<FUObjectToken> UObjectToken = StaticCastSharedRef<FUObjectToken>(InMessageToken);
diff --git a/Engine/Source/Editor/UnrealEd/Private/UnrealEdMisc.cpp b/Engine/Source/Editor/UnrealEd/Private/UnrealEdMisc.cpp
--- a/Engine/Source/Editor/UnrealEd/Private/UnrealEdMisc.cpp
+++ b/Engine/Source/Editor/UnrealEd/Private/UnrealEdMisc.cpp
@@ -54,6 +54,7 @@
 #include "FbxLibs.h"
 #include "Kismet2/CompilerResultsLog.h"
 #include "AssetRegistry/AssetRegistryModule.h"
+#include "AssetRegistry/AssetDataToken.h"
 #include "EngineAnalytics.h"
 #include "AnalyticsEventAttribute.h"
 #include "Interfaces/IAnalyticsProvider.h"
@@ -73,6 +74,8 @@
 #include "PerformanceMonitor.h"
 #include "Engine/WorldComposition.h"
 #include "WorldPartition/WorldPartition.h"
+#include "WorldPartition/WorldPartitionActorDesc.h"
+#include "WorldPartition/WorldPartitionActorDescUtils.h"
 #include "Interfaces/IProjectManager.h"
 #include "FeaturePackContentSource.h"
 #include "ProjectDescriptor.h"
@@ -517,6 +520,8 @@ void FUnrealEdMisc::OnInit()
 	FUObjectToken::DefaultOnGetObjectDisplayName().BindRaw(this, &FUnrealEdMisc::OnGetDisplayName);
 	FAssetNameToken::OnGotoAsset().BindRaw(this, &FUnrealEdMisc::OnGotoAsset);
 	FActorToken::DefaultOnMessageTokenActivated().BindRaw(this, &FUnrealEdMisc::OnActorTokenActivated);
+	FAssetDataToken::DefaultOnMessageTokenActivated().BindRaw(this, &FUnrealEdMisc::OnAssetDataTokenActivated);
+	FAssetDataToken::DefaultOnGetAssetDisplayName().BindRaw(this, &FUnrealEdMisc::OnGetAssetDataDisplayName);
 
 	// Register to receive notification of new key bindings
 	OnUserDefinedChordChangedDelegateHandle = FInputBindingManager::Get().RegisterUserDefinedChordChanged(FOnUserDefinedChordChanged::FDelegate::CreateRaw( this, &FUnrealEdMisc::OnUserDefinedChordChanged ));
@@ -1441,6 +1446,72 @@ void FUnrealEdMisc::OnMessageTokenActivated(const TSharedRef<IMessageToken>& Tok
 	}
 }
 
+/** Delegate used to get a display name for a message log FAssetData token  */
+FText FUnrealEdMisc::OnGetAssetDataDisplayName(const FAssetData& InObject, const bool bFullPath)
+{
+	static FName NAME_ActorLabel("ActorLabel");
+	FString ActorLabel;
+	TStringBuilder<FName::StringBufferSize> Buffer;
+	if (InObject.GetTagValue(NAME_ActorLabel, ActorLabel))
+	{
+		Buffer << ActorLabel << TEXT(" (");
+		InObject.AppendObjectPath(Buffer);
+		Buffer << TEXT(")");
+		return FText::FromStringView(Buffer.ToView());
+	}
+	else if (bFullPath)
+	{
+		InObject.AppendObjectPath(Buffer);
+		return FText::FromStringView(Buffer.ToView());
+	}
+	return FText::FromName(InObject.PackageName);
+}
+
+/** Delegate used on message log FAssetData token activation */
+void FUnrealEdMisc::OnAssetDataTokenActivated(const TSharedRef<class IMessageToken>& InToken)
+{
+	if (InToken->GetType() != EMessageToken::AssetData)
+	{
+		return;
+	}
+	
+	const TSharedRef<FAssetDataToken> Token = StaticCastSharedRef<FAssetDataToken>(InToken);
+	const FAssetData& AssetData = Token->GetAssetData();
+	// If this is a standalone asset, jump to it in the content browser
+	if (AssetData.GetOptionalOuterPathName().IsNone())
+	{
+		GEditor->SyncBrowserToObjects({ AssetData });
+	}
+	// If this is a loaded actor, resolve the pointer and select it
+	else if (AActor* Actor = Cast<AActor>(AssetData.GetSoftObjectPath().ResolveObject()))
+	{
+		SelectActorFromMessageToken(Actor);
+	}
+	else
+	{
+		FName PackageName = FName(FPackageName::ObjectPathToPackageName(WriteToString<FName::StringBufferSize>(AssetData.GetOptionalOuterPathName()).ToView()));
+		UWorld* EditorWorld = GEditor->GetEditorWorldContext().World();
+		// If this is an unloaded actor in the currently loaded level, select it as an unloaded actor
+		if (EditorWorld && EditorWorld->GetPackage()->GetFName() == PackageName)
+		{
+			TUniquePtr<FWorldPartitionActorDesc> Desc = FWorldPartitionActorDescUtils::GetActorDescriptorFromAssetData(AssetData);
+			if (Desc.IsValid())
+			{
+				GEditor->BroadcastSelectUnloadedActors({ Desc->GetGuid() });
+			}
+		}
+		// If this is an actor in an unloaded level, jump to the level in the content browser
+		else 
+		{
+			FAssetRegistryModule& AssetRegistryModule = FModuleManager::GetModuleChecked<FAssetRegistryModule>("AssetRegistry");
+			IAssetRegistry& AssetRegistry = AssetRegistryModule.Get();
+			TArray<FAssetData> OutAssets;
+			AssetRegistry.GetAssetsByPackageName(PackageName, OutAssets);
+			GEditor->SyncBrowserToObjects(OutAssets);
+		}
+	}
+}
+
 void FUnrealEdMisc::OnActorTokenActivated(const TSharedRef<class IMessageToken>& Token)
 {
 	if (Token->GetType() == EMessageToken::Actor)
diff --git a/Engine/Source/Editor/UnrealEd/Public/UnrealEdMisc.h b/Engine/Source/Editor/UnrealEd/Public/UnrealEdMisc.h
--- a/Engine/Source/Editor/UnrealEd/Public/UnrealEdMisc.h
+++ b/Engine/Source/Editor/UnrealEd/Public/UnrealEdMisc.h
@@ -258,6 +258,12 @@ public:
 	/** Delegate used on message log AActor token activation */
 	UNREALED_API void OnActorTokenActivated(const TSharedRef<class IMessageToken>& Token);
 
+	/** Delegate used to get a display name for a message log FAssetData token  */
+	UNREALED_API FText OnGetAssetDataDisplayName(const FAssetData& InObject, const bool bFullPath);
+
+	/** Delegate used on message log FAssetData token activation */
+	UNREALED_API void OnAssetDataTokenActivated(const TSharedRef<class IMessageToken>& Token);
+
 	/** Delegate used to update the map of asset update counts */
 	UNREALED_API void OnObjectSaved(UObject* SavedObject, FObjectPreSaveContext SaveContext);
 
diff --git a/Engine/Source/Runtime/Core/Public/Logging/TokenizedMessage.h b/Engine/Source/Runtime/Core/Public/Logging/TokenizedMessage.h
--- a/Engine/Source/Runtime/Core/Public/Logging/TokenizedMessage.h
+++ b/Engine/Source/Runtime/Core/Public/Logging/TokenizedMessage.h
@@ -37,6 +37,7 @@ namespace EMessageToken
 		Action,
 		Actor,
 		AssetName,
+		AssetData,
 		Documentation,
 		Image,
 		Object,
diff --git a/Engine/Source/Runtime/CoreUObject/Private/AssetRegistry/AssetDataToken.cpp b/Engine/Source/Runtime/CoreUObject/Private/AssetRegistry/AssetDataToken.cpp
new file mode 100644
--- /dev/null
+++ b/Engine/Source/Runtime/CoreUObject/Private/AssetRegistry/AssetDataToken.cpp
@@ -0,0 +1,43 @@
+// Copyright Epic Games, Inc. All Rights Reserved.
+
+#include "AssetRegistry/AssetDataToken.h" 
+
+FOnMessageTokenActivated FAssetDataToken::DefaultMessageTokenActivated;
+FAssetDataToken::FOnGetDisplayName FAssetDataToken::DefaultGetAssetDisplayName;
+
+TSharedRef<FAssetDataToken> FAssetDataToken::Create(const FAssetData& InAssetData, const FText& InLabelOverride)
+{
+    return MakeShared<FAssetDataToken>(FPrivateToken{}, InAssetData, InLabelOverride);
+}
+
+FAssetDataToken::FAssetDataToken(FPrivateToken, const FAssetData& InAssetData, const FText& InLabelOverride)
+    : AssetData(InAssetData)
+{
+	if ( !InLabelOverride.IsEmpty() )
+	{
+		CachedText = InLabelOverride;
+	}
+	else
+	{
+		if ( DefaultGetAssetDisplayName.IsBound() )
+		{
+			CachedText = DefaultGetAssetDisplayName.Execute(InAssetData, false);
+		}
+		else 
+		{
+			CachedText = FText::FromString(InAssetData.GetObjectPathString());
+		}
+	}
+}
+
+const FOnMessageTokenActivated& FAssetDataToken::GetOnMessageTokenActivated() const
+{
+	if(MessageTokenActivated.IsBound())
+	{
+		return MessageTokenActivated;
+	}
+	else
+	{
+		return DefaultMessageTokenActivated;
+	}
+}
\ No newline at end of file
diff --git a/Engine/Source/Runtime/CoreUObject/Public/AssetRegistry/AssetDataToken.h b/Engine/Source/Runtime/CoreUObject/Public/AssetRegistry/AssetDataToken.h
new file mode 100644
--- /dev/null
+++ b/Engine/Source/Runtime/CoreUObject/Public/AssetRegistry/AssetDataToken.h
@@ -0,0 +1,59 @@
+// Copyright Epic Games, Inc. All Rights Reserved.
+
+#pragma once
+
+#include "AssetData.h"
+#include "Logging/TokenizedMessage.h"
+
+/**
+ * A Message Log token that links to an asset, allowing a hyperlink to navigate to an asset in the content browser or an actor in a level.
+ */
+class FAssetDataToken : public IMessageToken
+{
+	// The private token allows only members or friends to call MakeShared.
+	struct FPrivateToken { explicit FPrivateToken() = default; };
+public:
+	/** Factory method, tokens can only be constructed as shared refs */
+	COREUOBJECT_API static TSharedRef<FAssetDataToken> Create(const FAssetData& AssetData, const FText& InLabelOverride = FText());
+
+	/** Private constructor (via FPrivateToken argument) */
+	COREUOBJECT_API FAssetDataToken(FPrivateToken, const FAssetData& InAssetData, const FText& InLabelOverride);
+
+	/** Begin IMessageToken interface */
+	virtual EMessageToken::Type GetType() const override
+	{
+		return EMessageToken::AssetData;
+	}
+
+	virtual const FOnMessageTokenActivated& GetOnMessageTokenActivated() const override;
+	/** End IMessageToken interface */
+
+    COREUOBJECT_API const FAssetData& GetAssetData() const
+    {
+        return AssetData;
+    }
+
+	/** Get the delegate for default token activation */
+	static FOnMessageTokenActivated& DefaultOnMessageTokenActivated()
+	{
+		return DefaultMessageTokenActivated;
+	}
+
+	/** Get the delegate for displaying the asset name */
+	DECLARE_DELEGATE_RetVal_TwoParams(FText, FOnGetDisplayName, const FAssetData&, bool /* Full path */);
+	static FOnGetDisplayName& DefaultOnGetAssetDisplayName()
+	{
+		return DefaultGetAssetDisplayName;
+	}
+
+private:
+
+	/** The asset data this token was created for */
+	FAssetData AssetData;
+
+	/** The default activation method, if any. Usually populated by another module in the engine. */
+	COREUOBJECT_API static FOnMessageTokenActivated DefaultMessageTokenActivated;
+
+	/** The default method for getting a display name for assets, if any. Usually populated by another module in the ending. */
+	COREUOBJECT_API static FOnGetDisplayName DefaultGetAssetDisplayName;
+};
