Oodle JPEG: Added ability to export directly to .ooj format.

diff --git a/Engine/Source/Editor/UnrealEd/Classes/Exporters/TextureExporterPNG.h b/Engine/Source/Editor/UnrealEd/Classes/Exporters/TextureExporterPNG.h
--- a/Engine/Source/Editor/UnrealEd/Classes/Exporters/TextureExporterPNG.h
+++ b/Engine/Source/Editor/UnrealEd/Classes/Exporters/TextureExporterPNG.h
@@ -43,6 +43,14 @@ class UTextureExporterJPEG : public UExporter
 };
 
 
+UCLASS()
+class UTextureExporterOOJPEG : public UTextureExporterGeneric
+{
+	GENERATED_UCLASS_BODY()
+	
+	virtual bool SupportsTexture(UTexture* Texture) const override;
+};
+
 
 
 
diff --git a/Engine/Source/Editor/UnrealEd/Private/Factories/EditorFactories.cpp b/Engine/Source/Editor/UnrealEd/Private/Factories/EditorFactories.cpp
--- a/Engine/Source/Editor/UnrealEd/Private/Factories/EditorFactories.cpp
+++ b/Engine/Source/Editor/UnrealEd/Private/Factories/EditorFactories.cpp
@@ -5002,6 +5002,23 @@ UVirtualTextureBuilderExporterPNG::UVirtualTextureBuilderExporterPNG(const FObje
 
 //========================================================
 
+UTextureExporterOOJPEG::UTextureExporterOOJPEG(const FObjectInitializer& ObjectInitializer)
+	: UTextureExporterGeneric(ObjectInitializer)
+{
+	FormatExtension.Add(TEXT("OOJ"));
+	FormatDescription.Add(TEXT("Oodle JPEG"));
+}
+
+bool UTextureExporterOOJPEG::SupportsTexture(UTexture* Texture) const
+{
+	ETextureSourceFormat TSF = Texture->Source.GetFormat();
+	ERawImageFormat::Type RawFormat = FImageCoreUtils::ConvertToRawImageFormat(TSF);
+	// supports all non-HDR formats :
+	return ! ERawImageFormat::IsHDR(RawFormat);
+}
+
+//========================================================
+
 UTextureExporterEXR::UTextureExporterEXR(const FObjectInitializer& ObjectInitializer)
 	: UTextureExporterGeneric(ObjectInitializer)
 {
