Fix crash duplicating blueprint with sparse class data

diff --git a/Engine/Source/Editor/UnrealEd/Private/Kismet2/BlueprintEditorUtils.cpp b/Engine/Source/Editor/UnrealEd/Private/Kismet2/BlueprintEditorUtils.cpp
--- a/Engine/Source/Editor/UnrealEd/Private/Kismet2/BlueprintEditorUtils.cpp
+++ b/Engine/Source/Editor/UnrealEd/Private/Kismet2/BlueprintEditorUtils.cpp
@@ -1690,7 +1690,11 @@ void FBlueprintEditorUtils::PostDuplicateBlueprint(UBlueprint* Blueprint, bool b
 
 							if(FProperty* Property = OutdatedReference.ResolveMember<FProperty>(OutdatedNode->GetBlueprintClassFromNode()))
 							{
-								TargetClass = Property->GetOwnerClass()->GetAuthoritativeClass();
+								// Properties that are owned by a Struct - e.g. sparse class data - will not have a class:
+								if (UClass* OwningClass = Property->GetOwnerClass())
+								{
+									TargetClass = OwningClass->GetAuthoritativeClass();
+								}
 							}
 							else
 							{
