UUnrealEdEngine::AddTemplateMap to append custom map templates to the New Level dialog

diff --git a/Engine/Source/Editor/UnrealEd/Classes/Editor/UnrealEdEngine.h b/Engine/Source/Editor/UnrealEd/Classes/Editor/UnrealEdEngine.h
--- a/Engine/Source/Editor/UnrealEd/Classes/Editor/UnrealEdEngine.h
+++ b/Engine/Source/Editor/UnrealEd/Classes/Editor/UnrealEdEngine.h
@@ -732,6 +732,11 @@ public:
 
 	UNREALED_API void RebuildTemplateMapData();
 
+	/**
+	 * Appends the specified template maps to the list of available templates.
+	 */
+	UNREALED_API void AppendTemplateMaps( const TArray<FTemplateMapInfo>& TemplateMapInfos );
+
 	/**
 	 * Returns true if the user is currently interacting with a viewport.
 	 */
diff --git a/Engine/Source/Editor/UnrealEd/Private/UnrealEdEngine.cpp b/Engine/Source/Editor/UnrealEd/Private/UnrealEdEngine.cpp
--- a/Engine/Source/Editor/UnrealEd/Private/UnrealEdEngine.cpp
+++ b/Engine/Source/Editor/UnrealEd/Private/UnrealEdEngine.cpp
@@ -1044,6 +1044,37 @@ void UUnrealEdEngine::RebuildTemplateMapData()
 	}
 }
 
+void UUnrealEdEngine::AppendTemplateMaps(const TArray<FTemplateMapInfo>& InTemplateMapInfos)
+{
+	bool bTemplateWasAdded = false;
+	
+	PRAGMA_DISABLE_DEPRECATION_WARNINGS
+	for (const FTemplateMapInfo& TemplateMapInfo : InTemplateMapInfos)
+	{
+		if (!TemplateMapInfos.ContainsByPredicate(
+			[&TemplateMapInfo](const FTemplateMapInfo& InTemplate)
+		{
+			return InTemplate.Map == TemplateMapInfo.Map
+				&& InTemplate.DisplayName.EqualTo(InTemplate.DisplayName);
+		}))
+		{
+			TemplateMapInfos.Emplace(TemplateMapInfo);
+			bTemplateWasAdded = true;
+		}
+		else
+		{
+			UE_LOG(LogUnrealEdEngine, Warning, TEXT("Attempted to register an already registered template map ('%s'). Skipping registration."), *TemplateMapInfo.Map.ToString());
+		}
+		PRAGMA_ENABLE_DEPRECATION_WARNINGS
+	}
+
+	if (bTemplateWasAdded)
+	{
+		// A new template was added, so refresh
+		RebuildTemplateMapData();
+	}
+}
+
 void UUnrealEdEngine::SetCurrentClass( UClass* InClass )
 {
 	USelection* SelectionSet = GetSelectedObjects();
