Fix ComponentVisualizers not drawing after NotifyPropertyModified called
* Fixes at least Spline and Zone Shape handle dragging causing component visualizer to stop drawing
* Root cause appears to be that property modifications are causing object reinstancing to occur unnecessarily.  This causes the object to be deselected and reselected.  In that operation, the SelectionSet OnPreChange delegate fires which fires ComponentVisManager.ClearActiveComponentVis() and clears out the VisualizersForSelection list.  The manager is responsible for routing mouse delta events to the visualizer, the list is responsible for drawing the visualizer.
* Visualizer would be re-added to the list in UUnrealEdEngine::OnEditorElementSelectionChanged (which is called by SelectionSet::OnChanged delegate)... but in the case of property modification, SelectionSet::OnElementUpdated supresses the OnChange delegate via a FScopedClearNewPendingChange and therefore that callback is never made.

diff --git a/Engine/Source/Editor/UnrealEd/Private/ComponentVisualizer.cpp b/Engine/Source/Editor/UnrealEd/Private/ComponentVisualizer.cpp
--- a/Engine/Source/Editor/UnrealEd/Private/ComponentVisualizer.cpp
+++ b/Engine/Source/Editor/UnrealEd/Private/ComponentVisualizer.cpp
@@ -357,11 +357,11 @@ void FComponentVisualizer::NotifyPropertiesModified(UActorComponent* Component,
 				InstanceOwner->PostEditMove(PropertyChangeType == EPropertyChangeType::ValueSet);
 			}
 		}
-	}
 
-	// Rerun construction script on preview actor
-	if (Owner)
-	{
-		Owner->PostEditMove(PropertyChangeType == EPropertyChangeType::ValueSet);
+		// Rerun construction script on preview actor
+		if (Owner)
+		{
+			Owner->PostEditMove(PropertyChangeType == EPropertyChangeType::ValueSet);
+		}
 	}
 }
