Fixed not being able to edit multiple assets from content browser asset context menu.
#rnx

diff --git a/Engine/Source/Editor/UnrealEd/Private/Subsystems/AssetEditorSubsystem.cpp b/Engine/Source/Editor/UnrealEd/Private/Subsystems/AssetEditorSubsystem.cpp
--- a/Engine/Source/Editor/UnrealEd/Private/Subsystems/AssetEditorSubsystem.cpp
+++ b/Engine/Source/Editor/UnrealEd/Private/Subsystems/AssetEditorSubsystem.cpp
@@ -647,7 +647,7 @@ bool UAssetEditorSubsystem::OpenEditorForAssets_Advanced(const TArray <UObject*>
 			// remove them from the list of assets to open an editor for
 			UPackage* Package = Asset->GetOutermost();
 			FText ErrorMessage;
-			if (FindEditorForAsset(Asset, true) != nullptr || CanOpenEditorForAsset(Asset, OpenedMethod, &ErrorMessage))
+			if (FindEditorForAsset(Asset, true) != nullptr || !CanOpenEditorForAsset(Asset, OpenedMethod, &ErrorMessage))
 			{
 				SkipOpenAssets.Add(Asset);
 			}
