Print shader compilation stats after the CompileShadersCommandlet runs.

- Turn Log into Display so they are visible

diff --git a/Engine/Source/Editor/UnrealEd/Private/Commandlets/CompileShadersTestBedCommandlet.cpp b/Engine/Source/Editor/UnrealEd/Private/Commandlets/CompileShadersTestBedCommandlet.cpp
--- a/Engine/Source/Editor/UnrealEd/Private/Commandlets/CompileShadersTestBedCommandlet.cpp
+++ b/Engine/Source/Editor/UnrealEd/Private/Commandlets/CompileShadersTestBedCommandlet.cpp
@@ -36,14 +36,14 @@ int32 UCompileShadersTestBedCommandlet::Main(const FString& Params)
 	// Display help
 	if (Switches.Contains("help"))
 	{
-		UE_LOG(LogCompileShadersTestBedCommandlet, Log, TEXT("CompileShadersTestBed"));
-		UE_LOG(LogCompileShadersTestBedCommandlet, Log, TEXT("This commandlet compiles global and default material shaders.  Used to profile and test shader compilation."));
-		UE_LOG(LogCompileShadersTestBedCommandlet, Log, TEXT(" Optional: -collection=<name>                (You can also specify a collection of assets to narrow down the results e.g. if you maintain a collection that represents the actually used in-game assets)."));
-		UE_LOG(LogCompileShadersTestBedCommandlet, Log, TEXT(" Optional: -materials=<path1>+<path2>        (You can also specify a list of material asset paths separated by a '+' to narrow down the results.)"));
-		UE_LOG(LogCompileShadersTestBedCommandlet, Log, TEXT(" Optional: -all                              (You can specify -all to compile all global/default shaders as well as shaders for all materials/material instances found in a project.)"))
-		UE_LOG(LogCompileShadersTestBedCommandlet, Log, TEXT(" Optional: -ExcludeGlobalShaders             (Skip the compilation of global shaders.)"))
-		UE_LOG(LogCompileShadersTestBedCommandlet, Log, TEXT(" Optional: -ExcludeDefaultMaterials          (Skip the compilation of default material shaders.)"))
-		UE_LOG(LogCompileShadersTestBedCommandlet, Log, TEXT(" Optional: -ExcludeMaterials                 (Skip the compilation of non default material shaders.)"))
+		UE_LOG(LogCompileShadersTestBedCommandlet, Display, TEXT("CompileShadersTestBed"));
+		UE_LOG(LogCompileShadersTestBedCommandlet, Display, TEXT("This commandlet compiles global and default material shaders.  Used to profile and test shader compilation."));
+		UE_LOG(LogCompileShadersTestBedCommandlet, Display, TEXT(" Optional: -collection=<name>                (You can also specify a collection of assets to narrow down the results e.g. if you maintain a collection that represents the actually used in-game assets)."));
+		UE_LOG(LogCompileShadersTestBedCommandlet, Display, TEXT(" Optional: -materials=<path1>+<path2>        (You can also specify a list of material asset paths separated by a '+' to narrow down the results.)"));
+		UE_LOG(LogCompileShadersTestBedCommandlet, Display, TEXT(" Optional: -all                              (You can specify -all to compile all global/default shaders as well as shaders for all materials/material instances found in a project.)"))
+		UE_LOG(LogCompileShadersTestBedCommandlet, Display, TEXT(" Optional: -ExcludeGlobalShaders             (Skip the compilation of global shaders.)"))
+		UE_LOG(LogCompileShadersTestBedCommandlet, Display, TEXT(" Optional: -ExcludeDefaultMaterials          (Skip the compilation of default material shaders.)"))
+		UE_LOG(LogCompileShadersTestBedCommandlet, Display, TEXT(" Optional: -ExcludeMaterials                 (Skip the compilation of non default material shaders.)"))
 		return 0;
 	}
 
@@ -61,7 +61,7 @@ int32 UCompileShadersTestBedCommandlet::Main(const FString& Params)
 		Filter.ClassPaths.Add(UMaterial::StaticClass()->GetClassPathName());
 		Filter.ClassPaths.Add(UMaterialInstance::StaticClass()->GetClassPathName());
 		AssetRegistry.GetAssets(Filter, MaterialList);
-		UE_LOG(LogCompileShadersTestBedCommandlet, Log, TEXT("Found %d materials/material instances in project."), MaterialList.Num());
+		UE_LOG(LogCompileShadersTestBedCommandlet, Display, TEXT("Found %d materials/material instances in project."), MaterialList.Num());
 	}
 	else // otherwise parse -collection and -materials arguments
 	{
@@ -81,7 +81,7 @@ int32 UCompileShadersTestBedCommandlet::Main(const FString& Params)
 				CollectionManagerModule.Get().GetObjectsInCollection(FName(*CollectionName), ECollectionShareType::CST_All, Filter.SoftObjectPaths, ECollectionRecursionFlags::SelfAndChildren);
 
 				AssetRegistry.GetAssets(Filter, MaterialList);
-				UE_LOG(LogCompileShadersTestBedCommandlet, Log, TEXT("Found %d materials/material instances from collection %s."), MaterialList.Num(), *CollectionName);
+				UE_LOG(LogCompileShadersTestBedCommandlet, Display, TEXT("Found %d materials/material instances from collection %s."), MaterialList.Num(), *CollectionName);
 			}
 		}
 		// Process -materials= switches separated by a '+'
@@ -108,7 +108,7 @@ int32 UCompileShadersTestBedCommandlet::Main(const FString& Params)
 			}
 
 			AssetRegistry.GetAssets(Filter, MaterialList);
-			UE_LOG(LogCompileShadersTestBedCommandlet, Log, TEXT("Found %d/%d requested materials/material instances."), MaterialList.Num() - MaterialsNumBefore, Filter.SoftObjectPaths.Num());
+			UE_LOG(LogCompileShadersTestBedCommandlet, Display, TEXT("Found %d/%d requested materials/material instances."), MaterialList.Num() - MaterialsNumBefore, Filter.SoftObjectPaths.Num());
 		}
 	}
 
@@ -120,12 +120,14 @@ int32 UCompileShadersTestBedCommandlet::Main(const FString& Params)
 
 	for (ITargetPlatform* Platform : Platforms)
 	{
+		UE_LOG(LogCompileShadersTestBedCommandlet, Display, TEXT("Compiling shaders for %s..."), *Platform->PlatformName());
+
 		// Compile default materials
 		if (!Switches.Contains(TEXT("ExcludeDefaultMaterials")))
 		{
 			TRACE_CPUPROFILER_EVENT_SCOPE(DefaultMaterials);
 
-			UE_LOG(LogCompileShadersTestBedCommandlet, Log, TEXT("Compile default materials"));
+			UE_LOG(LogCompileShadersTestBedCommandlet, Display, TEXT("Compile default materials"));
 
 			for (int32 Domain = 0; Domain < MD_MAX; ++Domain)
 			{
@@ -148,7 +150,7 @@ int32 UCompileShadersTestBedCommandlet::Main(const FString& Params)
 			TArray<FName> DesiredShaderFormats;
 			Platform->GetAllTargetedShaderFormats(DesiredShaderFormats);
 
-			UE_LOG(LogCompileShadersTestBedCommandlet, Log, TEXT("Compile global shaders"));
+			UE_LOG(LogCompileShadersTestBedCommandlet, Display, TEXT("Compile global shaders"));
 
 			for (int32 FormatIndex = 0; FormatIndex < DesiredShaderFormats.Num(); FormatIndex++)
 			{
@@ -169,7 +171,7 @@ int32 UCompileShadersTestBedCommandlet::Main(const FString& Params)
 			// Sort the material lists by name so the order is stable.
 			Algo::SortBy(MaterialList, [](const FAssetData& AssetData) { return AssetData.GetSoftObjectPath(); }, [](const FSoftObjectPath& A, const FSoftObjectPath& B) { return A.LexicalLess(B); });
 
-			UE_LOG(LogCompileShadersTestBedCommandlet, Log, TEXT("Begin Cache For Cooked PlatformData"));
+			UE_LOG(LogCompileShadersTestBedCommandlet, Display, TEXT("Begin Cache For Cooked PlatformData"));
 
 			for (const FAssetData& AssetData : MaterialList)
 			{
@@ -195,7 +197,7 @@ int32 UCompileShadersTestBedCommandlet::Main(const FString& Params)
 		{
 			TRACE_CPUPROFILER_EVENT_SCOPE(SubmitJobs);
 
-			UE_LOG(LogCompileShadersTestBedCommandlet, Log, TEXT("Submit Jobs"));
+			UE_LOG(LogCompileShadersTestBedCommandlet, Display, TEXT("Submit Jobs"));
 
 			while (MaterialsToCompile.Num())
 			{
@@ -214,7 +216,7 @@ int32 UCompileShadersTestBedCommandlet::Main(const FString& Params)
 					const int32 CurrentOutstandingJobs = GShaderCompilingManager->GetNumOutstandingJobs();
 					if (CurrentOutstandingJobs != PreviousOutstandingJobs)
 					{
-						UE_LOG(LogCompileShadersTestBedCommandlet, Log, TEXT("Outstanding Jobs: %d"), CurrentOutstandingJobs);
+						UE_LOG(LogCompileShadersTestBedCommandlet, Display, TEXT("Outstanding Jobs: %d"), CurrentOutstandingJobs);
 						PreviousOutstandingJobs = CurrentOutstandingJobs;
 					}
 
@@ -264,5 +266,7 @@ int32 UCompileShadersTestBedCommandlet::Main(const FString& Params)
 		}
 	}
 
+	GShaderCompilingManager->PrintStats();
+
 	return 0;
 }
