Cooker DiffOnly: When comparing diff headers, compare names case sensitively since they differences in capitalization will create diffs that need to be explained.
#rnx

diff --git a/Engine/Source/Editor/UnrealEd/Private/Cooker/DiffWriterArchive.cpp b/Engine/Source/Editor/UnrealEd/Private/Cooker/DiffWriterArchive.cpp
--- a/Engine/Source/Editor/UnrealEd/Private/Cooker/DiffWriterArchive.cpp
+++ b/Engine/Source/Editor/UnrealEd/Private/Cooker/DiffWriterArchive.cpp
@@ -1471,22 +1471,32 @@ void DumpPackageHeaderDiffs_ZenPackage(
 		return;
 	}
 
-	TArray<FName> SourceNames;
-	TArray<FName> DestNames;
+	TArray<FString> SourceNames;
+	TArray<FString> DestNames;
 	for (FDisplayNameEntryId Id : SourceHeader.GetPackageHeader().NameMap)
 	{
-		SourceNames.Add(Id.ToName(0));
+		SourceNames.Add(Id.ToName(0).ToString());
 	}
 	for (FDisplayNameEntryId Id : DestHeader.GetPackageHeader().NameMap)
 	{
-		DestNames.Add(Id.ToName(0));
+		DestNames.Add(Id.ToName(0).ToString());
 	}
-	Algo::Sort(SourceNames, FNameFastLess());
-	Algo::Sort(DestNames, FNameFastLess());
+	auto StringSortByNoCaseThenCase = [](const FString& A, const FString& B)
+	{
+		int32 NoCase = A.Compare(B, ESearchCase::IgnoreCase);
+		if (NoCase != 0)
+		{
+			return NoCase < 0;
+		}
+		int32 Case = A.Compare(B, ESearchCase::CaseSensitive);
+		return Case < 0;
+	};
+	Algo::Sort(SourceNames, StringSortByNoCaseThenCase);
+	Algo::Sort(DestNames, StringSortByNoCaseThenCase);
 
-	if (SourceNames != DestNames)
+	if (!SourceHeader.IsNameMapIdentical(DestHeader, SourceNames, DestNames))
 	{
-		DumpTableDifferences<FName>(SourceHeader, DestHeader, SourceNames, DestNames,
+		DumpTableDifferences<FString>(SourceHeader, DestHeader, SourceNames, DestNames,
 			*AssetFilename, TEXT("Name"), MaxDiffsToLog);
 	}
 
diff --git a/Engine/Source/Editor/UnrealEd/Private/Cooker/DiffWriterZenHeader.cpp b/Engine/Source/Editor/UnrealEd/Private/Cooker/DiffWriterZenHeader.cpp
--- a/Engine/Source/Editor/UnrealEd/Private/Cooker/DiffWriterZenHeader.cpp
+++ b/Engine/Source/Editor/UnrealEd/Private/Cooker/DiffWriterZenHeader.cpp
@@ -113,20 +113,38 @@ FStringView FDiffWriterZenHeader::GetObjectIndexPathName(FPackageObjectIndex Pac
 	return *PathName;
 }
 
-FString FDiffWriterZenHeader::GetTableKey(FName Id)
+bool FDiffWriterZenHeader::IsNameMapIdentical(FDiffWriterZenHeader& DestContext,
+	const TArray<FString>& SourceNames, const TArray<FString>& DestNames)
 {
-	return Id.ToString();
+	int32 NumNames = SourceNames.Num();
+	if (NumNames != DestNames.Num())
+	{
+		return false;
+	}
+	for (int32 Index = 0; Index < NumNames; ++Index)
+	{
+		if (SourceNames[Index].Compare(DestNames[Index], ESearchCase::CaseSensitive) != 0)
+		{
+			return false;
+		}
+	}
+	return true;
+}
+
+FString FDiffWriterZenHeader::GetTableKey(const FString& Id)
+{
+	return Id;
 }
 
 bool FDiffWriterZenHeader::CompareTableItem(FDiffWriterZenHeader& DestContext,
-	FName SourceName, FName DestName)
+	const FString& SourceName, const FString& DestName)
 {
-	return SourceName == DestName;
+	return SourceName.Compare(DestName, ESearchCase::CaseSensitive) == 0;
 }
 
-FString FDiffWriterZenHeader::ConvertItemToText(FName Id)
+FString FDiffWriterZenHeader::ConvertItemToText(const FString& Id)
 {
-	return Id.ToString();
+	return Id;
 }
 
 bool FDiffWriterZenHeader::IsImportMapIdentical(FDiffWriterZenHeader& DestContext)
diff --git a/Engine/Source/Editor/UnrealEd/Private/Cooker/DiffWriterZenHeader.h b/Engine/Source/Editor/UnrealEd/Private/Cooker/DiffWriterZenHeader.h
--- a/Engine/Source/Editor/UnrealEd/Private/Cooker/DiffWriterZenHeader.h
+++ b/Engine/Source/Editor/UnrealEd/Private/Cooker/DiffWriterZenHeader.h
@@ -53,9 +53,11 @@ public:
 	FStringView GetObjectIndexPathName(FPackageObjectIndex PackageObjectIndex);
 
 	// API for DumpTableDifferences
-	FString GetTableKey(FName Id);
-	bool CompareTableItem(FDiffWriterZenHeader& DestContext, FName SourceName, FName DestName);
-	FString ConvertItemToText(FName Id);
+	bool IsNameMapIdentical(FDiffWriterZenHeader& DestContext,
+		const TArray<FString>& SourceNames, const TArray<FString>& DestNames);
+	FString GetTableKey(const FString& Id);
+	bool CompareTableItem(FDiffWriterZenHeader& DestContext, const FString& SourceText, const FString& DestText);
+	FString ConvertItemToText(const FString& Id);
 
 	bool IsImportMapIdentical(FDiffWriterZenHeader& DestContext);
 	FString GetTableKey(FPackageObjectIndex Id);
