Apply padding around the button row in confirmation dialogs in places where it had been removed by a previous fix to SWindow.

diff --git a/Engine/Plugins/Animation/ControlRig/Source/ControlRigEditor/Private/Editor/SRigSpacePickerWidget.cpp b/Engine/Plugins/Animation/ControlRig/Source/ControlRigEditor/Private/Editor/SRigSpacePickerWidget.cpp
--- a/Engine/Plugins/Animation/ControlRig/Source/ControlRigEditor/Private/Editor/SRigSpacePickerWidget.cpp
+++ b/Engine/Plugins/Animation/ControlRig/Source/ControlRigEditor/Private/Editor/SRigSpacePickerWidget.cpp
@@ -1202,7 +1202,7 @@ void SRigSpacePickerBakeWidget::Construct(const FArguments& InArgs)
 
 			+SVerticalBox::Slot()
 			.AutoHeight()
-			.Padding(0.f, 16.f, 0.f, 0.f)
+			.Padding(0.f, 16.f, 0.f, 16.f)
 			[
 				SNew(SHorizontalBox)
 
diff --git a/Engine/Plugins/Animation/ControlRig/Source/ControlRigEditor/Private/Sequencer/ControlRigParameterTrackEditor.cpp b/Engine/Plugins/Animation/ControlRig/Source/ControlRigEditor/Private/Sequencer/ControlRigParameterTrackEditor.cpp
--- a/Engine/Plugins/Animation/ControlRig/Source/ControlRigEditor/Private/Sequencer/ControlRigParameterTrackEditor.cpp
+++ b/Engine/Plugins/Animation/ControlRig/Source/ControlRigEditor/Private/Sequencer/ControlRigParameterTrackEditor.cpp
@@ -4416,13 +4416,13 @@ void SCollapseControlsWidget::Construct(const FArguments& InArgs)
 				SNew(SVerticalBox)
 				+ SVerticalBox::Slot()
 				.AutoHeight()
-				.Padding(0.f, 8.f, 0.f, 0.f)
+				.Padding(0.f)
 				[
 					DetailsView->GetWidget().ToSharedRef()
 				]
 				+ SVerticalBox::Slot()
 				.AutoHeight()
-				.Padding(FMargin(0.0f, 0.f))
+				.Padding(16.f)
 				[
 
 					SNew(SHorizontalBox)
diff --git a/Engine/Plugins/Compositing/Composure/Source/ComposureLayersEditor/Private/Widgets/SCompElementPickerWindow.cpp b/Engine/Plugins/Compositing/Composure/Source/ComposureLayersEditor/Private/Widgets/SCompElementPickerWindow.cpp
--- a/Engine/Plugins/Compositing/Composure/Source/ComposureLayersEditor/Private/Widgets/SCompElementPickerWindow.cpp
+++ b/Engine/Plugins/Compositing/Composure/Source/ComposureLayersEditor/Private/Widgets/SCompElementPickerWindow.cpp
@@ -267,7 +267,7 @@ TSharedRef<SWidget> SCompElementColorPickerDialog::GeneratePreviewContent()
 		.AutoHeight()
 		.HAlign(HAlign_Right)
 		.VAlign(VAlign_Bottom)
-		.Padding(FMargin(8.f, 8.f, 8.f, 0.f))
+		.Padding(8.f)
 	[
 		SNew(SHorizontalBox)
 		+SHorizontalBox::Slot()
diff --git a/Engine/Plugins/Editor/WaveformEditor/Source/WaveformEditorWidgets/Private/SWaveformEditorMessageDialog.cpp b/Engine/Plugins/Editor/WaveformEditor/Source/WaveformEditorWidgets/Private/SWaveformEditorMessageDialog.cpp
--- a/Engine/Plugins/Editor/WaveformEditor/Source/WaveformEditorWidgets/Private/SWaveformEditorMessageDialog.cpp
+++ b/Engine/Plugins/Editor/WaveformEditor/Source/WaveformEditorWidgets/Private/SWaveformEditorMessageDialog.cpp
@@ -34,7 +34,7 @@ void SWaveformEditorMessageDialog::Construct(const FArguments& InArgs)
 				]
 				+ SVerticalBox::Slot()
 				.VAlign(VAlign_Bottom)
-				.Padding(8.0f, 0.f, 8.0f, 16.0f)
+				.Padding(8.0f, 16.f, 8.0f, 16.0f)
 				[
 					SNew(SHorizontalBox)
 					+ SHorizontalBox::Slot()
diff --git a/Engine/Plugins/Importers/USDImporter/Source/USDClassesEditor/Private/USDDefaultAssetCacheDialog.cpp b/Engine/Plugins/Importers/USDImporter/Source/USDClassesEditor/Private/USDDefaultAssetCacheDialog.cpp
--- a/Engine/Plugins/Importers/USDImporter/Source/USDClassesEditor/Private/USDDefaultAssetCacheDialog.cpp
+++ b/Engine/Plugins/Importers/USDImporter/Source/USDClassesEditor/Private/USDDefaultAssetCacheDialog.cpp
@@ -145,6 +145,7 @@ void SUsdDefaultAssetCacheDialog::Construct( const FArguments& InArgs )
 			]
 
 			+SVerticalBox::Slot()
+			.Padding(0.f, 32.f, 0.f, 0.f)
 			.AutoHeight()
 			[
 				SNew(SHorizontalBox)
@@ -153,7 +154,7 @@ void SUsdDefaultAssetCacheDialog::Construct( const FArguments& InArgs )
 				.FillWidth(1.0f)
 				.HAlign(HAlign_Left)
 				.VAlign(VAlign_Center)
-				.Padding(FMargin(16.f, 0.f, 0.f, 0.f))
+				.Padding(0)
 				[
 					SNew(SCheckBox)
 					.IsChecked(ECheckBoxState::Unchecked)
@@ -176,9 +177,10 @@ void SUsdDefaultAssetCacheDialog::Construct( const FArguments& InArgs )
 				.AutoWidth()
 				.HAlign(HAlign_Right)
 				.VAlign(VAlign_Center)
+				.Padding(16.f, 0.f, 0.f, 0.f)
 				[
 					SNew( SUniformGridPanel )
-					.SlotPadding( 2 )
+					.SlotPadding(FAppStyle::Get().GetMargin("StandardDialog.SlotPadding"))
 					.MinDesiredSlotWidth(FAppStyle::Get().GetFloat("StandardDialog.MinDesiredSlotWidth"))
 					.MinDesiredSlotHeight(FAppStyle::Get().GetFloat("StandardDialog.MinDesiredSlotHeight"))
 
diff --git a/Engine/Plugins/MovieScene/LevelSequenceEditor/Source/LevelSequenceEditor/Private/LevelSequenceEditorSubsystem.cpp b/Engine/Plugins/MovieScene/LevelSequenceEditor/Source/LevelSequenceEditor/Private/LevelSequenceEditorSubsystem.cpp
--- a/Engine/Plugins/MovieScene/LevelSequenceEditor/Source/LevelSequenceEditor/Private/LevelSequenceEditorSubsystem.cpp
+++ b/Engine/Plugins/MovieScene/LevelSequenceEditor/Source/LevelSequenceEditor/Private/LevelSequenceEditorSubsystem.cpp
@@ -741,7 +741,7 @@ void SBakeTransformWidget::Construct(const FArguments& InArgs)
 
 			+ SVerticalBox::Slot()
 			.AutoHeight()
-			.Padding(0.f, 16.f, 0.f, 0.f)
+			.Padding(0.f, 16.f, 0.f, 16.f)
 			[
 				SNew(SHorizontalBox)
 
diff --git a/Engine/Source/Developer/ToolWidgets/Private/Dialog/SMessageDialog.cpp b/Engine/Source/Developer/ToolWidgets/Private/Dialog/SMessageDialog.cpp
--- a/Engine/Source/Developer/ToolWidgets/Private/Dialog/SMessageDialog.cpp
+++ b/Engine/Source/Developer/ToolWidgets/Private/Dialog/SMessageDialog.cpp
@@ -36,7 +36,7 @@ void SMessageDialog::Construct(const FArguments& InArgs)
 		.ContentAreaPadding(FMargin(16.f, 0.f, 0.f, 0.f))
 		.UseScrollBox(InArgs._UseScrollBox)
 		.ScrollBoxMaxHeight(InArgs._ScrollBoxMaxHeight)
-		.ButtonAreaPadding(0.f)
+		.ButtonAreaPadding(FMargin(0.f, 32.f, 0.f, 0.f))
 		.OnClosed(InArgs._OnClosed)
 		.BeforeButtons()
 		[
diff --git a/Engine/Source/Developer/ToolWidgets/Public/Dialog/SCustomDialog.h b/Engine/Source/Developer/ToolWidgets/Public/Dialog/SCustomDialog.h
--- a/Engine/Source/Developer/ToolWidgets/Public/Dialog/SCustomDialog.h
+++ b/Engine/Source/Developer/ToolWidgets/Public/Dialog/SCustomDialog.h
@@ -65,7 +65,7 @@ public:
 		, _HAlignIcon(HAlign_Left)
 		, _VAlignIcon(VAlign_Center)
 		, _RootPadding(FMargin(4.f))
-		, _ButtonAreaPadding(FMargin(20.f, 0.f, 0.f, 0.f))
+		, _ButtonAreaPadding(FMargin(20.f, 16.f, 0.f, 0.f))
 		, _UseScrollBox(true)
 		, _ScrollBoxMaxHeight(300)
 		, _HAlignContent(HAlign_Left)
@@ -121,7 +121,7 @@ public:
 		/** Padding to apply to the widget embedded in the window, i.e. to all widgets contained in the window (default: {4,4,4,4} )*/
 		SLATE_ATTRIBUTE(FMargin, RootPadding)
 
-		/** Padding to apply around the layout holding the buttons (default: {20,0,0,0}) */
+		/** Padding to apply around the layout holding the buttons (default: {20,16,0,0}) */
 		SLATE_ATTRIBUTE(FMargin, ButtonAreaPadding)
 
 		/** Padding to apply to DialogContent - you can use it to move away from the icon (default: {0,0,0,0}) */
diff --git a/Engine/Source/Editor/UnrealEd/Private/Dialogs/Dialogs.cpp b/Engine/Source/Editor/UnrealEd/Private/Dialogs/Dialogs.cpp
--- a/Engine/Source/Editor/UnrealEd/Private/Dialogs/Dialogs.cpp
+++ b/Engine/Source/Editor/UnrealEd/Private/Dialogs/Dialogs.cpp
@@ -117,6 +117,7 @@ public:
 				]
 
 				+SVerticalBox::Slot()
+				.Padding(0.f, 32.f, 0.f, 0.f)
 				.AutoHeight()
 				[
 					SNew(SHorizontalBox)
@@ -163,7 +164,7 @@ public:
 					.VAlign(VAlign_Center)
 					[
 						SAssignNew( ButtonBox, SUniformGridPanel )
-						.SlotPadding(FMargin(16.f, 0.f, 0.f, 0.f))
+						.SlotPadding(FMargin(8.f, 0.f, 0.f, 0.f))
 						.MinDesiredSlotWidth(FAppStyle::Get().GetFloat("StandardDialog.MinDesiredSlotWidth"))
 						.MinDesiredSlotHeight(FAppStyle::Get().GetFloat("StandardDialog.MinDesiredSlotHeight"))
 					]
@@ -630,6 +631,7 @@ public:
 		ChildSlot
 		[
 			SNew(SBorder)
+			.Padding(16.f)
 			.BorderImage(FAppStyle::GetBrush("ToolPanel.GroupBorder"))
 			[
 				SNew( SVerticalBox )
@@ -649,7 +651,6 @@ public:
 						+SHorizontalBox::Slot()
 						.AutoWidth()
 						.VAlign(VAlign_Center)
-						.Padding(16) // currently hardcoded until we adjust StandardDialog.SlotPadding
 						[
 							SNew( SImage )
 							.Image(InArgs._Image)
@@ -659,7 +660,7 @@ public:
 						+SHorizontalBox::Slot()
 						.AutoWidth()
 						.VAlign(VAlign_Center)
-						.Padding(0)
+						.Padding(FMargin(16.f, 0.f, 0.f, 0.f))
 						[
 							SNew( STextBlock )
 							.WrapTextAt(512.0f)
@@ -672,6 +673,7 @@ public:
 				+SVerticalBox::Slot()
 				.AutoHeight()
 				.HAlign(HAlign_Fill)
+				.Padding(FMargin(0.f, 32.f, 0.f, 0.f))
 				[
 					ConstructConditionalInternals(InArgs)
 				]
@@ -693,7 +695,6 @@ public:
 		// checkbox with user specified text
 		HorizontalBox->AddSlot()
 		.HAlign(HAlign_Left)
-		.Padding(16) // currently hardcoded until we adjust StandardDialog.SlotPadding
 		.FillWidth(1.0)
 		[
 			SNew(SCheckBox)
@@ -708,7 +709,7 @@ public:
 		];
 		HorizontalBox->AddSlot()
 		.HAlign(HAlign_Right)
-		.Padding(16) // currently hardcoded until we adjust StandardDialog.SlotPadding
+		.Padding(FMargin(16.f, 0.f, 0.f, 0.f)) // currently hardcoded until we adjust StandardDialog.SlotPadding
 		.AutoWidth()
 		[
 			SNew(SPrimaryButton)
@@ -724,7 +725,7 @@ public:
 			// cancel/stop/abort button
 			HorizontalBox->AddSlot()
 			.HAlign(HAlign_Right)
-			.Padding(0, 16, 16, 16) // currently hardcoded until we adjust StandardDialog.SlotPadding
+			.Padding(FMargin(8.f, 0.f, 0.f, 0.f)) // currently hardcoded until we adjust StandardDialog.SlotPadding
 			.AutoWidth()
 			[
 				SNew( SButton )
@@ -1025,7 +1026,7 @@ void SGenericDialogWidget::Construct( const FArguments& InArgs )
 		+SVerticalBox::Slot()
 		.HAlign(HAlign_Right)
 		.AutoHeight()
-		.Padding(0.0f, 2.0f, 0.0f, 0.0f)
+		.Padding(0.f, 2.f, 0.f, 0.f)
 		[
 			SNew(SButton)
 			.Text( NSLOCTEXT("UnrealEd", "OK", "OK") )
diff --git a/Engine/Source/Programs/SlateUGS/Private/Widgets/SScheduledSyncWindow.cpp b/Engine/Source/Programs/SlateUGS/Private/Widgets/SScheduledSyncWindow.cpp
--- a/Engine/Source/Programs/SlateUGS/Private/Widgets/SScheduledSyncWindow.cpp
+++ b/Engine/Source/Programs/SlateUGS/Private/Widgets/SScheduledSyncWindow.cpp
@@ -96,7 +96,7 @@ void SScheduledSyncWindow::Construct(const FArguments& InArgs, UGSTab* InTab)
 		// Buttons
 		+SVerticalBox::Slot()
 		.AutoHeight()
-		.Padding(20.0f, 10.0f, 20.0f, 0.0f)
+		.Padding(20.0f, 10.0f, 20.0f, 20.0f)
 		.VAlign(VAlign_Bottom)
 		[
 			SNew(SHorizontalBox)
diff --git a/Engine/Source/Programs/ZenDashboard/Private/SMessageDialog.cpp b/Engine/Source/Programs/ZenDashboard/Private/SMessageDialog.cpp
--- a/Engine/Source/Programs/ZenDashboard/Private/SMessageDialog.cpp
+++ b/Engine/Source/Programs/ZenDashboard/Private/SMessageDialog.cpp
@@ -106,6 +106,7 @@ public:
 				]
 
 				+SVerticalBox::Slot()
+				.Padding(FMargin(0.f, 32.f, 0.f, 0.f))
 				.AutoHeight()
 				[
 					SNew(SHorizontalBox)
diff --git a/Engine/Source/Runtime/AppFramework/Private/Widgets/Colors/SColorPicker.cpp b/Engine/Source/Runtime/AppFramework/Private/Widgets/Colors/SColorPicker.cpp
--- a/Engine/Source/Runtime/AppFramework/Private/Widgets/Colors/SColorPicker.cpp
+++ b/Engine/Source/Runtime/AppFramework/Private/Widgets/Colors/SColorPicker.cpp
@@ -491,7 +491,6 @@ void SColorPicker::GenerateDefaultColorPickerContent( bool bAdvancedSectionExpan
 		+ SVerticalBox::Slot()
 			.AutoHeight()
 			.HAlign(HAlign_Right)
-			.VAlign(VAlign_Center)
 			.Padding(0.0f, 12.0f, 0.0f, 0.0f)
 			[
 				SNew(SUniformGridPanel)
