[BugFix] Fix diff crash when diffed objects are GCed

diff --git a/Engine/Source/Editor/UnrealEd/Private/DiffUtils.cpp b/Engine/Source/Editor/UnrealEd/Private/DiffUtils.cpp
--- a/Engine/Source/Editor/UnrealEd/Private/DiffUtils.cpp
+++ b/Engine/Source/Editor/UnrealEd/Private/DiffUtils.cpp
@@ -1111,7 +1111,13 @@ bool DiffUtils::Identical(const TSharedPtr<IPropertyHandle>& PropertyHandleA, co
 		const void* ValueB = ValuesB[IndexB];
 
 		const UObject* OwningOuterA = OwningOutersA.IsEmpty() ? HandleOutersA[IndexA] : OwningOutersA[IndexA].Get();
-		const UObject* OwningOuterB = OwningOutersB.IsEmpty() ? HandleOutersB[IndexB] : OwningOutersB[IndexB].Get();;
+		const UObject* OwningOuterB = OwningOutersB.IsEmpty() ? HandleOutersB[IndexB] : OwningOutersB[IndexB].Get();
+
+		if (!OwningOuterA || !OwningOuterB)
+		{
+			// objects were Garbage Collected!
+			return !OwningOuterA && !OwningOuterB;
+		}
 
 		// note that we're not directly calling FProperty::Identical because sub-object properties should be weakly compared based on
 		// their paths instead of their pointers or data
