Support for rich text messaging in SMessageDialog and utility for converting between message categories and Slate brushes

diff --git a/Engine/Source/Developer/ToolWidgets/Private/Dialog/DialogUtils.cpp b/Engine/Source/Developer/ToolWidgets/Private/Dialog/DialogUtils.cpp
new file mode 100644
--- /dev/null
+++ b/Engine/Source/Developer/ToolWidgets/Private/Dialog/DialogUtils.cpp
@@ -0,0 +1,56 @@
+// Copyright Epic Games, Inc. All Rights Reserved.
+
+#include "Dialog/DialogUtils.h"
+
+#include "GenericPlatform/GenericPlatformMisc.h"
+#include "Internationalization/Internationalization.h"
+#include "Styling/AppStyle.h"
+
+const FDialogButtonTexts& FDialogButtonTexts::Get()
+{
+	static FDialogButtonTexts ButtonTexts;
+	return ButtonTexts;
+}
+
+FDialogButtonTexts::FDialogButtonTexts()
+	: No(NSLOCTEXT("Dialogs", "EAppReturnTypeNo", "No"))
+	, Yes(NSLOCTEXT("Dialogs", "EAppReturnTypeYes", "Yes"))
+	, YesAll(NSLOCTEXT("Dialogs", "EAppReturnTypeYesAll", "Yes All"))
+	, NoAll(NSLOCTEXT("Dialogs", "EAppReturnTypeNoAll", "No All"))
+	, Cancel(NSLOCTEXT("Dialogs", "EAppReturnTypeCancel", "Cancel"))
+	, Ok(NSLOCTEXT("Dialogs", "EAppReturnTypeOk", "OK"))
+	, Retry(NSLOCTEXT("Dialogs", "EAppReturnTypeRetry", "Retry"))
+	, Continue(NSLOCTEXT("Dialogs", "EAppReturnTypeContinue", "Continue"))
+{
+}
+
+const FSlateBrush* FDialogUtils::GetMessageCategoryIcon(const EAppMsgCategory MessageCategory)
+{
+	FName MessageCategoryBrushName;
+
+	switch (MessageCategory)
+	{
+		case EAppMsgCategory::Info:
+		{
+			MessageCategoryBrushName = "Icons.InfoWithColor.Large";
+			break;
+		}
+		case EAppMsgCategory::Warning:
+		{
+			MessageCategoryBrushName = "Icons.WarningWithColor.Large";
+			break;
+		}
+		case EAppMsgCategory::Error:
+		{
+			MessageCategoryBrushName = "Icons.ErrorWithColor.Large";
+			break;
+		}
+		case EAppMsgCategory::Success:
+		{
+			MessageCategoryBrushName = "Icons.SuccessWithColor.Large";
+			break;
+		}
+	}
+
+	return !MessageCategoryBrushName.IsNone() ? FAppStyle::Get().GetBrush(MessageCategoryBrushName) : nullptr;
+}
diff --git a/Engine/Source/Developer/ToolWidgets/Private/Dialog/SMessageDialog.cpp b/Engine/Source/Developer/ToolWidgets/Private/Dialog/SMessageDialog.cpp
--- a/Engine/Source/Developer/ToolWidgets/Private/Dialog/SMessageDialog.cpp
+++ b/Engine/Source/Developer/ToolWidgets/Private/Dialog/SMessageDialog.cpp
@@ -6,7 +6,7 @@
 
 #include "Styling/AppStyle.h"
 #include "Styling/SlateBrush.h"
-#include "Widgets/Text/STextBlock.h"
+#include "Widgets/Text/SRichTextBlock.h"
 #include "Widgets/Input/SButton.h"
 #include "Widgets/Images/SImage.h"
 
@@ -14,14 +14,16 @@ void SMessageDialog::Construct(const FArguments& InArgs)
 {
 	Message = InArgs._Message;
 	
+	TSharedPtr<SRichTextBlock> RichTextBlock;
+
 	SCustomDialog::Construct(SCustomDialog::FArguments()
 		.Title(InArgs._Title)
 		.Content()
 		[
-			SNew(STextBlock)
-			.Text(InArgs._Message)
-			.Font(FAppStyle::Get().GetFontStyle("StandardDialog.LargeFont"))
+			SAssignNew(RichTextBlock, SRichTextBlock)
+			.Text(Message)
 			.WrapTextAt(InArgs._WrapMessageAt)
+			.Decorators(InArgs._Decorators)
 		]
 		.WindowArguments(InArgs._WindowArguments)
 		.RootPadding(16.f)
@@ -52,7 +54,12 @@ void SMessageDialog::Construct(const FArguments& InArgs)
 				.ColorAndOpacity(FSlateColor::UseForeground())
 			]
 		]
-		);
+	);
+
+	if (InArgs._DecoratorStyleSet)
+	{
+		RichTextBlock->SetDecoratorStyleSet(InArgs._DecoratorStyleSet);
+	}
 }
 
 FReply SMessageDialog::OnKeyDown(const FGeometry& MyGeometry, const FKeyEvent& InKeyEvent)
diff --git a/Engine/Source/Developer/ToolWidgets/Public/Dialog/DialogUtils.h b/Engine/Source/Developer/ToolWidgets/Public/Dialog/DialogUtils.h
new file mode 100644
--- /dev/null
+++ b/Engine/Source/Developer/ToolWidgets/Public/Dialog/DialogUtils.h
@@ -0,0 +1,34 @@
+// Copyright Epic Games, Inc. All Rights Reserved.
+
+#include "Internationalization/Text.h"
+
+enum class EAppMsgCategory : uint8;
+struct FSlateBrush;
+
+struct FDialogButtonTexts
+{
+	static TOOLWIDGETS_API const FDialogButtonTexts& Get();
+
+	/** Common dialog button texts */
+	const FText No;
+	const FText Yes;
+	const FText YesAll;
+	const FText NoAll;
+	const FText Cancel;
+	const FText Ok;
+	const FText Retry;
+	const FText Continue;
+
+private:
+	FDialogButtonTexts(const FDialogButtonTexts&) = delete;
+	FDialogButtonTexts& operator=(const FDialogButtonTexts&) = delete;
+
+	FDialogButtonTexts();
+};
+
+class TOOLWIDGETS_API FDialogUtils
+{
+public:
+	/** Returns the Slate brush corresponding to the provided message category */
+	static const FSlateBrush* GetMessageCategoryIcon(const EAppMsgCategory MessageCategory);
+};
diff --git a/Engine/Source/Developer/ToolWidgets/Public/Dialog/SMessageDialog.h b/Engine/Source/Developer/ToolWidgets/Public/Dialog/SMessageDialog.h
--- a/Engine/Source/Developer/ToolWidgets/Public/Dialog/SMessageDialog.h
+++ b/Engine/Source/Developer/ToolWidgets/Public/Dialog/SMessageDialog.h
@@ -41,6 +41,12 @@ public:
 		/** Whether to automatically close this window when any button is pressed (default: true) */
 		SLATE_ARGUMENT(bool, AutoCloseOnButtonPress)
 	
+		/** Text decorators used while parsing the rich text messages */
+		SLATE_ARGUMENT(TArray<TSharedRef<class ITextDecorator>>, Decorators)
+
+		/** Style set used to look up styles used by decorators for rich text messages */
+		SLATE_ARGUMENT(const ISlateStyle*, DecoratorStyleSet)
+		
 		/********** Cosmetic **********/
 	
 		/** Optional icon to display in the dialog. (default: empty) */
diff --git a/Engine/Source/Editor/UnrealEd/Private/Dialogs/Dialogs.cpp b/Engine/Source/Editor/UnrealEd/Private/Dialogs/Dialogs.cpp
--- a/Engine/Source/Editor/UnrealEd/Private/Dialogs/Dialogs.cpp
+++ b/Engine/Source/Editor/UnrealEd/Private/Dialogs/Dialogs.cpp
@@ -2,6 +2,8 @@
 
 
 #include "Dialogs/Dialogs.h"
+
+#include "Dialog/DialogUtils.h"
 #include "Dialogs/DialogsPrivate.h"
 #include "Misc/App.h"
 #include "Misc/AssertionMacros.h"
@@ -59,21 +61,7 @@ public:
 
 		TSharedPtr<SUniformGridPanel> ButtonBox;
 
-		const FSlateBrush* IconBrush = FAppStyle::Get().GetBrush("Icons.WarningWithColor.Large");
-		switch (InArgs._MessageCategory)
-		{
-		case EAppMsgCategory::Error:
-			IconBrush = FAppStyle::Get().GetBrush("Icons.ErrorWithColor.Large");
-			break;
-		case EAppMsgCategory::Success:
-			IconBrush = FAppStyle::Get().GetBrush("Icons.SuccessWithColor.Large");
-			break;
-		case EAppMsgCategory::Info:
-			IconBrush = FAppStyle::Get().GetBrush("Icons.InfoWithColor.Large");
-			break;
-		default:
-			break;
-		}
+		const FSlateBrush* IconBrush = FDialogUtils::GetMessageCategoryIcon(InArgs._MessageCategory);
 
 		this->ChildSlot
 		[	
@@ -279,21 +267,21 @@ public:
 		switch(ReturnType)
 		{
 		case EAppReturnType::No:
-			return LOCTEXT("EAppReturnTypeNo", "No");
+			return FDialogButtonTexts::Get().No;
 		case EAppReturnType::Yes:
-			return LOCTEXT("EAppReturnTypeYes", "Yes");
+			return FDialogButtonTexts::Get().Yes;
 		case EAppReturnType::YesAll:
-			return LOCTEXT("EAppReturnTypeYesAll", "Yes All");
+			return FDialogButtonTexts::Get().YesAll;
 		case EAppReturnType::NoAll:
-			return LOCTEXT("EAppReturnTypeNoAll", "No All");
+			return FDialogButtonTexts::Get().NoAll;
 		case EAppReturnType::Cancel:
-			return LOCTEXT("EAppReturnTypeCancel", "Cancel");
+			return FDialogButtonTexts::Get().Cancel;
 		case EAppReturnType::Ok:
-			return LOCTEXT("EAppReturnTypeOk", "OK");
+			return FDialogButtonTexts::Get().Ok;
 		case EAppReturnType::Retry:
-			return LOCTEXT("EAppReturnTypeRetry", "Retry");
+			return FDialogButtonTexts::Get().Retry;
 		case EAppReturnType::Continue:
-			return LOCTEXT("EAppReturnTypeContinue", "Continue");
+			return FDialogButtonTexts::Get().Continue;
 		default:
 			return LOCTEXT("MissingType", "MISSING RETURN TYPE");
 		}
