MPCook: Fix CookWorker commandlines to have the correct project dir when cooking projects in external directories (outside of UnrealEngine root) .
[RN] Bugfix, Foundation

diff --git a/Engine/Source/Editor/UnrealEd/Private/Cooker/CookDirector.cpp b/Engine/Source/Editor/UnrealEd/Private/Cooker/CookDirector.cpp
--- a/Engine/Source/Editor/UnrealEd/Private/Cooker/CookDirector.cpp
+++ b/Engine/Source/Editor/UnrealEd/Private/Cooker/CookDirector.cpp
@@ -1228,50 +1228,62 @@ FString FCookDirector::GetWorkerLogFileName(int32 ProfileId)
 
 FString FCookDirector::GetWorkerCommandLine(FWorkerId WorkerId, int32 ProfileId)
 {
-	FString CommandLine = FCommandLine::Get();
+	const TCHAR* CommandLine = FCommandLine::Get();
 
 	const TCHAR* ProjectName = FApp::GetProjectName();
 	checkf(ProjectName && ProjectName[0], TEXT("Expected UnrealEditor to be running with a non-empty project name"));
+
+	// Note that we need to handle quoted strings for e.g. a projectfile with spaces in it; FParse::Token does handle them
+	FString Token;
 	TArray<FString> Tokens;
-	UE::String::ParseTokensMultiple(CommandLine, { ' ', '\t', '\r', '\n' }, [&Tokens, ProfileId](FStringView Token)
-		{
-			if (Token.StartsWith(TEXT("-run=")) ||
-				Token == TEXT("-CookOnTheFly") ||
-				Token == TEXT("-CookWorker") ||
-				Token.StartsWith(TEXT("-CookCultures")) ||
-				Token.StartsWith(TEXT("-CookDirectorHost=")) ||
-				Token.StartsWith(TEXT("-MultiprocessId=")) ||
-				Token.StartsWith(TEXT("-CookProfileId=")) ||
-				Token.StartsWith(TEXT("-ShowCookWorker")) ||
-				Token.StartsWith(TEXT("-CoreLimit")) ||
-				Token.StartsWith(TEXT("-PhysicalCoreLimit")) ||
-				Token.StartsWith(TEXT("-CookProcessCount=")) ||
-				Token.StartsWith(TEXT("-abslog=")) ||
-				Token.StartsWith(TEXT("-unattended"))
-				)
-			{
-				return;
-			}
-			else if (Token.StartsWith(TEXT("-tracefile=")))
+	while (FParse::Token(CommandLine, Token, false /* bUseEscape */))
+	{
+		if (Token.IsEmpty())
+		{
+			continue;
+		}
+		if (Token.StartsWith(TEXT("-run=")) ||
+			Token == TEXT("-CookOnTheFly") ||
+			Token == TEXT("-CookWorker") ||
+			Token.StartsWith(TEXT("-CookCultures")) ||
+			Token.StartsWith(TEXT("-CookDirectorHost=")) ||
+			Token.StartsWith(TEXT("-MultiprocessId=")) ||
+			Token.StartsWith(TEXT("-CookProfileId=")) ||
+			Token.StartsWith(TEXT("-ShowCookWorker")) ||
+			Token.StartsWith(TEXT("-CoreLimit")) ||
+			Token.StartsWith(TEXT("-PhysicalCoreLimit")) ||
+			Token.StartsWith(TEXT("-CookProcessCount=")) ||
+			Token.StartsWith(TEXT("-abslog=")) ||
+			Token.StartsWith(TEXT("-unattended"))
+			)
+		{
+			continue;
+		}
+		else if (Token.StartsWith(TEXT("-tracefile=")))
+		{
+			FString TraceFile;
+			FString TokenString(Token);
+			if (FParse::Value(*TokenString, TEXT("-tracefile="), TraceFile) && !TraceFile.IsEmpty())
 			{
-				FString TraceFile;
-				FString TokenString(Token);
-				if (FParse::Value(*TokenString, TEXT("-tracefile="), TraceFile) && !TraceFile.IsEmpty())
-				{
-					FStringView BaseFilenameWithPath = FPathViews::GetBaseFilenameWithPath(TraceFile);
-					FStringView Extension = FPathViews::GetExtension(TraceFile, true /* bIncludeDot */);
-					Tokens.Add(FString::Printf(TEXT("-tracefile=\"%.*s_Worker%d%.*s\""),
-						BaseFilenameWithPath.Len(), BaseFilenameWithPath.GetData(),
-						ProfileId,
-						Extension.Len(), Extension.GetData()));
-					return;
-				}
+				FStringView BaseFilenameWithPath = FPathViews::GetBaseFilenameWithPath(TraceFile);
+				FStringView Extension = FPathViews::GetExtension(TraceFile, true /* bIncludeDot */);
+				Tokens.Add(FString::Printf(TEXT("-tracefile=\"%.*s_Worker%d%.*s\""),
+					BaseFilenameWithPath.Len(), BaseFilenameWithPath.GetData(),
+					ProfileId,
+					Extension.Len(), Extension.GetData()));
+				continue;
 			}
-			Tokens.Add(FString(Token));
-		}, UE::String::EParseTokensOptions::SkipEmpty);
-	if (Tokens[0] != ProjectName)
+		}
+		Tokens.Add(MoveTemp(Token));
+	}
+
+	if (Tokens[0] != ProjectName && !Tokens[0].EndsWith(TEXT(".uproject"), ESearchCase::IgnoreCase))
 	{
-		Tokens.Insert(ProjectName, 0);
+		FString ProjectFilePath = FPaths::GetProjectFilePath();
+		if (!FPaths::IsSamePath(Tokens[0], ProjectFilePath))
+		{
+			Tokens.Insert(ProjectFilePath, 0);
+		}
 	}
 	Tokens.Insert(TEXT("-run=cook"), 1);
 	Tokens.Insert(TEXT("-cookworker"), 2);
@@ -1286,6 +1298,19 @@ FString FCookDirector::GetWorkerCommandLine(FWorkerId WorkerId, int32 ProfileId)
 		Tokens.Add(FString::Printf(TEXT("-PhysicalCoreLimit=%d"), CoreLimit));
 	}
 
+	// We are joining the tokens back into a commandline string; wrap tokens with whitespace in quotes
+	for (FString& IterToken : Tokens)
+	{
+		int32 IndexOfWhitespace = UE::String::FindFirstOfAnyChar(IterToken, { ' ', '\r', '\n' });
+		if (IndexOfWhitespace != INDEX_NONE)
+		{
+			int32 IndexOfQuote;
+			if (!IterToken.FindChar('\"', IndexOfQuote))
+			{
+				IterToken = FString::Printf(TEXT("\"%s\""), *IterToken);
+			}
+		}
+	}
 	return FString::Join(Tokens, TEXT(" "));
 }
 
