LevelEditor, SceneGraph: Make the placement mode actor positioning code allow hits of non primitive components, to allow placement on top of entities (that have a physics component). This will be enabled by default but can be reverted to previous behavior via the PlacementMode.AllowNonPrimitiveComponentHits CVar.

diff --git a/Engine/Source/Editor/UnrealEd/Private/Editor/ActorPositioning.cpp b/Engine/Source/Editor/UnrealEd/Private/Editor/ActorPositioning.cpp
--- a/Engine/Source/Editor/UnrealEd/Private/Editor/ActorPositioning.cpp
+++ b/Engine/Source/Editor/UnrealEd/Private/Editor/ActorPositioning.cpp
@@ -13,12 +13,22 @@
 #include "Editor.h"
 #include "ActorFactories/ActorFactory.h"
 #include "EditorViewportClient.h"
+#include "HAL/IConsoleManager.h" // FAutoConsoleVariableRef
 #include "LevelEditorViewport.h"
 #include "SnappingUtils.h"
 #include "LandscapeHeightfieldCollisionComponent.h"
 #include "LandscapeComponent.h"
 #include "Editor/EditorPerProjectUserSettings.h"
 
+namespace ActorPositioningLocals
+{
+	bool bAllowNonPrimitiveComponentHits = true;
+	static FAutoConsoleVariableRef CVarAllowNonPrimitiveComponentHits(
+		TEXT("PlacementMode.AllowNonPrimitiveComponentHits"),
+		bAllowNonPrimitiveComponentHits,
+		TEXT("When raycasting the world in placement mode, allow hits of physics objects that are not tied to a UPrimitiveComponent (to work with non-actor workflows)."));
+}
+
 FActorPositionTraceResult FActorPositioning::TraceWorldForPositionWithDefault(const FViewportCursorLocation& Cursor, const FSceneView& View, const TArray<AActor*>* IgnoreActors)
 {
 	FActorPositionTraceResult Results = TraceWorldForPosition(Cursor, View, IgnoreActors);
@@ -73,11 +83,18 @@ bool IsHitIgnored(const FHitResult& InHit, const FSceneView& InSceneView)
 		PrimitiveComponent = CastChecked<ULandscapeHeightfieldCollisionComponent>(PrimitiveComponent)->GetRenderComponent();
 	}
 
-	if (InHit.bStartPenetrating || !PrimitiveComponent)
+	if (InHit.bStartPenetrating)
 	{
 		return true;
 	}
 
+	if (!PrimitiveComponent)
+	{
+		// If we don't have a primitive component, either ignore the hit, or pass it through if the CVar is set appropriately.
+		// i.e. ignoring is the inverse of the "allow non primitive hits" CVar.
+		return !ActorPositioningLocals::bAllowNonPrimitiveComponentHits;
+	}
+
 	// Ignore volumes and shapes
 	if (HitObjHandle.DoesRepresentClass(AVolume::StaticClass()))
 	{
