LevelEditor, SceneGraph: Don't make unknown hitproxy types prevent drops of objects in level editor- just use same fallback as elsewhere. Prevents entities from cancelling drop operations on top of themselves.

diff --git a/Engine/Source/Editor/UnrealEd/Private/LevelEditorViewport.cpp b/Engine/Source/Editor/UnrealEd/Private/LevelEditorViewport.cpp
--- a/Engine/Source/Editor/UnrealEd/Private/LevelEditorViewport.cpp
+++ b/Engine/Source/Editor/UnrealEd/Private/LevelEditorViewport.cpp
@@ -1641,6 +1641,37 @@ bool FLevelEditorViewportClient::DropObjectsAtCoordinates(int32 MouseX, int32 Mo
 			// Axis translation/rotation/scale widget - find out what's underneath the axis widget
 			bResult = DropObjectsOnWidget(View, Cursor, DroppedObjects);
 		}
+		// If the hit proxy was not one of the above types, we probably still don't want to cancel the drop.
+		else
+		{
+			// Ideally we would probably use some interface to get the information we need to confirm the drop. However, the fallback
+			// drop we use in the case of actors doesn't really need extra information. Instead, we'll just use the presence of a
+			// ITypedElementWorldInterface interface to mean that the hitproxy corresponds to something real in the world.
+			auto DoesElementHaveWorldInterface = [](const FTypedElementHandle& Handle) -> bool
+			{
+				UTypedElementRegistry* Registry = UTypedElementRegistry::GetInstance();
+				if (!Handle || !Registry)
+				{
+					return false;
+				}
+				return !!Registry->GetElementInterface<ITypedElementWorldInterface>(Handle);
+			};
+
+			if (DoesElementHaveWorldInterface(HitProxy->GetElementHandle()))
+			{
+				// This is the same thing that we do in other "drop objects" functions like FLevelEditorViewportClient::DropObjectsOnBackground
+				// and FLevelEditorViewportClient::DropObjectsOnActor (when not doing some kind of special "apply to actor" operation)
+				for (UObject* DroppedObject : DroppedObjects)
+				{
+					TArray<AActor*> NewActors = TryPlacingActorFromObject(GetWorld()->GetCurrentLevel(), DroppedObject, true, ObjectFlags, FactoryToUse, NAME_None, &Cursor);
+					if (NewActors.Num() > 0)
+					{
+						OutNewActors.Append(NewActors);
+						bResult = true;
+					}
+				}
+			}
+		}
 
 		if ( bResult )
 		{
