Ability to cook multiple plugins at once with the Cook DLCs flow.

#rnx

diff --git a/Engine/Source/Editor/UnrealEd/Classes/CookOnTheSide/CookOnTheFlyServer.h b/Engine/Source/Editor/UnrealEd/Classes/CookOnTheSide/CookOnTheFlyServer.h
--- a/Engine/Source/Editor/UnrealEd/Classes/CookOnTheSide/CookOnTheFlyServer.h
+++ b/Engine/Source/Editor/UnrealEd/Classes/CookOnTheSide/CookOnTheFlyServer.h
@@ -1259,8 +1259,16 @@ private:
 
 	UNREALED_API FString GetMountedAssetPathForDLC() const;
 
+	static FString GetMountedAssetPathForPlugin(const FString& InPluginName);
+
 	FString GetMetadataDirectory() const;
 
+	/**
+	 * In a recook build, gets the list of plugins that need to be recooked, which includes DLCName if cooking a DLC
+	 * and all the plugins specified by the -CookPlugins command line parameter.
+	 */
+	void GetPluginsToRecook(TSet<FString>& OutPlugins) const;
+
 	/**
 	 * Is the local CookOnTheFlyServer cooking a Project+Engine+EmbeddedPlugin Release that can be used as a
 	 * base for future DLC or Patches?
diff --git a/Engine/Source/Editor/UnrealEd/Private/CookOnTheFlyServer.cpp b/Engine/Source/Editor/UnrealEd/Private/CookOnTheFlyServer.cpp
--- a/Engine/Source/Editor/UnrealEd/Private/CookOnTheFlyServer.cpp
+++ b/Engine/Source/Editor/UnrealEd/Private/CookOnTheFlyServer.cpp
@@ -1212,13 +1212,18 @@ FString UCookOnTheFlyServer::GetBaseDirectoryForDLC() const
 
 FString UCookOnTheFlyServer::GetMountedAssetPathForDLC() const
 {
-	TSharedPtr<IPlugin> Plugin = IPluginManager::Get().FindPlugin(CookByTheBookOptions->DlcName);
+	return GetMountedAssetPathForPlugin(CookByTheBookOptions->DlcName);
+}
+
+FString UCookOnTheFlyServer::GetMountedAssetPathForPlugin(const FString& InPluginName)
+{
+	TSharedPtr<IPlugin> Plugin = IPluginManager::Get().FindPlugin(InPluginName);
 	if (Plugin.IsValid())
 	{
 		return Plugin->GetMountedAssetPath();
 	}
 
-	return FString::Printf(TEXT("/%s/"), *CookByTheBookOptions->DlcName);
+	return FString::Printf(TEXT("/%s/"), *InPluginName);
 }
 
 FString UCookOnTheFlyServer::GetContentDirectoryForDLC() const
@@ -1232,6 +1237,26 @@ FString UCookOnTheFlyServer::GetMetadataDirectory() const
 	return ProjectOrPluginRoot / TEXT("Metadata");
 }
 
+void UCookOnTheFlyServer::GetPluginsToRecook(TSet<FString>& OutPlugins) const
+{
+	OutPlugins.Empty();
+
+	// DLCName, if cooking a DLC.
+	if (IsCookingDLC())
+	{
+		OutPlugins.Add(CookByTheBookOptions->DlcName);
+	}
+
+	// Command line parameter -CookPlugins.
+	TArray<FString> PluginsList;
+	FString CookPluginsStr;
+	if (FParse::Value(FCommandLine::Get(), TEXT("CookPlugins="), CookPluginsStr, false))
+	{
+		CookPluginsStr.ParseIntoArray(PluginsList, TEXT(","));
+	}
+	OutPlugins.Append(PluginsList);
+}
+
 // allow for a command line to start async preloading a Development AssetRegistry if requested
 static FEventRef GPreloadAREvent(EEventMode::ManualReset);
 static FEventRef GPreloadARInfoEvent(EEventMode::ManualReset);
@@ -11553,11 +11578,37 @@ void UCookOnTheFlyServer::RecordDLCPackagesFromBaseGame(FBeginCookContext& Begin
 
 		TArray<FName>& PlatformBasedPackages = CookByTheBookOptions->BasedOnReleaseCookedPackages.FindOrAdd(PlatformName);
 		PlatformBasedPackages.Reset(ActivePackageList.Num());
-		FString PluginPathToSkip(!!(CookOptions & ECookByTheBookOptions::DlcRecook) ? GetMountedAssetPathForDLC() : "");
+
+		// Packages that are present in PlatformBasedPackages will be stripped from the AssetRegistry and pak files generated for this cook.
+		// If we are recooking plugins, make sure packages that live in those plugins are *not* added to PlatformBasedPackages to prevent them from being stripped.
+		TArray<FString> PathsToSkip;
+		if (!!(CookOptions & ECookByTheBookOptions::DlcRecook))
+		{
+			TSet<FString> Plugins;
+			GetPluginsToRecook(Plugins);
+			for (const FString& Plugin : Plugins)
+			{
+				PathsToSkip.Add(GetMountedAssetPathForPlugin(Plugin));
+			}
+		}
+
 		for (UE::Cook::FConstructPackageData& PackageData : ActivePackageList)
 		{
-			// If we are recooking a DLC, skip adding its assets to 'BasedOnReleaseCookedPackages' so that they don't get ignored in the final stages of the cook.
-			if ((PluginPathToSkip.Len() == 0) || !PackageData.PackageName.ToString().StartsWith(PluginPathToSkip))
+			bool bShouldSkip = false;
+			if (PathsToSkip.Num() > 0)
+			{
+				const FString PackageNameString = PackageData.PackageName.ToString();
+				for (const FString& Path : PathsToSkip)
+				{
+					if (PackageNameString.StartsWith(*Path))
+					{
+						bShouldSkip = true;
+						break;
+					}
+				}
+			}
+
+			if (!bShouldSkip)
 			{
 				PlatformBasedPackages.Add(PackageData.NormalizedFileName);
 			}
@@ -11586,10 +11637,15 @@ void UCookOnTheFlyServer::RecordDLCPackagesFromBaseGame(FBeginCookContext& Begin
 		}
 	}
 
-	if (IsCookingDLC() && (!!(CookOptions & ECookByTheBookOptions::DlcRecook)))
+	if (!!(CookOptions & ECookByTheBookOptions::DlcRecook))
 	{
-		// Mark all the packages in the DLC as not cooked to force them to be cooked again.
-		PackageDatas->ClearCookResultsForPlugin(CookByTheBookOptions->DlcName);
+		TSet<FString> Plugins;
+		GetPluginsToRecook(Plugins);
+		for (const FString& Plugin : Plugins)
+		{
+			// Mark all the packages in the plugin as not cooked to force them to be cooked again.
+			PackageDatas->ClearCookResultsForPlugin(Plugin);
+		}
 	}
 }
 
