Cooker: Optimize parsing of shortpackagenames on the commandline: use an AssetRegistry query instead of a directory search.
#rnx

diff --git a/Engine/Source/Editor/UnrealEd/Private/CookOnTheFlyServer.cpp b/Engine/Source/Editor/UnrealEd/Private/CookOnTheFlyServer.cpp
--- a/Engine/Source/Editor/UnrealEd/Private/CookOnTheFlyServer.cpp
+++ b/Engine/Source/Editor/UnrealEd/Private/CookOnTheFlyServer.cpp
@@ -7024,11 +7024,23 @@ bool UCookOnTheFlyServer::Exec_Editor(class UWorld* InWorld, const TCHAR* Cmd, F
 
 		if (FPackageName::IsShortPackageName(PackageName))
 		{
-			FString OutFilename;
-			if (FPackageName::SearchForPackageOnDisk(PackageName, NULL, &OutFilename))
+			TArray<FName> LongPackageNames;
+			AssetRegistry->GetPackagesByName(PackageName, LongPackageNames);
+			if (LongPackageNames.IsEmpty())
 			{
-				PackageName = OutFilename;
+				Ar.Logf(TEXT("No package found with leaf name %s."), *PackageName);
+				return true;
+			}
+			if (LongPackageNames.Num() > 1)
+			{
+				Ar.Logf(TEXT("Multiple packages found with leaf name %s. Specify the full LongPackageName."), *PackageName);
+				for (FName LongPackageName : LongPackageNames)
+				{
+					Ar.Logf(TEXT("\n\t%s"), *LongPackageName.ToString());
+				}
+				return true;
 			}
+			PackageName = LongPackageNames[0].ToString();
 		}
 
 		FName RawPackageName(*PackageName);
@@ -8477,13 +8489,21 @@ void UCookOnTheFlyServer::GenerateLongPackageNames(TArray<FName>& FilesInPath, T
 		{
 			FString LongPackageName;
 			FPackageName::EErrorCode FailureReason;
-			if (FPackageName::TryConvertToMountedPath(FileInPath, nullptr /* LocalPath */, &LongPackageName,
+			bool bFound = FPackageName::TryConvertToMountedPath(FileInPath, nullptr /* LocalPath */, &LongPackageName,
 				nullptr /* ObjectName */, nullptr /* SubObjectName */, nullptr /* Extension */,
-				nullptr /* FlexNameType */, &FailureReason)
-				||
-				(FPackageName::IsShortPackageName(FileInPath) &&
-					FPackageName::SearchForPackageOnDisk(FileInPath, &LongPackageName, nullptr))
-				)
+				nullptr /* FlexNameType */, &FailureReason);
+			if (!bFound && FPackageName::IsShortPackageName(FileInPath))
+			{
+				TArray<FName> LongPackageNames;
+				AssetRegistry->GetPackagesByName(FileInPath, LongPackageNames);
+				if (LongPackageNames.Num() == 1)
+				{
+					bFound = true;
+					LongPackageName = LongPackageNames[0].ToString();
+				}
+			}
+
+			if (bFound)
 			{
 				const FName LongPackageFName(*LongPackageName);
 				bool bIsAlreadyAdded;
@@ -8760,14 +8780,35 @@ void UCookOnTheFlyServer::CollectFilesToCook(TArray<FName>& FilesInPath, TMap<FN
 		UE_SCOPED_HIERARCHICAL_COOKTIMER(SearchForPackageOnDisk);
 		if (FPackageName::IsShortPackageName(CurrEntry))
 		{
-			FString OutFilename;
-			if (FPackageName::SearchForPackageOnDisk(CurrEntry, NULL, &OutFilename) == false)
+			TArray<FName> LongPackageNames;
+			AssetRegistry->GetPackagesByName(CurrEntry, LongPackageNames);
+			if (LongPackageNames.IsEmpty())
 			{
-				LogCookerMessage( FString::Printf(TEXT("Unable to find package for map %s."), *CurrEntry), EMessageSeverity::Warning);
+				LogCookerMessage(FString::Printf(TEXT("Unable to find package for map %s."), *CurrEntry), EMessageSeverity::Warning);
+			}
+			else if (LongPackageNames.Num() > 1)
+			{
+				constexpr int32 MaxMessageLen = 256;
+				TStringBuilder<256> Message;
+				Message.Appendf(TEXT("Multiple packages found for map %s; it will not be added. Specify the full LongPackageName. Packages found:"), *CurrEntry);
+				for (FName LongPackageName : LongPackageNames)
+				{
+					Message << TEXT("\n\t");
+					if (Message.Len() >= MaxMessageLen)
+					{
+						Message << TEXT("...");
+						break;
+					}
+					else
+					{
+						Message << LongPackageName;
+					}
+				}
+				LogCookerMessage(FString(*Message), EMessageSeverity::Warning);
 			}
 			else
 			{
-				AddFileToCook(FilesInPath, Instigators, OutFilename, EInstigator::CommandLinePackage);
+				AddFileToCook(FilesInPath, Instigators, LongPackageNames[0].ToString(), EInstigator::CommandLinePackage);
 			}
 		}
 		else
