Cooker: SkipOnlyEditorOnly and MPCook: Fix non-determinstic errors in the determiniation of whether packages are onlyeditoronly referenced.
1) When a worker queues an unsolicited package for discovery on the director because it was an unexpected load, the director needs to retest whether it was an unexpected load, since the worker does not have all discovery information.
2) When files are marked as EInstigator::SaveTimeSoftDependency by a FCookLoadScope, they should be added for cooking even if they were found to be an only-editor-only in the graph search of assetregistry dependencies.
#rnx

diff --git a/Engine/Source/Editor/UnrealEd/Private/CookOnTheFlyServer.cpp b/Engine/Source/Editor/UnrealEd/Private/CookOnTheFlyServer.cpp
--- a/Engine/Source/Editor/UnrealEd/Private/CookOnTheFlyServer.cpp
+++ b/Engine/Source/Editor/UnrealEd/Private/CookOnTheFlyServer.cpp
@@ -4232,6 +4232,12 @@ void UCookOnTheFlyServer::ProcessUnsolicitedPackages(TArray<FName>* OutDiscovere
 			Instigator.Category = EInstigator::ForceExplorableSaveTimeSoftDependency;
 			QueueDiscoveredPackage(*PackageData, FInstigator(Instigator), EDiscoveredPlatformSet::CopyFromInstigator);
 		}
+		else if (Instigator.Category != EInstigator::Unsolicited)
+		{
+			// A load undeclared in AssetRegistry dependencies, but one that has been marked up as a known issue and needed at runtime.
+			// Queue it for discovery and do not log a warning about it.
+			QueueDiscoveredPackage(*PackageData, FInstigator(Instigator), EDiscoveredPlatformSet::CopyFromInstigator);
+		}
 		else if (PackageData->FindOrAddPlatformData(CookerLoadingPlatformKey).IsReachable())
 		{
 			// This load was expected so we do not need to add a hidden dependency for it.
@@ -4241,12 +4247,6 @@ void UCookOnTheFlyServer::ProcessUnsolicitedPackages(TArray<FName>* OutDiscovere
 				QueueDiscoveredPackage(*PackageData, FInstigator(Instigator), EDiscoveredPlatformSet::CopyFromInstigator);
 			}
 		}
-		else if (Instigator.Category == EInstigator::SaveTimeHardDependency || Instigator.Category == EInstigator::SaveTimeSoftDependency)
-		{
-			// A load undeclared in AssetRegistry dependencies, but one that has been marked up as a known issue and needed at runtime.
-			// Queue it for discovery and do not log a warning about it.
-			QueueDiscoveredPackage(*PackageData, FInstigator(Instigator), EDiscoveredPlatformSet::CopyFromInstigator);
-		}
 		else
 		{
 			// Possibly an unexpected load, but it might be an expected load after we finish processing some other unexpected loads that were
diff --git a/Engine/Source/Editor/UnrealEd/Private/Cooker/CookWorkerServer.cpp b/Engine/Source/Editor/UnrealEd/Private/Cooker/CookWorkerServer.cpp
--- a/Engine/Source/Editor/UnrealEd/Private/Cooker/CookWorkerServer.cpp
+++ b/Engine/Source/Editor/UnrealEd/Private/Cooker/CookWorkerServer.cpp
@@ -901,7 +901,23 @@ void FCookWorkerServer::QueueDiscoveredPackage(FDiscoveredPackageReplication&& D
 	if (Instigator.Category != EInstigator::ForceExplorableSaveTimeSoftDependency &&
 		PackageData.HasReachablePlatforms(DiscoveredPlatforms))
 	{
-		// The CookWorker thought this was a new package, but the Director already knows about it; ignore the report
+		// The CookWorker thought there were some new reachable platforms, but the Director already knows about
+		// all of them; ignore the report
+		return;
+	}
+	if (COTFS.bSkipOnlyEditorOnly &&
+		Instigator.Category == EInstigator::Unsolicited &&
+		Platforms.GetSource() == EDiscoveredPlatformSet::CopyFromInstigator &&
+		PackageData.FindOrAddPlatformData(CookerLoadingPlatformKey).IsReachable())
+	{
+		// The CookWorker thought this package was new (previously unreachable even by editoronly references),
+		// and it is not marked as a known used-in-game or editor-only issue, so it fell back to reporting it
+		// as used-in-game-because-its-not-a-known-issue (see UCookOnTheFlyServer::ProcessUnsolicitedPackages's
+		// use of PackageData->FindOrAddPlatformData(CookerLoadingPlatformKey).IsReachable()).
+		// But we only do that fall back for unexpected packages not found by the search of editor-only AssetRegistry
+		// dependencies. And this package was found by that search; the director has already marked it as reachable by
+		// editoronly references. Correct the heuristic: ignore the unmarked load because the load is expected as an
+		// editor-only reference.
 		return;
 	}
 
