Report the number of shader type and vertex factory type permutations to cook analytics.

- Aggregate all the shader type permutations and report them as `Shaders_NumShaderTypePermutations`
- Aggregate all vertex factory type permutations and report them as `Shaders_NumVertexFactoryTypePermutations`

diff --git a/Engine/Source/Editor/UnrealEd/Private/ShaderStatsCollector.cpp b/Engine/Source/Editor/UnrealEd/Private/ShaderStatsCollector.cpp
--- a/Engine/Source/Editor/UnrealEd/Private/ShaderStatsCollector.cpp
+++ b/Engine/Source/Editor/UnrealEd/Private/ShaderStatsCollector.cpp
@@ -2,6 +2,7 @@
 
 #include "ShaderStatsCollector.h"
 #include "ShaderStats.h"
+#include "AnalyticsEventAttribute.h"
 
 static TAutoConsoleVariable<int> CVarShaderCompilerStatsPrintoutInterval(
 	TEXT("r.ShaderCompiler.StatsPrintoutInterval"),
@@ -50,6 +51,20 @@ void FShaderStatsFunctions::GatherShaderAnalytics(TArray<FAnalyticsEventAttribut
 	GShaderCompilingManager->GetLocalStats(AggregatedCompilerStats);
 	GShaderStatsReporter.AggregateStats(AggregatedCompilerStats);
 	AggregatedCompilerStats.GatherAnalytics(TEXT("Shaders_"), Attributes);
+
+	int32 TotalShaderTypePermutations = 0;
+	for (TLinkedList<FShaderType*>::TIterator ShaderTypeIt(FShaderType::GetTypeList()); ShaderTypeIt; ShaderTypeIt.Next())
+	{
+		TotalShaderTypePermutations += ShaderTypeIt->GetPermutationCount();
+	}
+	Attributes.Emplace(TEXT("Shaders_NumShaderTypePermutations"), TotalShaderTypePermutations);
+
+	int32 TotalVFTypePermutations = 0;
+	for (TLinkedList<FVertexFactoryType*>::TIterator VFTypeIt(FVertexFactoryType::GetTypeList()); VFTypeIt; VFTypeIt.Next())
+	{
+		TotalVFTypePermutations += 1;
+	}
+	Attributes.Emplace(TEXT("Shaders_NumVertexFactoryTypePermutations"), TotalVFTypePermutations);
 }
 
 void FShaderStatsFunctions::WriteShaderStats()
