Hiding Material Root node inline controls behind a console variable

To use:
MaterialGraph.EnableRootNodeInlineControls true/false

diff --git a/Engine/Source/Editor/UnrealEd/Private/MaterialGraphNode_Root.cpp b/Engine/Source/Editor/UnrealEd/Private/MaterialGraphNode_Root.cpp
--- a/Engine/Source/Editor/UnrealEd/Private/MaterialGraphNode_Root.cpp
+++ b/Engine/Source/Editor/UnrealEd/Private/MaterialGraphNode_Root.cpp
@@ -16,6 +16,13 @@
 
 #define LOCTEXT_NAMESPACE "MaterialGraphNode_Root"
 
+
+static TAutoConsoleVariable<bool> CVarEnableRootNodeInlineControls(
+	TEXT("MaterialGraph.EnableRootNodeInlineControls"),
+	false,
+	TEXT("Control if the Material Graph Root Node should show inline editing controls")
+	);
+
 /////////////////////////////////////////////////////
 // UMaterialGraphNode_Root
 
@@ -203,105 +210,113 @@ void UMaterialGraphNode_Root::CreateInputPins()
 		FName MaterialInputName = *MaterialInput.GetName().ToString();
 		FName PinSubCategory;
 		FString RefractionMethodStr;
-		switch (Property)
+		const bool ShouldEnableInlineControls = CVarEnableRootNodeInlineControls.GetValueOnAnyThread();
+		if (ShouldEnableInlineControls)
 		{
-		case MP_Metallic:
-		case MP_Specular:
-		case MP_Roughness:
-		case MP_Anisotropy:
-		case MP_Opacity:
-		case MP_OpacityMask:
-		case MP_CustomData0:
-		case MP_CustomData1:
-		case MP_AmbientOcclusion:
-		case MP_PixelDepthOffset:
-		case MP_Displacement:
-		case MP_SurfaceThickness:
-			PinSubCategory = UMaterialGraphSchema::PSC_Red;
-			break;
-			
-		case MP_Refraction:
-			PinSubCategory = UMaterialGraphSchema::PSC_Red;
-			
-			switch (MaterialGraph->Material->RefractionMethod)
+			switch (Property)
 			{
-			case ERefractionMode::RM_None:
-				RefractionMethodStr = TEXT("Disabled");
+			case MP_Metallic:
+			case MP_Specular:
+			case MP_Roughness:
+			case MP_Anisotropy:
+			case MP_Opacity:
+			case MP_OpacityMask:
+			case MP_CustomData0:
+			case MP_CustomData1:
+			case MP_AmbientOcclusion:
+			case MP_PixelDepthOffset:
+			case MP_Displacement:
+			case MP_SurfaceThickness:
+				PinSubCategory = UMaterialGraphSchema::PSC_Red;
 				break;
-			case ERefractionMode::RM_IndexOfRefraction:
-				RefractionMethodStr = TEXT("Index Of Refraction");
+				
+			case MP_Refraction:
+				PinSubCategory = UMaterialGraphSchema::PSC_Red;
+				switch (MaterialGraph->Material->RefractionMethod)
+				{
+				case ERefractionMode::RM_None:
+					RefractionMethodStr = TEXT("Disabled");
+					break;
+				case ERefractionMode::RM_IndexOfRefraction:
+					RefractionMethodStr = TEXT("Index Of Refraction");
+					break;
+				case ERefractionMode::RM_PixelNormalOffset:
+					RefractionMethodStr = TEXT("Pixel Normal Offset");
+					break;
+				case ERefractionMode::RM_2DOffset:
+					RefractionMethodStr = TEXT("2D Offset");
+					break;
+				default:
+					RefractionMethodStr = TEXT("UNKNOWN ERefractionMode");
+				}
+				
+				MaterialInputName = *FString::Printf(TEXT("%s (%s)"), *MaterialInput.GetName().ToString(), *RefractionMethodStr);
 				break;
-			case ERefractionMode::RM_PixelNormalOffset:
-				RefractionMethodStr = TEXT("Pixel Normal Offset");
+				
+			case MP_Normal:
+			case MP_Tangent:
+			case MP_WorldPositionOffset:
+				PinSubCategory = UMaterialGraphSchema::PSC_RGB;
 				break;
-			case ERefractionMode::RM_2DOffset:
-				RefractionMethodStr = TEXT("2D Offset");
+				
+			case MP_BaseColor:
+			case MP_EmissiveColor:
+			case MP_SubsurfaceColor:
+				PinSubCategory = UMaterialGraphSchema::PSC_RGBA;
 				break;
-			default:
-				RefractionMethodStr = TEXT("UNKNOWN ERefractionMode");
-			}
-			
-			MaterialInputName = *FString::Printf(TEXT("%s (%s)"), *MaterialInput.GetName().ToString(), *RefractionMethodStr);
-			break;
-		case MP_Normal:
-		case MP_Tangent:
-		case MP_WorldPositionOffset:
-			PinSubCategory = UMaterialGraphSchema::PSC_RGB;
-			break;
-		case MP_BaseColor:
-		case MP_EmissiveColor:
-		case MP_SubsurfaceColor:
-			PinSubCategory = UMaterialGraphSchema::PSC_RGBA;
-			break;
 
-		case MP_ShadingModel: // TODO FMaterialInput<uint32>
-		case MP_FrontMaterial: // TODO FMaterialInput<uint32>
-			// TODO: Not sure if we want to simply show a Num Pin UI here.
-			// Skipping these for now.
-			// If these are based on Enum types, Enum pin construction requires a UEnum object to create UI with
-			// which these constants don't have.
-			break;
-			
-		default:
-			if (Property >= MP_CustomizedUVs0 && Property <= MP_CustomizedUVs7)
-			{
-				PinSubCategory = UMaterialGraphSchema::PSC_RG;
+			case MP_ShadingModel: // TODO FMaterialInput<uint32>
+			case MP_FrontMaterial: // TODO FMaterialInput<uint32>
+				// TODO: Not sure if we want to simply show a Num Pin UI here.
+				// Skipping these for now.
+				// If these are based on Enum types, Enum pin construction requires a UEnum object to create UI with
+				// which these constants don't have.
+				break;
+				
+			default:
+				if (Property >= MP_CustomizedUVs0 && Property <= MP_CustomizedUVs7)
+				{
+					PinSubCategory = UMaterialGraphSchema::PSC_RG;
+				}
+				break;
 			}
-			break;
 		}
-
+		
 		UEdGraphPin* InputPin = CreatePin(EGPD_Input, UMaterialGraphSchema::PC_MaterialInput, PinSubCategory/*, *FString::Printf(TEXT("%d"), (int32)Property)*/, MaterialInputName);
 		InputPin->SourceIndex = Index;
 
-		switch (Property)
+		if (ShouldEnableInlineControls)
 		{
-		case MP_EmissiveColor:		InputPin->DefaultValue = EditorOnlyData->EmissiveColor.GetDefaultValue(); break;
-		case MP_Opacity:			InputPin->DefaultValue = EditorOnlyData->Opacity.GetDefaultValue(); break;
-		case MP_OpacityMask:		InputPin->DefaultValue = EditorOnlyData->OpacityMask.GetDefaultValue(); break;
-		case MP_BaseColor:			InputPin->DefaultValue = EditorOnlyData->BaseColor.GetDefaultValue(); break;
-		case MP_Metallic:			InputPin->DefaultValue = EditorOnlyData->Metallic.GetDefaultValue(); break;
-		case MP_Specular:			InputPin->DefaultValue = EditorOnlyData->Specular.GetDefaultValue(); break;
-		case MP_Roughness:			InputPin->DefaultValue = EditorOnlyData->Roughness.GetDefaultValue(); break;
-		case MP_Anisotropy:			InputPin->DefaultValue = EditorOnlyData->Anisotropy.GetDefaultValue(); break;
-		case MP_Normal:				InputPin->DefaultValue = EditorOnlyData->Normal.GetDefaultValue(); break;
-		case MP_Tangent:			InputPin->DefaultValue = EditorOnlyData->Tangent.GetDefaultValue(); break;
-		case MP_WorldPositionOffset:InputPin->DefaultValue = EditorOnlyData->WorldPositionOffset.GetDefaultValue(); break;
-		case MP_SubsurfaceColor:	InputPin->DefaultValue = EditorOnlyData->SubsurfaceColor.GetDefaultValue(); break;
-		case MP_CustomData0:		InputPin->DefaultValue = EditorOnlyData->ClearCoat.GetDefaultValue(); break;
-		case MP_CustomData1:		InputPin->DefaultValue = EditorOnlyData->ClearCoatRoughness.GetDefaultValue(); break;
-		case MP_AmbientOcclusion:	InputPin->DefaultValue = EditorOnlyData->AmbientOcclusion.GetDefaultValue(); break;
-		case MP_Refraction:			InputPin->DefaultValue = EditorOnlyData->Refraction.GetDefaultValue(); break;
-		case MP_PixelDepthOffset:	InputPin->DefaultValue = EditorOnlyData->PixelDepthOffset.GetDefaultValue(); break;
-		case MP_ShadingModel:		break; // TODO
-		case MP_FrontMaterial:		break; // TODO
-		case MP_SurfaceThickness:	InputPin->DefaultValue = EditorOnlyData->SurfaceThickness.GetDefaultValue(); break;
-		case MP_Displacement:		InputPin->DefaultValue = EditorOnlyData->Displacement.GetDefaultValue(); break;
-		default:
-			if (Property >= MP_CustomizedUVs0 && Property <= MP_CustomizedUVs7)
+			switch (Property)
 			{
-				InputPin->DefaultValue = EditorOnlyData->CustomizedUVs[Property - MP_CustomizedUVs0].GetDefaultValue(); break;
+			case MP_EmissiveColor:		InputPin->DefaultValue = EditorOnlyData->EmissiveColor.GetDefaultValue(); break;
+			case MP_Opacity:			InputPin->DefaultValue = EditorOnlyData->Opacity.GetDefaultValue(); break;
+			case MP_OpacityMask:		InputPin->DefaultValue = EditorOnlyData->OpacityMask.GetDefaultValue(); break;
+			case MP_BaseColor:			InputPin->DefaultValue = EditorOnlyData->BaseColor.GetDefaultValue(); break;
+			case MP_Metallic:			InputPin->DefaultValue = EditorOnlyData->Metallic.GetDefaultValue(); break;
+			case MP_Specular:			InputPin->DefaultValue = EditorOnlyData->Specular.GetDefaultValue(); break;
+			case MP_Roughness:			InputPin->DefaultValue = EditorOnlyData->Roughness.GetDefaultValue(); break;
+			case MP_Anisotropy:			InputPin->DefaultValue = EditorOnlyData->Anisotropy.GetDefaultValue(); break;
+			case MP_Normal:				InputPin->DefaultValue = EditorOnlyData->Normal.GetDefaultValue(); break;
+			case MP_Tangent:			InputPin->DefaultValue = EditorOnlyData->Tangent.GetDefaultValue(); break;
+			case MP_WorldPositionOffset:InputPin->DefaultValue = EditorOnlyData->WorldPositionOffset.GetDefaultValue(); break;
+			case MP_SubsurfaceColor:	InputPin->DefaultValue = EditorOnlyData->SubsurfaceColor.GetDefaultValue(); break;
+			case MP_CustomData0:		InputPin->DefaultValue = EditorOnlyData->ClearCoat.GetDefaultValue(); break;
+			case MP_CustomData1:		InputPin->DefaultValue = EditorOnlyData->ClearCoatRoughness.GetDefaultValue(); break;
+			case MP_AmbientOcclusion:	InputPin->DefaultValue = EditorOnlyData->AmbientOcclusion.GetDefaultValue(); break;
+			case MP_Refraction:			InputPin->DefaultValue = EditorOnlyData->Refraction.GetDefaultValue(); break;
+			case MP_PixelDepthOffset:	InputPin->DefaultValue = EditorOnlyData->PixelDepthOffset.GetDefaultValue(); break;
+			case MP_ShadingModel:		break; // TODO
+			case MP_FrontMaterial:		break; // TODO
+			case MP_SurfaceThickness:	InputPin->DefaultValue = EditorOnlyData->SurfaceThickness.GetDefaultValue(); break;
+			case MP_Displacement:		InputPin->DefaultValue = EditorOnlyData->Displacement.GetDefaultValue(); break;
+			default:
+				if (Property >= MP_CustomizedUVs0 && Property <= MP_CustomizedUVs7)
+				{
+					InputPin->DefaultValue = EditorOnlyData->CustomizedUVs[Property - MP_CustomizedUVs0].GetDefaultValue(); break;
+				}
+				break;
 			}
-			break;
 		}
 	}
 
