Fix Import filter for physics asset. It was using CreateClass instead of ImportClass filter.

#rnx

diff --git a/Engine/Source/Editor/UnrealEd/Private/Fbx/FbxFactory.cpp b/Engine/Source/Editor/UnrealEd/Private/Fbx/FbxFactory.cpp
--- a/Engine/Source/Editor/UnrealEd/Private/Fbx/FbxFactory.cpp
+++ b/Engine/Source/Editor/UnrealEd/Private/Fbx/FbxFactory.cpp
@@ -404,7 +404,7 @@ UObject* UFbxFactory::FactoryCreateFile
 				ImportOptions->bUpdateSkeletonReferencePose = false;
 			}
 		}
-		if (!FbxImporter->CanCreateClass(UPhysicsAsset::StaticClass()))
+		if (!FbxImporter->CanImportClass(UPhysicsAsset::StaticClass()))
 		{
 			ImportOptions->bCreatePhysicsAsset = false;
 		}
diff --git a/Engine/Source/Editor/UnrealEd/Private/Fbx/FbxMainImport.cpp b/Engine/Source/Editor/UnrealEd/Private/Fbx/FbxMainImport.cpp
--- a/Engine/Source/Editor/UnrealEd/Private/Fbx/FbxMainImport.cpp
+++ b/Engine/Source/Editor/UnrealEd/Private/Fbx/FbxMainImport.cpp
@@ -276,7 +276,7 @@ FBXImportOptions* GetImportOptions( UnFbx::FFbxImporter* FbxImporter, UFbxImport
 			}
 		}
 
-		if (!FbxImporter->CanCreateClass(UPhysicsAsset::StaticClass()))
+		if (!FbxImporter->CanImportClass(UPhysicsAsset::StaticClass()))
 		{
 			ImportUI->bCreatePhysicsAsset = false;
 		}
@@ -326,7 +326,7 @@ FBXImportOptions* GetImportOptions( UnFbx::FFbxImporter* FbxImporter, UFbxImport
 		UnFbx::FBXImportOptions::ResetOptions(ImportOptions);
 		ApplyImportUIToImportOptions(ImportUI, *ImportOptions);
 		
-		if (!FbxImporter->CanCreateClass(UPhysicsAsset::StaticClass()))
+		if (!FbxImporter->CanImportClass(UPhysicsAsset::StaticClass()))
 		{
 			ImportOptions->bCreatePhysicsAsset = false;
 		}
@@ -336,7 +336,7 @@ FBXImportOptions* GetImportOptions( UnFbx::FFbxImporter* FbxImporter, UFbxImport
 	else
 	{
 		
-		if (!FbxImporter->CanCreateClass(UPhysicsAsset::StaticClass()))
+		if (!FbxImporter->CanImportClass(UPhysicsAsset::StaticClass()))
 		{
 			FbxImporter->GetImportOptions()->bCreatePhysicsAsset = false;
 		}
diff --git a/Engine/Source/Editor/UnrealEd/Private/Fbx/FbxSkeletalMeshImport.cpp b/Engine/Source/Editor/UnrealEd/Private/Fbx/FbxSkeletalMeshImport.cpp
--- a/Engine/Source/Editor/UnrealEd/Private/Fbx/FbxSkeletalMeshImport.cpp
+++ b/Engine/Source/Editor/UnrealEd/Private/Fbx/FbxSkeletalMeshImport.cpp
@@ -2071,7 +2071,7 @@ USkeletalMesh* UnFbx::FFbxImporter::ImportSkeletalMesh(FImportSkeletalMeshArgs &
 
 		EARLY_RETURN_ON_CANCEL(false, CancelCleanup);
 
-		if ((!SkeletalMesh->GetResourceForRendering() || !SkeletalMesh->GetResourceForRendering()->LODRenderData.IsValidIndex(0)) && ImportOptions->bCreatePhysicsAsset && CanCreateClass(UPhysicsAsset::StaticClass()))
+		if ((!SkeletalMesh->GetResourceForRendering() || !SkeletalMesh->GetResourceForRendering()->LODRenderData.IsValidIndex(0)) && ImportOptions->bCreatePhysicsAsset && CanImportClass(UPhysicsAsset::StaticClass()))
 		{
 			//We need to have a valid render data to create physic asset
 			SkeletalMesh->Build();
@@ -2189,7 +2189,7 @@ USkeletalMesh* UnFbx::FFbxImporter::ImportSkeletalMesh(FImportSkeletalMeshArgs &
 		// We create the physic asset after we create the skeleton since we need the skeleton to correctly build it
 		if (ImportOptions->bCreatePhysicsAsset)
 		{
-			if (SkeletalMesh->GetPhysicsAsset() == NULL && CanCreateClass(UPhysicsAsset::StaticClass()))
+			if (SkeletalMesh->GetPhysicsAsset() == NULL && CanImportClass(UPhysicsAsset::StaticClass()))
 			{
 				FString ObjectName = FString::Printf(TEXT("%s_PhysicsAsset"), *SkeletalMesh->GetName());
 				UPhysicsAsset * NewPhysicsAsset = CreateAsset<UPhysicsAsset>(ImportSkeletalMeshArgs.InParent->GetName(), ObjectName, true);
