WorldPartitionHLODsBuilder will now also delete SpatialHashRuntimeGridInfo when deleting HLOD actors, as it is specific to WorldPArtitionRuntimeSpatialHash implementation.

diff --git a/Engine/Source/Editor/UnrealEd/Private/WorldPartition/WorldPartitionHLODsBuilder.cpp b/Engine/Source/Editor/UnrealEd/Private/WorldPartition/WorldPartitionHLODsBuilder.cpp
--- a/Engine/Source/Editor/UnrealEd/Private/WorldPartition/WorldPartitionHLODsBuilder.cpp
+++ b/Engine/Source/Editor/UnrealEd/Private/WorldPartition/WorldPartitionHLODsBuilder.cpp
@@ -618,12 +618,20 @@ bool UWorldPartitionHLODsBuilder::DeleteHLODActors()
 {
 	UE_LOG(LogWorldPartitionHLODsBuilder, Display, TEXT("#### Deleting HLOD actors ####"));
 
+	TArray<UClass*> HLODActorClasses =
+	{
+		AWorldPartitionHLOD::StaticClass(),
+		FindObject<UClass>(nullptr, TEXT("/Script/Engine.SpatialHashRuntimeGridInfo"))
+	};
+
 	TArray<FString> PackagesToDelete;
-	for (FActorDescContainerCollection::TIterator<AWorldPartitionHLOD> HLODIterator(WorldPartition); HLODIterator; ++HLODIterator)
+	for (FActorDescContainerCollection::TIterator<> ActorDescIterator(WorldPartition); ActorDescIterator; ++ActorDescIterator)
 	{
-		FWorldPartitionActorDesc* HLODActorDesc = *HLODIterator;
-		FString PackageName = HLODActorDesc->GetActorPackage().ToString();
-		PackagesToDelete.Add(PackageName);
+		if (HLODActorClasses.FindByPredicate([ActorClass = ActorDescIterator->GetActorNativeClass()](const UClass* HLODClass) { return ActorClass->IsChildOf(HLODClass); }))
+		{
+			FString PackageName = ActorDescIterator->GetActorPackage().ToString();
+			PackagesToDelete.Add(PackageName);
+		}
 	}
 
 	// Ensure we don't hold on to packages of always loaded actors
