Fix crash pressing 'w' or 'e' in material editor to due hidden viewport moment triggering a draw of hidden videwports

diff --git a/Engine/Source/Editor/UnrealEd/Private/EditorViewportClient.cpp b/Engine/Source/Editor/UnrealEd/Private/EditorViewportClient.cpp
--- a/Engine/Source/Editor/UnrealEd/Private/EditorViewportClient.cpp
+++ b/Engine/Source/Editor/UnrealEd/Private/EditorViewportClient.cpp
@@ -3856,6 +3856,19 @@ bool FEditorViewportClient::InputWidgetDelta(FViewport* InViewport, EAxisList::T
 
 void FEditorViewportClient::SetWidgetMode(UE::Widget::EWidgetMode NewMode)
 {
+	// Don't set hit proxies or redraw collapsed viewport widgets
+	if (TSharedPtr<SEditorViewport> EditorViewportWidgetPinned = EditorViewportWidget.Pin())
+	{
+		if (!EditorViewportWidgetPinned->IsVisible() || EditorViewportWidgetPinned->GetVisibility() != EVisibility::Visible)
+		{
+			return;
+		}
+	}
+	else
+	{
+		return;
+	}
+
 	if (!ModeTools->IsTracking() && !IsFlightCameraActive())
 	{
 		ModeTools->SetWidgetMode(NewMode);
