Cooker: Make "Cooker has been blocked" message severity configurable per project so it can be downgraded to display for contended farm machines.
#rnx

diff --git a/Engine/Config/BaseEditor.ini b/Engine/Config/BaseEditor.ini
--- a/Engine/Config/BaseEditor.ini
+++ b/Engine/Config/BaseEditor.ini
@@ -337,6 +337,12 @@ BulkDataRegistryEnabled=true
 ; the message can be downgraded using this variable.
 CookContentMissingSeverity=Error
 
+; The cooker can become stuck forever due to system-specific bugs in IsCachedCookedPlatformDataLoaded bugs.
+; It can also just take a long time for that work to complete if e.g. the machine is under load.
+; We log a warning after a timeout (Cook.display.warnbusytime) to warn about the possible softlock.
+; But this warning can be spurious on farm machines so the project may want to downgrade it to display.
+CookerIdleWarningSeverity=Warning
+
 ; HybridIterative uses the TargetDomain and BuildDependencies to decide when packages need to be recooked. If enabled,
 ; every cook is iterative for packages with allowed classes (IterativeClassAllowList) unless -noiterative is passed.
 HybridIterativeEnabled=false
diff --git a/Engine/Source/Editor/UnrealEd/Classes/CookOnTheSide/CookOnTheFlyServer.h b/Engine/Source/Editor/UnrealEd/Classes/CookOnTheSide/CookOnTheFlyServer.h
--- a/Engine/Source/Editor/UnrealEd/Classes/CookOnTheSide/CookOnTheFlyServer.h
+++ b/Engine/Source/Editor/UnrealEd/Classes/CookOnTheSide/CookOnTheFlyServer.h
@@ -8,6 +8,7 @@
 #include "CookPackageSplitter.h"
 #include "HAL/PlatformMemory.h"
 #include "INetworkFileSystemModule.h"
+#include "Logging/LogMacros.h"
 #include "Misc/EnumClassFlags.h"
 #include "Misc/Optional.h"
 #include "Misc/PackageAccessTracking.h"
@@ -1451,6 +1452,7 @@ private:
 	/** Classes (and all subclasses) that were listed as the only classes that should be cooked in the filter settings. */
 	TSet<FName> CookFilterIncludedClasses;
 
+	ELogVerbosity::Type CookerIdleWarningSeverity = ELogVerbosity::Warning;
 	/** True when PumpLoads has detected it is blocked on async work and CookOnTheFlyServer should do work elsewhere. */
 	bool bLoadBusy = false;
 	/** True when PumpSaves has detected it is blocked on async work and CookOnTheFlyServer should do work elsewhere. */
diff --git a/Engine/Source/Editor/UnrealEd/Private/CookOnTheFlyServer.cpp b/Engine/Source/Editor/UnrealEd/Private/CookOnTheFlyServer.cpp
--- a/Engine/Source/Editor/UnrealEd/Private/CookOnTheFlyServer.cpp
+++ b/Engine/Source/Editor/UnrealEd/Private/CookOnTheFlyServer.cpp
@@ -1648,14 +1648,14 @@ void UCookOnTheFlyServer::SetSaveBusy(bool bInBusy)
 			
 			FString Message = FString::Printf(TEXT("Cooker has been blocked from saving the current packages for %.0f seconds."),
 				(float)CurrentTime - SaveBusyStartTimeSeconds);
-			if (NonExpectedObjects.IsEmpty())
+			ELogVerbosity::Type MessageSeverity = ELogVerbosity::Display;
+			if (ExpectedObjects.IsEmpty() || !NonExpectedObjects.IsEmpty())
 			{
-				UE_LOG(LogCook, Display, TEXT("%s"), *Message);
-			}
-			else
-			{
-				UE_LOG(LogCook, Warning, TEXT("%s"), *Message);
+				MessageSeverity = CookerIdleWarningSeverity;
 			}
+#if !NO_LOGGING
+			FMsg::Logf(__FILE__, __LINE__, LogCook.GetCategoryName(), MessageSeverity, TEXT("%s"), *Message);
+#endif
 
 			UE_LOG(LogCook, Display, TEXT("%d packages in the savequeue: "), SaveQueue.Num());
 			int DisplayCount = 0;
@@ -1734,8 +1734,11 @@ void UCookOnTheFlyServer::SetLoadBusy(bool bInLoadBusy)
 			int DisplayCount = 0;
 			const int DisplayMax = 10;
 			FLoadPrepareQueue& LoadPrepareQueue = PackageDatas->GetLoadPrepareQueue();
-			UE_LOG(LogCook, Warning, TEXT("Cooker has been blocked from loading the current packages for %.0f seconds. %d packages in the loadqueue:"),
+#if !NO_LOGGING
+			FMsg::Logf(__FILE__, __LINE__, LogCook.GetCategoryName(), CookerIdleWarningSeverity,
+				TEXT("Cooker has been blocked from loading the current packages for %.0f seconds. %d packages in the loadqueue:"),
 				(float)(CurrentTime - LoadBusyStartTimeSeconds), LoadPrepareQueue.PreloadingQueue.Num() + LoadPrepareQueue.EntryQueue.Num());
+#endif
 			for (FPackageData* PackageData : LoadPrepareQueue.PreloadingQueue)
 			{
 				if (DisplayCount == DisplayMax)
@@ -6896,6 +6899,10 @@ void UCookOnTheFlyServer::SetInitializeConfigSettings(UE::Cook::FInitializeConfi
 
 	bIgnoreUnsolicitedPackages = FParse::Param(FCommandLine::Get(), TEXT("odsc"));
 	bSkipSave = FParse::Param(FCommandLine::Get(), TEXT("CookSkipSave"));
+
+	FString Severity;
+	GConfig->GetString(TEXT("CookSettings"), TEXT("CookerIdleWarningSeverity"), Severity, GEditorIni);
+	CookerIdleWarningSeverity = ParseLogVerbosityFromString(Severity);
 }
 
 void UCookOnTheFlyServer::ParseCookFilters()
