[UnrealEd]
* Added dll export to Set/ClearLogger which was used by an inlined function in a friend class

diff --git a/Engine/Source/Editor/UnrealEd/Public/FbxImporter.h b/Engine/Source/Editor/UnrealEd/Public/FbxImporter.h
--- a/Engine/Source/Editor/UnrealEd/Public/FbxImporter.h
+++ b/Engine/Source/Editor/UnrealEd/Public/FbxImporter.h
@@ -1812,8 +1812,8 @@ private:
 
 	// logger set/clear function
 	class FFbxLogger * Logger;
-	void SetLogger(class FFbxLogger * InLogger);
-	void ClearLogger();
+	UNREALED_API void SetLogger(class FFbxLogger * InLogger);
+	UNREALED_API void ClearLogger();
 
 	FImportedMaterialData ImportedMaterialData;
 
