Cooker: Add HardEditorOnlyReference to the list of Instigator causes. This is useful for validating that packages omitted by SkipOnlyEditorOnly were indeed OnlyEditorOnly in the non-SOEO cook.
#rnx

diff --git a/Engine/Source/Editor/UnrealEd/Private/CookOnTheFlyServer.cpp b/Engine/Source/Editor/UnrealEd/Private/CookOnTheFlyServer.cpp
--- a/Engine/Source/Editor/UnrealEd/Private/CookOnTheFlyServer.cpp
+++ b/Engine/Source/Editor/UnrealEd/Private/CookOnTheFlyServer.cpp
@@ -7126,6 +7126,7 @@ TArray<UE::Cook::FInstigator> UCookOnTheFlyServer::GetInstigatorChain(FName Pack
 		{
 			case EInstigator::Dependency: bGetNext = true; break;
 			case EInstigator::HardDependency: bGetNext = true; break;
+			case EInstigator::HardEditorOnlyDependency: bGetNext = true; break;
 			case EInstigator::SoftDependency: bGetNext = true; break;
 			case EInstigator::Unsolicited: bGetNext = true; break;
 			case EInstigator::EditorOnlyLoad: bGetNext = true; break;
diff --git a/Engine/Source/Editor/UnrealEd/Private/Cooker/CookRequestCluster.cpp b/Engine/Source/Editor/UnrealEd/Private/Cooker/CookRequestCluster.cpp
--- a/Engine/Source/Editor/UnrealEd/Private/Cooker/CookRequestCluster.cpp
+++ b/Engine/Source/Editor/UnrealEd/Private/Cooker/CookRequestCluster.cpp
@@ -984,9 +984,11 @@ void FRequestCluster::FGraphSearch::ExploreVertexEdges(FVertexData& Vertex)
 	}
 
 	TArray<FName>& HardGameDependencies(Scratch.HardGameDependencies);
+	TArray<FName>& HardEditorDependencies(Scratch.HardEditorDependencies);
 	TArray<FName>& SoftGameDependencies(Scratch.SoftGameDependencies);
 	TSet<FName>& HardDependenciesSet(Scratch.HardDependenciesSet);
 	HardGameDependencies.Reset();
+	HardEditorDependencies.Reset();
 	SoftGameDependencies.Reset();
 	HardDependenciesSet.Reset();
 	FPackageData& PackageData = *Vertex.PackageData;
@@ -998,16 +1000,21 @@ void FRequestCluster::FGraphSearch::ExploreVertexEdges(FVertexData& Vertex)
 		EDependencyQuery FlagsForHardDependencyQuery;
 		if (Cluster.COTFS.bSkipOnlyEditorOnly)
 		{
-			FlagsForHardDependencyQuery = EDependencyQuery::Game | EDependencyQuery::Hard;
+			Cluster.AssetRegistry.GetDependencies(PackageName, HardGameDependencies, EDependencyCategory::Package,
+				EDependencyQuery::Game | EDependencyQuery::Hard);
+			HardDependenciesSet.Append(HardGameDependencies);
 		}
 		else
 		{
 			// We're not allowed to skip editoronly imports, so include all hard dependencies
 			FlagsForHardDependencyQuery = EDependencyQuery::Hard;
+			Cluster.AssetRegistry.GetDependencies(PackageName, HardGameDependencies, EDependencyCategory::Package,
+				EDependencyQuery::Game | EDependencyQuery::Hard);
+			Cluster.AssetRegistry.GetDependencies(PackageName, HardEditorDependencies, EDependencyCategory::Package,
+				EDependencyQuery::EditorOnly | EDependencyQuery::Hard);
+			HardDependenciesSet.Append(HardGameDependencies);
+			HardDependenciesSet.Append(HardEditorDependencies);
 		}
-		Cluster.AssetRegistry.GetDependencies(PackageName, HardGameDependencies, EDependencyCategory::Package,
-			FlagsForHardDependencyQuery);
-		HardDependenciesSet.Append(HardGameDependencies);
 		if (DiscoveredDependencies)
 		{
 			HardDependenciesSet.Append(*DiscoveredDependencies);
@@ -1040,26 +1047,47 @@ void FRequestCluster::FGraphSearch::ExploreVertexEdges(FVertexData& Vertex)
 	int32 LocalNumFetchPlatforms = NumFetchPlatforms();
 	TMap<FName, FScratchPlatformDependencyBits>& PlatformDependencyMap(Scratch.PlatformDependencyMap);
 	PlatformDependencyMap.Reset();
-	auto AddPlatformDependency = [&PlatformDependencyMap, LocalNumFetchPlatforms](FName DependencyName, int32 PlatformIndex, bool bHardDependency)
+	auto AddPlatformDependency = [&PlatformDependencyMap, LocalNumFetchPlatforms]
+	(FName DependencyName, int32 PlatformIndex, EInstigator InstigatorType)
 	{
 		FScratchPlatformDependencyBits& PlatformDependencyBits = PlatformDependencyMap.FindOrAdd(DependencyName);
 		if (PlatformDependencyBits.HasPlatformByIndex.Num() != LocalNumFetchPlatforms)
 		{
 			PlatformDependencyBits.HasPlatformByIndex.Init(false, LocalNumFetchPlatforms);
-			PlatformDependencyBits.bHardDependency = false;
+			PlatformDependencyBits.InstigatorType = EInstigator::SoftDependency;
 		}
 		PlatformDependencyBits.HasPlatformByIndex[PlatformIndex] = true;
-		if (bHardDependency)
+
+		// Calculate PlatformDependencyType.InstigatorType == Max(InstigatorType, PlatformDependencyType.InstigatorType)
+		// based on the enum values, from least required to most: [ Soft, HardEditorOnly, Hard ]
+		switch (InstigatorType)
 		{
-			PlatformDependencyBits.bHardDependency = true;
+		case EInstigator::HardDependency:
+			PlatformDependencyBits.InstigatorType = InstigatorType;
+			break;
+		case EInstigator::HardEditorOnlyDependency:
+			if (PlatformDependencyBits.InstigatorType != EInstigator::HardDependency)
+			{
+				PlatformDependencyBits.InstigatorType = InstigatorType;
+			}
+			break;
+		case EInstigator::SoftDependency:
+			// New value is minimum, so keep the old value
+			break;
+		case EInstigator::InvalidCategory:
+			// Caller indicated they do not want to set the InstigatorType
+			break;
+		default:
+			checkNoEntry();
+			break;
 		}
-
 	};
-	auto AddPlatformDependencyRange = [&AddPlatformDependency](TConstArrayView<FName> Range, int32 PlatformIndex, bool bHardDependency)
+	auto AddPlatformDependencyRange = [&AddPlatformDependency]
+	(TConstArrayView<FName> Range, int32 PlatformIndex, EInstigator InstigatorType)
 	{
 		for (FName DependencyName : Range)
 		{
-			AddPlatformDependency(DependencyName, PlatformIndex, bHardDependency);
+			AddPlatformDependency(DependencyName, PlatformIndex, InstigatorType);
 		}
 	};
 
@@ -1085,11 +1113,11 @@ void FRequestCluster::FGraphSearch::ExploreVertexEdges(FVertexData& Vertex)
 				if (bExploreDependencies)
 				{
 					AddPlatformDependencyRange(PlatformAttachments.BuildDependencies, PlatformIndex,
-						true /* bHardDependency */);
+						EInstigator::HardDependency);
 					if (Cluster.bAllowSoftDependencies)
 					{
 						AddPlatformDependencyRange(PlatformAttachments.RuntimeOnlyDependencies, PlatformIndex,
-							true /* bHardDependency */);
+							EInstigator::HardDependency);
 					}
 				}
 
@@ -1172,19 +1200,20 @@ void FRequestCluster::FGraphSearch::ExploreVertexEdges(FVertexData& Vertex)
 				Cluster.AssetRegistry.GetDependencies(PackageName, CookerLoadingDependencies, EDependencyCategory::Package,
 					EDependencyQuery::Build);
 			}
-			// CookerLoadingPlatform does not cause SetInstigator so it does not modify bHardDependency
-			AddPlatformDependencyRange(CookerLoadingDependencies, PlatformIndex, false /* bHardDependency */);
+			// CookerLoadingPlatform does not cause SetInstigator so it does not modify the platformdependency's InstigatorType
+			AddPlatformDependencyRange(CookerLoadingDependencies, PlatformIndex, EInstigator::InvalidCategory);
 		}
 		else
 		{
-			AddPlatformDependencyRange(HardGameDependencies, PlatformIndex, true /* bHardDependency */);
-			AddPlatformDependencyRange(SoftGameDependencies, PlatformIndex, false /* bHardDependency */);
+			AddPlatformDependencyRange(HardGameDependencies, PlatformIndex, EInstigator::HardDependency);
+			AddPlatformDependencyRange(HardEditorDependencies, PlatformIndex, EInstigator::HardEditorOnlyDependency);
+			AddPlatformDependencyRange(SoftGameDependencies, PlatformIndex, EInstigator::SoftDependency);
 			ProcessPlatformAttachments(PlatformIndex, TargetPlatform, FetchPlatformData, PackagePlatformData,
 				QueryPlatformData.CookAttachments, true /* bExploreDependencies  */);
 		}
 		if (DiscoveredDependencies)
 		{
-			AddPlatformDependencyRange(*DiscoveredDependencies, PlatformIndex, true /* bHardDependency */);
+			AddPlatformDependencyRange(*DiscoveredDependencies, PlatformIndex, EInstigator::HardDependency);
 		}
 	}
 	if (PlatformDependencyMap.IsEmpty())
@@ -1197,7 +1226,7 @@ void FRequestCluster::FGraphSearch::ExploreVertexEdges(FVertexData& Vertex)
 	{
 		FName DependencyName = PlatformDependencyPair.Key;
 		TBitArray<>& HasPlatformByIndex = PlatformDependencyPair.Value.HasPlatformByIndex;
-		bool bHardDependency = PlatformDependencyPair.Value.bHardDependency;
+		EInstigator InstigatorType = PlatformDependencyPair.Value.InstigatorType;
 
 		// Process any CoreRedirects before checking whether the package exists
 		FName Redirected = FCoreRedirects::GetRedirectedName(ECoreRedirectFlags::Type_Package,
@@ -1237,7 +1266,6 @@ void FRequestCluster::FGraphSearch::ExploreVertexEdges(FVertexData& Vertex)
 				PlatformData.SetReachable(true);
 				if (!DependencyPackageData.HasInstigator() && TargetPlatform != CookerLoadingPlatformKey)
 				{
-					EInstigator InstigatorType = bHardDependency ? EInstigator::HardDependency : EInstigator::SoftDependency;
 					DependencyPackageData.SetInstigator(Cluster, FInstigator(InstigatorType, PackageName));
 				}
 			}
diff --git a/Engine/Source/Editor/UnrealEd/Private/Cooker/CookRequestCluster.h b/Engine/Source/Editor/UnrealEd/Private/Cooker/CookRequestCluster.h
--- a/Engine/Source/Editor/UnrealEd/Private/Cooker/CookRequestCluster.h
+++ b/Engine/Source/Editor/UnrealEd/Private/Cooker/CookRequestCluster.h
@@ -233,11 +233,12 @@ private:
 		struct FScratchPlatformDependencyBits
 		{
 			TBitArray<> HasPlatformByIndex;
-			bool bHardDependency = false;
+			EInstigator InstigatorType = EInstigator::SoftDependency;
 		};
 		struct FScratch
 		{
 			TArray<FName> HardGameDependencies;
+			TArray<FName> HardEditorDependencies;
 			TArray<FName> SoftGameDependencies;
 			TArray<FName> CookerLoadingDependencies;
 			TMap<FName, FScratchPlatformDependencyBits> PlatformDependencyMap;
diff --git a/Engine/Source/Runtime/CoreUObject/Public/UObject/ICookInfo.h b/Engine/Source/Runtime/CoreUObject/Public/UObject/ICookInfo.h
--- a/Engine/Source/Runtime/CoreUObject/Public/UObject/ICookInfo.h
+++ b/Engine/Source/Runtime/CoreUObject/Public/UObject/ICookInfo.h
@@ -47,6 +47,7 @@ namespace UE::Cook
 	CallbackMacro(RequestPackageFunction, true) \
 	CallbackMacro(Dependency, false) \
 	CallbackMacro(HardDependency, false) \
+	CallbackMacro(HardEditorOnlyDependency, false) \
 	CallbackMacro(SoftDependency, false) \
 	CallbackMacro(Unsolicited, false) \
 	CallbackMacro(EditorOnlyLoad, false) \
