Cooker: Add -cookfaststartup to reduce iteration time for users iterating with -cooksinglepackagenorefs.
#rnx

diff --git a/Engine/Source/Editor/UnrealEd/Classes/CookOnTheSide/CookOnTheFlyServer.h b/Engine/Source/Editor/UnrealEd/Classes/CookOnTheSide/CookOnTheFlyServer.h
--- a/Engine/Source/Editor/UnrealEd/Classes/CookOnTheSide/CookOnTheFlyServer.h
+++ b/Engine/Source/Editor/UnrealEd/Classes/CookOnTheSide/CookOnTheFlyServer.h
@@ -1333,7 +1333,7 @@ private:
 	UNREALED_API void PopulateCookedPackages(const TConstArrayView<const ITargetPlatform*> TargetPlatforms);
 
 	/** Waits for the AssetRegistry to complete so that we know any missing assets are missing on disk */
-	UNREALED_API void BlockOnAssetRegistry();
+	UNREALED_API void BlockOnAssetRegistry(TConstArrayView<FString> CommandlinePackages);
 
 	/** Construct or refresh-for-filechanges the platform-specific asset registry for the given platforms */
 	UNREALED_API void RefreshPlatformAssetRegistries(const TArrayView<const ITargetPlatform* const>& TargetPlatforms);
@@ -1495,6 +1495,8 @@ private:
 	bool bCookFilter = false;
 	/** True if commandline arguments specify that packages on commandline should be cooked first. */
 	bool bCookFirst = false;
+	/** True if experimental optimizations for fast startup should be used. */
+	bool bCookFastStartup = false;
 	/**
 	 * Experimental feature to correctly invoke the BeginCacheForCookedPlatformData contracts for
 	 * objects created by PreSave
diff --git a/Engine/Source/Editor/UnrealEd/Private/Commandlets/AssetRegistryGenerator.cpp b/Engine/Source/Editor/UnrealEd/Private/Commandlets/AssetRegistryGenerator.cpp
--- a/Engine/Source/Editor/UnrealEd/Private/Commandlets/AssetRegistryGenerator.cpp
+++ b/Engine/Source/Editor/UnrealEd/Private/Commandlets/AssetRegistryGenerator.cpp
@@ -1169,15 +1169,14 @@ void FAssetRegistryGenerator::Initialize(const TArray<FName> &InStartupPackages,
 
 	FAssetRegistrySerializationOptions SaveOptions;
 
-	// If the asset registry is still doing it's background scan, we need to wait for it to finish and tick it so that the results are flushed out
-	AssetRegistry.WaitForCompletion();
-
-	ensureMsgf(!AssetRegistry.IsLoadingAssets(), TEXT("Cannot initialize asset registry generator while asset registry is still scanning source assets "));
-
 	AssetRegistry.InitializeSerializationOptions(SaveOptions, TargetPlatform->IniPlatformName());
 
 	if (bInitializeFromExisting)
 	{
+		// If the asset registry is still doing its background scan, we need to wait for it to finish and tick it so that the results are flushed out
+		AssetRegistry.WaitForCompletion();
+		ensureMsgf(!AssetRegistry.IsLoadingAssets(), TEXT("Cannot initialize asset registry generator while asset registry is still scanning source assets "));
+
 		bClonedGlobalAssetRegistry = true;
 		AssetRegistry.InitializeTemporaryAssetRegistryState(State, SaveOptions);
 	}
diff --git a/Engine/Source/Editor/UnrealEd/Private/CookOnTheFlyServer.cpp b/Engine/Source/Editor/UnrealEd/Private/CookOnTheFlyServer.cpp
--- a/Engine/Source/Editor/UnrealEd/Private/CookOnTheFlyServer.cpp
+++ b/Engine/Source/Editor/UnrealEd/Private/CookOnTheFlyServer.cpp
@@ -769,7 +769,7 @@ void UCookOnTheFlyServer::StartCookOnTheFlySessionFromGameThread(ITargetPlatform
 	// Blocking on the AssetRegistry needs to wait until the session starts because it needs all plugins loaded.
 	// AddCookOnTheFlyPlatformFromGameThread can be called on cooker startup which occurs in UUnrealEdEngine::Init
 	// before all plugins are loaded.
-	BlockOnAssetRegistry();
+	BlockOnAssetRegistry(TConstArrayView<FString>());
 
 	if (CookOnTheFlyRequestManager)
 	{
@@ -6903,6 +6903,8 @@ void UCookOnTheFlyServer::SetInitializeConfigSettings(UE::Cook::FInitializeConfi
 	FString Severity;
 	GConfig->GetString(TEXT("CookSettings"), TEXT("CookerIdleWarningSeverity"), Severity, GEditorIni);
 	CookerIdleWarningSeverity = ParseLogVerbosityFromString(Severity);
+
+	bCookFastStartup = FParse::Param(FCommandLine::Get(), TEXT("cookfaststartup"));
 }
 
 void UCookOnTheFlyServer::ParseCookFilters()
@@ -8384,7 +8386,7 @@ void DumpAssetRegistryForCooker(IAssetRegistry* AssetRegistry)
 #endif
 
 
-void UCookOnTheFlyServer::BlockOnAssetRegistry()
+void UCookOnTheFlyServer::BlockOnAssetRegistry(TConstArrayView<FString> CommandlinePackages)
 {
 	if (!bFirstCookInThisProcess)
 	{
@@ -8393,10 +8395,27 @@ void UCookOnTheFlyServer::BlockOnAssetRegistry()
 	TRACE_CPUPROFILER_EVENT_SCOPE(UCookOnTheFlyServer::BlockOnAssetRegistry);
 	COOK_STAT(FScopedDurationTimer TickTimer(DetailedCookStats::BlockOnAssetRegistryTimeSec));
 
+	bool bAssetGatherCompleted = true;
 	UE_LOG(LogCook, Display, TEXT("Waiting for Asset Registry"));
 	// Blocking on the AssetRegistry has to be done on the game thread since some AssetManager functions require it
 	check(IsInGameThread());
-	if (ShouldPopulateFullAssetRegistry())
+	if (EnumHasAnyFlags(CookByTheBookOptions->StartupOptions, ECookByTheBookOptions::SkipHardReferences) &&
+		!CommandlinePackages.IsEmpty() &&
+		bCookFastStartup)
+	{
+		TArray<FString> PackageNames;
+		for (const FString& FileNameOrPackageName : CommandlinePackages)
+		{
+			FString PackageName;
+			if (FPackageName::TryConvertFilenameToLongPackageName(FileNameOrPackageName, PackageName))
+			{
+				PackageNames.Add(MoveTemp(PackageName));
+			}
+		}
+		AssetRegistry->ScanFilesSynchronous(PackageNames);
+		bAssetGatherCompleted = false;
+	}
+	else if (ShouldPopulateFullAssetRegistry())
 	{
 		// Trigger or wait for completion the primary AssetRegistry scan.
 		// Additionally scan any cook-specific paths from ini
@@ -8432,7 +8451,10 @@ void UCookOnTheFlyServer::BlockOnAssetRegistry()
 	{
 		FAssetRegistryGenerator::UpdateAssetManagerDatabase();
 	}
-	AssetRegistry->ClearGathererCache();
+	if (bAssetGatherCompleted)
+	{
+		AssetRegistry->ClearGathererCache();
+	}
 }
 
 void UCookOnTheFlyServer::RefreshPlatformAssetRegistries(const TArrayView<const ITargetPlatform* const>& TargetPlatforms)
@@ -8463,8 +8485,15 @@ void UCookOnTheFlyServer::RefreshPlatformAssetRegistries(const TArrayView<const
 		{
 			// if we are cooking DLC, we will just spend a lot of time removing the shipped packages from the AR,
 			// so we don't bother copying them over. can easily save 10 seconds on a large project
-			bool bInitalizeFromExisting = !IsCookingDLC();
-			PlatformData->RegistryGenerator->Initialize(CookByTheBookOptions->StartupPackages, bInitalizeFromExisting);
+			bool bInitializeFromExisting = !IsCookingDLC();
+			if (EnumHasAnyFlags(CookByTheBookOptions->StartupOptions, ECookByTheBookOptions::SkipHardReferences) &&
+				bCookFastStartup)
+			{
+				// We don't want to wait on the AssetRegistry when just testing the cook of a single file
+				bInitializeFromExisting = false;
+			}
+
+			PlatformData->RegistryGenerator->Initialize(CookByTheBookOptions->StartupPackages, bInitializeFromExisting);
 		}
 	}
 }
@@ -10978,7 +11007,7 @@ void UCookOnTheFlyServer::StartCookByTheBook( const FCookByTheBookStartupOptions
 	// Functions in this section are ordered and can depend on the functions before them
 	InitializeSession();
 	FBeginCookContext BeginContext = CreateBeginCookByTheBookContext(EffectiveStartupOptions);
-	BlockOnAssetRegistry();
+	BlockOnAssetRegistry(BeginContext.StartupOptions->CookMaps);
 	CreateSandboxFile(BeginContext);
 	LoadBeginCookConfigSettings(BeginContext);
 	SelectSessionPlatforms(BeginContext);
@@ -11156,7 +11185,7 @@ void UCookOnTheFlyServer::StartCookAsCookWorker()
 	InitializeSession();
 	FBeginCookContext BeginContext = CreateCookWorkerContext();
 	// MPCOOKTODO: Load serialized AssetRegistry from Director
-	BlockOnAssetRegistry();
+	BlockOnAssetRegistry(TConstArrayView<FString>());
 	CreateSandboxFile(BeginContext);
 	LoadBeginCookConfigSettings(BeginContext);
 	SelectSessionPlatforms(BeginContext);
@@ -11928,7 +11957,20 @@ TArray<FName> UCookOnTheFlyServer::GetNeverCookPackageFileNames(TArrayView<const
 		// Rather than blocking on the AssetRegistry now, fallback to scanning the directories on disk
 		// TODO: Change CookOnTheFlyStartup in the editor to delay most of its startup until the AssetRegistry has
 		// finished loading so we can block on the AssetRegistry before calling this function.
-		FPackageName::FindPackagesInDirectories(NeverCookPackagesPaths, NeverCookDirectories);
+		bool bUseDirectoryScanFallback = true;
+
+		if (EnumHasAnyFlags(CookByTheBookOptions->StartupOptions, ECookByTheBookOptions::SkipHardReferences) &&
+			bCookFastStartup)
+		{
+			// When using -cooksinglepackagenorefs, skip the calculation of NeverCook packages since it requires
+			// waiting for the full AssetRegistry scan
+			bUseDirectoryScanFallback = false;
+		}
+
+		if (bUseDirectoryScanFallback)
+		{
+			FPackageName::FindPackagesInDirectories(NeverCookPackagesPaths, NeverCookDirectories);
+		}
 	}
 
 	TArray<FName> NeverCookNormalizedFileNames;
