Avoid misuse of CLASS_CompiledFromBlueprint in CanCreateBlueprintOfClass, this resolves redundancy between 2398267 and 2422476/2427167/2428539

diff --git a/Engine/Source/Editor/UnrealEd/Private/Kismet2/Kismet2.cpp b/Engine/Source/Editor/UnrealEd/Private/Kismet2/Kismet2.cpp
--- a/Engine/Source/Editor/UnrealEd/Private/Kismet2/Kismet2.cpp
+++ b/Engine/Source/Editor/UnrealEd/Private/Kismet2/Kismet2.cpp
@@ -1036,7 +1036,7 @@ bool FKismetEditorUtilities::CanCreateBlueprintOfClass(const UClass* Class)
 
 		const bool bIsValidClass = Class->GetBoolMetaDataHierarchical(FBlueprintMetadata::MD_IsBlueprintBase)
 			|| (Class == UObject::StaticClass())
-			|| (Class->HasAnyClassFlags(CLASS_CompiledFromBlueprint) || Class == USceneComponent::StaticClass() || Class == UActorComponent::StaticClass())
+			|| (Class == USceneComponent::StaticClass() || Class == UActorComponent::StaticClass())
 			|| bIsBPGC;  // BPs are always considered inheritable
 			
 		bCanCreateBlueprint &= bIsValidClass;
