EditorDomain: Suppress message about missing digest causing inability to load from editordomain when EditorDomain is being used to get digests but not for package reading.
#rnx

diff --git a/Engine/Source/Editor/UnrealEd/Private/EditorDomain/EditorDomain.cpp b/Engine/Source/Editor/UnrealEd/Private/EditorDomain/EditorDomain.cpp
--- a/Engine/Source/Editor/UnrealEd/Private/EditorDomain/EditorDomain.cpp
+++ b/Engine/Source/Editor/UnrealEd/Private/EditorDomain/EditorDomain.cpp
@@ -279,9 +279,12 @@ bool FEditorDomain::TryFindOrAddPackageSource(FScopeLock& ScopeLock, bool& bOutR
 		}
 		return false;
 	default:
-		UE_LOG(LogEditorDomain, Warning,
-			TEXT("Could not load package %s from EditorDomain; it will be loaded from the WorkspaceDomain: %s"),
-			*WriteToString<256>(PackageName), *PackageDigest.GetStatusString());
+		if (bEditorDomainReadEnabled)
+		{
+			UE_LOG(LogEditorDomain, Warning,
+				TEXT("Could not load package %s from EditorDomain; it will be loaded from the WorkspaceDomain: %s"),
+				*WriteToString<256>(PackageName), *PackageDigest.GetStatusString());
+		}
 		PackageSource = new FPackageSource();
 		PackageSource->Source = EPackageSource::Workspace;
 		PackageSource->Digest = MoveTemp(PackageDigest);
