Modifications to tab factory systems

Allow multiple document types to be bound to the same tab stack - defers to the tab factory when document ID is not supplied, otherwise default to "Document" as before.
Add the ability for workflow tabs to display a tab suffix (allows tabs to display dirty state)

diff --git a/Engine/Source/Editor/UnrealEd/Private/WorkflowOrientedApp/WorkflowTabFactory.cpp b/Engine/Source/Editor/UnrealEd/Private/WorkflowOrientedApp/WorkflowTabFactory.cpp
--- a/Engine/Source/Editor/UnrealEd/Private/WorkflowOrientedApp/WorkflowTabFactory.cpp
+++ b/Engine/Source/Editor/UnrealEd/Private/WorkflowOrientedApp/WorkflowTabFactory.cpp
@@ -55,6 +55,7 @@ TSharedRef<SDockTab> FWorkflowTabFactory::SpawnTab(const FWorkflowTabSpawnInfo&
 	TSharedRef<SDockTab> NewTab = SNew(SDockTab)
 		.TabRole(this->TabRole)
 		.Label(ConstructTabName(Info))
+		.LabelSuffix(ConstructTabLabelSuffix(Info))
 		.ShouldAutosize(bShouldAutosize)
 		[
 			TabBody
@@ -98,6 +99,7 @@ void FWorkflowTabFactory::UpdateTab(TSharedPtr<SDockTab> InDockTab, const FWorkf
 
 	InDockTab->SetContent(TabBody);
 	InDockTab->SetLabel(ConstructTabName(Info));
+	InDockTab->SetTabLabelSuffix(ConstructTabLabelSuffix(Info));
 	InDockTab->SetTabIcon(GetTabIcon(Info));
 	InDockTab->SetTabToolTipWidget(CreateTabToolTipWidget(Info));
 }
@@ -158,6 +160,11 @@ TAttribute<FText> FWorkflowTabFactory::ConstructTabName(const FWorkflowTabSpawnI
 	return TabLabel;
 }
 
+TAttribute<FText> FWorkflowTabFactory::ConstructTabLabelSuffix(const FWorkflowTabSpawnInfo& Info) const
+{
+	return FText::GetEmpty();
+}
+
 TSharedRef<SWidget> FWorkflowTabFactory::CreateTabBody(const FWorkflowTabSpawnInfo& Info) const
 {
 	return SNew(STextBlock).Text(NSLOCTEXT("WorkflowTabFactory", "EmptyTabMessage", "EmptyBody"));
diff --git a/Engine/Source/Editor/UnrealEd/Private/WorkflowOrientedApp/WorkflowTabManager.cpp b/Engine/Source/Editor/UnrealEd/Private/WorkflowOrientedApp/WorkflowTabManager.cpp
--- a/Engine/Source/Editor/UnrealEd/Private/WorkflowOrientedApp/WorkflowTabManager.cpp
+++ b/Engine/Source/Editor/UnrealEd/Private/WorkflowOrientedApp/WorkflowTabManager.cpp
@@ -413,7 +413,8 @@ void FDocumentTracker::RegisterDocumentFactory(TSharedPtr<class FDocumentTabFact
 	PotentialTabFactories.Add(NewIdentifier, Factory);
 }
 
-FDocumentTracker::FDocumentTracker()
+FDocumentTracker::FDocumentTracker(FName InDefaultDocumentId)
+	: DefaultDocumentId(InDefaultDocumentId)
 {
 	// Make sure we know when tabs become active
 	OnActiveTabChangedDelegateHandle = FGlobalTabmanager::Get()->OnActiveTabChanged_Subscribe( FOnActiveTabChanged::FDelegate::CreateRaw( this, &FDocumentTracker::OnActiveTabChanged ) );
@@ -768,13 +769,15 @@ TSharedPtr<SDockTab> FDocumentTracker::OpenNewTab(TSharedPtr<FGenericTabHistory>
 			NewTabInfo->AddTabHistory(InTabHistory);
 		}
 
+		const FName DocumentId = DefaultDocumentId != NAME_None ? DefaultDocumentId : Factory->GetIdentifier();
+
 		if (InOpenCause == ForceOpenNewDocument  || InOpenCause == OpenNewDocument)
 		{
-			TabManager->InsertNewDocumentTab( "Document", FTabManager::ESearchPreference::RequireClosedTab, NewTab.ToSharedRef() );
+			TabManager->InsertNewDocumentTab( DocumentId, FTabManager::ESearchPreference::RequireClosedTab, NewTab.ToSharedRef() );
 		}
 		else if (InOpenCause == RestorePreviousDocument)
 		{
-			TabManager->RestoreDocumentTab( "Document", FTabManager::ESearchPreference::RequireClosedTab, NewTab.ToSharedRef() );
+			TabManager->RestoreDocumentTab( DocumentId, FTabManager::ESearchPreference::RequireClosedTab, NewTab.ToSharedRef() );
 
 			// Clear tab history before this so previous restores don't show up
 			History.RemoveAt(0, History.Num() - 1, true);
diff --git a/Engine/Source/Editor/UnrealEd/Public/WorkflowOrientedApp/WorkflowTabFactory.h b/Engine/Source/Editor/UnrealEd/Public/WorkflowOrientedApp/WorkflowTabFactory.h
--- a/Engine/Source/Editor/UnrealEd/Public/WorkflowOrientedApp/WorkflowTabFactory.h
+++ b/Engine/Source/Editor/UnrealEd/Public/WorkflowOrientedApp/WorkflowTabFactory.h
@@ -123,6 +123,9 @@ public:
 	// Creates the label for the tab
 	UNREALED_API virtual TAttribute<FText> ConstructTabName(const FWorkflowTabSpawnInfo& Info) const;
 
+	// Creates the label suffix for the tab
+	UNREALED_API virtual TAttribute<FText> ConstructTabLabelSuffix(const FWorkflowTabSpawnInfo& Info) const;
+
 	/** Returns TRUE if the payload is valid */
 	virtual bool IsPayloadValid(TSharedRef<FTabPayload> Payload) const { return false; }
 
diff --git a/Engine/Source/Editor/UnrealEd/Public/WorkflowOrientedApp/WorkflowTabManager.h b/Engine/Source/Editor/UnrealEd/Public/WorkflowOrientedApp/WorkflowTabManager.h
--- a/Engine/Source/Editor/UnrealEd/Public/WorkflowOrientedApp/WorkflowTabManager.h
+++ b/Engine/Source/Editor/UnrealEd/Public/WorkflowOrientedApp/WorkflowTabManager.h
@@ -246,7 +246,7 @@ DECLARE_DELEGATE_OneParam(FOnTabActivated, TSharedPtr<SDockTab>);
 class FDocumentTracker : public TSharedFromThis<FDocumentTracker>
 {
 public:
-	UNREALED_API FDocumentTracker();
+	UNREALED_API explicit FDocumentTracker(FName InDefaultDocumentId = "Document");
 	UNREALED_API ~FDocumentTracker();
 
 	UNREALED_API void ClearDocumentFactories();
@@ -354,6 +354,9 @@ private:
 	/** Current history index */
 	int32 CurrentHistoryIndex;
 
+	/**  Document ID to use when inserting documents into the tab manager. If this is NAME_None, the tab factory's TabIdentifier will be used. */
+	FName DefaultDocumentId;
+
 private:
 	// Clean the spawned list 
 	UNREALED_API FTabList& GetSpawnedList();
