TEv1 copy/paste improvements:

Add a way to propagate the TEv1 selection set to TEDS
Update ActorElementEditorCopyAndPaste to more accurately match what the level editor's copy/paste does
Automatically create a UTypedElementCommonActions for all asset editors

diff --git a/Engine/Source/Editor/UnrealEd/Private/Elements/Actor/ActorElementEditorCopyAndPaste.cpp b/Engine/Source/Editor/UnrealEd/Private/Elements/Actor/ActorElementEditorCopyAndPaste.cpp
--- a/Engine/Source/Editor/UnrealEd/Private/Elements/Actor/ActorElementEditorCopyAndPaste.cpp
+++ b/Engine/Source/Editor/UnrealEd/Private/Elements/Actor/ActorElementEditorCopyAndPaste.cpp
@@ -51,7 +51,11 @@ bool UActorElementsExporterT3D::ExportText(const FExportObjectInnerContext* Cont
 		ULevel* Level = Actor->GetLevel();
 		AActor* ParentActor = Actor->GetAttachParentActor();
 		FName SocketName = Actor->GetAttachParentSocketName();
-		Actor->DetachFromActor(FDetachmentTransformRules::KeepWorldTransform);
+
+		// We set DetachmentTransformRules.bCallModify to false to match the level editor (don't need to modify on copy)
+		FDetachmentTransformRules DetachmentTransformRules = FDetachmentTransformRules::KeepWorldTransform;
+		DetachmentTransformRules.bCallModify = false;
+		Actor->DetachFromActor(DetachmentTransformRules);
 
 		FString ParentActorString = (ParentActor
 			                             ? FString::Printf(TEXT(" ParentActor=%s"), *ParentActor->GetName())
@@ -423,15 +427,21 @@ void FActorElementEditorPasteImporter::PostImportProcess()
 			}
 		}
 
-		// Do the attachment of actors
+		// Do the attachment of actors 
 		for (FImportedActor& ImportedActor : ImportedActors)
 		{
 			if (ImportedActor.ActorPtr)
 			{
+				// Try to find the parent in our imported actors
 				if (AActor** Parent = NameToImportedActors.Find(ImportedActor.ParentActor))
 				{
 					GEditor->ParentActors(*Parent, ImportedActor.ActorPtr, ImportedActor.SocketNameString);
 				}
+				// Otherwise try to find the parent in the world
+				else if(AActor* ParentActor = FindObject<AActor>( ImportedActor.ActorPtr->GetWorld()->GetCurrentLevel(), *ImportedActor.ParentActor.ToString() ))
+				{
+					GEditor->ParentActors(ParentActor, ImportedActor.ActorPtr, ImportedActor.SocketNameString);
+				}
 			}
 		}
 
diff --git a/Engine/Source/Editor/UnrealEd/Private/Elements/Actor/ActorElementEditorSelectionInterface.cpp b/Engine/Source/Editor/UnrealEd/Private/Elements/Actor/ActorElementEditorSelectionInterface.cpp
--- a/Engine/Source/Editor/UnrealEd/Private/Elements/Actor/ActorElementEditorSelectionInterface.cpp
+++ b/Engine/Source/Editor/UnrealEd/Private/Elements/Actor/ActorElementEditorSelectionInterface.cpp
@@ -58,7 +58,11 @@ bool UActorElementEditorSelectionInterface::SelectElement(const FTypedElementHan
 		{
 			if (ITypedElementDataStorageInterface* DataStorage = UTypedElementRegistry::GetInstance()->GetMutableDataStorage())
 			{
-				DataStorage->AddColumn(Row, FTypedElementSelectionColumn::StaticStruct());
+				// TODO: We currently don't have a way to add the selection set if the column add is deferred
+				if(FTypedElementSelectionColumn* SelectionColumn = DataStorage->AddOrGetColumn<FTypedElementSelectionColumn>(Row))
+				{
+					SelectionColumn->SelectionSet = InSelectionOptions.GetNameForTEDSIntegration();
+				}
 			}
 		}
 	}
diff --git a/Engine/Source/Editor/UnrealEd/Private/Toolkits/SStandaloneAssetEditorToolkitHost.cpp b/Engine/Source/Editor/UnrealEd/Private/Toolkits/SStandaloneAssetEditorToolkitHost.cpp
--- a/Engine/Source/Editor/UnrealEd/Private/Toolkits/SStandaloneAssetEditorToolkitHost.cpp
+++ b/Engine/Source/Editor/UnrealEd/Private/Toolkits/SStandaloneAssetEditorToolkitHost.cpp
@@ -17,6 +17,7 @@
 #include "EditorModeManager.h"
 #include "EditorModes.h"
 #include "WidgetDrawerConfig.h"
+#include "Elements/Framework/TypedElementCommonActions.h"
 
 #define LOCTEXT_NAMESPACE "StandaloneAssetEditorToolkit"
 
@@ -34,6 +35,9 @@ void SStandaloneAssetEditorToolkitHost::Construct( const SStandaloneAssetEditorT
 	StatusBarName = FName(AppName, ++StatusBarIdGenerator);
 
 	MyTabManager = InTabManager;
+
+	CommonActions = NewObject<UTypedElementCommonActions>();
+	CommonActions->AddToRoot();
 }
 
 void SStandaloneAssetEditorToolkitHost::SetupInitialContent( const TSharedRef<FTabManager::FLayout>& DefaultLayout, const TSharedPtr<SDockTab>& InHostTab, const bool bCreateDefaultStandaloneMenu )
@@ -297,6 +301,8 @@ void SStandaloneAssetEditorToolkitHost::UnbindEditorCloseRequestFromHostTab()
 SStandaloneAssetEditorToolkitHost::~SStandaloneAssetEditorToolkitHost()
 {
 	ShutdownToolkitHost();
+	
+	CommonActions->RemoveFromRoot();
 }
 
 
@@ -384,7 +390,7 @@ UWorld* SStandaloneAssetEditorToolkitHost::GetWorld() const
 
 UTypedElementCommonActions* SStandaloneAssetEditorToolkitHost::GetCommonActions() const
 {
-	return nullptr;
+	return CommonActions.Get();
 }
 
 void SStandaloneAssetEditorToolkitHost::AddViewportOverlayWidget(TSharedRef<SWidget> InOverlaidWidget, TSharedPtr<IAssetViewport> InViewport)
diff --git a/Engine/Source/Editor/UnrealEd/Private/Toolkits/SStandaloneAssetEditorToolkitHost.h b/Engine/Source/Editor/UnrealEd/Private/Toolkits/SStandaloneAssetEditorToolkitHost.h
--- a/Engine/Source/Editor/UnrealEd/Private/Toolkits/SStandaloneAssetEditorToolkitHost.h
+++ b/Engine/Source/Editor/UnrealEd/Private/Toolkits/SStandaloneAssetEditorToolkitHost.h
@@ -158,6 +158,8 @@ private:
 	TSharedPtr<SWidget> StatusBarWidget;
 
 	SVerticalBox::FSlot* ToolbarSlot;
+
+	TObjectPtr<UTypedElementCommonActions> CommonActions;
 };
 
 
diff --git a/Engine/Source/Runtime/TypedElementRuntime/Public/Elements/Interfaces/TypedElementSelectionInterface.h b/Engine/Source/Runtime/TypedElementRuntime/Public/Elements/Interfaces/TypedElementSelectionInterface.h
--- a/Engine/Source/Runtime/TypedElementRuntime/Public/Elements/Interfaces/TypedElementSelectionInterface.h
+++ b/Engine/Source/Runtime/TypedElementRuntime/Public/Elements/Interfaces/TypedElementSelectionInterface.h
@@ -74,7 +74,10 @@ public:
 
 	FTypedElementSelectionOptions& SetChildElementInclusionMethod(const ETypedElementChildInclusionMethod InChildElementInclusionMethod) { ChildElementInclusionMethod = InChildElementInclusionMethod; return *this; }
 	ETypedElementChildInclusionMethod GetChildElementInclusionMethod() const { return ChildElementInclusionMethod; }
-	
+
+	// Set the selection set name that will be passed into the selection column in TEDS (if it is enabled)
+	FTypedElementSelectionOptions& SetNameForTEDSIntegration(const FName& InTEDSIntegrationSelectionSetName) { TEDSIntegrationSelectionSetName = InTEDSIntegrationSelectionSetName; return *this; }
+	FName GetNameForTEDSIntegration() const { return TEDSIntegrationSelectionSetName; }
 private:
 	UPROPERTY(BlueprintReadWrite, Category="TypedElementInterfaces|Selection|SelectionOptions", meta=(AllowPrivateAccess=true))
 	bool bAllowHidden = false;
@@ -90,6 +93,9 @@ private:
 
 	UPROPERTY(BlueprintReadWrite, Category="TypedElementInterfaces|Selection|SelectionOptions", meta=(AllowPrivateAccess=true))
 	ETypedElementChildInclusionMethod ChildElementInclusionMethod = ETypedElementChildInclusionMethod::None;
+
+	UPROPERTY(BlueprintReadWrite, Category="TypedElementInterfaces|Selection|SelectionOptions", meta=(AllowPrivateAccess=true))
+	FName TEDSIntegrationSelectionSetName = FName();
 };
 
 class ITypedElementTransactedElement
