TEDS Selection: Fix a bug where the level editor selection through TEDS would not work correctly when the column add was deferred because AddOrGetColumn does not support deferred addition yet.

diff --git a/Engine/Source/Editor/UnrealEd/Private/Elements/Actor/ActorElementEditorSelectionInterface.cpp b/Engine/Source/Editor/UnrealEd/Private/Elements/Actor/ActorElementEditorSelectionInterface.cpp
--- a/Engine/Source/Editor/UnrealEd/Private/Elements/Actor/ActorElementEditorSelectionInterface.cpp
+++ b/Engine/Source/Editor/UnrealEd/Private/Elements/Actor/ActorElementEditorSelectionInterface.cpp
@@ -58,8 +58,11 @@ bool UActorElementEditorSelectionInterface::SelectElement(const FTypedElementHan
 		{
 			if (ITypedElementDataStorageInterface* DataStorage = UTypedElementRegistry::GetInstance()->GetMutableDataStorage())
 			{
-				// TODO: We currently don't have a way to add the selection set if the column add is deferred
-				if(FTypedElementSelectionColumn* SelectionColumn = DataStorage->AddOrGetColumn<FTypedElementSelectionColumn>(Row))
+				DataStorage->AddColumn<FTypedElementSelectionColumn>(Row);
+
+				// TODO: If the column add is deferred because the row has not been assigned yet, we do not have a way to add the selection set currently
+				// as GetColumn will return null here
+				if(FTypedElementSelectionColumn* SelectionColumn = DataStorage->GetColumn<FTypedElementSelectionColumn>(Row))
 				{
 					SelectionColumn->SelectionSet = InSelectionOptions.GetNameForTEDSIntegration();
 				}
