Disable "Goto Definition" feature, it is causing crash in 5.3 and has stopped working since at least 5.1. Solution pending

diff --git a/Engine/Source/Editor/UnrealEd/Private/SourceCodeNavigation.cpp b/Engine/Source/Editor/UnrealEd/Private/SourceCodeNavigation.cpp
--- a/Engine/Source/Editor/UnrealEd/Private/SourceCodeNavigation.cpp
+++ b/Engine/Source/Editor/UnrealEd/Private/SourceCodeNavigation.cpp
@@ -552,181 +552,184 @@ void FSourceCodeNavigationImpl::NavigateToFunctionSource( const FString& Functio
 
 #elif PLATFORM_MAC
 	
-	for(uint32 Index = 0; Index < _dyld_image_count(); Index++)
-	{
-		char const* IndexName = _dyld_get_image_name(Index);
-		FString FullModulePath(IndexName);
-		FString Name = FPaths::GetBaseFilename(FullModulePath);
-		if(Name == FunctionModuleName)
-		{
-			struct mach_header_64 const* IndexModule64 = NULL;
-			struct load_command const* LoadCommands = NULL;
-			
-			struct mach_header const* IndexModule32 = _dyld_get_image_header(Index);
-			check(IndexModule32->magic == MH_MAGIC_64);
-			
-			IndexModule64 = (struct mach_header_64 const*)IndexModule32;
-			LoadCommands = (struct load_command const*)(IndexModule64 + 1);
-			struct load_command const* Command = LoadCommands;
-			struct symtab_command const* SymbolTable = nullptr;
-			struct dysymtab_command const* DsymTable = nullptr;
-			struct uuid_command* UUIDCommand = nullptr;
-			for(uint32 CommandIndex = 0; CommandIndex < IndexModule64->ncmds; CommandIndex++)
-			{
-				if (Command && Command->cmd == LC_SYMTAB)
-				{
-					SymbolTable = (struct symtab_command const*)Command;
-				}
-				else if(Command && Command->cmd == LC_DYSYMTAB)
-				{
-					DsymTable = (struct dysymtab_command const*)Command;
-				}
-				else if (Command && Command->cmd == LC_UUID)
-				{
-					UUIDCommand = (struct uuid_command*)Command;
-				}
-				Command = (struct load_command const*)(((char const*)Command) + Command->cmdsize);
-			}
-			
-			check(SymbolTable && DsymTable && UUIDCommand);
-			
-			IPlatformFile& PlatformFile = IPlatformFile::GetPlatformPhysical();
-			IFileHandle* File = PlatformFile.OpenRead(*FullModulePath);
-			if(File)
-			{
-				struct nlist_64* SymbolEntries = new struct nlist_64[SymbolTable->nsyms];
-				check(SymbolEntries);
-				char* StringTable = new char[SymbolTable->strsize];
-				check(StringTable);
-				
-				bool FileOK = File->Seek(SymbolTable->symoff+(DsymTable->iextdefsym*sizeof(struct nlist_64)));
-				FileOK &= File->Read((uint8*)SymbolEntries, DsymTable->nextdefsym*sizeof(struct nlist_64));
-				
-				FileOK &= File->Seek(SymbolTable->stroff);
-				FileOK &= File->Read((uint8*)StringTable, SymbolTable->strsize);
-				
-				delete File;
-				
-				for(uint32 SymbolIndex = 0; FileOK && SymbolIndex < DsymTable->nextdefsym; SymbolIndex++)
-				{
-					struct nlist_64 const& SymbolEntry = SymbolEntries[SymbolIndex];
-					// All the entries in the mach-o external table are functions.
-					// The local table contains the minimal debug stabs used by dsymutil to create the DWARF dsym.
-					if(SymbolEntry.n_un.n_strx)
-					{
-						if (SymbolEntry.n_value)
-						{
-							char const* MangledSymbolName = (StringTable+SymbolEntry.n_un.n_strx);
-							// Remove leading '_'
-							MangledSymbolName += 1;
-							
-							int32 Status = 0;
-							char* DemangledName = abi::__cxa_demangle(MangledSymbolName, NULL, 0, &Status);
-							
-							FString SymbolName;
-							if (DemangledName)
-							{
-								// C++ function
-								SymbolName = DemangledName;
-								free(DemangledName);
-								
-								// This contains return & arguments, it would seem that the DbgHelp API doesn't.
-								// So we shall strip them.
-								int32 ArgumentIndex = -1;
-								if(SymbolName.FindLastChar(TCHAR('('), ArgumentIndex))
-								{
-									SymbolName.LeftInline(ArgumentIndex, false);
-									int32 TemplateNesting = 0;
-									
-									int32 Pos = SymbolName.Len();
-									// Cast operators are special & include spaces, whereas normal functions don't.
-									int32 OperatorIndex = SymbolName.Find("operator");
-									if(OperatorIndex >= 0)
-									{
-										// Trim from before the 'operator'
-										Pos = OperatorIndex;
-									}
-									
-									for(; Pos > 0; --Pos)
-									{
-										TCHAR Character = SymbolName[Pos - 1];
-										if(Character == TCHAR(' ') && TemplateNesting == 0)
-										{
-											SymbolName.MidInline(Pos, MAX_int32, false);
-											break;
-										}
-										else if(Character == TCHAR('>'))
-										{
-											TemplateNesting++;
-										}
-										else if(Character == TCHAR('<'))
-										{
-											TemplateNesting--;
-										}
-									}
-								}
-							}
-							else
-							{
-								// C function
-								SymbolName = MangledSymbolName;
-							}
-							
-							if(FunctionSymbolName == SymbolName)
-							{
-								CFUUIDBytes UUIDBytes;
-								FMemory::Memcpy(&UUIDBytes, UUIDCommand->uuid, sizeof(CFUUIDBytes));
-								CFUUIDRef UUIDRef = CFUUIDCreateFromUUIDBytes(kCFAllocatorDefault, UUIDBytes);
-								CFStringRef UUIDString = CFUUIDCreateString(kCFAllocatorDefault, UUIDRef);
-								FString UUID((NSString*)UUIDString);
-								CFRelease(UUIDString);
-								CFRelease(UUIDRef);
-							
-								uint64 Address = SymbolEntry.n_value;
-								uint64 BaseAddress = (uint64)IndexModule64;
-								FString AtoSCommand = FString::Printf(TEXT("\"%s\" -s %s -l 0x%lx 0x%lx"), *FullModulePath, *UUID, BaseAddress, Address);
-								int32 ReturnCode = 0;
-								FString Results;
-								
-								const FString AtoSPath = FString::Printf(TEXT("%sBinaries/Mac/UnrealAtoS"), *FPaths::EngineDir() );
-								FPlatformProcess::ExecProcess( *AtoSPath, *AtoSCommand, &ReturnCode, &Results, NULL );
-								if(ReturnCode == 0)
-								{
-									bool bSourceFileOpened = false;
-									int32 FirstIndex = -1;
-									int32 LastIndex = -1;
-									if(Results.FindChar(TCHAR('('), FirstIndex) && Results.FindLastChar(TCHAR('('), LastIndex) && FirstIndex != LastIndex)
-									{
-										int32 CloseIndex = -1;
-										int32 ColonIndex = -1;
-										if(Results.FindLastChar(TCHAR(':'), ColonIndex) && Results.FindLastChar(TCHAR(')'), CloseIndex) && CloseIndex > ColonIndex && LastIndex < ColonIndex)
-										{
-											int32 FileNamePos = LastIndex+1;
-											int32 FileNameLen = ColonIndex-FileNamePos;
-											FString FileName = Results.Mid(FileNamePos, FileNameLen);
-											FString LineNumber = Results.Mid(ColonIndex + 1, CloseIndex-(ColonIndex + 1));
-											bSourceFileOpened = SourceCodeAccessor.OpenFileAtLine( FileName, FCString::Atoi(*LineNumber), 0 );
-										}
-									}
-#if !NO_LOGGING
-									if (!bSourceFileOpened)
-									{
-										UE_LOG(LogSelectionDetails, Warning, TEXT("NavigateToFunctionSource:  Unable to find source file and line number for '%s'"), *FunctionSymbolName);
-									}
-#endif
-								}
-								break;
-							}
-						}
-					}
-				}
-				
-				delete [] StringTable;
-				delete [] SymbolEntries;
-			}
-			break;
-		}
-	}
+    // UE-195221, no easy solution for now, it is causing crash and wouldn't work anyways even without the crash
+    // Solution pending
+    
+//	for(uint32 Index = 0; Index < _dyld_image_count(); Index++)
+//	{
+//		char const* IndexName = _dyld_get_image_name(Index);
+//		FString FullModulePath(IndexName);
+//		FString Name = FPaths::GetBaseFilename(FullModulePath);
+//		if(Name == FunctionModuleName)
+//		{
+//			struct mach_header_64 const* IndexModule64 = NULL;
+//			struct load_command const* LoadCommands = NULL;
+//			
+//			struct mach_header const* IndexModule32 = _dyld_get_image_header(Index);
+//			check(IndexModule32->magic == MH_MAGIC_64);
+//			
+//			IndexModule64 = (struct mach_header_64 const*)IndexModule32;
+//			LoadCommands = (struct load_command const*)(IndexModule64 + 1);
+//			struct load_command const* Command = LoadCommands;
+//			struct symtab_command const* SymbolTable = nullptr;
+//			struct dysymtab_command const* DsymTable = nullptr;
+//			struct uuid_command* UUIDCommand = nullptr;
+//			for(uint32 CommandIndex = 0; CommandIndex < IndexModule64->ncmds; CommandIndex++)
+//			{
+//				if (Command && Command->cmd == LC_SYMTAB)
+//				{
+//					SymbolTable = (struct symtab_command const*)Command;
+//				}
+//				else if(Command && Command->cmd == LC_DYSYMTAB)
+//				{
+//					DsymTable = (struct dysymtab_command const*)Command;
+//				}
+//				else if (Command && Command->cmd == LC_UUID)
+//				{
+//					UUIDCommand = (struct uuid_command*)Command;
+//				}
+//				Command = (struct load_command const*)(((char const*)Command) + Command->cmdsize);
+//			}
+//			
+//			check(SymbolTable && DsymTable && UUIDCommand);
+//			
+//			IPlatformFile& PlatformFile = IPlatformFile::GetPlatformPhysical();
+//			IFileHandle* File = PlatformFile.OpenRead(*FullModulePath);
+//			if(File)
+//			{
+//				struct nlist_64* SymbolEntries = new struct nlist_64[SymbolTable->nsyms];
+//				check(SymbolEntries);
+//				char* StringTable = new char[SymbolTable->strsize];
+//				check(StringTable);
+//				
+//				bool FileOK = File->Seek(SymbolTable->symoff+(DsymTable->iextdefsym*sizeof(struct nlist_64)));
+//				FileOK &= File->Read((uint8*)SymbolEntries, DsymTable->nextdefsym*sizeof(struct nlist_64));
+//				
+//				FileOK &= File->Seek(SymbolTable->stroff);
+//				FileOK &= File->Read((uint8*)StringTable, SymbolTable->strsize);
+//				
+//				delete File;
+//				
+//				for(uint32 SymbolIndex = 0; FileOK && SymbolIndex < DsymTable->nextdefsym; SymbolIndex++)
+//				{
+//					struct nlist_64 const& SymbolEntry = SymbolEntries[SymbolIndex];
+//					// All the entries in the mach-o external table are functions.
+//					// The local table contains the minimal debug stabs used by dsymutil to create the DWARF dsym.
+//					if(SymbolEntry.n_un.n_strx)
+//					{
+//						if (SymbolEntry.n_value)
+//						{
+//							char const* MangledSymbolName = (StringTable+SymbolEntry.n_un.n_strx);
+//							// Remove leading '_'
+//							MangledSymbolName += 1;
+//							
+//							int32 Status = 0;
+//							char* DemangledName = abi::__cxa_demangle(MangledSymbolName, NULL, 0, &Status);
+//							
+//							FString SymbolName;
+//							if (DemangledName)
+//							{
+//								// C++ function
+//								SymbolName = DemangledName;
+//								free(DemangledName);
+//								
+//								// This contains return & arguments, it would seem that the DbgHelp API doesn't.
+//								// So we shall strip them.
+//								int32 ArgumentIndex = -1;
+//								if(SymbolName.FindLastChar(TCHAR('('), ArgumentIndex))
+//								{
+//									SymbolName.LeftInline(ArgumentIndex, false);
+//									int32 TemplateNesting = 0;
+//									
+//									int32 Pos = SymbolName.Len();
+//									// Cast operators are special & include spaces, whereas normal functions don't.
+//									int32 OperatorIndex = SymbolName.Find("operator");
+//									if(OperatorIndex >= 0)
+//									{
+//										// Trim from before the 'operator'
+//										Pos = OperatorIndex;
+//									}
+//									
+//									for(; Pos > 0; --Pos)
+//									{
+//										TCHAR Character = SymbolName[Pos - 1];
+//										if(Character == TCHAR(' ') && TemplateNesting == 0)
+//										{
+//											SymbolName.MidInline(Pos, MAX_int32, false);
+//											break;
+//										}
+//										else if(Character == TCHAR('>'))
+//										{
+//											TemplateNesting++;
+//										}
+//										else if(Character == TCHAR('<'))
+//										{
+//											TemplateNesting--;
+//										}
+//									}
+//								}
+//							}
+//							else
+//							{
+//								// C function
+//								SymbolName = MangledSymbolName;
+//							}
+//							
+//							if(FunctionSymbolName == SymbolName)
+//							{
+//								CFUUIDBytes UUIDBytes;
+//								FMemory::Memcpy(&UUIDBytes, UUIDCommand->uuid, sizeof(CFUUIDBytes));
+//								CFUUIDRef UUIDRef = CFUUIDCreateFromUUIDBytes(kCFAllocatorDefault, UUIDBytes);
+//								CFStringRef UUIDString = CFUUIDCreateString(kCFAllocatorDefault, UUIDRef);
+//								FString UUID((NSString*)UUIDString);
+//								CFRelease(UUIDString);
+//								CFRelease(UUIDRef);
+//							
+//								uint64 Address = SymbolEntry.n_value;
+//								uint64 BaseAddress = (uint64)IndexModule64;
+//								FString AtoSCommand = FString::Printf(TEXT("\"%s\" -s %s -l 0x%lx 0x%lx"), *FullModulePath, *UUID, BaseAddress, Address);
+//								int32 ReturnCode = 0;
+//								FString Results;
+//								
+//								const FString AtoSPath = FString::Printf(TEXT("%sBinaries/Mac/UnrealAtoS"), *FPaths::EngineDir() );
+//								FPlatformProcess::ExecProcess( *AtoSPath, *AtoSCommand, &ReturnCode, &Results, NULL );
+//								if(ReturnCode == 0)
+//								{
+//									bool bSourceFileOpened = false;
+//									int32 FirstIndex = -1;
+//									int32 LastIndex = -1;
+//									if(Results.FindChar(TCHAR('('), FirstIndex) && Results.FindLastChar(TCHAR('('), LastIndex) && FirstIndex != LastIndex)
+//									{
+//										int32 CloseIndex = -1;
+//										int32 ColonIndex = -1;
+//										if(Results.FindLastChar(TCHAR(':'), ColonIndex) && Results.FindLastChar(TCHAR(')'), CloseIndex) && CloseIndex > ColonIndex && LastIndex < ColonIndex)
+//										{
+//											int32 FileNamePos = LastIndex+1;
+//											int32 FileNameLen = ColonIndex-FileNamePos;
+//											FString FileName = Results.Mid(FileNamePos, FileNameLen);
+//											FString LineNumber = Results.Mid(ColonIndex + 1, CloseIndex-(ColonIndex + 1));
+//											bSourceFileOpened = SourceCodeAccessor.OpenFileAtLine( FileName, FCString::Atoi(*LineNumber), 0 );
+//										}
+//									}
+//#if !NO_LOGGING
+//									if (!bSourceFileOpened)
+//									{
+//										UE_LOG(LogSelectionDetails, Warning, TEXT("NavigateToFunctionSource:  Unable to find source file and line number for '%s'"), *FunctionSymbolName);
+//									}
+//#endif
+//								}
+//								break;
+//							}
+//						}
+//					}
+//				}
+//				
+//				delete [] StringTable;
+//				delete [] SymbolEntries;
+//			}
+//			break;
+//		}
+//	}
 	
 #endif	// PLATFORM_WINDOWS
 }
