Add overrall editor scale as a user setting.

diff --git a/Engine/Source/Editor/UnrealEd/Classes/Settings/EditorStyleSettings.h b/Engine/Source/Editor/UnrealEd/Classes/Settings/EditorStyleSettings.h
--- a/Engine/Source/Editor/UnrealEd/Classes/Settings/EditorStyleSettings.h
+++ b/Engine/Source/Editor/UnrealEd/Classes/Settings/EditorStyleSettings.h
@@ -47,6 +47,12 @@ public:
 	 */
 	UPROPERTY(EditAnywhere, Category=UserInterface, meta = (ConfigRestartRequired = true, DisplayName="Enable High DPI Support"))
 	bool bEnableHighDPIAwareness;
+	
+	/** 
+	 * Scales the entire editor interface up or down. 
+	 */
+	UPROPERTY(EditAnywhere, Config, Category=UserInterface, meta=(ClampMin=0.1))
+	float ApplicationScale = 1.0f;
 
 	/**
 	 * Whether to enable the Editor UI Layout configuration tools for the user.
diff --git a/Engine/Source/Editor/UnrealEd/Private/Settings/EditorStyleClasses.cpp b/Engine/Source/Editor/UnrealEd/Private/Settings/EditorStyleClasses.cpp
--- a/Engine/Source/Editor/UnrealEd/Private/Settings/EditorStyleClasses.cpp
+++ b/Engine/Source/Editor/UnrealEd/Private/Settings/EditorStyleClasses.cpp
@@ -1,6 +1,7 @@
 // Copyright Epic Games, Inc. All Rights Reserved.
 
 #include "CoreGlobals.h"
+#include "Framework/Application/SlateApplication.h"
 #include "HAL/IConsoleManager.h"
 #include "HAL/Platform.h"
 #include "Math/Color.h"
@@ -66,6 +67,8 @@ void UEditorStyleSettings::Init()
 		CurrentAppliedTheme = USlateThemeManager::Get().GetCurrentThemeID();
 		SaveConfig(); 
 	}
+	
+	FSlateApplication::Get().SetApplicationScale(ApplicationScale);
 
 	// Set from CVar 
 	IConsoleVariable* CVar = IConsoleManager::Get().FindConsoleVariable(TEXT("EnableHighDPIAwareness"));
@@ -93,6 +96,10 @@ void UEditorStyleSettings::PostEditChangeProperty(struct FPropertyChangedEvent&
 	{
 		GConfig->SetBool(TEXT("HDPI"), TEXT("EnableHighDPIAwareness"), bEnableHighDPIAwareness, GEditorSettingsIni);
 	}
+	else if (PropertyName == GET_MEMBER_NAME_CHECKED(UEditorStyleSettings, ApplicationScale))
+	{
+		FSlateApplication::Get().SetApplicationScale(ApplicationScale);
+	}
 
 //	if (!FUnrealEdMisc::Get().IsDeletePreferences())
 	{
