Level Editor: Improvements to drag-duplicate functionality

  - Pressing ALT during a drag or drag-duplicate operation will stamp a duplicate of the current selection at their current positions, allowing for much quicker duplication on a grid
  - Releasing the mouse and initiating a new drag operation while keeping ALT pressed will create a new drag-duplicate operation. Previously ALT would need to be released and re-pressed to start a new duplicate.

diff --git a/Engine/Source/Editor/UnrealEd/Private/LevelEditorViewport.cpp b/Engine/Source/Editor/UnrealEd/Private/LevelEditorViewport.cpp
--- a/Engine/Source/Editor/UnrealEd/Private/LevelEditorViewport.cpp
+++ b/Engine/Source/Editor/UnrealEd/Private/LevelEditorViewport.cpp
@@ -3142,11 +3142,38 @@ bool FLevelEditorViewportClient::InputKey(const FInputKeyEventArgs& InEventArgs)
 		GEditor->SetPreviewMeshMode(false);
 	}
 
-	// Clear Duplicate Actors mode when ALT and all mouse buttons are released
-	if ( !InputState.IsAltButtonPressed() && !InputState.IsAnyMouseButtonDown() )
+	// Clear Duplicate Actors mode when mouse buttons are released
+	if ( !InputState.IsAnyMouseButtonDown() )
 	{
 		bDuplicateActorsInProgress = false;
 	}
+	// Stamp a duplicate at the current position during a drag when ALT is pressed
+	else if (InEventArgs.Event == IE_Pressed && bDraggingByHandle && InputState.IsAltButtonPressed())
+	{
+		TSharedPtr<ILevelEditor> LevelEditor = ParentLevelEditor.Pin();
+		UTypedElementCommonActions* CommonActions = LevelEditor ? LevelEditor->GetCommonActions() : nullptr;
+		if (CommonActions)
+		{
+			FTypedElementListConstRef ElementsToManipulate = GetElementsToManipulate();
+			TArray<FTypedElementHandle> DuplicatedElements;
+			{
+				// Do not used the cached manipulation list here here, as it will have removed attachments, and we do want to duplicate those
+				DuplicatedElements = CommonActions->DuplicateSelectedElements(GetSelectionSet(), GetWorld(), FVector::ZeroVector);
+			}
+
+			// Do not select the duplicate actors so that the drag operation can continue freely
+			if (DuplicatedElements.Num() > 0)
+			{
+				// Although we did not change the selection, we still need to notify legacy mode tools of the duplication
+				TArray<AActor*> SelectedActors = GetSelectionSet()->GetSelectedObjects<AActor>();
+				constexpr bool bDidOffsetDuplicate = false;
+				ModeTools->ActorsDuplicatedNotify(SelectedActors, SelectedActors, bDidOffsetDuplicate);
+			}
+
+			RedrawAllViewportsIntoThisScene();
+		}
+	}
+
 	
 	return bHandled;
 }
