Add a mechanism to adjust core subscription via commandline when executing MPCook.

diff --git a/Engine/Source/Editor/UnrealEd/Private/Cooker/CookDirector.cpp b/Engine/Source/Editor/UnrealEd/Private/Cooker/CookDirector.cpp
--- a/Engine/Source/Editor/UnrealEd/Private/Cooker/CookDirector.cpp
+++ b/Engine/Source/Editor/UnrealEd/Private/Cooker/CookDirector.cpp
@@ -1057,6 +1057,14 @@ void FCookDirector::ActivateMachineResourceReduction()
 	int32 HyperThreadCount = FPlatformMisc::NumberOfCoresIncludingHyperthreads();
 	int32 NumberOfHyperThreadsPerCore = HyperThreadCount / NumberOfCores;
 	CoreLimit = FMath::Max(NumberOfCores / NumProcesses, 1);
+
+	const TCHAR* CommandLine = FCommandLine::Get();
+	float CoreOversubscription = 1.0f;
+	if (FParse::Value(CommandLine, TEXT("-MPCookCoreSubscription="), CoreOversubscription))
+	{
+		CoreLimit = FMath::Clamp(static_cast<int32>(CoreLimit*CoreOversubscription), 1, NumberOfCores);
+	}
+
 	int32 CoreIncludingHyperthreadsLimit = CoreLimit * NumberOfHyperThreadsPerCore;
 	int32 NumberOfWorkers = FMath::Max(CoreLimit - 1, 1) * NumberOfHyperThreadsPerCore;
 
@@ -1252,6 +1260,7 @@ FString FCookDirector::GetWorkerCommandLine(FWorkerId WorkerId, int32 ProfileId)
 			Token.StartsWith(TEXT("-ShowCookWorker")) ||
 			Token.StartsWith(TEXT("-CoreLimit")) ||
 			Token.StartsWith(TEXT("-PhysicalCoreLimit")) ||
+			Token.StartsWith(TEXT("-MPCookCoreSubscription")) ||
 			Token.StartsWith(TEXT("-CookProcessCount=")) ||
 			Token.StartsWith(TEXT("-abslog=")) ||
 			Token.StartsWith(TEXT("-unattended"))
