Fix root cause of an assertion on Actor spawn that leads to an editor crash after undoing certain property changes that precede a Blueprint compile.

Change summary:
- Added FReplaceReferenceHelper::IncludeDSOs()
- Modified FReplaceReferenceHelper::IncludeCDO() to call IncludeDSOs() where appropriate (i.e. when a template mapping is not otherwise explicitly provided)
- Modified FBlueprintCompileReinstancer::ReplaceInstancesOfClass_Inner() to always include old->new template mappings for delegate listeners (e.g. undo buffer), but continue to restrict explicit search-and-replace unless otherwise enabled

diff --git a/Engine/Source/Editor/UnrealEd/Private/Kismet2/KismetReinstanceUtilities.cpp b/Engine/Source/Editor/UnrealEd/Private/Kismet2/KismetReinstanceUtilities.cpp
--- a/Engine/Source/Editor/UnrealEd/Private/Kismet2/KismetReinstanceUtilities.cpp
+++ b/Engine/Source/Editor/UnrealEd/Private/Kismet2/KismetReinstanceUtilities.cpp
@@ -121,7 +121,26 @@ struct FReplaceReferenceHelper
 			*NamesOfClasses, *NamesOfObjects);
 	}
 
-	static void IncludeCDO(UClass* OldClass, UClass* NewClass, TMap<UObject*, UObject*> &OldToNewInstanceMap, TArray<UObject*> &SourceObjects, UObject* OriginalCDO, TMap<UClass*, TMap<UObject*, UObject*>>* OldToNewTemplates = nullptr)
+	static void IncludeDSOs(UObject* OldOuter, UObject* NewOuter, TMap<UObject*, UObject*>& OldToNewInstanceMap, TArray<UObject*>& SourceObjects)
+	{
+		TArray<UObject*> OldSubObjArray;
+		constexpr bool bIncludeNestedObjects = false;
+		GetObjectsWithOuter(OldOuter, OldSubObjArray, bIncludeNestedObjects);
+		for (UObject* OldSubObj : OldSubObjArray)
+		{
+			if (UObject* NewSubObj = NewOuter->GetDefaultSubobjectByName(OldSubObj->GetFName()))
+			{
+				ensure(!OldToNewInstanceMap.Contains(OldSubObj));
+				OldToNewInstanceMap.Add(OldSubObj, NewSubObj);
+				SourceObjects.Add(OldSubObj);
+
+				// Recursively include any nested DSOs
+				IncludeDSOs(OldSubObj, NewSubObj, OldToNewInstanceMap, SourceObjects);
+			}
+		}
+	}
+
+	static void IncludeCDO(UClass* OldClass, UClass* NewClass, TMap<UObject*, UObject*>& OldToNewInstanceMap, TArray<UObject*>& SourceObjects, UObject* OriginalCDO, TMap<UClass*, TMap<UObject*, UObject*>>* OldToNewTemplates = nullptr)
 	{
 		UObject* OldCDO = OldClass->GetDefaultObject();
 		UObject* NewCDO = NewClass->GetDefaultObject();
@@ -140,13 +159,15 @@ struct FReplaceReferenceHelper
 			OldToNewInstanceMap.Add(OldCDO, NewCDO);
 			// Add in the old CDO to this pass, so CDO references are fixed up
 			SourceObjects.Add(OldCDO);
+			// Add any old->new CDO default subobject mappings
+			IncludeDSOs(OldCDO, NewCDO, OldToNewInstanceMap, SourceObjects);
 		}
 
-
 		if (OriginalCDO && OriginalCDO != OldCDO)
 		{
 			OldToNewInstanceMap.Add(OriginalCDO, NewCDO);
 			SourceObjects.Add(OriginalCDO);
+			IncludeDSOs(OriginalCDO, NewCDO, OldToNewInstanceMap, SourceObjects);
 		}
 	}
 
@@ -2613,17 +2634,26 @@ void FBlueprintCompileReinstancer::ReplaceInstancesOfClass_Inner(const TMap<UCla
 	TArray<UObject*> SourceObjects;
 	OldToNewInstanceMap.GenerateKeyArray(SourceObjects);
 	
-	if (bReplaceReferencesToOldCDOs)
-	{
-		for (TPair<UClass*, UClass*> OldToNewClass : InOldToNewClassMap)
-		{
-			UClass* OldClass = OldToNewClass.Key;
-			UClass* NewClass = OldToNewClass.Value;
-			check(OldClass && NewClass);
-			check(OldClass != NewClass || IsReloadActive());
-
-			FReplaceReferenceHelper::IncludeCDO(OldClass, NewClass, OldToNewInstanceMap, SourceObjects, InOriginalCDO, Params.OldToNewTemplates);
-
+	TArray<UObject*> OldCDOSourceObjects;
+	for (TPair<UClass*, UClass*> OldToNewClass : InOldToNewClassMap)
+	{
+		UClass* OldClass = OldToNewClass.Key;
+		UClass* NewClass = OldToNewClass.Value;
+		check(OldClass && NewClass);
+		check(OldClass != NewClass || IsReloadActive());
+
+		// Always map old to new instances of CDOs along with any owned subobject(s). This allows delegates to be
+		// notified that these instances have been replaced. However, we don't proactively find and replace those
+		// references ourselves unless input parameters have explicitly configured this path to do so (see below).
+		FReplaceReferenceHelper::IncludeCDO(OldClass, NewClass, OldToNewInstanceMap, OldCDOSourceObjects, InOriginalCDO, Params.OldToNewTemplates);
+		if (bReplaceReferencesToOldCDOs)
+		{
+			// This means we'll proactively find and replace references to old CDOs and any owned subobject(s). It
+			// has an additional cost and is not enabled by default, since most systems don't store these references;
+			// those that do (e.g. the editor's transaction buffer) may do their own reference replacement pass instead.
+			SourceObjects.Append(OldCDOSourceObjects);
+
+			// This is part of the legacy reload path; it is only enabled if we're also replacing references to old CDOs.
 			if (bClassObjectReplaced)
 			{
 				FReplaceReferenceHelper::IncludeClass(OldClass, NewClass, OldToNewInstanceMap, SourceObjects, ObjectsReplaced);
