MPCook: Retraction: Change RetractionHandler to always be allocated so that it can carry information across multiple attempts to retract.
Change RetractionHandler to tick within the communication lock; this increases contention a small amount but makes the multithreading simpler.
#rnx

diff --git a/Engine/Source/Editor/UnrealEd/Private/Cooker/CookDirector.cpp b/Engine/Source/Editor/UnrealEd/Private/Cooker/CookDirector.cpp
--- a/Engine/Source/Editor/UnrealEd/Private/Cooker/CookDirector.cpp
+++ b/Engine/Source/Editor/UnrealEd/Private/Cooker/CookDirector.cpp
@@ -68,11 +68,25 @@ public:
 	void Initialize();
 	/** Initialize to handle an unexpected RetractionResultsMessage. */
 	void InitializeForResultsMessage(const FWorkerId& FromWorker);
-	void TickFromSchedulerThread(bool bAnyIdle, bool& bOutComplete);
+	void TickFromSchedulerThread(bool bAllWorkersConnected, bool bAnyIdle, int32 BusiestNumAssignments);
 	/** Hook called by the director when a retraction message comes in. */
-	void HandleRetractionMessage(const FWorkerId& FromWorker, TConstArrayView<FName> Packages);
+	void HandleRetractionMessage(FMPCollectorServerMessageContext& Context, bool bReadSuccessful,
+		FRetractionResultsMessage&& Message);
 
 private:
+	enum class ERetractionState : uint8
+	{
+		Idle,
+		WantToRetract,
+		WaitingForResponse,
+	};
+
+private:
+	/** Try to select a worker for retraction */
+	ERetractionState TickWantToRetract();
+	/** Tick the asynchronous wait for the message to come in, and synchronously handle it when it does. */
+	ERetractionState TickWaitingForResponse();
+
 	/**
 	 * Pick workers to give the retracted packages to, and assign those packages to the worker
 	 * in the local and remote state.
@@ -82,12 +96,14 @@ private:
 	TArray<FWorkerId> CalculateWorkersToSplitOver(int32 NumPackages, const FWorkerId& FromWorker,
 		TConstArrayView<TRefCountPtr<FCookWorkerServer>> LocalRemoteWorkers);
 
+private:
 	FCookDirector& Director;
 	FWorkerId ExpectedWorker;
 	TMap<FWorkerId, TArray<FName>> PackagesToRetract;
 	FWorkerId WorkerWithResults;
 	double MessageSentTimeSeconds = 0.;
 	double LastWarnTimeSeconds = 0.;
+	ERetractionState RetractionState = ERetractionState::Idle;
 };
 
 FCookDirector::FCookDirector(UCookOnTheFlyServer& InCOTFS, int32 CookProcessCount)
@@ -99,6 +115,7 @@ FCookDirector::FCookDirector(UCookOnTheFlyServer& InCOTFS, int32 CookProcessCoun
 	WorkersStalledWarnTimeSeconds = MAX_flt;
 	ShutdownEvent->Reset();
 	LocalWorkerProfileData = MakeUnique<FCookWorkerProfileData>();
+	RetractionHandler = MakeUnique<FRetractionHandler>(*this);
 
 	bool bConfigValid;
 	ParseConfig(CookProcessCount, bConfigValid);
@@ -124,7 +141,8 @@ FCookDirector::FCookDirector(UCookOnTheFlyServer& InCOTFS, int32 CookProcessCoun
 	Register(new IMPCollectorCbServerMessage<FRetractionResultsMessage>([this]
 	(FMPCollectorServerMessageContext& Context, bool bReadSuccessful, FRetractionResultsMessage&& Message)
 		{
-			HandleRetractionMessage(Context, bReadSuccessful, MoveTemp(Message));
+			// Called from inside CommunicationLock
+			RetractionHandler->HandleRetractionMessage(Context, bReadSuccessful, MoveTemp(Message));
 		}, TEXT("HandleRetractionMessage")));
 	Register(new IMPCollectorCbServerMessage<FHeartbeatMessage>([this]
 	(FMPCollectorServerMessageContext& Context, bool bReadSuccessful, FHeartbeatMessage&& Message)
@@ -526,10 +544,7 @@ void FCookDirector::TickFromSchedulerThread()
 			}
 		}
 		ReassignAbortedPackages(DeferredPackagesToReassign);
-	}
-	if (bAllWorkersConnected && (RetractionHandler.IsValid() || bAnyIdle))
-	{
-		TickRetractionFromSchedulerThread(bAnyIdle, BusiestNumAssignments);
+		RetractionHandler->TickFromSchedulerThread(bAllWorkersConnected, bAnyIdle, BusiestNumAssignments);
 	}
 
 	bool bIsStalled = bLocalWorkerIdle && !COTFS.PackageDatas->GetAssignedToWorkerSet().IsEmpty() && WorkersWithMessage.IsEmpty();
@@ -1462,54 +1477,12 @@ void FCookDirector::LogCookStats(FCookStatsManager::AddStatFuncRef AddStat)
 }
 #endif
 
-void FCookDirector::HandleRetractionMessage(FMPCollectorServerMessageContext& Context, bool bReadSuccessful,
-	FRetractionResultsMessage&& Message)
-{
-	if (!bReadSuccessful)
-	{
-		UE_LOG(LogCook, Error, TEXT("Corrupt RetractionResultsMessage received from CookWorker %d. It will be ignored and packages may fail to cook."),
-			Context.GetProfileId());
-		return;
-	}
-
-	if (!RetractionHandler)
-	{
-		UE_LOG(LogCook, Warning, TEXT("Retractionmessage received from CookWorker %d when we were not expecting one."), Context.GetProfileId());
-		RetractionHandler = MakeUnique<FRetractionHandler>(*this);
-		RetractionHandler->InitializeForResultsMessage(Context.GetWorkerId());
-	}
-	RetractionHandler->HandleRetractionMessage(Context.GetWorkerId(), Message.ReturnedPackages);
-}
-
-void FCookDirector::TickRetractionFromSchedulerThread(bool bAnyIdle, int32 BusiestNumAssignments)
-{
-	if (!RetractionHandler)
-	{
-		if (bAnyIdle && BusiestNumAssignments > RetractionMinimumNumAssignments)
-		{
-			RetractionHandler = MakeUnique<FRetractionHandler>(*this);
-			RetractionHandler->Initialize();
-		}
-		else
-		{
-			return;
-		}
-	}
-
-	bool bComplete;
-	RetractionHandler->TickFromSchedulerThread(bAnyIdle, bComplete);
-	if (bComplete)
-	{
-		RetractionHandler.Reset();
-	}
-}
-
 FCookDirector::FRetractionHandler::FRetractionHandler(FCookDirector& InDirector)
 	: Director(InDirector)
 {
 }
 
-void FCookDirector::FRetractionHandler::Initialize()
+FCookDirector::FRetractionHandler::ERetractionState FCookDirector::FRetractionHandler::TickWantToRetract()
 {
 	FWorkerId BusiestWorker;
 	TArray<FWorkerId> IdleWorkers;
@@ -1543,8 +1516,8 @@ void FCookDirector::FRetractionHandler::Initialize()
 	if (IdleWorkers.IsEmpty() || BusiestNumAssignments < RetractionMinimumNumAssignments)
 	{
 		// Worker loads changed after the point where we decided to initialize the RetractionHandler,
-		// and retraction is no longer needed. Early exit now and this will be deleted later in Tick.
-		return;
+		// and retraction is no longer needed.
+		return ERetractionState::Idle;
 	}
 
 	// Plan to divide the assignments evenly between all idle workers and the one busiest worker. This means
@@ -1580,6 +1553,8 @@ void FCookDirector::FRetractionHandler::Initialize()
 		MessageSentTimeSeconds = FPlatformTime::Seconds();
 		LastWarnTimeSeconds = MessageSentTimeSeconds;
 	}
+
+	return ERetractionState::WaitingForResponse;
 }
 
 void FCookDirector::FRetractionHandler::InitializeForResultsMessage(const FWorkerId& FromWorker)
@@ -1587,14 +1562,43 @@ void FCookDirector::FRetractionHandler::InitializeForResultsMessage(const FWorke
 	ExpectedWorker = FromWorker;
 }
 
-void FCookDirector::FRetractionHandler::TickFromSchedulerThread(bool bAnyIdle, bool& bOutComplete)
+void FCookDirector::FRetractionHandler::TickFromSchedulerThread(bool bAllWorkersConnected, bool bAnyIdle, int32 BusiestNumAssignments)
+{
+	switch (RetractionState)
+	{
+	case ERetractionState::Idle:
+		if (!bAnyIdle || !bAllWorkersConnected || BusiestNumAssignments <= RetractionMinimumNumAssignments)
+		{
+			break;
+		}
+		RetractionState = ERetractionState::WantToRetract;
+		RetractionState = TickWantToRetract();
+		break;
+	case ERetractionState::WantToRetract:
+		if (!bAnyIdle || !bAllWorkersConnected || BusiestNumAssignments <= RetractionMinimumNumAssignments)
+		{
+			RetractionState = ERetractionState::Idle;
+			break;
+		}
+		RetractionState = TickWantToRetract();
+		break;
+	case ERetractionState::WaitingForResponse:
+		RetractionState = TickWaitingForResponse();
+		break;
+	default:
+		checkNoEntry();
+		break;
+	}
+}
+
+FCookDirector::FRetractionHandler::ERetractionState FCookDirector::FRetractionHandler::TickWaitingForResponse()
 {
+	// Called from inside CommunicationLock
 	if (ExpectedWorker.IsInvalid())
 	{
 		// We decided to cancel
 		checkf(PackagesToRetract.IsEmpty(), TEXT("We should not have any packages when we cancelled."));
-		bOutComplete = true;
-		return;
+		return ERetractionState::Idle;
 	}
 	if (WorkerWithResults.IsInvalid())
 	{
@@ -1603,27 +1607,23 @@ void FCookDirector::FRetractionHandler::TickFromSchedulerThread(bool bAnyIdle, b
 
 		if (static_cast<float>(CurrentTime - LastWarnTimeSeconds) < WarnDuration)
 		{
-			bOutComplete = false;
-			return;
+			return ERetractionState::WaitingForResponse;
 		}
 		check(ExpectedWorker.IsRemote());
 		{
-			FScopeLock CommunicationScopeLock(&Director.CommunicationLock);
 			TRefCountPtr<FCookWorkerServer>* RemoteWorkerPtr = Director.RemoteWorkers.Find(ExpectedWorker.GetRemoteIndex());
 			if (!RemoteWorkerPtr)
 			{
 				// The CookWorker aborted and we already reassigned all of its packages; stop waiting for a retraction message from it.
 				check(PackagesToRetract.IsEmpty()); // Otherwise WorkerWithResults would have been set
 				ExpectedWorker = FWorkerId::Invalid();
-				bOutComplete = true;
-				return;
+				return ERetractionState::Idle;
 			}
 		}
-		bOutComplete = false;
 		UE_CLOG(!IsCookIgnoreTimeouts(), LogCook, Display, TEXT("%s has not responded to a RetractionRequest message for %.1f seconds. Continuing to wait..."),
 			*Director.GetDisplayName(ExpectedWorker), static_cast<float>(CurrentTime - MessageSentTimeSeconds));
 		LastWarnTimeSeconds = CurrentTime;
-		return;
+		return ERetractionState::WaitingForResponse;
 	}
 
 	// Convert names to packagedatas and collect results from all CookWorkers who sent a message.
@@ -1633,7 +1633,6 @@ void FCookDirector::FRetractionHandler::TickFromSchedulerThread(bool bAnyIdle, b
 		TRefCountPtr<FCookWorkerServer> RemoteWorker;
 		if (Pair.Key.IsRemote())
 		{
-			FScopeLock CommunicationScopeLock(&Director.CommunicationLock);
 			TRefCountPtr<FCookWorkerServer>* FoundRemoteWorker = Director.RemoteWorkers.Find(Pair.Key.GetRemoteIndex());
 			if (FoundRemoteWorker)
 			{
@@ -1666,15 +1665,33 @@ void FCookDirector::FRetractionHandler::TickFromSchedulerThread(bool bAnyIdle, b
 	ExpectedWorker = FWorkerId::Invalid();
 	WorkerWithResults = FWorkerId::Invalid();
 	PackagesToRetract.Empty();
-	bOutComplete = true;
+	return ERetractionState::Idle;
 }
 
-void FCookDirector::FRetractionHandler::HandleRetractionMessage(const FWorkerId& FromWorker, TConstArrayView<FName> Packages)
+void FCookDirector::FRetractionHandler::HandleRetractionMessage(FMPCollectorServerMessageContext& Context,
+	bool bReadSuccessful, FRetractionResultsMessage&& Message)
 {
+	// Called from inside CommunicationLock
+	if (!bReadSuccessful)
+	{
+		UE_LOG(LogCook, Error,
+			TEXT("Corrupt RetractionResultsMessage received from CookWorker %d. It will be ignored and packages may fail to cook."),
+			Context.GetProfileId());
+		return;
+	}
+
+	if (RetractionState != ERetractionState::WaitingForResponse)
+	{
+		UE_LOG(LogCook, Warning, TEXT("Retractionmessage received from CookWorker %d when we were not expecting one."),
+			Context.GetProfileId());
+		InitializeForResultsMessage(Context.GetWorkerId());
+		RetractionState = ERetractionState::WaitingForResponse;
+	}
+
 	UE_CLOG(WorkerWithResults.IsValid(), LogCook, Error,
 		TEXT("Unexpectedly received RetractionResults message from multiple CookWorkers. Merging the results."));
-	WorkerWithResults = FromWorker;
-	PackagesToRetract.FindOrAdd(FromWorker).Append(Packages);
+	WorkerWithResults = Context.GetWorkerId();
+	PackagesToRetract.FindOrAdd(WorkerWithResults).Append(Message.ReturnedPackages);
 }
 
 void FCookDirector::FRetractionHandler::ReassignPackages(const FWorkerId& FromWorker, TConstArrayView<FPackageData*> Packages)
diff --git a/Engine/Source/Editor/UnrealEd/Private/Cooker/CookDirector.h b/Engine/Source/Editor/UnrealEd/Private/Cooker/CookDirector.h
--- a/Engine/Source/Editor/UnrealEd/Private/Cooker/CookDirector.h
+++ b/Engine/Source/Editor/UnrealEd/Private/Cooker/CookDirector.h
@@ -202,9 +202,6 @@ private:
 #if ENABLE_COOK_STATS
 	void LogCookStats(FCookStatsManager::AddStatFuncRef AddStat);
 #endif
-	void TickRetractionFromSchedulerThread(bool bAnyIdle, int32 BusiestNumAssignments);
-	void HandleRetractionMessage(FMPCollectorServerMessageContext& Context, bool bReadSuccessful,
-		FRetractionResultsMessage&& Message);
 	void AssignRequests(TArray<FWorkerId>&& InWorkers, TArray<TRefCountPtr<FCookWorkerServer>>& InRemoteWorkers, 
 		TArrayView<FPackageData*> Requests, TArray<FWorkerId>& OutAssignments,
 		TMap<FPackageData*, TArray<FPackageData*>>&& RequestGraph);
@@ -250,6 +247,7 @@ private:
 	int32 CoreLimit = 0;
 	EShowWorker ShowWorkerOption = EShowWorker::CombinedLogs;
 	ELoadBalanceAlgorithm LoadBalanceAlgorithm = ELoadBalanceAlgorithm::CookBurden;
+	/** Whether the director is allowed to cook any packages. True by default, false by commandline parameter. */
 	bool bAllowLocalCooks = true;
 
 	// Data only accessible from the CommunicationThread (or if the CommunicationThread is inactive)
