Animation: Increase the slow task amount of work by 1, since EnterProgressFrame is called for AnimationLength and the slow task dialog reports ensures when the amount of work >= AnimationLength

#rnx

diff --git a/Engine/Source/Editor/UnrealEd/Private/Fbx/FbxAnimationExport.cpp b/Engine/Source/Editor/UnrealEd/Private/Fbx/FbxAnimationExport.cpp
--- a/Engine/Source/Editor/UnrealEd/Private/Fbx/FbxAnimationExport.cpp
+++ b/Engine/Source/Editor/UnrealEd/Private/Fbx/FbxAnimationExport.cpp
@@ -587,7 +587,7 @@ void FFbxExporter::ExportAnimTrack(IAnimTrackAdapter& AnimTrackAdapter, AActor*
 
 	const float TickRate = 1.0f/FrameRate;
 
-	FScopedSlowTask SlowTask(AnimationLength, NSLOCTEXT("UnrealEd", "ExportAnimationProgress", "Exporting Animation"));
+	FScopedSlowTask SlowTask(AnimationLength + 1, NSLOCTEXT("UnrealEd", "ExportAnimationProgress", "Exporting Animation"));
 	SlowTask.MakeDialog(true);
 
 	for (int32 FrameCount = 0; FrameCount <= AnimationLength; ++FrameCount)
