Detect bad SCS component templates in diff tool - we have seen assets that are referencing CDOs as component templates

diff --git a/Engine/Source/Editor/Kismet/Private/DiffControl.cpp b/Engine/Source/Editor/Kismet/Private/DiffControl.cpp
--- a/Engine/Source/Editor/Kismet/Private/DiffControl.cpp
+++ b/Engine/Source/Editor/Kismet/Private/DiffControl.cpp
@@ -254,7 +254,10 @@ void FSCSDiffControl::GenerateTreeEntries(TArray<TSharedPtr<FBlueprintDifference
 		{
 			return SNew(STextBlock)
 					.Text(DiffViewUtils::SCSDiffMessage(Entry, ObjectName))
-					.ColorAndOpacity(DiffViewUtils::Differs());
+					.ColorAndOpacity(Entry.DiffType == ETreeDiffType::NODE_CORRUPTED ?
+						DiffViewUtils::Conflicting() :
+						DiffViewUtils::Differs()
+						);
 		};
 
 		for (const FSCSDiffEntry& Difference : DifferingProperties.Entries)
diff --git a/Engine/Source/Editor/UnrealEd/Private/DiffUtils.cpp b/Engine/Source/Editor/UnrealEd/Private/DiffUtils.cpp
--- a/Engine/Source/Editor/UnrealEd/Private/DiffUtils.cpp
+++ b/Engine/Source/Editor/UnrealEd/Private/DiffUtils.cpp
@@ -750,6 +750,23 @@ void DiffUtils::CompareUnrelatedSCS(const UBlueprint* Old, const TArray< FSCSRes
 				OutDifferingEntries.Entries.Push(Diff);
 			}
 
+			// did it become corrupted? or stop being corrupted?
+			const bool bNewIsCorrupt = New->GeneratedClass && NewEntry->Object && !NewEntry->Object->IsIn(New->GeneratedClass);
+			const bool bOldIsCorrupt = Old->GeneratedClass && OldNode.Object && !OldNode.Object->IsIn(Old->GeneratedClass);
+			if (bNewIsCorrupt != bOldIsCorrupt)
+			{
+				if (bNewIsCorrupt)
+				{
+					FSCSDiffEntry Diff = { NewEntry->Identifier, ETreeDiffType::NODE_CORRUPTED, FSingleObjectDiffEntry() };
+					OutDifferingEntries.Entries.Push(Diff);
+				}
+				else
+				{
+					FSCSDiffEntry Diff = { OldNode.Identifier, ETreeDiffType::NODE_CORRUPTED, FSingleObjectDiffEntry() };
+					OutDifferingEntries.Entries.Push(Diff);
+				}
+			}
+
 			// no change! Do nothing.
 		}
 		else
@@ -1505,6 +1522,9 @@ FText DiffViewUtils::SCSDiffMessage(const FSCSDiffEntry& Difference, FText Objec
 	case ETreeDiffType::NODE_MOVED:
 		Text = FText::Format(NSLOCTEXT("DiffViewUtils", "NodeMoved", "Moved Node {0} in {1}"), NodeName, ObjectName);
 		break;
+	case ETreeDiffType::NODE_CORRUPTED:
+		Text = FText::Format(NSLOCTEXT("DiffViewUtils", "NodeMoved", "Node {0} in {1} has corrupt outer - must be recreated"), NodeName, ObjectName);
+		break;
 	}
 	return Text;
 }
diff --git a/Engine/Source/Editor/UnrealEd/Public/DiffUtils.h b/Engine/Source/Editor/UnrealEd/Public/DiffUtils.h
--- a/Engine/Source/Editor/UnrealEd/Public/DiffUtils.h
+++ b/Engine/Source/Editor/UnrealEd/Public/DiffUtils.h
@@ -220,6 +220,7 @@ namespace ETreeDiffType
 		NODE_TYPE_CHANGED,
 		NODE_PROPERTY_CHANGED,
 		NODE_MOVED,
+		NODE_CORRUPTED
 		/** We could potentially try to identify hierarchy reorders separately from add/remove */
 	};
 }
