Add directionality to node corruption detection - time moves forward and people appreciate it when you remember that. Also fix a localization key collision

diff --git a/Engine/Source/Editor/UnrealEd/Private/DiffUtils.cpp b/Engine/Source/Editor/UnrealEd/Private/DiffUtils.cpp
--- a/Engine/Source/Editor/UnrealEd/Private/DiffUtils.cpp
+++ b/Engine/Source/Editor/UnrealEd/Private/DiffUtils.cpp
@@ -762,7 +762,7 @@ void DiffUtils::CompareUnrelatedSCS(const UBlueprint* Old, const TArray< FSCSRes
 				}
 				else
 				{
-					FSCSDiffEntry Diff = { OldNode.Identifier, ETreeDiffType::NODE_CORRUPTED, FSingleObjectDiffEntry() };
+					FSCSDiffEntry Diff = { OldNode.Identifier, ETreeDiffType::NODE_FIXED, FSingleObjectDiffEntry() };
 					OutDifferingEntries.Entries.Push(Diff);
 				}
 			}
@@ -1523,7 +1523,10 @@ FText DiffViewUtils::SCSDiffMessage(const FSCSDiffEntry& Difference, FText Objec
 		Text = FText::Format(NSLOCTEXT("DiffViewUtils", "NodeMoved", "Moved Node {0} in {1}"), NodeName, ObjectName);
 		break;
 	case ETreeDiffType::NODE_CORRUPTED:
-		Text = FText::Format(NSLOCTEXT("DiffViewUtils", "NodeMoved", "Node {0} in {1} has corrupt outer - must be recreated"), NodeName, ObjectName);
+		Text = FText::Format(NSLOCTEXT("DiffViewUtils", "NodeCorrupted", "Node {0} in {1} has corrupt outer - must be recreated"), NodeName, ObjectName);
+		break;
+	case ETreeDiffType::NODE_FIXED:
+		Text = FText::Format(NSLOCTEXT("DiffViewUtils", "NodeFixed", "Node {0} in {1} has been recreated to correct outer"), NodeName, ObjectName);
 		break;
 	}
 	return Text;
diff --git a/Engine/Source/Editor/UnrealEd/Public/DiffUtils.h b/Engine/Source/Editor/UnrealEd/Public/DiffUtils.h
--- a/Engine/Source/Editor/UnrealEd/Public/DiffUtils.h
+++ b/Engine/Source/Editor/UnrealEd/Public/DiffUtils.h
@@ -220,7 +220,8 @@ namespace ETreeDiffType
 		NODE_TYPE_CHANGED,
 		NODE_PROPERTY_CHANGED,
 		NODE_MOVED,
-		NODE_CORRUPTED
+		NODE_CORRUPTED,
+		NODE_FIXED
 		/** We could potentially try to identify hierarchy reorders separately from add/remove */
 	};
 }
