WorldPartitionDataLayerToAssetCommandlet: Avoid deleting / reusing DataLayerAssets with matching names from other folders when converting

#rnx

diff --git a/Engine/Source/Editor/UnrealEd/Classes/Commandlets/WorldPartitionDataLayerToAssetCommandLet.h b/Engine/Source/Editor/UnrealEd/Classes/Commandlets/WorldPartitionDataLayerToAssetCommandLet.h
--- a/Engine/Source/Editor/UnrealEd/Classes/Commandlets/WorldPartitionDataLayerToAssetCommandLet.h
+++ b/Engine/Source/Editor/UnrealEd/Classes/Commandlets/WorldPartitionDataLayerToAssetCommandLet.h
@@ -5,6 +5,7 @@
 #include "CoreMinimal.h"
 #include "UObject/ObjectMacros.h"
 #include "UObject/UObjectAnnotation.h"
+#include "UObject/SoftObjectPath.h"
 #include "UObject/StrongObjectPtr.h"
 #include "Commandlets/Commandlet.h"
 #include "Commandlets/WorldPartitionCommandletHelpers.h"
@@ -167,7 +168,7 @@ private:
 	UNREALED_API bool CommitConversion(TStrongObjectPtr<UDataLayerToAssetCommandletContext>& CommandletContext, FPackageSourceControlHelper& PackageHelper);
 
 	UNREALED_API FString GetConversionFolder() const;
-	UNREALED_API bool IsAssetInConversionFolder(const TObjectPtr<UDataLayerAsset> DataLayerAsset);
+	UNREALED_API bool IsAssetInConversionFolder(const FSoftObjectPath& DataLayerAsset);
 
 	UPROPERTY(Config)
 	FString DestinationFolder;
diff --git a/Engine/Source/Editor/UnrealEd/Private/Commandlets/WorldPartitionDataLayerToAssetCommandlet.cpp b/Engine/Source/Editor/UnrealEd/Private/Commandlets/WorldPartitionDataLayerToAssetCommandlet.cpp
--- a/Engine/Source/Editor/UnrealEd/Private/Commandlets/WorldPartitionDataLayerToAssetCommandlet.cpp
+++ b/Engine/Source/Editor/UnrealEd/Private/Commandlets/WorldPartitionDataLayerToAssetCommandlet.cpp
@@ -263,7 +263,7 @@ int32 UDataLayerToAssetCommandlet::Main(const FString& Params)
 		return EReturnCode::CommandletInitializationError;
 	}
 
-	ConversionFolder = DestinationFolder + "/" + MainWorld->GetName();
+	ConversionFolder = DestinationFolder + "/" + MainWorld->GetName() + "/";
 
 	TStrongObjectPtr<UDataLayerToAssetCommandletContext> Context(NewObject<UDataLayerToAssetCommandletContext>());
 
@@ -348,7 +348,10 @@ bool UDataLayerToAssetCommandlet::BuildConversionInfos(TStrongObjectPtr<UDataLay
 	AssetRegistry.GetAssetsByClass(UDataLayerAsset::StaticClass()->GetFName(), ExistingDataLayerAssets);
 	for (FAssetData& AssetData : ExistingDataLayerAssets)
 	{
-		CommandletContext->StoreExistingDataLayer(AssetData);
+		if (IsAssetInConversionFolder(AssetData.GetSoftObjectPath()))
+		{
+			CommandletContext->StoreExistingDataLayer(AssetData);
+		}
 	}
 
 	AWorldDataLayers* WorldDataLayers = MainWorld->GetWorldDataLayers();
@@ -749,9 +752,9 @@ bool UDataLayerToAssetCommandlet::CommitConversion(TStrongObjectPtr<UDataLayerTo
 	return true;
 }
 
-bool UDataLayerToAssetCommandlet::IsAssetInConversionFolder(const TObjectPtr<UDataLayerAsset> DataLayerAsset)
+bool UDataLayerToAssetCommandlet::IsAssetInConversionFolder(const FSoftObjectPath& DataLayerAsset)
 {
-	return FCString::Stristr(*DataLayerAsset->GetPackage()->GetName(), *ConversionFolder) != nullptr;
+	return DataLayerAsset.GetAssetPathString().StartsWith(ConversionFolder);
 }
 
 PRAGMA_ENABLE_DEPRECATION_WARNINGS
