Back out 28667039 while I investigate drag issues

diff --git a/Engine/Source/Editor/UnrealEd/Private/LevelEditorViewport.cpp b/Engine/Source/Editor/UnrealEd/Private/LevelEditorViewport.cpp
--- a/Engine/Source/Editor/UnrealEd/Private/LevelEditorViewport.cpp
+++ b/Engine/Source/Editor/UnrealEd/Private/LevelEditorViewport.cpp
@@ -3142,38 +3142,11 @@ bool FLevelEditorViewportClient::InputKey(const FInputKeyEventArgs& InEventArgs)
 		GEditor->SetPreviewMeshMode(false);
 	}
 
-	// Clear Duplicate Actors mode when mouse buttons are released
-	if ( !InputState.IsAnyMouseButtonDown() )
+	// Clear Duplicate Actors mode when ALT and all mouse buttons are released
+	if ( !InputState.IsAltButtonPressed() && !InputState.IsAnyMouseButtonDown() )
 	{
 		bDuplicateActorsInProgress = false;
 	}
-	// Stamp a duplicate at the current position during a drag when ALT is pressed
-	else if (InEventArgs.Event == IE_Pressed && bDraggingByHandle && InputState.IsAltButtonPressed())
-	{
-		TSharedPtr<ILevelEditor> LevelEditor = ParentLevelEditor.Pin();
-		UTypedElementCommonActions* CommonActions = LevelEditor ? LevelEditor->GetCommonActions() : nullptr;
-		if (CommonActions)
-		{
-			FTypedElementListConstRef ElementsToManipulate = GetElementsToManipulate();
-			TArray<FTypedElementHandle> DuplicatedElements;
-			{
-				// Do not used the cached manipulation list here here, as it will have removed attachments, and we do want to duplicate those
-				DuplicatedElements = CommonActions->DuplicateSelectedElements(GetSelectionSet(), GetWorld(), FVector::ZeroVector);
-			}
-
-			// Do not select the duplicate actors so that the drag operation can continue freely
-			if (DuplicatedElements.Num() > 0)
-			{
-				// Although we did not change the selection, we still need to notify legacy mode tools of the duplication
-				TArray<AActor*> SelectedActors = GetSelectionSet()->GetSelectedObjects<AActor>();
-				constexpr bool bDidOffsetDuplicate = false;
-				ModeTools->ActorsDuplicatedNotify(SelectedActors, SelectedActors, bDidOffsetDuplicate);
-			}
-
-			RedrawAllViewportsIntoThisScene();
-		}
-	}
-
 	
 	return bHandled;
 }
