Clamp max application scale to 3.0, make undo work.

diff --git a/Engine/Source/Editor/UnrealEd/Classes/Settings/EditorStyleSettings.h b/Engine/Source/Editor/UnrealEd/Classes/Settings/EditorStyleSettings.h
--- a/Engine/Source/Editor/UnrealEd/Classes/Settings/EditorStyleSettings.h
+++ b/Engine/Source/Editor/UnrealEd/Classes/Settings/EditorStyleSettings.h
@@ -51,7 +51,7 @@ public:
 	/** 
 	 * Scales the entire editor interface up or down. 
 	 */
-	UPROPERTY(EditAnywhere, Config, Category=UserInterface, meta=(ClampMin=0.1))
+	UPROPERTY(EditAnywhere, Config, Category=UserInterface, meta=(ClampMin=0.1, ClampMax=3.0))
 	float ApplicationScale = 1.0f;
 
 	/**
diff --git a/Engine/Source/Editor/UnrealEd/Private/Settings/EditorStyleClasses.cpp b/Engine/Source/Editor/UnrealEd/Private/Settings/EditorStyleClasses.cpp
--- a/Engine/Source/Editor/UnrealEd/Private/Settings/EditorStyleClasses.cpp
+++ b/Engine/Source/Editor/UnrealEd/Private/Settings/EditorStyleClasses.cpp
@@ -99,7 +99,7 @@ void UEditorStyleSettings::PostEditChangeProperty(struct FPropertyChangedEvent&
 	{
 		GConfig->SetBool(TEXT("HDPI"), TEXT("EnableHighDPIAwareness"), bEnableHighDPIAwareness, GEditorSettingsIni);
 	}
-	else if (PropertyName == GET_MEMBER_NAME_CHECKED(UEditorStyleSettings, ApplicationScale))
+	else if (PropertyName.IsNone() || PropertyName == GET_MEMBER_NAME_CHECKED(UEditorStyleSettings, ApplicationScale))
 	{
 		if (FSlateApplication::IsInitialized())
 		{
