Fix for NavigateToFunctionSource crash/error, use xcode built in atos instead of UnrealAtoS as it can't handle universal build.

diff --git a/Engine/Build/BatchFiles/Mac/SourceCodeLookup.sh b/Engine/Build/BatchFiles/Mac/SourceCodeLookup.sh
new file mode 100755
--- /dev/null
+++ b/Engine/Build/BatchFiles/Mac/SourceCodeLookup.sh
@@ -0,0 +1,12 @@
+#!/bin/sh
+# Copyright Epic Games, Inc. All Rights Reserved.
+
+FunctionName=$1
+DylibFullPath=$2
+
+if [ -f "$DylibFullPath" ]; then
+    nm -Cg "$DylibFullPath" | grep "\b$FunctionName\b" | cut -d' ' -f1 | atos --fullPath -o "$DylibFullPath"
+else
+    echo "$DylibFullPath does not exist"
+    exit 1
+fi
diff --git a/Engine/Build/InstalledEngineFilters.xml b/Engine/Build/InstalledEngineFilters.xml
--- a/Engine/Build/InstalledEngineFilters.xml
+++ b/Engine/Build/InstalledEngineFilters.xml
@@ -342,6 +342,7 @@
 		Engine/Build/BatchFiles/Mac/ThinApp.sh
 		Engine/Build/BatchFiles/Mac/UpdateVersionAfterBuild.sh
 		Engine/Build/BatchFiles/Mac/ThinApp.sh
+        Engine/Build/BatchFiles/Mac/SourceCodeLookup.sh
 		Engine/Build/BatchFiles/RunDotnet.sh
 		Engine/Build/BatchFiles/RunDotnetMSBuild.sh
 		Engine/Build/BatchFiles/RunUAT.command
diff --git a/Engine/Plugins/Developer/XCodeSourceCodeAccess/Source/XCodeSourceCodeAccess/Private/XCodeSourceCodeAccessor.cpp b/Engine/Plugins/Developer/XCodeSourceCodeAccess/Source/XCodeSourceCodeAccess/Private/XCodeSourceCodeAccessor.cpp
--- a/Engine/Plugins/Developer/XCodeSourceCodeAccess/Source/XCodeSourceCodeAccess/Private/XCodeSourceCodeAccessor.cpp
+++ b/Engine/Plugins/Developer/XCodeSourceCodeAccess/Source/XCodeSourceCodeAccess/Private/XCodeSourceCodeAccessor.cpp
@@ -167,112 +167,23 @@ bool FXCodeSourceCodeAccessor::DoesSolutionExist() const
     return FPaths::FileExists( FullPath );
 }
 
+// ColumnNumber is not supported
 bool FXCodeSourceCodeAccessor::OpenFileAtLine(const FString& FullPath, int32 LineNumber, int32 ColumnNumber)
 {
-	ISourceCodeAccessModule& SourceCodeAccessModule = FModuleManager::LoadModuleChecked<ISourceCodeAccessModule>(TEXT("SourceCodeAccess"));
-
-	// column & line numbers are 1-based, so dont allow zero
-	if(LineNumber == 0)
-	{
-		LineNumber++;
-	}
-	if(ColumnNumber == 0)
-	{
-		ColumnNumber++;
-	}
-		 
-	bool ExecutionSucceeded = false;
-	
-    FString SolutionPath = GetSolutionPath();
-	{
-		const FString ProjPath = IFileManager::Get().ConvertToAbsolutePathForExternalAppForRead( *SolutionPath );
-		if ( FPaths::FileExists( ProjPath ) )
-		{
-			FString XcodePath = FPlatformMisc::GetXcodePath();
-			XcodePath.RemoveFromEnd(TEXT("/Contents/Developer"));
-			
-			NSURL* ProjectPath = [NSURL fileURLWithPath:[ProjPath.GetNSString() stringByDeletingLastPathComponent]];
-			[[NSWorkspace sharedWorkspace]
-				openURLs:[NSArray arrayWithObject: ProjectPath]
-				withApplicationAtURL:[NSURL URLWithString:XcodePath.GetNSString()]
-				configuration: [NSWorkspaceOpenConfiguration configuration]
-				completionHandler:nil
-			];
-			
-			NSAppleScript* AppleScript = nil;
-			
-			NSString* AppleScriptString = [NSString stringWithCString:OpenXCodeAtFileAndLineAppleScript encoding:NSUTF8StringEncoding];
-			AppleScriptString = [AppleScriptString stringByReplacingOccurrencesOfString:@"{XCODE_PATH}" withString:XcodePath.GetNSString()];
-			AppleScript = [[NSAppleScript alloc] initWithSource:AppleScriptString];
-			
-			int PID = [[NSProcessInfo processInfo] processIdentifier];
-			NSAppleEventDescriptor* ThisApplication = [NSAppleEventDescriptor descriptorWithDescriptorType:typeKernelProcessID bytes:&PID length:sizeof(PID)];
-			
-			NSAppleEventDescriptor* ContainerEvent = [NSAppleEventDescriptor appleEventWithEventClass:'ascr' eventID:'psbr' targetDescriptor:ThisApplication returnID:kAutoGenerateReturnID transactionID:kAnyTransactionID];
-			
-			[ContainerEvent setParamDescriptor:[NSAppleEventDescriptor descriptorWithString:@"OpenXcodeAtFileAndLine"] forKeyword:'snam'];
-			
-			{
-				NSAppleEventDescriptor* Arguments = [[NSAppleEventDescriptor alloc] initListDescriptor];
-				
-				CFStringRef FileString = FPlatformString::TCHARToCFString(*FullPath);
-				NSString* Path = (NSString*)FileString;
-				
-				if([Path isAbsolutePath] == NO)
-				{
-					NSString* CurDir = [[NSFileManager defaultManager] currentDirectoryPath];
-					NSString* ResolvedPath = [[NSString stringWithFormat:@"%@/%@", CurDir, Path] stringByResolvingSymlinksInPath];
-					if([[NSFileManager defaultManager] fileExistsAtPath:ResolvedPath])
-					{
-						Path = ResolvedPath;
-					}
-					else // If it doesn't exist, supply only the filename, we'll use Open Quickly to try and find it from Xcode
-					{
-						Path = [Path lastPathComponent];
-					}
-				}
-				
-				[Arguments insertDescriptor:[NSAppleEventDescriptor descriptorWithString:Path] atIndex:([Arguments numberOfItems] + 1)];
-				CFRelease(FileString);
-				
-				CFStringRef LineString = FPlatformString::TCHARToCFString(*FString::FromInt(LineNumber));
-				if(LineString)
-				{
-					[Arguments insertDescriptor:[NSAppleEventDescriptor descriptorWithString:(NSString*)LineString] atIndex:([Arguments numberOfItems] + 1)];
-					CFRelease(LineString);
-				}
-				else
-				{
-					[Arguments insertDescriptor:[NSAppleEventDescriptor descriptorWithString:@"1"] atIndex:([Arguments numberOfItems] + 1)];
-				}
-				
-				[ContainerEvent setParamDescriptor:Arguments forKeyword:keyDirectObject];
-				[Arguments release];
-			}
-			
-			NSDictionary* ExecutionError = nil;
-			[AppleScript executeAppleEvent:ContainerEvent error:&ExecutionError];
-			if(ExecutionError == nil)
-			{
-				ExecutionSucceeded = true;
-			}
-			else
-			{
-				UE_LOG(LogXcodeAccessor, Error, TEXT("%s"), *FString([ExecutionError description]));
-			}
-			
-			[AppleScript release];
-			
-			// Fallback to trivial implementation when something goes wrong (like not having permission for UI scripting)
-			if(ExecutionSucceeded == false)
-			{
-				FPlatformProcess::LaunchFileInDefaultExternalApplication(*FullPath);
-				ExecutionSucceeded = true;
-			}
-		}
-	}
-		
-	return ExecutionSucceeded;
+    // We use xed to open a file at specified line number through Xcode, it also reuses your currently opened workspace if file exists in project
+    int32 ReturnCode = 0;
+    FString Errors;
+    FString XedCommand = FString::Printf(TEXT("xed -l %d \"%s\""), LineNumber, *FullPath);
+    FPlatformProcess::ExecProcess(TEXT("/usr/bin/env"), *XedCommand, &ReturnCode, NULL, &Errors);
+    
+    if(ReturnCode != 0)
+    {
+        UE_LOG(LogXcodeAccessor, Warning, TEXT("FXCodeSourceCodeAccessor::OpenFileAtLine failed to open xcode with %s"), *XedCommand);
+        // Fallback to trivial implementation when something goes wrong
+        FPlatformProcess::LaunchFileInDefaultExternalApplication(*FullPath);
+    }
+    
+    return true;
 }
 
 bool FXCodeSourceCodeAccessor::OpenSourceFiles(const TArray<FString>& AbsoluteSourcePaths)
diff --git a/Engine/Source/Editor/UnrealEd/Private/SourceCodeNavigation.cpp b/Engine/Source/Editor/UnrealEd/Private/SourceCodeNavigation.cpp
--- a/Engine/Source/Editor/UnrealEd/Private/SourceCodeNavigation.cpp
+++ b/Engine/Source/Editor/UnrealEd/Private/SourceCodeNavigation.cpp
@@ -551,186 +551,57 @@ void FSourceCodeNavigationImpl::NavigateToFunctionSource( const FString& Functio
 #endif // !NO_LOGGING
 
 #elif PLATFORM_MAC
-	
-    // UE-195221, no easy solution for now, it is causing crash and wouldn't work anyways even without the crash
-    // Solution pending
     
-//	for(uint32 Index = 0; Index < _dyld_image_count(); Index++)
-//	{
-//		char const* IndexName = _dyld_get_image_name(Index);
-//		FString FullModulePath(IndexName);
-//		FString Name = FPaths::GetBaseFilename(FullModulePath);
-//		if(Name == FunctionModuleName)
-//		{
-//			struct mach_header_64 const* IndexModule64 = NULL;
-//			struct load_command const* LoadCommands = NULL;
-//			
-//			struct mach_header const* IndexModule32 = _dyld_get_image_header(Index);
-//			check(IndexModule32->magic == MH_MAGIC_64);
-//			
-//			IndexModule64 = (struct mach_header_64 const*)IndexModule32;
-//			LoadCommands = (struct load_command const*)(IndexModule64 + 1);
-//			struct load_command const* Command = LoadCommands;
-//			struct symtab_command const* SymbolTable = nullptr;
-//			struct dysymtab_command const* DsymTable = nullptr;
-//			struct uuid_command* UUIDCommand = nullptr;
-//			for(uint32 CommandIndex = 0; CommandIndex < IndexModule64->ncmds; CommandIndex++)
-//			{
-//				if (Command && Command->cmd == LC_SYMTAB)
-//				{
-//					SymbolTable = (struct symtab_command const*)Command;
-//				}
-//				else if(Command && Command->cmd == LC_DYSYMTAB)
-//				{
-//					DsymTable = (struct dysymtab_command const*)Command;
-//				}
-//				else if (Command && Command->cmd == LC_UUID)
-//				{
-//					UUIDCommand = (struct uuid_command*)Command;
-//				}
-//				Command = (struct load_command const*)(((char const*)Command) + Command->cmdsize);
-//			}
-//			
-//			check(SymbolTable && DsymTable && UUIDCommand);
-//			
-//			IPlatformFile& PlatformFile = IPlatformFile::GetPlatformPhysical();
-//			IFileHandle* File = PlatformFile.OpenRead(*FullModulePath);
-//			if(File)
-//			{
-//				struct nlist_64* SymbolEntries = new struct nlist_64[SymbolTable->nsyms];
-//				check(SymbolEntries);
-//				char* StringTable = new char[SymbolTable->strsize];
-//				check(StringTable);
-//				
-//				bool FileOK = File->Seek(SymbolTable->symoff+(DsymTable->iextdefsym*sizeof(struct nlist_64)));
-//				FileOK &= File->Read((uint8*)SymbolEntries, DsymTable->nextdefsym*sizeof(struct nlist_64));
-//				
-//				FileOK &= File->Seek(SymbolTable->stroff);
-//				FileOK &= File->Read((uint8*)StringTable, SymbolTable->strsize);
-//				
-//				delete File;
-//				
-//				for(uint32 SymbolIndex = 0; FileOK && SymbolIndex < DsymTable->nextdefsym; SymbolIndex++)
-//				{
-//					struct nlist_64 const& SymbolEntry = SymbolEntries[SymbolIndex];
-//					// All the entries in the mach-o external table are functions.
-//					// The local table contains the minimal debug stabs used by dsymutil to create the DWARF dsym.
-//					if(SymbolEntry.n_un.n_strx)
-//					{
-//						if (SymbolEntry.n_value)
-//						{
-//							char const* MangledSymbolName = (StringTable+SymbolEntry.n_un.n_strx);
-//							// Remove leading '_'
-//							MangledSymbolName += 1;
-//							
-//							int32 Status = 0;
-//							char* DemangledName = abi::__cxa_demangle(MangledSymbolName, NULL, 0, &Status);
-//							
-//							FString SymbolName;
-//							if (DemangledName)
-//							{
-//								// C++ function
-//								SymbolName = DemangledName;
-//								free(DemangledName);
-//								
-//								// This contains return & arguments, it would seem that the DbgHelp API doesn't.
-//								// So we shall strip them.
-//								int32 ArgumentIndex = -1;
-//								if(SymbolName.FindLastChar(TCHAR('('), ArgumentIndex))
-//								{
-//									SymbolName.LeftInline(ArgumentIndex, false);
-//									int32 TemplateNesting = 0;
-//									
-//									int32 Pos = SymbolName.Len();
-//									// Cast operators are special & include spaces, whereas normal functions don't.
-//									int32 OperatorIndex = SymbolName.Find("operator");
-//									if(OperatorIndex >= 0)
-//									{
-//										// Trim from before the 'operator'
-//										Pos = OperatorIndex;
-//									}
-//									
-//									for(; Pos > 0; --Pos)
-//									{
-//										TCHAR Character = SymbolName[Pos - 1];
-//										if(Character == TCHAR(' ') && TemplateNesting == 0)
-//										{
-//											SymbolName.MidInline(Pos, MAX_int32, false);
-//											break;
-//										}
-//										else if(Character == TCHAR('>'))
-//										{
-//											TemplateNesting++;
-//										}
-//										else if(Character == TCHAR('<'))
-//										{
-//											TemplateNesting--;
-//										}
-//									}
-//								}
-//							}
-//							else
-//							{
-//								// C function
-//								SymbolName = MangledSymbolName;
-//							}
-//							
-//							if(FunctionSymbolName == SymbolName)
-//							{
-//								CFUUIDBytes UUIDBytes;
-//								FMemory::Memcpy(&UUIDBytes, UUIDCommand->uuid, sizeof(CFUUIDBytes));
-//								CFUUIDRef UUIDRef = CFUUIDCreateFromUUIDBytes(kCFAllocatorDefault, UUIDBytes);
-//								CFStringRef UUIDString = CFUUIDCreateString(kCFAllocatorDefault, UUIDRef);
-//								FString UUID((NSString*)UUIDString);
-//								CFRelease(UUIDString);
-//								CFRelease(UUIDRef);
-//							
-//								uint64 Address = SymbolEntry.n_value;
-//								uint64 BaseAddress = (uint64)IndexModule64;
-//								FString AtoSCommand = FString::Printf(TEXT("\"%s\" -s %s -l 0x%lx 0x%lx"), *FullModulePath, *UUID, BaseAddress, Address);
-//								int32 ReturnCode = 0;
-//								FString Results;
-//								
-//								const FString AtoSPath = FString::Printf(TEXT("%sBinaries/Mac/UnrealAtoS"), *FPaths::EngineDir() );
-//								FPlatformProcess::ExecProcess( *AtoSPath, *AtoSCommand, &ReturnCode, &Results, NULL );
-//								if(ReturnCode == 0)
-//								{
-//									bool bSourceFileOpened = false;
-//									int32 FirstIndex = -1;
-//									int32 LastIndex = -1;
-//									if(Results.FindChar(TCHAR('('), FirstIndex) && Results.FindLastChar(TCHAR('('), LastIndex) && FirstIndex != LastIndex)
-//									{
-//										int32 CloseIndex = -1;
-//										int32 ColonIndex = -1;
-//										if(Results.FindLastChar(TCHAR(':'), ColonIndex) && Results.FindLastChar(TCHAR(')'), CloseIndex) && CloseIndex > ColonIndex && LastIndex < ColonIndex)
-//										{
-//											int32 FileNamePos = LastIndex+1;
-//											int32 FileNameLen = ColonIndex-FileNamePos;
-//											FString FileName = Results.Mid(FileNamePos, FileNameLen);
-//											FString LineNumber = Results.Mid(ColonIndex + 1, CloseIndex-(ColonIndex + 1));
-//											bSourceFileOpened = SourceCodeAccessor.OpenFileAtLine( FileName, FCString::Atoi(*LineNumber), 0 );
-//										}
-//									}
-//#if !NO_LOGGING
-//									if (!bSourceFileOpened)
-//									{
-//										UE_LOG(LogSelectionDetails, Warning, TEXT("NavigateToFunctionSource:  Unable to find source file and line number for '%s'"), *FunctionSymbolName);
-//									}
-//#endif
-//								}
-//								break;
-//							}
-//						}
-//					}
-//				}
-//				
-//				delete [] StringTable;
-//				delete [] SymbolEntries;
-//			}
-//			break;
-//		}
-//	}
-	
+    int32 ReturnCode = 0;
+    FString Results;
+    FString Errors;
+	for(uint32 Index = 0; Index < _dyld_image_count(); Index++)
+	{
+		char const* IndexName = _dyld_get_image_name(Index);
+		FString FullModulePath(IndexName);
+		FString Name = FPaths::GetBaseFilename(FullModulePath);
+		if(Name == FunctionModuleName)
+		{
+            const FString SourceCodeLookupCommand = FString::Printf(TEXT("%sBuild/BatchFiles/Mac/SourceCodeLookup.sh \"%s\" \"%s\""), *FPaths::EngineDir(), *FunctionSymbolName, *FullModulePath);
+            FPlatformProcess::ExecProcess( TEXT("/bin/sh"), *SourceCodeLookupCommand, &ReturnCode, &Results, &Errors );
+            if(ReturnCode == 0 && !Results.IsEmpty())
+            {
+                // find the last occurance of (filepath:linenumber)
+                int32 OpenIndex = -1;
+                int32 ColonIndex = -1;
+                int32 CloseIndex = -1;
+                if(Results.FindLastChar(TCHAR('('), OpenIndex) && Results.FindLastChar(TCHAR(':'), ColonIndex) && Results.FindLastChar(TCHAR(')'), CloseIndex)
+                   && CloseIndex > ColonIndex && OpenIndex < ColonIndex)
+                {
+                    int32 FileNamePos = OpenIndex + 1;
+                    int32 FileNameLen = ColonIndex - FileNamePos;
+                    FString FileName = Results.Mid(FileNamePos, FileNameLen);
+                    int32 LineNumberPos = ColonIndex + 1;
+                    int32 LineNumberLen = CloseIndex - LineNumberPos;
+                    FString LineNumber = Results.Mid(LineNumberPos, LineNumberLen);
+                    if (SourceCodeAccessor.OpenFileAtLine( FileName, FCString::Atoi(*LineNumber), 0 ))
+                    {
+                        return;
+                    }
+                    else
+                    {
+                        UE_LOG(LogSelectionDetails, Warning, TEXT("NavigateToFunctionSource: Unable to open file %s and line number %s"), *FileName, *LineNumber);
+                    }
+                }
+                else
+                {
+                    UE_LOG(LogSelectionDetails, Warning, TEXT("NavigateToFunctionSource: Unexpected SourceCodeLookup.sh output: %s"), *Results);
+                }
+            }
+            else
+            {
+                UE_LOG(LogSelectionDetails, Warning, TEXT("NavigateToFunctionSource: SourceCodeLookup.sh error: %s"), *Errors);
+            }
+		}
+	}
+    
+    UE_LOG(LogSelectionDetails, Warning, TEXT("NavigateToFunctionSource: Unable to look up symbol: %s in module:%s"), *FunctionSymbolName, *FunctionModuleName);
+    
 #endif	// PLATFORM_WINDOWS
 }
 
