Clamp application scale to 0.5 to avoid completely losing ability to change it back without using undo

diff --git a/Engine/Source/Editor/UnrealEd/Classes/Settings/EditorStyleSettings.h b/Engine/Source/Editor/UnrealEd/Classes/Settings/EditorStyleSettings.h
--- a/Engine/Source/Editor/UnrealEd/Classes/Settings/EditorStyleSettings.h
+++ b/Engine/Source/Editor/UnrealEd/Classes/Settings/EditorStyleSettings.h
@@ -51,7 +51,7 @@ public:
 	/** 
 	 * Scales the entire editor interface up or down. 
 	 */
-	UPROPERTY(EditAnywhere, Config, Category=UserInterface, meta=(ClampMin=0.1, ClampMax=3.0))
+	UPROPERTY(EditAnywhere, Config, Category=UserInterface, meta=(ClampMin=0.5, ClampMax=3.0))
 	float ApplicationScale = 1.0f;
 
 	/**
