ActorGroupingUtils: Add a way for external systems to add a delegate to check whether a set of actors can be grouped together

diff --git a/Engine/Source/Editor/LevelEditor/Private/LevelEditor.cpp b/Engine/Source/Editor/LevelEditor/Private/LevelEditor.cpp
--- a/Engine/Source/Editor/LevelEditor/Private/LevelEditor.cpp
+++ b/Engine/Source/Editor/LevelEditor/Private/LevelEditor.cpp
@@ -1029,7 +1029,7 @@ void FLevelEditorModule::BindGlobalLevelEditorCommands()
 	ActionList.MapAction(
 		Commands.RegroupActors,
 		FExecuteAction::CreateStatic( &FLevelEditorActionCallbacks::RegroupActor_Clicked ),
-		FCanExecuteAction::CreateStatic( &FLevelEditorActionCallbacks::ActorSelected_CanExecute )
+		FCanExecuteAction::CreateStatic( &FLevelEditorActionCallbacks::GroupActors_CanExecute )
 		);
 
 	ActionList.MapAction(
@@ -1049,7 +1049,8 @@ void FLevelEditorModule::BindGlobalLevelEditorCommands()
 
 	ActionList.MapAction(
 		Commands.AddActorsToGroup,
-		FExecuteAction::CreateStatic( &FLevelEditorActionCallbacks::AddActorsToGroup_Clicked )
+		FExecuteAction::CreateStatic( &FLevelEditorActionCallbacks::AddActorsToGroup_Clicked ),
+		FCanExecuteAction::CreateStatic( &FLevelEditorActionCallbacks::GroupActors_CanExecute )
 		);
 
 	ActionList.MapAction(
diff --git a/Engine/Source/Editor/LevelEditor/Private/LevelEditorActions.cpp b/Engine/Source/Editor/LevelEditor/Private/LevelEditorActions.cpp
--- a/Engine/Source/Editor/LevelEditor/Private/LevelEditorActions.cpp
+++ b/Engine/Source/Editor/LevelEditor/Private/LevelEditorActions.cpp
@@ -2153,6 +2153,11 @@ void FLevelEditorActionCallbacks::OnSurfaceAlignment( ETexAlign AlignmentMode )
 	GTexAlignTools.GetAligner( AlignmentMode )->Align( GetWorld(), AlignmentMode );
 }
 
+bool FLevelEditorActionCallbacks::GroupActors_CanExecute()
+{
+	return UActorGroupingUtils::Get()->CanGroupSelectedActors();
+}
+
 void FLevelEditorActionCallbacks::RegroupActor_Clicked()
 {
 	UActorGroupingUtils::Get()->GroupSelected();
diff --git a/Engine/Source/Editor/LevelEditor/Public/LevelEditorActions.h b/Engine/Source/Editor/LevelEditor/Public/LevelEditorActions.h
--- a/Engine/Source/Editor/LevelEditor/Public/LevelEditorActions.h
+++ b/Engine/Source/Editor/LevelEditor/Public/LevelEditorActions.h
@@ -1060,6 +1060,12 @@ public:
 	 * Called to apply a material to selected surfaces
 	 */
 	static void OnApplyMaterialToSurface();
+
+	/**
+	 * Checks to see if the selected actors can be grouped
+	 *	@return true if it can execute.
+	 */
+	static bool GroupActors_CanExecute();
 	
 	/**
 	 * Called when the RegroupActor command is executed
diff --git a/Engine/Source/Editor/UnrealEd/Private/ActorGroupingUtils.cpp b/Engine/Source/Editor/UnrealEd/Private/ActorGroupingUtils.cpp
--- a/Engine/Source/Editor/UnrealEd/Private/ActorGroupingUtils.cpp
+++ b/Engine/Source/Editor/UnrealEd/Private/ActorGroupingUtils.cpp
@@ -10,6 +10,7 @@
 #include "Widgets/Notifications/SNotificationList.h"
 
 bool UActorGroupingUtils::bGroupingActive = true;
+TMap<FName, UActorGroupingUtils::FCanGroupActors> UActorGroupingUtils::CanGroupActorsDelegates;
 
 void UActorGroupingUtils::SetGroupingActive(bool bInGroupingActive)
 {
@@ -24,7 +25,7 @@ UActorGroupingUtils* UActorGroupingUtils::Get()
 
 AGroupActor* UActorGroupingUtils::GroupSelected()
 {
-	if (IsGroupingActive())
+	if (IsGroupingActive() && CanGroupSelectedActors())
 	{
 		TArray<AActor*> ActorsToAdd;
 		for (FSelectionIterator It(GEditor->GetSelectedActorIterator()); It; ++It)
@@ -50,7 +51,7 @@ AGroupActor* UActorGroupingUtils::GroupSelected()
 
 AGroupActor* UActorGroupingUtils::GroupActors(const TArray<AActor*>& ActorsToGroup)
 {
-	if(IsGroupingActive())
+	if(IsGroupingActive() && CanGroupActors(ActorsToGroup))
 	{
 		ULevel* ActorLevel = nullptr;
 		TArray<AActor*> FinalActorList;
@@ -206,7 +207,7 @@ void UActorGroupingUtils::UnlockSelectedGroups()
 
 void UActorGroupingUtils::AddSelectedToGroup()
 {
-	if (IsGroupingActive())
+	if (IsGroupingActive() && CanGroupSelectedActors())
 	{
 		AGroupActor::AddSelectedActorsToSelectedGroup();
 	}
@@ -279,3 +280,43 @@ void UActorGroupingUtils::RemoveSelectedFromGroup()
 	}
 
 }
+
+void UActorGroupingUtils::AddCanGroupActorsDelegate(const FName& Owner, const FCanGroupActors& InGroupActorsFilter)
+{
+	CanGroupActorsDelegates.Emplace(Owner, InGroupActorsFilter);
+}
+
+void UActorGroupingUtils::RemoveCanGroupActorsDelegate(const FName& Owner)
+{
+	CanGroupActorsDelegates.Remove(Owner);
+}
+
+bool UActorGroupingUtils::CanGroupActors(const TArray<AActor*>& ActorsToGroup) const
+{
+	for(const TPair<FName, FCanGroupActors>& Filter : CanGroupActorsDelegates)
+	{
+		if(!Filter.Value.Execute(ActorsToGroup))
+		{
+			return false;
+		}
+	}
+	
+	return true;
+}
+
+bool UActorGroupingUtils::CanGroupSelectedActors() const
+{
+	USelection* ActorSelection = GEditor->GetSelectedActors();
+	TArray<AActor*> Actors;
+	ActorSelection->GetSelectedObjects(Actors);
+	
+	for(const TPair<FName, FCanGroupActors>& Filter : CanGroupActorsDelegates)
+	{
+		if(!Filter.Value.Execute(Actors))
+		{
+			return false;
+		}
+	}
+
+	return true;
+}
diff --git a/Engine/Source/Editor/UnrealEd/Public/ActorGroupingUtils.h b/Engine/Source/Editor/UnrealEd/Public/ActorGroupingUtils.h
--- a/Engine/Source/Editor/UnrealEd/Public/ActorGroupingUtils.h
+++ b/Engine/Source/Editor/UnrealEd/Public/ActorGroupingUtils.h
@@ -75,9 +75,37 @@ public:
 	UFUNCTION(BlueprintCallable, Category = "Editor Scripting | Actor Grouping")
 	UNREALED_API virtual void RemoveSelectedFromGroup();
 
+	/**
+	 * Check if the provided list of actors can be grouped together
+	 */
+	UFUNCTION(BlueprintCallable, Category = "Editor Scripting | Actor Grouping")
+	UNREALED_API bool CanGroupActors(const TArray<AActor*>& ActorsToGroup) const;
+
+	/*
+	 * Check if the currently selected actors can be grouped together
+	 */
+	UFUNCTION(BlueprintCallable, Category = "Editor Scripting | Actor Grouping")
+	UNREALED_API bool CanGroupSelectedActors() const;
+
+	/**
+	 * External delegates that are run to check if a set of actors can be grouped together
+	 */
+	DECLARE_DELEGATE_RetVal_OneParam(bool, FCanGroupActors, const TArray<AActor*>&);
+	
+	/*
+	 * Add a filter that is checked to see if a set of actors can be grouped together
+	 */
+	UNREALED_API void AddCanGroupActorsDelegate(const FName& Owner, const FCanGroupActors& InGroupActorsDelegate);
+
+	/*
+	 * Remove the filter checked before grouping actors registered by the provided owner
+	 */
+	UNREALED_API void RemoveCanGroupActorsDelegate(const FName& Owner);
+
 protected:
 	static UNREALED_API bool bGroupingActive;
 	static UNREALED_API FSoftClassPath ClassToUse;
+	static UNREALED_API TMap<FName, FCanGroupActors> CanGroupActorsDelegates;
 };
 
 
