MPCook: Retraction: Keep a record of workers that report no packages available for retraction and do not try to retract from them again until they get new packages assigned.
#rnx

diff --git a/Engine/Source/Editor/UnrealEd/Private/Cooker/CookDirector.cpp b/Engine/Source/Editor/UnrealEd/Private/Cooker/CookDirector.cpp
--- a/Engine/Source/Editor/UnrealEd/Private/Cooker/CookDirector.cpp
+++ b/Engine/Source/Editor/UnrealEd/Private/Cooker/CookDirector.cpp
@@ -80,7 +80,11 @@ private:
 		WantToRetract,
 		WaitingForResponse,
 	};
-
+	enum class ERetractionResult : uint8
+	{
+		NoneAvailable,
+		Retracted
+	};
 private:
 	/** Try to select a worker for retraction */
 	ERetractionState TickWantToRetract();
@@ -91,15 +95,19 @@ private:
 	 * Pick workers to give the retracted packages to, and assign those packages to the worker
 	 * in the local and remote state.
 	 */
-	void ReassignPackages(const FWorkerId& WorkerId, TConstArrayView<FPackageData*> Packages);
+	ERetractionResult ReassignPackages(const FWorkerId& WorkerId, TConstArrayView<FPackageData*> Packages);
 	/** Pick workers to give the retracted packages to. */
 	TArray<FWorkerId> CalculateWorkersToSplitOver(int32 NumPackages, const FWorkerId& FromWorker,
 		TConstArrayView<TRefCountPtr<FCookWorkerServer>> LocalRemoteWorkers);
 
+	void SetRetractionState(ERetractionState NewState);
+	bool IsAvailableForRetraction(const FWorkerId& WorkerId);
+
 private:
 	FCookDirector& Director;
 	FWorkerId ExpectedWorker;
 	TMap<FWorkerId, TArray<FName>> PackagesToRetract;
+	TMap<FWorkerId, int32> WorkersUnavailableForRetract;
 	FWorkerId WorkerWithResults;
 	double MessageSentTimeSeconds = 0.;
 	double LastWarnTimeSeconds = 0.;
@@ -475,7 +483,7 @@ void FCookDirector::RemoveFromWorker(FPackageData& PackageData)
 	TRefCountPtr<FCookWorkerServer> OwningWorker;
 	{
 		FScopeLock CommunicationScopeLock(&CommunicationLock);
-		TRefCountPtr<FCookWorkerServer>* RemoteWorkerPtr = RemoteWorkers.Find(WorkerId.GetRemoteIndex());
+		const TRefCountPtr<FCookWorkerServer>* RemoteWorkerPtr = FindRemoteWorkerInLock(WorkerId);
 		if (!RemoteWorkerPtr)
 		{
 			return;
@@ -589,7 +597,7 @@ FString FCookDirector::GetDisplayName(const FWorkerId& WorkerId, int32 Preferred
 		const TRefCountPtr<FCookWorkerServer>* RemoteWorker = nullptr;
 		{
 			FScopeLock CommunicationScopeLock(&CommunicationLock);
-			RemoteWorker = RemoteWorkers.Find(WorkerId.GetRemoteIndex());
+			RemoteWorker = FindRemoteWorkerInLock(WorkerId);
 			if (!RemoteWorker)
 			{
 				for (const TPair<FCookWorkerServer*, TRefCountPtr<FCookWorkerServer>>& Pair : ShuttingDownWorkers)
@@ -624,6 +632,15 @@ FString FCookDirector::GetDisplayName(const FCookWorkerServer& RemoteWorker, int
 	return Result;
 }
 
+const TRefCountPtr<FCookWorkerServer>* FCookDirector::FindRemoteWorkerInLock(const FWorkerId& WorkerId) const
+{
+	if (!WorkerId.IsRemote())
+	{
+		return nullptr;
+	}
+	return RemoteWorkers.Find(WorkerId.GetRemoteIndex());
+}
+
 void FCookDirector::TickCommunication(ECookDirectorThread TickThread)
 {
 	bool bHasShutdownWorkers = false;
@@ -1491,8 +1508,14 @@ FCookDirector::FRetractionHandler::ERetractionState FCookDirector::FRetractionHa
 	if (Director.bAllowLocalCooks)
 	{
 		int32 NumAssignments = Director.COTFS.NumMultiprocessLocalWorkerAssignments();
-		BusiestWorker = FWorkerId::Local();
-		BusiestNumAssignments = NumAssignments;
+		if (NumAssignments > BusiestNumAssignments)
+		{
+			if (IsAvailableForRetraction(FWorkerId::Local()))
+			{
+				BusiestWorker = FWorkerId::Local();
+				BusiestNumAssignments = NumAssignments;
+			}
+		}
 		if (NumAssignments == 0)
 		{
 			IdleWorkers.Add(FWorkerId::Local());
@@ -1503,10 +1526,13 @@ FCookDirector::FRetractionHandler::ERetractionState FCookDirector::FRetractionHa
 	for (const TRefCountPtr<FCookWorkerServer>& RemoteWorker : LocalRemoteWorkers)
 	{
 		int32 NumAssignments = RemoteWorker->NumAssignments();
-		if (BusiestWorker.IsInvalid() || NumAssignments > BusiestNumAssignments)
+		if (NumAssignments > BusiestNumAssignments)
 		{
-			BusiestWorker = RemoteWorker->GetWorkerId();
-			BusiestNumAssignments = NumAssignments;
+			if (IsAvailableForRetraction(RemoteWorker->GetWorkerId()))
+			{
+				BusiestWorker = RemoteWorker->GetWorkerId();
+				BusiestNumAssignments = NumAssignments;
+			}
 		}
 		if (NumAssignments == 0)
 		{
@@ -1516,9 +1542,11 @@ FCookDirector::FRetractionHandler::ERetractionState FCookDirector::FRetractionHa
 	if (IdleWorkers.IsEmpty() || BusiestNumAssignments < RetractionMinimumNumAssignments)
 	{
 		// Worker loads changed after the point where we decided to initialize the RetractionHandler,
-		// and retraction is no longer needed.
-		return ERetractionState::Idle;
+		// or all workers with packages assigned are unavailable for retraction, so retraction is not
+		// currently possible. Try again later.
+		return ERetractionState::WantToRetract;
 	}
+	check(!BusiestWorker.IsInvalid());
 
 	// Plan to divide the assignments evenly between all idle workers and the one busiest worker. This means
 	// retracting all but 1/(N+1) packages from the busiest worker.
@@ -1567,30 +1595,89 @@ void FCookDirector::FRetractionHandler::TickFromSchedulerThread(bool bAllWorkers
 	switch (RetractionState)
 	{
 	case ERetractionState::Idle:
+	{
 		if (!bAnyIdle || !bAllWorkersConnected || BusiestNumAssignments <= RetractionMinimumNumAssignments)
 		{
 			break;
 		}
-		RetractionState = ERetractionState::WantToRetract;
-		RetractionState = TickWantToRetract();
+		SetRetractionState(ERetractionState::WantToRetract);
+		ERetractionState NewState = TickWantToRetract();
+		SetRetractionState(NewState);
 		break;
+	}
 	case ERetractionState::WantToRetract:
+	{
 		if (!bAnyIdle || !bAllWorkersConnected || BusiestNumAssignments <= RetractionMinimumNumAssignments)
 		{
-			RetractionState = ERetractionState::Idle;
+			SetRetractionState(ERetractionState::Idle);
 			break;
 		}
-		RetractionState = TickWantToRetract();
+		ERetractionState NewState = TickWantToRetract();
+		SetRetractionState(NewState);
 		break;
+	}
 	case ERetractionState::WaitingForResponse:
-		RetractionState = TickWaitingForResponse();
+	{
+		ERetractionState NewState = TickWaitingForResponse();
+		SetRetractionState(NewState);
 		break;
+	}
 	default:
 		checkNoEntry();
 		break;
 	}
 }
 
+void FCookDirector::FRetractionHandler::SetRetractionState(ERetractionState NewState)
+{
+	if (RetractionState == NewState)
+	{
+		return;
+	}
+
+	RetractionState = NewState;
+	if (NewState == ERetractionState::Idle)
+	{
+		WorkersUnavailableForRetract.Empty();
+	}
+}
+
+bool FCookDirector::FRetractionHandler::IsAvailableForRetraction(const FWorkerId& WorkerId)
+{
+	// Called from inside CommunicationLock
+	int32* AssignedPackagesFence = WorkersUnavailableForRetract.Find(WorkerId);
+	if (!AssignedPackagesFence)
+	{
+		return true;
+	}
+
+	int32 CurrentFenceMarker;
+	if (WorkerId.IsLocal())
+	{
+		CurrentFenceMarker = Director.COTFS.PackageDatas->GetMonitor().GetMPCookAssignedFenceMarker();
+	}
+	else
+	{
+		const TRefCountPtr<FCookWorkerServer>* RemoteWorker = Director.FindRemoteWorkerInLock(WorkerId);
+		if (!RemoteWorker)
+		{
+			WorkersUnavailableForRetract.Remove(WorkerId);
+			return true;
+		}
+		CurrentFenceMarker = (*RemoteWorker)->GetPackagesAssignedFenceMarker();
+	}
+
+	if (*AssignedPackagesFence == CurrentFenceMarker)
+	{
+		// FenceMarker has not changed since we recorded the worker as unavailable for retraction at that fence marker
+		// The worker is still unavailable for retraction
+		return false;
+	}
+
+	WorkersUnavailableForRetract.Remove(WorkerId);
+	return true;
+}
+
 FCookDirector::FRetractionHandler::ERetractionState FCookDirector::FRetractionHandler::TickWaitingForResponse()
 {
 	// Called from inside CommunicationLock
@@ -1611,7 +1698,7 @@ FCookDirector::FRetractionHandler::ERetractionState FCookDirector::FRetractionHa
 		}
 		check(ExpectedWorker.IsRemote());
 		{
-			TRefCountPtr<FCookWorkerServer>* RemoteWorkerPtr = Director.RemoteWorkers.Find(ExpectedWorker.GetRemoteIndex());
+			const TRefCountPtr<FCookWorkerServer>* RemoteWorkerPtr = Director.FindRemoteWorkerInLock(ExpectedWorker);
 			if (!RemoteWorkerPtr)
 			{
 				// The CookWorker aborted and we already reassigned all of its packages; stop waiting for a retraction message from it.
@@ -1633,7 +1720,7 @@ FCookDirector::FRetractionHandler::ERetractionState FCookDirector::FRetractionHa
 		TRefCountPtr<FCookWorkerServer> RemoteWorker;
 		if (Pair.Key.IsRemote())
 		{
-			TRefCountPtr<FCookWorkerServer>* FoundRemoteWorker = Director.RemoteWorkers.Find(Pair.Key.GetRemoteIndex());
+			const TRefCountPtr<FCookWorkerServer>* FoundRemoteWorker = Director.FindRemoteWorkerInLock(Pair.Key);
 			if (FoundRemoteWorker)
 			{
 				RemoteWorker = *FoundRemoteWorker;
@@ -1659,13 +1746,35 @@ FCookDirector::FRetractionHandler::ERetractionState FCookDirector::FRetractionHa
 	}
 
 	// Reassign the packages
-	ReassignPackages(WorkerWithResults, PackageDatasToReassign);
+	ERetractionResult Result = ReassignPackages(WorkerWithResults, PackageDatasToReassign);
+	if (Result == ERetractionResult::NoneAvailable)
+	{
+		TOptional<int32> AssignedPackagesFence;
+		if (WorkerWithResults.IsLocal())
+		{
+			AssignedPackagesFence.Emplace(Director.COTFS.PackageDatas->GetMonitor().GetMPCookAssignedFenceMarker());
+		}
+		else
+		{
+			const TRefCountPtr<FCookWorkerServer>* RemoteWorker = Director.FindRemoteWorkerInLock(WorkerWithResults);
+			if (RemoteWorker)
+			{
+				AssignedPackagesFence.Emplace((*RemoteWorker)->GetPackagesAssignedFenceMarker());
+			}
+		}
+		if (AssignedPackagesFence.IsSet())
+		{
+			WorkersUnavailableForRetract.Add(WorkerWithResults, *AssignedPackagesFence);
+		}
+	}
 
 	// Mark that we are no longer waiting
 	ExpectedWorker = FWorkerId::Invalid();
 	WorkerWithResults = FWorkerId::Invalid();
 	PackagesToRetract.Empty();
-	return ERetractionState::Idle;
+
+	// Return to WantToRetract state; that state will handle returning to idle if the retraction was sufficient
+	return ERetractionState::WantToRetract;
 }
 
 void FCookDirector::FRetractionHandler::HandleRetractionMessage(FMPCollectorServerMessageContext& Context,
@@ -1685,7 +1794,7 @@ void FCookDirector::FRetractionHandler::HandleRetractionMessage(FMPCollectorServ
 		UE_LOG(LogCook, Warning, TEXT("Retractionmessage received from CookWorker %d when we were not expecting one."),
 			Context.GetProfileId());
 		InitializeForResultsMessage(Context.GetWorkerId());
-		RetractionState = ERetractionState::WaitingForResponse;
+		SetRetractionState(ERetractionState::WaitingForResponse);
 	}
 
 	UE_CLOG(WorkerWithResults.IsValid(), LogCook, Error,
@@ -1694,7 +1803,8 @@ void FCookDirector::FRetractionHandler::HandleRetractionMessage(FMPCollectorServ
 	PackagesToRetract.FindOrAdd(WorkerWithResults).Append(Message.ReturnedPackages);
 }
 
-void FCookDirector::FRetractionHandler::ReassignPackages(const FWorkerId& FromWorker, TConstArrayView<FPackageData*> Packages)
+FCookDirector::FRetractionHandler::ERetractionResult
+FCookDirector::FRetractionHandler::ReassignPackages(const FWorkerId& FromWorker, TConstArrayView<FPackageData*> Packages)
 {
 	TArray<TRefCountPtr<FCookWorkerServer>> LocalRemoteWorkers = Director.CopyRemoteWorkers();
 
@@ -1726,7 +1836,7 @@ void FCookDirector::FRetractionHandler::ReassignPackages(const FWorkerId& FromWo
 		UE_LOG(LogCook, Display, TEXT("Retraction results message received from %s; no packages were available for retraction."),
 			*Director.GetDisplayName(FromWorker));
 		Director.DisplayRemainingPackages();
-		return;
+		return ERetractionResult::NoneAvailable;
 	}
 
 	TArray<FWorkerId> WorkersToSplitOver = CalculateWorkersToSplitOver(AssignmentPackages.Num(), FromWorker, LocalRemoteWorkers);
@@ -1742,7 +1852,7 @@ void FCookDirector::FRetractionHandler::ReassignPackages(const FWorkerId& FromWo
 		UE_LOG(LogCook, Display, TEXT("%d packages retracted from %s. No workers are currently idle so the packages were assigned evenly to all CookWorkers."),
 			AssignmentPackages.Num(), *Director.GetDisplayName(FromWorker));
 		Director.DisplayRemainingPackages();
-		return;
+		return ERetractionResult::Retracted;
 	}
 	for (const FWorkerId& WorkerId : WorkersRequiredByConstraint)
 	{
@@ -1789,6 +1899,8 @@ void FCookDirector::FRetractionHandler::ReassignPackages(const FWorkerId& FromWo
 		// that we didn't consider during SoftGC
 		Director.COTFS.PackageTracker->ClearExpectedNeverLoadPackages();
 	}
+
+	return ERetractionResult::Retracted;
 }
 
 TArray<FWorkerId> FCookDirector::FRetractionHandler::CalculateWorkersToSplitOver(int32 NumPackages, const FWorkerId& FromWorker,
diff --git a/Engine/Source/Editor/UnrealEd/Private/Cooker/CookDirector.h b/Engine/Source/Editor/UnrealEd/Private/Cooker/CookDirector.h
--- a/Engine/Source/Editor/UnrealEd/Private/Cooker/CookDirector.h
+++ b/Engine/Source/Editor/UnrealEd/Private/Cooker/CookDirector.h
@@ -210,6 +210,7 @@ private:
 	void DisplayRemainingPackages() const;
 	FString GetDisplayName(const FWorkerId& WorkerId, int32 PreferredWidth = -1) const;
 	FString GetDisplayName(const FCookWorkerServer& RemoteWorker, int32 PreferredWidth=-1) const;
+	const TRefCountPtr<FCookWorkerServer>* FindRemoteWorkerInLock(const FWorkerId& WorkerId) const;
 
 private:
 	// Synchronization primitives that can be used from any thread
diff --git a/Engine/Source/Editor/UnrealEd/Private/Cooker/CookPackageData.cpp b/Engine/Source/Editor/UnrealEd/Private/Cooker/CookPackageData.cpp
--- a/Engine/Source/Editor/UnrealEd/Private/Cooker/CookPackageData.cpp
+++ b/Engine/Source/Editor/UnrealEd/Private/Cooker/CookPackageData.cpp
@@ -2581,6 +2581,14 @@ void FPackageDataMonitor::OnStateChanged(FPackageData& PackageData, EPackageStat
 
 	TrackUrgentRequests(OldState, -1);
 	TrackUrgentRequests(PackageData.GetState(), 1);
+
+	EPackageState NewState = PackageData.GetState();
+	bool bOldStateAssignedToLocal = OldState != EPackageState::Idle && OldState != EPackageState::AssignedToWorker;
+	bool bNewStateAssignedToLocal = NewState != EPackageState::Idle && NewState != EPackageState::AssignedToWorker;
+	if (bOldStateAssignedToLocal != bNewStateAssignedToLocal)
+	{
+		++(bNewStateAssignedToLocal ? MPCookAssignedFenceMarker : MPCookRetiredFenceMarker);
+	}
 }
 
 void FPackageDataMonitor::TrackUrgentRequests(EPackageState State, int32 Delta)
@@ -2590,6 +2598,15 @@ void FPackageDataMonitor::TrackUrgentRequests(EPackageState State, int32 Delta)
 	check(NumUrgentInState[static_cast<uint32>(State) - static_cast<uint32>(EPackageState::Min)] >= 0);
 }
 
+int32 FPackageDataMonitor::GetMPCookAssignedFenceMarker() const
+{
+	return MPCookAssignedFenceMarker;
+}
+
+int32 FPackageDataMonitor::GetMPCookRetiredFenceMarker() const
+{
+	return MPCookRetiredFenceMarker;
+}
 
 //////////////////////////////////////////////////////////////////////////
 // FPackageDatas
diff --git a/Engine/Source/Editor/UnrealEd/Private/Cooker/CookPackageData.h b/Engine/Source/Editor/UnrealEd/Private/Cooker/CookPackageData.h
--- a/Engine/Source/Editor/UnrealEd/Private/Cooker/CookPackageData.h
+++ b/Engine/Source/Editor/UnrealEd/Private/Cooker/CookPackageData.h
@@ -1115,6 +1115,9 @@ public:
 	/** Callback called from FPackageData when it has changed its state. */
 	void OnStateChanged(FPackageData& PackageData, EPackageState OldState);
 
+	int32 GetMPCookAssignedFenceMarker() const;
+	int32 GetMPCookRetiredFenceMarker() const;
+
 private:
 	/** Increment or decrement the NumUrgent counter for the given state. */
 	void TrackUrgentRequests(EPackageState State, int32 Delta);
@@ -1123,6 +1126,8 @@ private:
 	int32 NumCooked[(uint8)ECookResult::Count]{};
 	int32 NumPreloadAllocated = 0;
 	int32 NumUrgentInState[static_cast<uint32>(EPackageState::Count)];
+	int32 MPCookAssignedFenceMarker = 0;
+	int32 MPCookRetiredFenceMarker = 0;
 };
 
 struct FDiscoveryQueueElement
diff --git a/Engine/Source/Editor/UnrealEd/Private/Cooker/CookWorkerServer.cpp b/Engine/Source/Editor/UnrealEd/Private/Cooker/CookWorkerServer.cpp
--- a/Engine/Source/Editor/UnrealEd/Private/Cooker/CookWorkerServer.cpp
+++ b/Engine/Source/Editor/UnrealEd/Private/Cooker/CookWorkerServer.cpp
@@ -226,6 +226,7 @@ void FCookWorkerServer::ShutdownRemoteProcess()
 void FCookWorkerServer::AppendAssignments(TArrayView<FPackageData*> Assignments, ECookDirectorThread TickThread)
 {
 	FCommunicationScopeLock ScopeLock(this, TickThread, ETickAction::Queue);
+	++PackagesAssignedFenceMarker;
 	PackagesToAssign.Append(Assignments);
 }
 
@@ -254,6 +255,7 @@ void FCookWorkerServer::AbortAllAssignmentsInLock(TSet<FPackageData*>& OutPendin
 	}
 	OutPendingPackages.Append(PackagesToAssign);
 	PackagesToAssign.Empty();
+	++PackagesRetiredFenceMarker;
 }
 
 void FCookWorkerServer::AbortAssignment(FPackageData& PackageData, ECookDirectorThread TickThread,
@@ -282,6 +284,7 @@ void FCookWorkerServer::AbortAssignments(TConstArrayView<FPackageData*> PackageD
 
 		PackagesToAssign.Remove(PackageData);
 	}
+	++PackagesRetiredFenceMarker;
 	if (!PackageNamesToMessage.IsEmpty())
 	{
 		SendMessageInLock(FAbortPackagesMessage(MoveTemp(PackageNamesToMessage)));
@@ -383,6 +386,17 @@ void FCookWorkerServer::SetLastReceivedHeartbeatNumberInLock(int32 InHeartbeatNu
 	LastReceivedHeartbeatNumber = InHeartbeatNumber;
 }
 
+int32 FCookWorkerServer::GetPackagesAssignedFenceMarker() const
+{
+	FScopeLock CommunicationScopeLock(&CommunicationLock);
+	return PackagesAssignedFenceMarker;
+}
+
+int32 FCookWorkerServer::GetPackagesRetiredFenceMarker() const
+{
+	FScopeLock CommunicationScopeLock(&CommunicationLock);
+	return PackagesRetiredFenceMarker;
+}
 
 bool FCookWorkerServer::TryHandleConnectMessage(FWorkerConnectMessage& Message, FSocket* InSocket, TArray<UE::CompactBinaryTCP::FMarshalledMessage>&& OtherPacketMessages, ECookDirectorThread TickThread)
 {
@@ -801,6 +815,7 @@ void FCookWorkerServer::RecordResults(FPackageResultsMessage& Message)
 {
 	check(TickState.TickThread == ECookDirectorThread::SchedulerThread);
 
+	bool bRetiredAnyPackages = false;
 	for (FPackageRemoteResult& Result : Message.Results)
 	{
 		FPackageData* PackageData = COTFS.PackageDatas->FindPackageDataByPackageName(Result.GetPackageName());
@@ -816,6 +831,7 @@ void FCookWorkerServer::RecordResults(FPackageResultsMessage& Message)
 				ProfileId, *Result.GetPackageName().ToString());
 			continue;
 		}
+		bRetiredAnyPackages = true;
 		PackageData->SetWorkerAssignment(FWorkerId::Invalid(), ESendFlags::QueueNone);
 
 		// MPCOOKTODO: Refactor FSaveCookedPackageContext::FinishPlatform and ::FinishPackage so we can call them from here
@@ -869,6 +885,10 @@ void FCookWorkerServer::RecordResults(FPackageResultsMessage& Message)
 		}
 	}
 	Director.ResetFinalIdleHeartbeatFence();
+	if (bRetiredAnyPackages)
+	{
+		++PackagesRetiredFenceMarker;
+	}
 }
 
 void FCookWorkerServer::LogInvalidMessage(const TCHAR* MessageTypeName)
diff --git a/Engine/Source/Editor/UnrealEd/Private/Cooker/CookWorkerServer.h b/Engine/Source/Editor/UnrealEd/Private/Cooker/CookWorkerServer.h
--- a/Engine/Source/Editor/UnrealEd/Private/Cooker/CookWorkerServer.h
+++ b/Engine/Source/Editor/UnrealEd/Private/Cooker/CookWorkerServer.h
@@ -99,6 +99,9 @@ public:
 	/** Set the LastReceivedHeartbeatNumber. Assumes lock is already entered; can only be called from with a HandleReceivedMessages callback */
 	void SetLastReceivedHeartbeatNumberInLock(int32 InHeartbeatNumber);
 
+	int32 GetPackagesAssignedFenceMarker() const;
+	int32 GetPackagesRetiredFenceMarker() const;
+
 private:
 	enum class EConnectStatus
 	{
@@ -189,6 +192,8 @@ private:
 	uint32 CookWorkerProcessId = 0;
 	int32 ProfileId = 0;
 	int32 LastReceivedHeartbeatNumber = 0;
+	int32 PackagesAssignedFenceMarker = 0;
+	int32 PackagesRetiredFenceMarker = 0;
 	double ConnectStartTimeSeconds = 0.;
 	double ConnectTestStartTimeSeconds = 0.;
 	FWorkerId WorkerId = FWorkerId::Invalid();
