Fix null dereference in UAssetEditorSubsystem::FindEditorsForAssetAndSubObjects
#rnx

diff --git a/Engine/Source/Editor/UnrealEd/Private/Subsystems/AssetEditorSubsystem.cpp b/Engine/Source/Editor/UnrealEd/Private/Subsystems/AssetEditorSubsystem.cpp
--- a/Engine/Source/Editor/UnrealEd/Private/Subsystems/AssetEditorSubsystem.cpp
+++ b/Engine/Source/Editor/UnrealEd/Private/Subsystems/AssetEditorSubsystem.cpp
@@ -288,14 +288,17 @@ TArray<IAssetEditorInstance*> UAssetEditorSubsystem::FindEditorsForAssetAndSubOb
 {
 	TArray<IAssetEditorInstance*> EditorInstances;
 
-	for (const TPair<FAssetEntry, IAssetEditorInstance*>& Pair : OpenedAssets)
+	if (Asset)
 	{
-		if (Pair.Key.RawPtr == Asset || (Pair.Key.ObjectPtr.IsValid() && Pair.Key.ObjectPtr.Get()->IsIn(Asset)))
-		{		
-			EditorInstances.Add(Pair.Value);
+		for (const TPair<FAssetEntry, IAssetEditorInstance*>& Pair : OpenedAssets)
+		{
+			if (Pair.Key.RawPtr == Asset || (Pair.Key.ObjectPtr.IsValid() && Pair.Key.ObjectPtr.Get()->IsIn(Asset)))
+			{		
+				EditorInstances.Add(Pair.Value);
+			}
 		}
 	}
-
+	
 	return EditorInstances;
 }
 
