MPCook: Add global function UE::GetMultiprocessId to handle the commandline parsing of MultiprocessId for CookWorkers.
#rnx

diff --git a/Engine/Source/Developer/DerivedDataCache/Private/ZenCacheStore.cpp b/Engine/Source/Developer/DerivedDataCache/Private/ZenCacheStore.cpp
--- a/Engine/Source/Developer/DerivedDataCache/Private/ZenCacheStore.cpp
+++ b/Engine/Source/Developer/DerivedDataCache/Private/ZenCacheStore.cpp
@@ -1654,8 +1654,7 @@ ILegacyCacheStore* CreateZenCacheStore(const TCHAR* NodeName, const TCHAR* Confi
 		AutoLaunchSettings.bLimitProcessLifetime = true;
 
 		// Flush the cache if requested.
-		uint32 MultiprocessId = 0;
-		FParse::Value(FCommandLine::Get(), TEXT("-MultiprocessId="), MultiprocessId);
+		uint32 MultiprocessId = UE::GetMultiprocessId();
 		if (bFlush && (MultiprocessId == 0))
 		{
 			bool bStopped = true;
diff --git a/Engine/Source/Developer/ShaderCompilerCommon/Private/ShaderSymbolExport.cpp b/Engine/Source/Developer/ShaderCompilerCommon/Private/ShaderSymbolExport.cpp
--- a/Engine/Source/Developer/ShaderCompilerCommon/Private/ShaderSymbolExport.cpp
+++ b/Engine/Source/Developer/ShaderCompilerCommon/Private/ShaderSymbolExport.cpp
@@ -90,8 +90,7 @@ void FShaderSymbolExport::Initialize()
 		else
 		{
 			// setup multiproc data in case we need it
-			uint32 MultiprocessId = 0;
-			FParse::Value(FCommandLine::Get(), TEXT("-MultiprocessId="), MultiprocessId);
+			uint32 MultiprocessId = UE::GetMultiprocessId();
 			bMultiprocessOwner = MultiprocessId == 0;
 
 			// Check if the export mode is to an uncompressed archive or loose files.
diff --git a/Engine/Source/Editor/UnrealEd/Private/Cooker/CookDirector.cpp b/Engine/Source/Editor/UnrealEd/Private/Cooker/CookDirector.cpp
--- a/Engine/Source/Editor/UnrealEd/Private/Cooker/CookDirector.cpp
+++ b/Engine/Source/Editor/UnrealEd/Private/Cooker/CookDirector.cpp
@@ -220,8 +220,9 @@ void FCookDirector::ParseConfig(int32 CookProcessCount, bool& bOutValid)
 
 	bAllowLocalCooks = !FParse::Param(CommandLine, TEXT("CookForceRemote"));
 
-	int32 MultiprocessId;
-	if (FParse::Value(CommandLine, TEXT("-MultiprocessId="), MultiprocessId))
+	
+	int32 MultiprocessId = UE::GetMultiprocessId();
+	if (MultiprocessId != 0)
 	{
 		bOutValid = false;
 		UE_LOG(LogCook, Error, TEXT("CookMultiprocess is incompatible with -MultiprocessId on the CookDirector's commandline. The CookDirector needs to be able to specify all MultiprocessIds."));
@@ -1362,8 +1363,8 @@ bool FDirectorConnectionInfo::TryParseCommandLine()
 		UE_LOG(LogCook, Error, TEXT("CookWorker startup failed: no CookDirector specified on commandline."));
 		return false;
 	}
-	uint32 MultiprocessId;
-	if (!FParse::Value(FCommandLine::Get(), TEXT("-MultiprocessId="), MultiprocessId))
+	uint32 MultiprocessId = UE::GetMultiprocessId();
+	if (MultiprocessId == 0 && !FParse::Value(FCommandLine::Get(), TEXT("-MultiprocessId="), MultiprocessId))
 	{
 		UE_LOG(LogCook, Error, TEXT("CookWorker startup failed: no MultiprocessId specified on commandline."));
 		return false;
diff --git a/Engine/Source/Editor/UnrealEd/Private/Cooker/CookProfiling.cpp b/Engine/Source/Editor/UnrealEd/Private/Cooker/CookProfiling.cpp
--- a/Engine/Source/Editor/UnrealEd/Private/Cooker/CookProfiling.cpp
+++ b/Engine/Source/Editor/UnrealEd/Private/Cooker/CookProfiling.cpp
@@ -998,8 +998,7 @@ void SendLogCookStats(ECookMode::Type CookMode)
 		FString CookStatsFileName;
 		if (FParse::Value(FCommandLine::Get(), TEXT("-CookStatsFile="), CookStatsFileName))
 		{
-			uint32 MultiprocessId = 0;
-			FParse::Value(FCommandLine::Get(), TEXT("-MultiprocessId="), MultiprocessId);
+			uint32 MultiprocessId = UE::GetMultiprocessId();
 			if (MultiprocessId != 0)
 			{
 				// Suppress the file creation on CookWorkers
diff --git a/Engine/Source/Runtime/AssetRegistry/Private/AssetDataGatherer.cpp b/Engine/Source/Runtime/AssetRegistry/Private/AssetDataGatherer.cpp
--- a/Engine/Source/Runtime/AssetRegistry/Private/AssetDataGatherer.cpp
+++ b/Engine/Source/Runtime/AssetRegistry/Private/AssetDataGatherer.cpp
@@ -3052,8 +3052,7 @@ struct FPreloadSettings
 		bool bNoAssetRegistryCache = FParse::Param(FCommandLine::Get(), TEXT("NoAssetRegistryCache"));
 		bool bNoAssetRegistryCacheRead = FParse::Param(FCommandLine::Get(), TEXT("NoAssetRegistryCacheRead"));
 		bool bNoAssetRegistryCacheWrite = FParse::Param(FCommandLine::Get(), TEXT("NoAssetRegistryCacheWrite"));
-		uint32 MultiprocessId = 0;
-		FParse::Value(FCommandLine::Get(), TEXT("multiprocessid="), MultiprocessId);
+		uint32 MultiprocessId = UE::GetMultiprocessId();
 		bool bMultiprocess = MultiprocessId > 0 || FParse::Param(FCommandLine::Get(), TEXT("multiprocess"));
 		bCacheReadEnabled = !bNoAssetRegistryCache && !bNoAssetRegistryCacheRead;
 		bCacheWriteEnabled = !bNoAssetRegistryCache && !bNoAssetRegistryCacheWrite && !bMultiprocess;
diff --git a/Engine/Source/Runtime/Core/Private/Misc/CoreGlobals.cpp b/Engine/Source/Runtime/Core/Private/Misc/CoreGlobals.cpp
--- a/Engine/Source/Runtime/Core/Private/Misc/CoreGlobals.cpp
+++ b/Engine/Source/Runtime/Core/Private/Misc/CoreGlobals.cpp
@@ -176,6 +176,10 @@ bool GIsRunningUnattendedScript = false;
 #if WITH_EDITOR
 bool					PRIVATE_GIsRunningCookCommandlet	= false;				/** Whether this executable is running the cook commandlet */
 bool					PRIVATE_GIsRunningDLCCookCommandlet = false;				/** Whether this executable is running the cook commandlet on a DLC plugin */
+namespace UE::Private
+{
+int32					GMultiprocessId = 0;
+}
 #endif
 
 #if WITH_ENGINE
@@ -874,3 +878,29 @@ bool IsRunningCookOnTheFly()
 	return false;
 #endif
 }
+
+namespace UE
+{
+
+int32 GetMultiprocessId()
+{
+#if WITH_EDITOR
+	return UE::Private::GMultiprocessId;
+#else
+	return 0;
+#endif
+}
+
+}
+
+namespace UE::Private
+{
+
+void SetMultiprocessId(int32 MultiprocessId)
+{
+#if WITH_EDITOR
+	UE::Private::GMultiprocessId = MultiprocessId;
+#endif
+}
+
+}
diff --git a/Engine/Source/Runtime/Core/Public/CoreGlobals.h b/Engine/Source/Runtime/Core/Public/CoreGlobals.h
--- a/Engine/Source/Runtime/Core/Public/CoreGlobals.h
+++ b/Engine/Source/Runtime/Core/Public/CoreGlobals.h
@@ -310,6 +310,23 @@ FORCEINLINE bool IsAllowCommandletAudio()
 #endif
 }
 
+namespace UE
+{
+/**
+ * If multiple copies of the process are running and communicating together, the director process has MultiprocessId 0
+ * and the worker processes have MultiprocessId > 0.
+ * When only a single process is running, GetMultiprocessId returns 0.
+ * Based on commandline. Primary example is cookworkers in multiprocess cook.
+ * Systems that need to write a separate file per running process should use this id as part of their file descriptor.
+ */
+CORE_API int32 GetMultiprocessId();
+}
+
+namespace UE::Private
+{
+CORE_API void SetMultiprocessId(int32 MultiprocessId);
+}
+
 class FIsDuplicatingClassForReinstancing
 {
 public:
diff --git a/Engine/Source/Runtime/Engine/Private/Materials/HLSLMaterialTranslator.cpp b/Engine/Source/Runtime/Engine/Private/Materials/HLSLMaterialTranslator.cpp
--- a/Engine/Source/Runtime/Engine/Private/Materials/HLSLMaterialTranslator.cpp
+++ b/Engine/Source/Runtime/Engine/Private/Materials/HLSLMaterialTranslator.cpp
@@ -99,8 +99,7 @@ struct FCsvLogFile
 	{
 		if (GEnableMaterialTranslationLogFile)
 		{
-			uint32 MultiprocessId = 0;
-			FParse::Value(FCommandLine::Get(), TEXT("-MultiprocessId="), MultiprocessId);
+			uint32 MultiprocessId = UE::GetMultiprocessId();
 			FString FilePath = GShaderCompilingManager->GetAbsoluteShaderDebugInfoDirectory() / FString::Printf(TEXT("MaterialTranslationLog-%d.csv"), MultiprocessId);
 			if (!FFileHelper::SaveStringToFile(LogContent, *FilePath))
 			{
diff --git a/Engine/Source/Runtime/Launch/Private/LaunchEngineLoop.cpp b/Engine/Source/Runtime/Launch/Private/LaunchEngineLoop.cpp
--- a/Engine/Source/Runtime/Launch/Private/LaunchEngineLoop.cpp
+++ b/Engine/Source/Runtime/Launch/Private/LaunchEngineLoop.cpp
@@ -2417,7 +2417,15 @@ int32 FEngineLoop::PreInitPreStartupScreen(const TCHAR* CmdLine)
 			}
 		}
 	}
-#endif // UE_EDITOR || WITH_ENGINE
+#endif // UE_EDITOR || WITH_ENGINE || WITH_EDITOR
+#if WITH_EDITOR
+	int32 MultiprocessId;
+	if (FParse::Value(CmdLine, TEXT("-MultiprocessId="), MultiprocessId))
+	{
+		UE::Private::SetMultiprocessId(MultiprocessId);
+	}
+#endif
+
 
 	// In the commandlet case, we have set the Token to something other than the first token from commandline.
 	// In all other cases we should check for the first token being the Project Specifier
