Enabling root node inline controls by default (removing CVar controlling it)

diff --git a/Engine/Source/Editor/UnrealEd/Private/MaterialGraphNode_Root.cpp b/Engine/Source/Editor/UnrealEd/Private/MaterialGraphNode_Root.cpp
--- a/Engine/Source/Editor/UnrealEd/Private/MaterialGraphNode_Root.cpp
+++ b/Engine/Source/Editor/UnrealEd/Private/MaterialGraphNode_Root.cpp
@@ -16,35 +16,18 @@
 
 #define LOCTEXT_NAMESPACE "MaterialGraphNode_Root"
 
-
-static TAutoConsoleVariable<bool> CVarEnableRootNodeInlineControls(
-	TEXT("MaterialGraph.EnableRootNodeInlineControls"),
-	false,
-	TEXT("Control if the Material Graph Root Node should show inline editing controls")
-	);
-
 /////////////////////////////////////////////////////
 // UMaterialGraphNode_Root
 
 UMaterialGraphNode_Root::UMaterialGraphNode_Root(const FObjectInitializer& ObjectInitializer)
 	: Super(ObjectInitializer)
 {
-	CVarEnableRootNodeInlineControls->SetOnChangedCallback(FConsoleVariableDelegate::CreateLambda([this](IConsoleVariable* InVariable)
-		{
-			for (UEdGraphPin* Pin : this->Pins)
-			{
-				UpdateInputUseConstant(Pin, InVariable->GetBool());
-			}
-		}));
 }
 
 void UMaterialGraphNode_Root::UpdateInputUseConstant(UEdGraphPin* Pin, bool bUseConstant)
 {
 	const UMaterialGraph* MaterialGraph = CastChecked<UMaterialGraph>(GetGraph());
 	const FMaterialInputInfo& MaterialInput = MaterialGraph->MaterialInputs[Pin->SourceIndex];
-
-	const bool ShouldEnableInlineControls = CVarEnableRootNodeInlineControls.GetValueOnAnyThread();
-	bUseConstant &= ShouldEnableInlineControls;
 	
 	UMaterialEditorOnlyData* EditorOnlyData = Material->GetEditorOnlyData();
 	EMaterialProperty Property = MaterialInput.GetProperty();
@@ -220,11 +203,8 @@ void UMaterialGraphNode_Root::CreateInputPins()
 		FName MaterialInputName = *MaterialInput.GetName().ToString();
 		FName PinSubCategory;
 		FString RefractionMethodStr;
-		const bool ShouldEnableInlineControls = CVarEnableRootNodeInlineControls.GetValueOnAnyThread();
-		if (ShouldEnableInlineControls)
+		switch (Property)
 		{
-			switch (Property)
-			{
 			case MP_Metallic:
 			case MP_Specular:
 			case MP_Roughness:
@@ -289,16 +269,13 @@ void UMaterialGraphNode_Root::CreateInputPins()
 					PinSubCategory = UMaterialGraphSchema::PSC_RG;
 				}
 				break;
-			}
 		}
 		
 		UEdGraphPin* InputPin = CreatePin(EGPD_Input, UMaterialGraphSchema::PC_MaterialInput, PinSubCategory/*, *FString::Printf(TEXT("%d"), (int32)Property)*/, MaterialInputName);
 		InputPin->SourceIndex = Index;
 
-		if (ShouldEnableInlineControls)
+		switch (Property)
 		{
-			switch (Property)
-			{
 			case MP_EmissiveColor:		InputPin->DefaultValue = EditorOnlyData->EmissiveColor.GetDefaultValue(); break;
 			case MP_Opacity:			InputPin->DefaultValue = EditorOnlyData->Opacity.GetDefaultValue(); break;
 			case MP_OpacityMask:		InputPin->DefaultValue = EditorOnlyData->OpacityMask.GetDefaultValue(); break;
@@ -326,10 +303,8 @@ void UMaterialGraphNode_Root::CreateInputPins()
 					InputPin->DefaultValue = EditorOnlyData->CustomizedUVs[Property - MP_CustomizedUVs0].GetDefaultValue(); break;
 				}
 				break;
-			}
 		}
 	}
-
 }
 
 #undef LOCTEXT_NAMESPACE
