Cooker: DiffOnly: Add feedback for Error: Indeterministic header size... by dumping the header differences.
#rnx

diff --git a/Engine/Source/Editor/UnrealEd/Private/Cooker/DiffWriterArchive.cpp b/Engine/Source/Editor/UnrealEd/Private/Cooker/DiffWriterArchive.cpp
--- a/Engine/Source/Editor/UnrealEd/Private/Cooker/DiffWriterArchive.cpp
+++ b/Engine/Source/Editor/UnrealEd/Private/Cooker/DiffWriterArchive.cpp
@@ -564,6 +564,13 @@ void FAccumulator::OnFirstSaveComplete(FStringView LooseFilePath, int64 InHeader
 	ExportsCallstacks.Reset();
 
 	GenerateDiffMap();
+	if (HasDifferences())
+	{
+		// Make a copy of the LinkerArchive for comparison in case it differs even in the second memory save
+		check(LinkerArchive);
+		FirstSaveLinkerData.Empty(LinkerArchive->TotalSize());
+		FirstSaveLinkerData.Append(LinkerArchive->GetData(), LinkerArchive->TotalSize());
+	}
 
 	LinkerCallstacks.Reset();
 	bFirstSaveComplete = true;
@@ -589,8 +596,30 @@ void FAccumulator::OnSecondSaveComplete(int64 InHeaderSize)
 	if (HeaderSize != InHeaderSize)
 	{
 		MessageCallback(ELogVerbosity::Error, FString::Printf(
-			TEXT("%s: Indeterministic header size. When saving the package twice into memory, first header size %d != second header size %d. Callstacks for indeterminism in the exports will be incorrect."),
+			TEXT("%s: Indeterministic header size. When saving the package twice into memory, first header size %d != second header size %d. Callstacks for indeterminism in the exports will be incorrect.")
+			TEXT("\n\tDumping differences from first and second memory saves."),
 			*this->Filename, HeaderSize, InHeaderSize));
+
+		check(bFirstSaveComplete);
+		check(FirstSaveLinkerData.Num() >= HeaderSize);
+		check(LinkerArchive && LinkerArchive->TotalSize() >= InHeaderSize);
+
+		FPackageData FirstSaveHeader{ FirstSaveLinkerData.GetData(), HeaderSize};
+		FPackageData SecondSaveHeader{ LinkerArchive->GetData(), InHeaderSize };
+		int32 NumHeaderDiffMessages = 0;
+		DumpPackageHeaderDiffs(Globals, FirstSaveHeader, SecondSaveHeader, Filename, MaxDiffsToLog,
+			PackageHeaderFormat,
+			[&NumHeaderDiffMessages, this](ELogVerbosity::Type Verbosity, FStringView Message)
+			{
+				MessageCallback(Verbosity, Message);
+				++NumHeaderDiffMessages;
+			});
+		if (NumHeaderDiffMessages == 0)
+		{
+			MessageCallback(ELogVerbosity::Warning, FString::Printf(
+				TEXT("%s: headers are different, but DumpPackageHeaderDiffs does not yet implement describing the difference."),
+				*Filename));
+		}
 	}
 
 	if (IsWriterUsingPostSaveTransforms())
diff --git a/Engine/Source/Editor/UnrealEd/Private/Cooker/DiffWriterArchive.h b/Engine/Source/Editor/UnrealEd/Private/Cooker/DiffWriterArchive.h
--- a/Engine/Source/Editor/UnrealEd/Private/Cooker/DiffWriterArchive.h
+++ b/Engine/Source/Editor/UnrealEd/Private/Cooker/DiffWriterArchive.h
@@ -280,6 +280,8 @@ private:
 	ICookedPackageWriter::FPreviousCookedBytesData PreviousPackageData;
 	FDiffArchive* LinkerArchive = nullptr;
 	FDiffArchive* ExportsArchive = nullptr;
+	TArray<uint8> FirstSaveLinkerData;
+	int64 FirstSaveLinkerSize = 0;
 	FAccumulatorGlobals& Globals;
 
 	FDiffMap DiffMap;
@@ -323,10 +325,7 @@ protected:
 };
 
 /**
- * The archive written to when SavePackage is writing the Serialize blobs for exports.
- * When cooking, exports are serialized into a separate archive. We collect the serialization
- * callstack offsets and stack traces are into a separate callstack collection and append it
- * at the proper offset to the overall callstacks for the entire linker archive.
+ * The archive written to by SavePackage, includes the header and exports.
  */
 class FDiffArchiveForLinker : public FDiffArchive
 {
@@ -338,6 +337,12 @@ public:
 	FORCENOINLINE virtual void Serialize(void* InData, int64 Num) override; // FORCENOINLINE so it can be counted during StackTrace
 };
 
+/**
+ * The archive written to when SavePackage is writing the Serialize blobs for exports.
+ * When cooking, exports are serialized into a separate archive. We collect the serialization
+ * callstack offsets and stack traces into a separate callstack collection and append it
+ * at the proper offset to the overall callstacks for the entire linker archive.
+ */
 class FDiffArchiveForExports : public FDiffArchive
 {
 public:
