Cooker: WorldPartition: Improve feedback for "Error: CookPackageSplitter failure: We are demoting a generator package..." by adding an indicator of which SendToState callsite caused the demotion.
#rnx

diff --git a/Engine/Source/Editor/UnrealEd/Classes/CookOnTheSide/CookOnTheFlyServer.h b/Engine/Source/Editor/UnrealEd/Classes/CookOnTheSide/CookOnTheFlyServer.h
--- a/Engine/Source/Editor/UnrealEd/Classes/CookOnTheSide/CookOnTheFlyServer.h
+++ b/Engine/Source/Editor/UnrealEd/Classes/CookOnTheSide/CookOnTheFlyServer.h
@@ -172,7 +172,7 @@ namespace UE::Cook
 	class IWorkerRequests;
 	enum class ECachedCookedPlatformDataEvent : uint8;
 	enum class EPollStatus : uint8;
-	enum class EReleaseSaveReason : uint8;
+	enum class EStateChangeReason : uint8;
 	enum class ESuppressCookReason : uint8;
 	enum class ESendFlags : uint8;
 	struct FBeginCookConfigSettings;
@@ -1168,7 +1168,7 @@ private:
 	 *        the save again. If true, all data will be wiped.
 	 * @param ReleaseSaveReason Why the save data is being released, allows specifying how much to tear down
 	 */
-	UNREALED_API void ReleaseCookedPlatformData(UE::Cook::FPackageData& PackageData, UE::Cook::EReleaseSaveReason ReleaseSaveReason);
+	UNREALED_API void ReleaseCookedPlatformData(UE::Cook::FPackageData& PackageData, UE::Cook::EStateChangeReason ReleaseSaveReason);
 
 	/**
 	 * Poll the PendingCookedPlatformDatas and release their resources when they are complete.
diff --git a/Engine/Source/Editor/UnrealEd/Private/CookOnTheFlyServer.cpp b/Engine/Source/Editor/UnrealEd/Private/CookOnTheFlyServer.cpp
--- a/Engine/Source/Editor/UnrealEd/Private/CookOnTheFlyServer.cpp
+++ b/Engine/Source/Editor/UnrealEd/Private/CookOnTheFlyServer.cpp
@@ -2528,7 +2528,7 @@ bool UCookOnTheFlyServer::TryCreateRequestCluster(UE::Cook::FPackageData& Packag
 	using namespace UE::Cook;
 	if (!PackageData.AreAllReachablePlatformsVisitedByCluster())
 	{
-		PackageData.SendToState(EPackageState::Request, ESendFlags::QueueAdd);
+		PackageData.SendToState(EPackageState::Request, ESendFlags::QueueAdd, EStateChangeReason::Discovered);
 		return true;
 	}
 	return false;
@@ -2631,7 +2631,7 @@ void UCookOnTheFlyServer::PumpRequests(UE::Cook::FTickStackData& StackData, int3
 			DemoteToIdle(*PackageData, ESendFlags::QueueAdd, SuppressCookReason);
 			continue;
 		}
-		PackageData->SendToState(EPackageState::LoadPrepare, ESendFlags::QueueAdd);
+		PackageData->SendToState(EPackageState::LoadPrepare, ESendFlags::QueueAdd, EStateChangeReason::Requested);
 		++NumInBatch;
 	}
 	OutNumPushed += NumInBatch;
@@ -2666,7 +2666,7 @@ void UCookOnTheFlyServer::AssignRequests(TArrayView<UE::Cook::FPackageData*> Req
 			}
 			else
 			{
-				PackageData->SendToState(EPackageState::AssignedToWorker, ESendFlags::QueueAdd);
+				PackageData->SendToState(EPackageState::AssignedToWorker, ESendFlags::QueueAdd, EStateChangeReason::Requested);
 				PackageData->SetWorkerAssignment(Assignment);
 			}
 		}
@@ -2732,14 +2732,14 @@ void UCookOnTheFlyServer::DemoteToIdle(UE::Cook::FPackageData& PackageData, UE::
 			}
 		}
 	}
-	PackageData.SendToState(UE::Cook::EPackageState::Idle, SendFlags);
+	PackageData.SendToState(UE::Cook::EPackageState::Idle, SendFlags, EStateChangeReason::CookSuppressed);
 }
 
 void UCookOnTheFlyServer::PromoteToSaveComplete(UE::Cook::FPackageData& PackageData, UE::Cook::ESendFlags SendFlags)
 {
 	check(PackageData.IsInProgress());
 	WorkerRequests->ReportPromoteToSaveComplete(PackageData);
-	PackageData.SendToState(UE::Cook::EPackageState::Idle, SendFlags);
+	PackageData.SendToState(UE::Cook::EPackageState::Idle, SendFlags, UE::Cook::EStateChangeReason::Saved);
 }
 
 void UCookOnTheFlyServer::PumpLoads(UE::Cook::FTickStackData& StackData, uint32 DesiredQueueLength, int32& OutNumPushed, bool& bOutBusy)
@@ -2814,7 +2814,7 @@ void UCookOnTheFlyServer::PumpPreloadCompletes()
 		{
 			// Ready to go
 			PreloadingQueue.PopFront();
-			PackageData->SendToState(EPackageState::LoadReady, ESendFlags::QueueAdd);
+			PackageData->SendToState(EPackageState::LoadReady, ESendFlags::QueueAdd, EStateChangeReason::Loaded);
 			continue;
 		}
 		break;
@@ -2945,7 +2945,7 @@ void UCookOnTheFlyServer::LoadPackageInQueue(UE::Cook::FPackageData& PackageData
 
 	PostLoadPackageFixup(PackageData, LoadedPackage);
 	PackageData.SetPackage(LoadedPackage);
-	PackageData.SendToState(EPackageState::Save, ESendFlags::QueueAdd);
+	PackageData.SendToState(EPackageState::Save, ESendFlags::QueueAdd, EStateChangeReason::Loaded);
 	++OutNumPushed;
 }
 
@@ -3916,7 +3916,7 @@ UE::Cook::EPollStatus UCookOnTheFlyServer::CallBeginCacheOnObjects(UE::Cook::FPa
 	return EPollStatus::Success;
 }
 
-void UCookOnTheFlyServer::ReleaseCookedPlatformData(UE::Cook::FPackageData& PackageData, UE::Cook::EReleaseSaveReason ReleaseSaveReason)
+void UCookOnTheFlyServer::ReleaseCookedPlatformData(UE::Cook::FPackageData& PackageData, UE::Cook::EStateChangeReason ReleaseSaveReason)
 {
 	using namespace UE::Cook;
 
@@ -3934,7 +3934,7 @@ void UCookOnTheFlyServer::ReleaseCookedPlatformData(UE::Cook::FPackageData& Pack
 	FCookGenerationInfo* GenerationInfo = Generator ? Generator->FindInfo(PackageData) : nullptr;
 
 	// For every BeginCacheForCookedPlatformData call we made we need to call ClearAllCachedCookedPlatformData
-	if (ReleaseSaveReason == EReleaseSaveReason::Completed)
+	if (ReleaseSaveReason == EStateChangeReason::Completed)
 	{
 		// Since we have completed CookedPlatformData, we know we called BeginCacheForCookedPlatformData on all
 		// objects in the package, and none are pending
@@ -4018,7 +4018,7 @@ void UCookOnTheFlyServer::ReleaseCookedPlatformData(UE::Cook::FPackageData& Pack
 			PackageData.SetInitializedGeneratorSave(false);
 		}
 
-		if (ReleaseSaveReason == EReleaseSaveReason::Completed)
+		if (ReleaseSaveReason == EStateChangeReason::Completed)
 		{
 			Generator->SetPackageSaved(*GenerationInfo, PackageData);
 			if (Generator->IsComplete())
@@ -4040,7 +4040,7 @@ void UCookOnTheFlyServer::ReleaseCookedPlatformData(UE::Cook::FPackageData& Pack
 
 	PackageData.ClearCookedPlatformData();
 
-	if (ReleaseSaveReason != EReleaseSaveReason::RecreateObjectCache)
+	if (ReleaseSaveReason != EStateChangeReason::RecreateObjectCache)
 	{
 		if (!IsCookOnTheFlyMode() && !IsCookingInEditor())
 		{
@@ -4452,7 +4452,7 @@ void UCookOnTheFlyServer::PumpSaves(UE::Cook::FTickStackData& StackData, uint32
 			if (PrepareSaveStatus == EPollStatus::Error)
 			{
 				check(PackageData.HasPrepareSaveFailed()); // Should have been set by PrepareSave; we rely on this for cleanup
-				ReleaseCookedPlatformData(PackageData, EReleaseSaveReason::AbortSave);
+				ReleaseCookedPlatformData(PackageData, EStateChangeReason::SaveError);
 				PackageData.SetPlatformsCooked(PlatformsForPackage, ECookResult::Failed);
 				DemoteToIdle(PackageData, ESendFlags::QueueAdd, ESuppressCookReason::SaveError);
 				++OutNumPushed;
@@ -4542,7 +4542,7 @@ void UCookOnTheFlyServer::PumpSaves(UE::Cook::FTickStackData& StackData, uint32
 		{
 			// Timeouts can occur because of new objects created during the save, so we need to update our object cache,
 			// so we call ReleaseCookedPlatformData and ClearObjectCache to clear it and recache on next attempt.
-			ReleaseCookedPlatformData(PackageData, EReleaseSaveReason::RecreateObjectCache);
+			ReleaseCookedPlatformData(PackageData, EStateChangeReason::RecreateObjectCache);
 			PackageData.ClearObjectCache();
 			if (PackageData.GetIsUrgent())
 			{
@@ -4555,7 +4555,7 @@ void UCookOnTheFlyServer::PumpSaves(UE::Cook::FTickStackData& StackData, uint32
 			continue;
 		}
 
-		ReleaseCookedPlatformData(PackageData, !Context.bHasRetryErrorCode ? EReleaseSaveReason::Completed : EReleaseSaveReason::DoneForNow);
+		ReleaseCookedPlatformData(PackageData, !Context.bHasRetryErrorCode ? EStateChangeReason::Completed : EStateChangeReason::DoneForNow);
 		PromoteToSaveComplete(PackageData, ESendFlags::QueueAdd);
 		++OutNumPushed;
 	}
@@ -5110,7 +5110,7 @@ void UCookOnTheFlyServer::PreGarbageCollect()
 			Generator->PreGarbageCollect(*Info, GCKeepObjects, GCKeepPackages, GCKeepPackageDatas, bShouldDemote);
 			if (bShouldDemote)
 			{
-				ReleaseCookedPlatformData(*PackageData, UE::Cook::EReleaseSaveReason::Demoted);
+				ReleaseCookedPlatformData(*PackageData, UE::Cook::EStateChangeReason::GeneratorPreGarbageCollected);
 			}
 		}
 	}
@@ -5311,7 +5311,7 @@ void UCookOnTheFlyServer::PostGarbageCollect()
 	}
 	for (FPackageData* PackageData : Demotes)
 	{
-		PackageData->SendToState(EPackageState::Request, ESendFlags::QueueRemove);
+		PackageData->SendToState(EPackageState::Request, ESendFlags::QueueRemove, EStateChangeReason::GarbageCollected);
 		PackageDatas->GetRequestQueue().AddRequest(PackageData, /* bForceUrgent */ true);
 	}
 
@@ -5684,7 +5684,8 @@ void UCookOnTheFlyServer::MarkPackageDirtyForCookerFromSchedulerThread(const FNa
 		PackageData->ClearCookResults();
 		if (PackageData->IsInProgress())
 		{
-			PackageData->SendToState(UE::Cook::EPackageState::Request, UE::Cook::ESendFlags::QueueAddAndRemove);
+			PackageData->SendToState(UE::Cook::EPackageState::Request,
+				UE::Cook::ESendFlags::QueueAddAndRemove, UE::Cook::EStateChangeReason::ForceRecook);
 		}
 		else if (IsCookByTheBookMode() && IsInSession() && bHadCookedPlatforms)
 		{
diff --git a/Engine/Source/Editor/UnrealEd/Private/Cooker/CookDirector.cpp b/Engine/Source/Editor/UnrealEd/Private/Cooker/CookDirector.cpp
--- a/Engine/Source/Editor/UnrealEd/Private/Cooker/CookDirector.cpp
+++ b/Engine/Source/Editor/UnrealEd/Private/Cooker/CookDirector.cpp
@@ -245,7 +245,7 @@ FCookDirector::~FCookDirector()
 	{
 		check(PackageData->IsInProgress()); // Packages that were assigned to workers should be in the AssignedToWorker state
 		PackageData->SetWorkerAssignment(FWorkerId::Invalid(), ESendFlags::QueueNone);
-		PackageData->SendToState(UE::Cook::EPackageState::Request, ESendFlags::QueueAddAndRemove);
+		PackageData->SendToState(UE::Cook::EPackageState::Request, ESendFlags::QueueAddAndRemove, EStateChangeReason::CookerShutdown);
 	}
 	RemoteWorkers.Empty();
 	RemoteWorkerProfileDatas.Empty();
@@ -1440,7 +1440,7 @@ void FCookDirector::ReassignAbortedPackages(TArray<FPackageData*>& PackagesToRea
 	{
 		check(PackageData->IsInProgress()); // Packages that were assigned to a worker should be in the AssignedToWorker state
 		PackageData->SetWorkerAssignment(FWorkerId::Invalid());
-		PackageData->SendToState(UE::Cook::EPackageState::Request, ESendFlags::QueueAddAndRemove);
+		PackageData->SendToState(UE::Cook::EPackageState::Request, ESendFlags::QueueAddAndRemove, EStateChangeReason::ReassignAbortedPackages);
 	}
 	PackagesToReassign.Empty();
 }
@@ -1830,7 +1830,7 @@ FCookDirector::FRetractionHandler::ReassignPackages(const FWorkerId& FromWorker,
 		}
 
 		AssignmentPackages.Add(PackageData);
-		PackageData->SendToState(EPackageState::Request, ESendFlags::QueueRemove);
+		PackageData->SendToState(EPackageState::Request, ESendFlags::QueueRemove, EStateChangeReason::Retraction);
 	}
 	if (AssignmentPackages.IsEmpty())
 	{
@@ -1889,7 +1889,7 @@ FCookDirector::FRetractionHandler::ReassignPackages(const FWorkerId& FromWorker,
 		}
 		else
 		{
-			PackageData->SendToState(EPackageState::AssignedToWorker, ESendFlags::QueueAdd);
+			PackageData->SendToState(EPackageState::AssignedToWorker, ESendFlags::QueueAdd, EStateChangeReason::Retraction);
 			PackageData->SetWorkerAssignment(Assignment);
 		}
 	}
diff --git a/Engine/Source/Editor/UnrealEd/Private/Cooker/CookPackageData.cpp b/Engine/Source/Editor/UnrealEd/Private/Cooker/CookPackageData.cpp
--- a/Engine/Source/Editor/UnrealEd/Private/Cooker/CookPackageData.cpp
+++ b/Engine/Source/Editor/UnrealEd/Private/Cooker/CookPackageData.cpp
@@ -98,7 +98,7 @@ FPackageData::FPackageData(FPackageDatas& PackageDatas, const FName& InPackageNa
 	, bWasCookedThisSession(0)
 {
 	SetState(EPackageState::Idle);
-	SendToState(EPackageState::Idle, ESendFlags::QueueAdd);
+	SendToState(EPackageState::Idle, ESendFlags::QueueAdd, EStateChangeReason::Discovered);
 }
 
 FPackageData::~FPackageData()
@@ -108,7 +108,7 @@ FPackageData::~FPackageData()
 	// We need to send OnLastCookedPlatformRemoved message to the monitor, so call SetPlatformsNotCooked
 	ClearCookResults();
 	// Update the monitor's counters and call exit functions
-	SendToState(EPackageState::Idle, ESendFlags::QueueNone);
+	SendToState(EPackageState::Idle, ESendFlags::QueueNone, EStateChangeReason::CookerShutdown);
 }
 
 void FPackageData::ClearReferences()
@@ -264,7 +264,7 @@ void FPackageData::AddUrgency(bool bUrgent, bool bAllowUpdateState)
 	SetIsUrgent(true);
 	if (!bWasUrgent && bAllowUpdateState)
 	{
-		SendToState(GetState(), ESendFlags::QueueAddAndRemove);
+		SendToState(GetState(), ESendFlags::QueueAddAndRemove, EStateChangeReason::UrgencyUpdated);
 	}
 }
 
@@ -565,7 +565,7 @@ struct FStateProperties
 	}
 };
 
-void FPackageData::SendToState(EPackageState NextState, ESendFlags SendFlags)
+void FPackageData::SendToState(EPackageState NextState, ESendFlags SendFlags, EStateChangeReason ReleaseSaveReason)
 {
 	EPackageState OldState = GetState();
 	switch (OldState)
@@ -606,7 +606,7 @@ void FPackageData::SendToState(EPackageState NextState, ESendFlags SendFlags)
 		{
 			ensure(PackageDatas.GetSaveQueue().Remove(this) == 1);
 		}
-		OnExitSave();
+		OnExitSave(ReleaseSaveReason);
 		break;
 	default:
 		check(false);
@@ -850,9 +850,9 @@ void FPackageData::OnEnterSave()
 	CheckCookedPlatformDataEmpty();
 }
 
-void FPackageData::OnExitSave()
+void FPackageData::OnExitSave(EStateChangeReason ReleaseSaveReason)
 {
-	PackageDatas.GetCookOnTheFlyServer().ReleaseCookedPlatformData(*this, EReleaseSaveReason::Demoted);
+	PackageDatas.GetCookOnTheFlyServer().ReleaseCookedPlatformData(*this, ReleaseSaveReason);
 	ClearObjectCache();
 	SetHasPrepareSaveFailed(false);
 	SetIsPrepareSaveRequiresGC(false);
@@ -1372,7 +1372,7 @@ void FPackageData::RemapTargetPlatforms(const TMap<ITargetPlatform*, ITargetPlat
 		}
 		if (bDemote)
 		{
-			SendToState(EPackageState::Request, ESendFlags::QueueAddAndRemove);
+			SendToState(EPackageState::Request, ESendFlags::QueueAddAndRemove, EStateChangeReason::ForceRecook);
 		}
 	}
 	PlatformDatas = MoveTemp(NewPlatformDatas);
@@ -1951,7 +1951,7 @@ bool FGeneratorPackage::IsComplete() const
 	return RemainingToPopulate == 0;
 }
 
-void FGeneratorPackage::ResetSaveState(FCookGenerationInfo& Info, UPackage* Package, EReleaseSaveReason ReleaseSaveReason)
+void FGeneratorPackage::ResetSaveState(FCookGenerationInfo& Info, UPackage* Package, EStateChangeReason ReleaseSaveReason)
 {
 	check(IsInitialized());
 	if (Info.GetSaveState() > FCookGenerationInfo::ESaveState::CallPopulate)
@@ -1987,9 +1987,8 @@ void FGeneratorPackage::ResetSaveState(FCookGenerationInfo& Info, UPackage* Pack
 
 	if (Info.IsGenerator())
 	{
-		if (ReleaseSaveReason == EReleaseSaveReason::RecreateObjectCache ||
-			ReleaseSaveReason == EReleaseSaveReason::Demoted ||
-			ReleaseSaveReason == EReleaseSaveReason::DoneForNow)
+		if (ReleaseSaveReason == EStateChangeReason::RecreateObjectCache ||
+			ReleaseSaveReason == EStateChangeReason::DoneForNow)
 		{
 			if (Info.GetSaveState() >= FCookGenerationInfo::ESaveState::StartPopulate)
 			{
@@ -1998,7 +1997,7 @@ void FGeneratorPackage::ResetSaveState(FCookGenerationInfo& Info, UPackage* Pack
 			else
 			{
 				// Redo all the steps since we didn't make it to the FinishCachePreObjectsToMove.
-				// Restarting in the middle of that flow after a GarbageCollect is not robust
+				// Restarting in the middle of that flow is not robust
 				Info.SetSaveState(FCookGenerationInfo::ESaveState::StartGenerate);
 			}
 		}
@@ -2015,7 +2014,8 @@ void FGeneratorPackage::ResetSaveState(FCookGenerationInfo& Info, UPackage* Pack
 	{
 		if (Info.PackageData && Info.PackageData->GetCachedObjectsInOuter().Num() != 0 &&
 			IsUseInternalReferenceToAvoidGarbageCollect() &&
-			(ReleaseSaveReason == EReleaseSaveReason::Demoted || ReleaseSaveReason == EReleaseSaveReason::RecreateObjectCache))
+			(ReleaseSaveReason != EStateChangeReason::Completed && ReleaseSaveReason != EStateChangeReason::DoneForNow &&
+			 ReleaseSaveReason != EStateChangeReason::SaveError && ReleaseSaveReason != EStateChangeReason::CookerShutdown))
 		{
 			UE_LOG(LogCook, Error, TEXT("CookPackageSplitter failure: We are demoting a %s package from save and removing our references that keep its objects loaded.\n")
 				TEXT("This will allow the objects to be garbage collected and cause failures in the splitter which expects them to remain loaded.\n")
diff --git a/Engine/Source/Editor/UnrealEd/Private/Cooker/CookPackageData.h b/Engine/Source/Editor/UnrealEd/Private/Cooker/CookPackageData.h
--- a/Engine/Source/Editor/UnrealEd/Private/Cooker/CookPackageData.h
+++ b/Engine/Source/Editor/UnrealEd/Private/Cooker/CookPackageData.h
@@ -468,8 +468,9 @@ public:
 	 *                  the new and old states. Callers may want to manage queue membership directly for better
 	 *                  performance; removing from the middle is more expensive than popping from the front.
 	 *                  See definition of ESendFlags for a description of the behavior controlled by SendFlags.
+	 * @param ReleaseSaveReason Explanation for why the state is changing, used for debugging.
 	 */
-	void SendToState(EPackageState NextState, ESendFlags SendFlags);
+	void SendToState(EPackageState NextState, ESendFlags SendFlags, EStateChangeReason ReleaseSaveReason);
 	/* Debug-only code to assert that this PackageData is contained by the container matching its current state. */
 	void CheckInContainer() const;
 	/**
@@ -730,7 +731,7 @@ private:
 	void OnEnterLoadReady();
 	void OnExitLoadReady();
 	void OnEnterSave();
-	void OnExitSave();
+	void OnExitSave(EStateChangeReason ReleaseSaveReason);
 	/* Entry/Exit gates for Properties shared between multiple states */
 	void OnExitInProgress();
 	void OnEnterInProgress();
@@ -947,7 +948,7 @@ public:
 	 */
 	UObject* FindSplitDataObject() const;
 
-	void ResetSaveState(FCookGenerationInfo& Info, UPackage* Package, UE::Cook::EReleaseSaveReason ReleaseSaveReason);
+	void ResetSaveState(FCookGenerationInfo& Info, UPackage* Package, UE::Cook::EStateChangeReason ReleaseSaveReason);
 
 	int32& GetNextPopulateIndex() { check(IsInitialized()); return NextPopulateIndex; }
 
diff --git a/Engine/Source/Editor/UnrealEd/Private/Cooker/CookRequestCluster.cpp b/Engine/Source/Editor/UnrealEd/Private/Cooker/CookRequestCluster.cpp
--- a/Engine/Source/Editor/UnrealEd/Private/Cooker/CookRequestCluster.cpp
+++ b/Engine/Source/Editor/UnrealEd/Private/Cooker/CookRequestCluster.cpp
@@ -152,7 +152,7 @@ FRequestCluster::FRequestCluster(UCookOnTheFlyServer& InCOTFS, TRingBuffer<FDisc
 			Discovery = &PoppedDiscovery;
 			if (!PackageData.IsInProgress() && PackageData.GetPlatformsNeedingCookingNum() == 0)
 			{
-				PackageData.SendToState(EPackageState::Request, ESendFlags::QueueRemove);
+				PackageData.SendToState(EPackageState::Request, ESendFlags::QueueRemove, EStateChangeReason::RequestCluster);
 				OwnedPackageDatas.Add(&PackageData, ESuppressCookReason::NotSuppressed);
 			}
 			continue;
@@ -205,7 +205,7 @@ FRequestCluster::FRequestCluster(UCookOnTheFlyServer& InCOTFS, TRingBuffer<FDisc
 
 		// Send it to the Request state if it's not already there, remove it from its old container
 		// and add it to this cluster.
-		PackageData.SendToState(EPackageState::Request, ESendFlags::QueueRemove);
+		PackageData.SendToState(EPackageState::Request, ESendFlags::QueueRemove, EStateChangeReason::RequestCluster);
 		PackageData.AddUrgency(Discovery->bUrgent, false /* bAllowUpdateState */);
 		OwnedPackageDatas.Add(&PackageData, ESuppressCookReason::NotSuppressed);
 	}
@@ -319,7 +319,7 @@ void FRequestCluster::PullIntoCluster(FPackageData& PackageData)
 		}
 		else
 		{
-			PackageData.SendToState(EPackageState::Request, ESendFlags::QueueRemove);
+			PackageData.SendToState(EPackageState::Request, ESendFlags::QueueRemove, EStateChangeReason::RequestCluster);
 		}
 		Existing = ESuppressCookReason::NotSuppressed;
 	}
diff --git a/Engine/Source/Editor/UnrealEd/Private/Cooker/CookTypes.cpp b/Engine/Source/Editor/UnrealEd/Private/Cooker/CookTypes.cpp
--- a/Engine/Source/Editor/UnrealEd/Private/Cooker/CookTypes.cpp
+++ b/Engine/Source/Editor/UnrealEd/Private/Cooker/CookTypes.cpp
@@ -23,15 +23,28 @@ LLM_DEFINE_TAG(Cooker_CachedPlatformData);
 namespace UE::Cook
 {
 
-const TCHAR* LexToString(EReleaseSaveReason Reason)
+const TCHAR* LexToString(EStateChangeReason Reason)
 {
 	switch (Reason)
 	{
-	case EReleaseSaveReason::Completed: return TEXT("Completed");
-	case EReleaseSaveReason::DoneForNow: return TEXT("DoneForNow");
-	case EReleaseSaveReason::Demoted: return TEXT("Demoted");
-	case EReleaseSaveReason::AbortSave: return TEXT("AbortSave");
-	case EReleaseSaveReason::RecreateObjectCache: return TEXT("RecreateObjectCache");
+	case EStateChangeReason::Completed: return TEXT("Completed");
+	case EStateChangeReason::DoneForNow: return TEXT("DoneForNow");
+	case EStateChangeReason::SaveError: return TEXT("SaveError");
+	case EStateChangeReason::RecreateObjectCache: return TEXT("RecreateObjectCache");
+	case EStateChangeReason::CookerShutdown: return TEXT("CookerShutdown");
+	case EStateChangeReason::ReassignAbortedPackages: return TEXT("ReassignAbortedPackages");
+	case EStateChangeReason::Retraction: return TEXT("Retraction");
+	case EStateChangeReason::Discovered: return TEXT("Discovered");
+	case EStateChangeReason::Requested: return TEXT("Requested");
+	case EStateChangeReason::RequestCluster: return TEXT("RequestCluster");
+	case EStateChangeReason::DirectorRequest: return TEXT("DirectorRequest");
+	case EStateChangeReason::Loaded: return TEXT("Loaded");
+	case EStateChangeReason::Saved: return TEXT("Saved");
+	case EStateChangeReason::CookSuppressed: return TEXT("CookSuppressed");
+	case EStateChangeReason::GarbageCollected: return TEXT("GarbageCollected");
+	case EStateChangeReason::GeneratorPreGarbageCollected: return TEXT("GeneratorPreGarbageCollected");
+	case EStateChangeReason::ForceRecook: return TEXT("ForceRecook");
+	case EStateChangeReason::UrgencyUpdated: return TEXT("UrgencyUpdated");
 	default: return TEXT("Invalid");
 	}
 }
diff --git a/Engine/Source/Editor/UnrealEd/Private/Cooker/CookTypes.h b/Engine/Source/Editor/UnrealEd/Private/Cooker/CookTypes.h
--- a/Engine/Source/Editor/UnrealEd/Private/Cooker/CookTypes.h
+++ b/Engine/Source/Editor/UnrealEd/Private/Cooker/CookTypes.h
@@ -131,19 +131,29 @@ namespace UE::Cook
 		Incomplete,
 	};
 
-	/**
-	 * The possible reasons that the save-state data on an FPackageData might be released.
-	 * Different levels of teardown will happen based on the reason.
-	 */
-	enum class EReleaseSaveReason : uint8
+	/** The reasons that a FPackageData can change its state, used for diagnostics and some control flow. */
+	enum class EStateChangeReason : uint8
 	{
 		Completed,
 		DoneForNow,
-		Demoted,
-		AbortSave,
+		SaveError,
 		RecreateObjectCache,
+		CookerShutdown,
+		ReassignAbortedPackages,
+		Retraction,
+		Discovered,
+		Requested,
+		RequestCluster,
+		DirectorRequest,
+		Loaded,
+		Saved,
+		CookSuppressed,
+		GarbageCollected,
+		GeneratorPreGarbageCollected,
+		ForceRecook,
+		UrgencyUpdated,
 	};
-	const TCHAR* LexToString(UE::Cook::EReleaseSaveReason Reason);
+	const TCHAR* LexToString(UE::Cook::EStateChangeReason Reason);
 
 	enum class ESuppressCookReason : uint8
 	{
diff --git a/Engine/Source/Editor/UnrealEd/Private/Cooker/CookWorkerClient.cpp b/Engine/Source/Editor/UnrealEd/Private/Cooker/CookWorkerClient.cpp
--- a/Engine/Source/Editor/UnrealEd/Private/Cooker/CookWorkerClient.cpp
+++ b/Engine/Source/Editor/UnrealEd/Private/Cooker/CookWorkerClient.cpp
@@ -793,7 +793,7 @@ void FCookWorkerClient::AssignPackages(FAssignPackagesMessage& Message)
 		}
 		PackageData.FindOrAddPlatformData(CookerLoadingPlatformKey).MarkCookableForWorker(*this);
 		PackageData.SetInstigator(*this, FInstigator(AssignData.Instigator));
-		PackageData.SendToState(EPackageState::Request, ESendFlags::QueueAddAndRemove);
+		PackageData.SendToState(EPackageState::Request, ESendFlags::QueueAddAndRemove, EStateChangeReason::DirectorRequest);
 	}
 
 	// Clear the SoftGC diagnostic ExpectedNeverLoadPackages because we have new assigned packages
