Fixed crash when exiting PIE with the mesh paint mode active and the controls not being attached to the player controller: Validate that the PIE context is valid before dereferencing it.

diff --git a/Engine/Source/Editor/UnrealEd/Private/EditorModeManager.cpp b/Engine/Source/Editor/UnrealEd/Private/EditorModeManager.cpp
--- a/Engine/Source/Editor/UnrealEd/Private/EditorModeManager.cpp
+++ b/Engine/Source/Editor/UnrealEd/Private/EditorModeManager.cpp
@@ -232,7 +232,7 @@ void FEditorModeTools::RestoreSelection(FName SelectionStoreKey)
 UWorld* FEditorModeTools::GetWorld() const
 {
 	// When in 'Simulate' mode, the editor mode tools will actually interact with the PIE world
-	if( GEditor->bIsSimulatingInEditor )
+	if( GEditor->bIsSimulatingInEditor && GEditor->GetPIEWorldContext() )
 	{
 		return GEditor->GetPIEWorldContext()->World();
 	}
