Cooker: DiffOnly: Fix failed assertion when using ZenWriter: "Illegal call to StaticFindObjectFast() while serializing object". The scope that allows us to call FindObject needs to also be around OnSecondSaveComplete.
#rnx

diff --git a/Engine/Source/Editor/UnrealEd/Private/Cooker/DiffPackageWriter.cpp b/Engine/Source/Editor/UnrealEd/Private/Cooker/DiffPackageWriter.cpp
--- a/Engine/Source/Editor/UnrealEd/Private/Cooker/DiffPackageWriter.cpp
+++ b/Engine/Source/Editor/UnrealEd/Private/Cooker/DiffPackageWriter.cpp
@@ -180,12 +180,12 @@ void FDiffPackageWriter::WritePackageData(const FPackageInfo& Info, FLargeMemory
 	}
 	else
 	{
-		Accumulator.OnSecondSaveComplete(LocalInfo.HeaderSize);
-
-		// Avoid an assert when calling StaticFindObject during save to retrieve the actor's class.
+		// Avoid an assert when calling StaticFindObject during save, which we do to list the "exports" from a package.
 		// We are not writing the discovered objects into the saved package, so the call to StaticFindObject is legal.
 		TGuardValue<bool> GIsSavingPackageGuard(GIsSavingPackage, false);
 
+		Accumulator.OnSecondSaveComplete(LocalInfo.HeaderSize);
+
 		TMap<FName, FArchiveDiffStats> PackageDiffStats;
 		const TCHAR* CutoffString = TEXT("UEditorEngine::Save()");
 		Accumulator.CompareWithPrevious(CutoffString, PackageDiffStats);
