REINSTANCER CHANGE - Gather all referenced contained objects for clone, excluding those referenced by transient properties
Resubmit/rework of 28989821, which did not skip transient properties (caused a crash when cloing ControlRig for a template project) and also did not include objects referenced across the object DAG

diff --git a/Engine/Source/Editor/UnrealEd/Private/Kismet2/KismetReinstanceUtilities.cpp b/Engine/Source/Editor/UnrealEd/Private/Kismet2/KismetReinstanceUtilities.cpp
--- a/Engine/Source/Editor/UnrealEd/Private/Kismet2/KismetReinstanceUtilities.cpp
+++ b/Engine/Source/Editor/UnrealEd/Private/Kismet2/KismetReinstanceUtilities.cpp
@@ -40,6 +40,7 @@
 #include "Subsystems/AssetEditorSubsystem.h"
 #include "Engine/ScopedMovementUpdate.h"
 #include "InstancedReferenceSubobjectHelper.h"
+#include "UObject/PropertyOptional.h"
 
 DECLARE_CYCLE_STAT(TEXT("Replace Instances"), EKismetReinstancerStats_ReplaceInstancesOfClass, STATGROUP_KismetReinstancer );
 DECLARE_CYCLE_STAT(TEXT("Find Referencers"), EKismetReinstancerStats_FindReferencers, STATGROUP_KismetReinstancer );
@@ -282,6 +283,119 @@ struct FReplaceReferenceHelper
 			KVP.Key->Actor = KVP.Value;
 		}
 	}
+	
+	// Others may want this simple iteration function, but hiding it here for now:
+	static void ForEachSubObject(const FProperty* TargetProp, const UObject* Outer, const UObject* Root, const void* ContainerAddress, TFunctionRef<void(const UObject*)> ObjRefFunc)
+	{
+		check(ContainerAddress && Outer);
+		if (TargetProp->HasAnyPropertyFlags(CPF_Transient))
+		{
+			return;
+		}
+
+		if (const FArrayProperty* ArrayProperty = CastField<const FArrayProperty>(TargetProp))
+		{
+			FScriptArrayHelper ArrayHelper(ArrayProperty, ContainerAddress);
+			for (int32 ElementIndex = 0; ElementIndex < ArrayHelper.Num(); ++ElementIndex)
+			{
+				const void* ValueAddress = ArrayHelper.GetRawPtr(ElementIndex);
+
+				ForEachSubObject(ArrayProperty->Inner, Outer, Root, ValueAddress, ObjRefFunc);
+			}
+		}
+		else if (const FMapProperty* MapProperty = CastField<const FMapProperty>(TargetProp))
+		{
+			// Exit now if the map doesn't contain any instanced references.
+			int32 LogicalIndex = 0;
+			FScriptMapHelper MapHelper(MapProperty, ContainerAddress);
+			for (int32 ElementIndex = 0; ElementIndex < MapHelper.GetMaxIndex(); ++ElementIndex)
+			{
+				if (MapHelper.IsValidIndex(ElementIndex))
+				{
+					const void* KeyAddress = MapHelper.GetKeyPtr(ElementIndex);
+					const void* ValueAddress = MapHelper.GetValuePtr(ElementIndex);
+
+					// Note: Keep these as the logical (Nth) index in case the map changes internally after we construct the path or in case we resolve using a different object.
+					ForEachSubObject(MapProperty->KeyProp, Outer, Root, KeyAddress, ObjRefFunc);
+					ForEachSubObject(MapProperty->ValueProp, Outer, Root, ValueAddress, ObjRefFunc);
+
+					++LogicalIndex;
+				}
+			}
+		}
+		else if (const FSetProperty* SetProperty = CastField<const FSetProperty>(TargetProp))
+		{
+			int32 LogicalIndex = 0;
+			FScriptSetHelper SetHelper(SetProperty, ContainerAddress);
+			for (int32 ElementIndex = 0; ElementIndex < SetHelper.GetMaxIndex(); ++ElementIndex)
+			{
+				if (SetHelper.IsValidIndex(ElementIndex))
+				{
+					const void* ValueAddress = SetHelper.GetElementPtr(ElementIndex);
+
+					// Note: Keep this as the logical (Nth) index in case the set changes internally after we construct the path or in case we resolve using a different object.
+					ForEachSubObject(SetProperty->ElementProp, Outer, Root, ValueAddress, ObjRefFunc);
+
+					++LogicalIndex;
+				}
+			}
+		}
+		else if (const FOptionalProperty* OptionalProperty = CastField<FOptionalProperty>(TargetProp))
+		{
+			if (const void* ValueAddress = static_cast<const void*>(OptionalProperty->GetValuePointerForReadOrReplaceIfSet(ContainerAddress)))
+			{
+				ForEachSubObject(OptionalProperty->GetValueProperty(), Outer, Root, ValueAddress, ObjRefFunc);
+			}
+		}
+		else if (const FStructProperty* StructProperty = CastField<const FStructProperty>(TargetProp))
+		{
+			for (FProperty* StructProp = StructProperty->Struct->RefLink; StructProp; StructProp = StructProp->NextRef)
+			{
+				for (int32 ArrayIdx = 0; ArrayIdx < StructProp->ArrayDim; ++ArrayIdx)
+				{
+					const void* ValueAddress = StructProp->ContainerPtrToValuePtr<uint8>(ContainerAddress, ArrayIdx);
+
+					ForEachSubObject(StructProp, Outer, Root, ValueAddress, ObjRefFunc);
+				}
+			}
+		}
+		else if (const FObjectProperty* ObjectProperty = CastField<const FObjectProperty>(TargetProp))
+		{
+			if (UObject* ObjectValue = ObjectProperty->GetObjectPropertyValue(ContainerAddress))
+			{
+				if (ObjectValue->IsIn(Root))
+				{
+					// don't need to push to PropertyPath, since this property is already at its head
+					ObjRefFunc(ObjectValue);
+				}
+			}
+		}
+	}
+
+	static void GetOwnedSubobjectsRecursive(const UObject* Container, TSet<UObject*>& OutObjects, const UObject* Root = nullptr)
+	{
+		if (Root == nullptr)
+		{
+			Root = Container;
+		}
+
+		const UClass* ContainerClass = Container->GetClass();
+		for (FProperty* Prop = ContainerClass->RefLink; Prop; Prop = Prop->NextRef)
+		{
+			for (int32 ArrayIdx = 0; ArrayIdx < Prop->ArrayDim; ++ArrayIdx)
+			{
+				const uint8* ValuePtr = Prop->ContainerPtrToValuePtr<uint8>(Container, ArrayIdx);
+				ForEachSubObject(Prop, Container, Root, ValuePtr, [&OutObjects, Root](const UObject* Ref)
+					{
+						if (!OutObjects.Contains(Ref))
+						{
+							OutObjects.Add(const_cast<UObject*>(Ref)); // consumer is not const correct
+							GetOwnedSubobjectsRecursive(Ref, OutObjects, Root);
+						}
+					});
+			}
+		}
+	}
 };
 
 struct FArchetypeReinstanceHelper
@@ -2914,7 +3028,7 @@ void FBlueprintCompileReinstancer::CopyPropertiesForUnrelatedObjects(UObject* Ol
 void FBlueprintCompileReinstancer::PreCreateSubObjectsForReinstantiation(const TMap<UClass*, UClass*>& OldToNewClassMap, UObject* OldObject, UObject* NewUObject, TMap<UObject*, UObject*>& CreatedInstanceMap, const TMap<UObject*, UObject*>* OldToNewInstanceMap/* = nullptr*/)
 {
 	TSet<UObject*> OldInstancedSubObjects;
-	FFindInstancedReferenceSubobjectHelper::GetInstancedSubObjectsRecursive(OldObject, OldInstancedSubObjects);
+	FReplaceReferenceHelper::GetOwnedSubobjectsRecursive(OldObject, OldInstancedSubObjects);
 
 	PreCreateSubObjectsForReinstantiation_Inner(OldInstancedSubObjects, OldToNewClassMap, OldObject, NewUObject, CreatedInstanceMap, OldToNewInstanceMap);
 }
