Cooker: MPCook: Move CompactBinaryTCP::IMessage into the public API for IMPCollector.
#rnx

diff --git a/Engine/Source/Editor/UnrealEd/Private/Cooker/CompactBinaryTCP.cpp b/Engine/Source/Editor/UnrealEd/Private/Cooker/CompactBinaryTCP.cpp
--- a/Engine/Source/Editor/UnrealEd/Private/Cooker/CompactBinaryTCP.cpp
+++ b/Engine/Source/Editor/UnrealEd/Private/Cooker/CompactBinaryTCP.cpp
@@ -337,30 +337,9 @@ EConnectionStatus TryWritePacket(FSocket* Socket, FSendBuffer& Buffer,
 		TArrayView<FMarshalledMessage>(&AppendMessage, 1), MaxPacketSize);
 }
 
-EConnectionStatus TryWritePacket(FSocket* Socket, FSendBuffer& Buffer,
-	const IMessage& AppendMessage, uint64 MaxPacketSize)
+void QueueMessage(FSendBuffer& Buffer, FMarshalledMessage&& Message)
 {
-	FMarshalledMessage Marshalled;
-	Marshalled.MessageType = AppendMessage.GetMessageType();
-	FCbWriter Writer;
-	Writer.BeginObject();
-	AppendMessage.Write(Writer);
-	Writer.EndObject();
-	Marshalled.Object = Writer.Save().AsObject();
-	return FCompactBinaryTCPImpl::TryWritePacket(Socket, Buffer,
-		TArrayView<FMarshalledMessage>(&Marshalled, 1), MaxPacketSize);
-}
-
-void QueueMessage(FSendBuffer& Buffer, const IMessage& AppendMessage)
-{
-	FMarshalledMessage Marshalled;
-	Marshalled.MessageType = AppendMessage.GetMessageType();
-	FCbWriter Writer;
-	Writer.BeginObject();
-	AppendMessage.Write(Writer);
-	Writer.EndObject();
-	Marshalled.Object = Writer.Save().AsObject();
-	FCompactBinaryTCPImpl::QueueMessage(Buffer, MoveTemp(Marshalled));
+	FCompactBinaryTCPImpl::QueueMessage(Buffer, MoveTemp(Message));
 }
 
 EConnectionStatus TryFlushBuffer(FSocket* Socket, FSendBuffer& Buffer, uint64 MaxPacketSize)
diff --git a/Engine/Source/Editor/UnrealEd/Private/Cooker/CompactBinaryTCP.h b/Engine/Source/Editor/UnrealEd/Private/Cooker/CompactBinaryTCP.h
--- a/Engine/Source/Editor/UnrealEd/Private/Cooker/CompactBinaryTCP.h
+++ b/Engine/Source/Editor/UnrealEd/Private/Cooker/CompactBinaryTCP.h
@@ -26,7 +26,7 @@ struct FSendBuffer;
 // and to allow room for the 16 byte PacketHeader.
 constexpr uint64 MaxOSPacketSize = (1 << 30) - 1;
 
-enum EConnectionStatus
+enum class EConnectionStatus
 {
 	/** Connection is still okay or operation succeeded. */
 	Okay,
@@ -41,23 +41,7 @@ enum EConnectionStatus
 };
 const TCHAR* DescribeStatus(EConnectionStatus Status);
 
-/**
- * The base class of messages that can be sent through CompactBinaryTCP. Messages are identified to
- * the remote connection by the Guid identifier from GetMessageType.
- */
-class IMessage
-{
-public:
-	virtual ~IMessage() {}
-	/** Marshall the message to a CompactBinaryObject. */
-	virtual void Write(FCbWriter& Writer) const = 0;
-	/** Unmarshall the message from a CompactBinaryObject. */
-	virtual bool TryRead(FCbObjectView Object) = 0;
-	/** Return the Guid that identifies the message to the remote connection. */
-	virtual FGuid GetMessageType() const = 0;
-};
-
-/** IMessages are marshalled into a Guid and a CompactBinaryObject for serialization through the socket. */
+/** The atom of data that can be sent through CompactBinaryTCP: guid identifier and compact binary payload. */
 struct FMarshalledMessage
 {
 	FGuid MessageType;
@@ -92,9 +76,7 @@ EConnectionStatus TryWritePacket(FSocket* Socket, FSendBuffer& Buffer,
 	TArray<FMarshalledMessage>&& AppendMessages, uint64 MaxPacketSize = 0);
 EConnectionStatus TryWritePacket(FSocket* Socket, FSendBuffer& Buffer,
 	FMarshalledMessage&& AppendMessage, uint64 MaxPacketSize = 0);
-EConnectionStatus TryWritePacket(FSocket* Socket, FSendBuffer& Buffer,
-	const IMessage& AppendMessage, uint64 MaxPacketSize = 0);
-void QueueMessage(FSendBuffer& Buffer, const IMessage& AppendMessage);
+void QueueMessage(FSendBuffer& Buffer, FMarshalledMessage&& Message);
 
 
 /** Attempt to finish sending packets that were previously queued by TryWritePacket */
diff --git a/Engine/Source/Editor/UnrealEd/Private/Cooker/CookDirector.cpp b/Engine/Source/Editor/UnrealEd/Private/Cooker/CookDirector.cpp
--- a/Engine/Source/Editor/UnrealEd/Private/Cooker/CookDirector.cpp
+++ b/Engine/Source/Editor/UnrealEd/Private/Cooker/CookDirector.cpp
@@ -146,17 +146,17 @@ FCookDirector::FCookDirector(UCookOnTheFlyServer& InCOTFS, int32 CookProcessCoun
 		RequestedCookWorkerCount+1, RequestedCookWorkerCount, RequestedCookWorkerCount > 1 ? TEXT("CookWorkers") : TEXT("CookWorker"));
 
 	Register(new FLogMessagesMessageHandler());
-	Register(new IMPCollectorCbServerMessage<FRetractionResultsMessage>([this]
+	Register(new TMPCollectorServerMessageCallback<FRetractionResultsMessage>([this]
 	(FMPCollectorServerMessageContext& Context, bool bReadSuccessful, FRetractionResultsMessage&& Message)
 		{
 			// Called from inside CommunicationLock
 			RetractionHandler->HandleRetractionMessage(Context, bReadSuccessful, MoveTemp(Message));
-		}, TEXT("HandleRetractionMessage")));
-	Register(new IMPCollectorCbServerMessage<FHeartbeatMessage>([this]
+		}));
+	Register(new TMPCollectorServerMessageCallback<FHeartbeatMessage>([this]
 	(FMPCollectorServerMessageContext& Context, bool bReadSuccessful, FHeartbeatMessage&& Message)
 		{
 			HandleHeartbeatMessage(Context, bReadSuccessful, MoveTemp(Message));
-		}, TEXT("HandleHeartbeatMessage")));
+		}));
 	Register(new FAssetRegistryMPCollector(COTFS));
 	Register(new FPackageWriterMPCollector(COTFS));
 
diff --git a/Engine/Source/Editor/UnrealEd/Private/Cooker/CookDirector.h b/Engine/Source/Editor/UnrealEd/Private/Cooker/CookDirector.h
--- a/Engine/Source/Editor/UnrealEd/Private/Cooker/CookDirector.h
+++ b/Engine/Source/Editor/UnrealEd/Private/Cooker/CookDirector.h
@@ -7,6 +7,7 @@
 #include "Containers/ArrayView.h"
 #include "Containers/Map.h"
 #include "Containers/UnrealString.h"
+#include "CookMPCollector.h"
 #include "CookSockets.h"
 #include "CookTypes.h"
 #include "HAL/CriticalSection.h"
@@ -27,8 +28,6 @@ class FCbWriter;
 class FRunnableThread;
 class UCookOnTheFlyServer;
 namespace UE::Cook { class FCookWorkerServer; }
-namespace UE::Cook { class FMPCollectorServerMessageContext; }
-namespace UE::Cook { class IMPCollector; }
 namespace UE::Cook { struct FCookWorkerProfileData; }
 namespace UE::Cook { struct FHeartbeatMessage; }
 namespace UE::Cook { struct FInitialConfigMessage; }
@@ -274,12 +273,13 @@ struct FDirectorConnectionInfo
 };
 
 /** Message sent from a CookWorker to the Director to report that it is ready for setup messages and cooking. */
-struct FWorkerConnectMessage : public UE::CompactBinaryTCP::IMessage
+struct FWorkerConnectMessage : public IMPCollectorMessage
 {
 public:
 	virtual void Write(FCbWriter& Writer) const override;
 	virtual bool TryRead(FCbObjectView Object) override;
 	virtual FGuid GetMessageType() const override { return MessageType; }
+	virtual const TCHAR* GetDebugName() const override { return TEXT("WorkerConnectMessage"); }
 
 public:
 	int32 RemoteIndex = 0;
@@ -290,11 +290,12 @@ public:
  * Message sent from CookDirector to a CookWorker to cancel some of its assigned packages and return them
  * dispatch to idle workers.
  */
-struct FRetractionRequestMessage : public UE::CompactBinaryTCP::IMessage
+struct FRetractionRequestMessage : public IMPCollectorMessage
 {
 	virtual void Write(FCbWriter& Writer) const override;
 	virtual bool TryRead(FCbObjectView Object) override;
 	virtual FGuid GetMessageType() const override { return MessageType; }
+	virtual const TCHAR* GetDebugName() const override { return TEXT("RetractionRequestMessage"); }
 
 public:
 	int32 RequestedCount = 0;
@@ -305,11 +306,12 @@ public:
  * Message sent from CookWorker to CookDirector identifying which assigned packages it chose to satisfy a
  * a RetractionRequest.
  */
-struct FRetractionResultsMessage : public UE::CompactBinaryTCP::IMessage
+struct FRetractionResultsMessage : public IMPCollectorMessage
 {
 	virtual void Write(FCbWriter& Writer) const override;
 	virtual bool TryRead(FCbObjectView Object) override;
 	virtual FGuid GetMessageType() const override { return MessageType; }
+	virtual const TCHAR* GetDebugName() const override { return TEXT("RetractionResultsMessage"); }
 
 public:
 	TArray<FName> ReturnedPackages;
diff --git a/Engine/Source/Editor/UnrealEd/Private/Cooker/CookMPCollector.h b/Engine/Source/Editor/UnrealEd/Private/Cooker/CookMPCollector.h
--- a/Engine/Source/Editor/UnrealEd/Private/Cooker/CookMPCollector.h
+++ b/Engine/Source/Editor/UnrealEd/Private/Cooker/CookMPCollector.h
@@ -5,7 +5,6 @@
 #include "Async/Future.h"
 #include "Containers/Array.h"
 #include "Containers/ArrayView.h"
-#include "Cooker/CompactBinaryTCP.h"
 #include "Cooker/CookTypes.h"
 #include "HAL/Platform.h"
 #include "Misc/Guid.h"
@@ -131,92 +130,107 @@ public:
 	virtual void ServerReceiveMessage(FMPCollectorServerMessageContext& Context, FCbObjectView Message) {}
 };
 
-/** A subinterface of IMPCollector that uses a UE::CompactBinaryTCP::IMessage subclass to serialize the message. */
-template <typename CompactBinaryTCPMessageType>
-class IMPCollectorCbMessage : public IMPCollector
+
+/** Baseclass for messages used by IMPCollectors that want to interpret their messages as c++ structs. */
+class IMPCollectorMessage
+{
+public:
+	virtual ~IMPCollectorMessage() {}
+
+	/** Marshall the message to a CompactBinaryObject. */
+	virtual void Write(FCbWriter& Writer) const = 0;
+	/** Unmarshall the message from a CompactBinaryObject. */
+	virtual bool TryRead(FCbObjectView Object) = 0;
+	/** Return the Guid that identifies the message to the remote connection. */
+	virtual FGuid GetMessageType() const = 0;
+	/** Return the debugname for diagnostics. */
+	virtual const TCHAR* GetDebugName() const = 0;
+};
+
+
+/** A subinterface of IMPCollector that uses a ICollectorMessage subclass to serialize the message. */
+template <typename MessageType>
+class IMPCollectorForMessage : public IMPCollector
 {
 public:
 	virtual void ClientReceiveMessage(FMPCollectorClientMessageContext& Context, bool bReadSuccessful,
-		CompactBinaryTCPMessageType&& Message) {}
+		MessageType&& Message) {}
 	virtual void ServerReceiveMessage(FMPCollectorServerMessageContext& Context, bool bReadSuccessful,
-		CompactBinaryTCPMessageType&& Message) {}
+		MessageType&& Message) {}
 
 	virtual FGuid GetMessageType() const override
 	{
-		return CompactBinaryTCPMessageType::MessageType;
+		return MessageType().GetMessageType();
+	}
+
+	virtual const TCHAR* GetDebugName() const override
+	{
+		return MessageType().GetDebugName();
 	}
 
 	virtual void ClientReceiveMessage(FMPCollectorClientMessageContext& Context, FCbObjectView Message) override
 	{
-		CompactBinaryTCPMessageType CBMessage;
+		MessageType CBMessage;
 		bool bReadSuccessful = CBMessage.TryRead(Message);
 		ClientReceiveMessage(Context, bReadSuccessful, MoveTemp(CBMessage));
 	}
 
 	virtual void ServerReceiveMessage(FMPCollectorServerMessageContext& Context, FCbObjectView Message) override
 	{
-		CompactBinaryTCPMessageType CBMessage;
+		MessageType CBMessage;
 		bool bReadSuccessful = CBMessage.TryRead(Message);
 		ServerReceiveMessage(Context, bReadSuccessful, MoveTemp(CBMessage));
 	}
 };
 
 /**
- * An implementation of IMPCollector that uses a UE::CompactBinaryTCP::IMessage subclass to serialize the message,
+ * An implementation of IMPCollector that uses an ICollectorMessage subclass to serialize the message,
  * and directs messages received on the client to the given callback.
  */
-template <typename CompactBinaryTCPMessageType>
-class IMPCollectorCbClientMessage : public IMPCollectorCbMessage<CompactBinaryTCPMessageType>
+template <typename MessageType>
+class TMPCollectorClientMessageCallback : public IMPCollectorForMessage<MessageType>
 {
 public:
-	IMPCollectorCbClientMessage(TUniqueFunction<void(FMPCollectorClientMessageContext& Context, bool bReadSuccessful,
-		CompactBinaryTCPMessageType&& Message)>&& InCallback, const TCHAR* InDebugName)
+	TMPCollectorClientMessageCallback(TUniqueFunction<void(FMPCollectorClientMessageContext& Context, bool bReadSuccessful,
+		MessageType&& Message)>&& InCallback)
 		: Callback(MoveTemp(InCallback))
-		, DebugName(InDebugName)
 	{
 		check(Callback);
 	}
 
-	virtual const TCHAR* GetDebugName() const override { return DebugName; }
-
 	virtual void ClientReceiveMessage(FMPCollectorClientMessageContext& Context, bool bReadSuccessful,
-		CompactBinaryTCPMessageType&& Message) override
+		MessageType&& Message) override
 	{
 		Callback(Context, bReadSuccessful, MoveTemp(Message));
 	}
 
 private:
-	TUniqueFunction<void(FMPCollectorClientMessageContext& Context, bool bReadSuccessful, CompactBinaryTCPMessageType&& Message)> Callback;
-	const TCHAR* DebugName;
+	TUniqueFunction<void(FMPCollectorClientMessageContext& Context, bool bReadSuccessful, MessageType&& Message)> Callback;
 };
 
 /**
- * An implementation of IMPCollector that uses a UE::CompactBinaryTCP::IMessage subclass to serialize the message,
+ * An implementation of IMPCollector that uses an IMPCollectorMessage subclass to serialize the message,
  * and directs messages received on the server to the given callback.
  */
-template <typename CompactBinaryTCPMessageType>
-class IMPCollectorCbServerMessage : public IMPCollectorCbMessage<CompactBinaryTCPMessageType>
+template <typename MessageType>
+class TMPCollectorServerMessageCallback : public IMPCollectorForMessage<MessageType>
 {
 public:
-	IMPCollectorCbServerMessage(TUniqueFunction<void(FMPCollectorServerMessageContext& Context, bool bReadSuccessful,
-		CompactBinaryTCPMessageType&& Message)>&& InCallback, const TCHAR* InDebugName)
+	TMPCollectorServerMessageCallback(TUniqueFunction<void(FMPCollectorServerMessageContext& Context, bool bReadSuccessful,
+		MessageType&& Message)>&& InCallback)
 		: Callback(MoveTemp(InCallback))
-		, DebugName(InDebugName)
 	{
 		check(Callback);
 	}
 
-	virtual const TCHAR* GetDebugName() const override { return DebugName; }
-
 	virtual void ServerReceiveMessage(FMPCollectorServerMessageContext& Context, bool bReadSuccessful,
-		CompactBinaryTCPMessageType&& Message) override
+		MessageType&& Message) override
 	{
 		Callback(Context, bReadSuccessful, MoveTemp(Message));
 	}
 
 private:
-	TUniqueFunction<void(FMPCollectorServerMessageContext& Context, bool bReadSuccessful, CompactBinaryTCPMessageType&& Message)> Callback;
-	const TCHAR* DebugName;
+	TUniqueFunction<void(FMPCollectorServerMessageContext& Context, bool bReadSuccessful, MessageType&& Message)> Callback;
 };
 
 }
\ No newline at end of file
diff --git a/Engine/Source/Editor/UnrealEd/Private/Cooker/CookWorkerClient.cpp b/Engine/Source/Editor/UnrealEd/Private/Cooker/CookWorkerClient.cpp
--- a/Engine/Source/Editor/UnrealEd/Private/Cooker/CookWorkerClient.cpp
+++ b/Engine/Source/Editor/UnrealEd/Private/Cooker/CookWorkerClient.cpp
@@ -34,21 +34,21 @@ FCookWorkerClient::FCookWorkerClient(UCookOnTheFlyServer& InCOTFS)
 	LogMessageHandler = new FLogMessagesMessageHandler();
 	LogMessageHandler->InitializeClient();
 	Register(LogMessageHandler);
-	Register(new IMPCollectorCbClientMessage<FRetractionRequestMessage>([this]
+	Register(new TMPCollectorClientMessageCallback<FRetractionRequestMessage>([this]
 	(FMPCollectorClientMessageContext& Context, bool bReadSuccessful, FRetractionRequestMessage&& Message)
 		{
 			HandleRetractionMessage(Context, bReadSuccessful, MoveTemp(Message));
-		}, TEXT("HandleRetractionMessage")));
-	Register(new IMPCollectorCbClientMessage<FAbortPackagesMessage>([this]
+		}));
+	Register(new TMPCollectorClientMessageCallback<FAbortPackagesMessage>([this]
 	(FMPCollectorClientMessageContext& Context, bool bReadSuccessful, FAbortPackagesMessage&& Message)
 		{
 			HandleAbortPackagesMessage(Context, bReadSuccessful, MoveTemp(Message));
-		}, TEXT("HandleAbortPackagesMessage")));
-	Register(new IMPCollectorCbClientMessage<FHeartbeatMessage>([this]
+		}));
+	Register(new TMPCollectorClientMessageCallback<FHeartbeatMessage>([this]
 	(FMPCollectorClientMessageContext& Context, bool bReadSuccessful, FHeartbeatMessage&& Message)
 		{
 			HandleHeartbeatMessage(Context, bReadSuccessful, MoveTemp(Message));
-		}, TEXT("HandleHeartbeatMessage")));
+		}));
 	Register(new FAssetRegistryMPCollector(COTFS));
 	Register(new FPackageWriterMPCollector(COTFS));
 }
@@ -398,7 +398,7 @@ void FCookWorkerClient::CreateServerSocket(const FDirectorConnectionInfo& Connec
 
 	FWorkerConnectMessage ConnectMessage;
 	ConnectMessage.RemoteIndex = ConnectInfo.RemoteIndex;
-	EConnectionStatus Status = TryWritePacket(ServerSocket, SendBuffer, ConnectMessage);
+	EConnectionStatus Status = TryWritePacket(ServerSocket, SendBuffer, MarshalToCompactBinaryTCP(ConnectMessage));
 	if (Status == EConnectionStatus::Incomplete)
 	{
 		SendToState(EConnectStatus::PollWriteConnectMessage);
@@ -554,7 +554,7 @@ void FCookWorkerClient::LogConnected()
 void FCookWorkerClient::PumpSendMessages()
 {
 	UE::CompactBinaryTCP::EConnectionStatus Status = UE::CompactBinaryTCP::TryFlushBuffer(ServerSocket, SendBuffer);
-	if (Status == UE::CompactBinaryTCP::Failed)
+	if (Status == UE::CompactBinaryTCP::EConnectionStatus::Failed)
 	{
 		UE_LOG(LogCook, Error, TEXT("CookWorkerClient failed to write message to Director. We will abort the CookAsCookWorker commandlet."));
 		SendToState(EConnectStatus::LostConnection);
@@ -732,9 +732,10 @@ void FCookWorkerClient::PumpDisconnect(FTickStackData& StackData)
 	}
 }
 
-void FCookWorkerClient::SendMessage(const UE::CompactBinaryTCP::IMessage& Message)
+void FCookWorkerClient::SendMessage(const IMPCollectorMessage& Message)
 {
-	UE::CompactBinaryTCP::TryWritePacket(ServerSocket, SendBuffer, Message);
+	UE::CompactBinaryTCP::TryWritePacket(ServerSocket, SendBuffer,
+		MarshalToCompactBinaryTCP(Message));
 }
 
 void FCookWorkerClient::SendToState(EConnectStatus TargetStatus)
diff --git a/Engine/Source/Editor/UnrealEd/Private/Cooker/CookWorkerClient.h b/Engine/Source/Editor/UnrealEd/Private/Cooker/CookWorkerClient.h
--- a/Engine/Source/Editor/UnrealEd/Private/Cooker/CookWorkerClient.h
+++ b/Engine/Source/Editor/UnrealEd/Private/Cooker/CookWorkerClient.h
@@ -4,6 +4,7 @@
 
 #include "CompactBinaryTCP.h"
 #include "Containers/Array.h"
+#include "CookMPCollector.h"
 #include "CookOnTheSide/CookOnTheFlyServer.h"
 #include "CookTypes.h"
 #include "HAL/CriticalSection.h"
@@ -109,7 +110,7 @@ private:
 	/** Helper for Tick, pump Send/Receive and check for whether we are done shutting down. */
 	void PumpDisconnect(FTickStackData& StackData);
 	/** Send the message immediately to the Socket. If cannot complete immediately, it will be finished during Tick. */
-	void SendMessage(const UE::CompactBinaryTCP::IMessage& Message);
+	void SendMessage(const IMPCollectorMessage& Message);
 	/** Send this into the given state. Update any state-dependent variables. */
 	void SendToState(EConnectStatus TargetStatus);
 	void LogInvalidMessage(const TCHAR* MessageTypeName);
diff --git a/Engine/Source/Editor/UnrealEd/Private/Cooker/CookWorkerServer.cpp b/Engine/Source/Editor/UnrealEd/Private/Cooker/CookWorkerServer.cpp
--- a/Engine/Source/Editor/UnrealEd/Private/Cooker/CookWorkerServer.cpp
+++ b/Engine/Source/Editor/UnrealEd/Private/Cooker/CookWorkerServer.cpp
@@ -611,7 +611,7 @@ void FCookWorkerServer::TickWaitForDisconnect()
 void FCookWorkerServer::PumpSendMessages()
 {
 	UE::CompactBinaryTCP::EConnectionStatus Status = UE::CompactBinaryTCP::TryFlushBuffer(Socket, SendBuffer);
-	if (Status == UE::CompactBinaryTCP::Failed)
+	if (Status == UE::CompactBinaryTCP::EConnectionStatus::Failed)
 	{
 		UE_LOG(LogCook, Error, TEXT("CookWorkerCrash: CookWorker %d failed to write to socket, we will shutdown the remote process. Assigned packages will be returned to the director."),
 			ProfileId);
@@ -792,22 +792,22 @@ void FCookWorkerServer::HandleReceivedPackagePlatformMessages(FPackageData& Pack
 	}
 }
 
-void FCookWorkerServer::SendMessage(const UE::CompactBinaryTCP::IMessage& Message, ECookDirectorThread TickThread)
+void FCookWorkerServer::SendMessage(const IMPCollectorMessage& Message, ECookDirectorThread TickThread)
 {
 	FCommunicationScopeLock ScopeLock(this, TickThread, ETickAction::Tick);
 	SendMessageInLock(Message);
 }
 
-void FCookWorkerServer::SendMessageInLock(const UE::CompactBinaryTCP::IMessage& Message)
+void FCookWorkerServer::SendMessageInLock(const IMPCollectorMessage& Message)
 {
 	if (TickState.TickAction == ETickAction::Tick)
 	{
-		UE::CompactBinaryTCP::TryWritePacket(Socket, SendBuffer, Message);
+		UE::CompactBinaryTCP::TryWritePacket(Socket, SendBuffer, MarshalToCompactBinaryTCP(Message));
 	}
 	else
 	{
 		check(TickState.TickAction == ETickAction::Queue);
-		UE::CompactBinaryTCP::QueueMessage(SendBuffer, Message);
+		UE::CompactBinaryTCP::QueueMessage(SendBuffer, MarshalToCompactBinaryTCP(Message));
 	}
 }
 
@@ -975,6 +975,18 @@ FCookWorkerServer::FCommunicationScopeLock::~FCommunicationScopeLock()
 	Server.TickState.TickAction = ETickAction::Invalid;
 }
 
+UE::CompactBinaryTCP::FMarshalledMessage MarshalToCompactBinaryTCP(const IMPCollectorMessage& Message)
+{
+	UE::CompactBinaryTCP::FMarshalledMessage Marshalled;
+	Marshalled.MessageType = Message.GetMessageType();
+	FCbWriter Writer;
+	Writer.BeginObject();
+	Message.Write(Writer);
+	Writer.EndObject();
+	Marshalled.Object = Writer.Save().AsObject();
+	return Marshalled;
+}
+
 FAssignPackagesMessage::FAssignPackagesMessage(TArray<FAssignPackageData>&& InPackageDatas)
 	: PackageDatas(MoveTemp(InPackageDatas))
 {
diff --git a/Engine/Source/Editor/UnrealEd/Private/Cooker/CookWorkerServer.h b/Engine/Source/Editor/UnrealEd/Private/Cooker/CookWorkerServer.h
--- a/Engine/Source/Editor/UnrealEd/Private/Cooker/CookWorkerServer.h
+++ b/Engine/Source/Editor/UnrealEd/Private/Cooker/CookWorkerServer.h
@@ -67,7 +67,7 @@ public:
 	bool TryHandleConnectMessage(FWorkerConnectMessage& Message, FSocket* InSocket, TArray<UE::CompactBinaryTCP::FMarshalledMessage>&& OtherPacketMessages, ECookDirectorThread TickThread);
 
 	/** Send the message immediately to the Socket. If cannot complete immediately, it will be finished during Tick. */
-	void SendMessage(const UE::CompactBinaryTCP::IMessage& Message, ECookDirectorThread TickThread);
+	void SendMessage(const IMPCollectorMessage& Message, ECookDirectorThread TickThread);
 
 	/** Periodic Tick function to send and receive messages to the Client. */
 	void TickCommunication(ECookDirectorThread TickThread);
@@ -154,7 +154,7 @@ private:
 	/** The main implementation of AbortAllAssignments, only callable from inside the lock. */
 	void AbortAllAssignmentsInLock(TSet<FPackageData*>& OutPendingPackages);
 	/** Send the message immediately to the Socket. If cannot complete immediately, it will be finished during Tick. */
-	void SendMessageInLock(const UE::CompactBinaryTCP::IMessage& Message);
+	void SendMessageInLock(const IMPCollectorMessage& Message);
 	/** Send this into the given state. Update any state-dependent variables. */
 	void SendToState(EConnectStatus TargetStatus);
 	/** Close the connection and connection resources to the remote process. Does not kill the process. */
@@ -202,6 +202,8 @@ private:
 	bool bNeedCrashDiagnostics = false;
 };
 
+UE::CompactBinaryTCP::FMarshalledMessage MarshalToCompactBinaryTCP(const IMPCollectorMessage& Message);
+
 /** Information about a PackageData the director sends to cookworkers. */
 struct FAssignPackageData
 {
@@ -218,7 +220,7 @@ FCbWriter& operator<<(FCbWriter& Writer, const FInstigator& Instigator);
 bool LoadFromCompactBinary(FCbFieldView Field, FInstigator& Instigator);
 
 /** Message from Server to Client to cook the given packages. */
-struct FAssignPackagesMessage : public UE::CompactBinaryTCP::IMessage
+struct FAssignPackagesMessage : public IMPCollectorMessage
 {
 public:
 	FAssignPackagesMessage() = default;
@@ -227,6 +229,7 @@ public:
 	virtual void Write(FCbWriter& Writer) const override;
 	virtual bool TryRead(FCbObjectView Object) override;
 	virtual FGuid GetMessageType() const override { return MessageType; }
+	virtual const TCHAR* GetDebugName() const override { return TEXT("AssignPackagesMessage"); }
 
 public:
 	TArray<FAssignPackageData> PackageDatas;
@@ -235,7 +238,7 @@ public:
 };
 
 /** Message from Server to Client to cancel the cook of the given packages. */
-struct FAbortPackagesMessage : public UE::CompactBinaryTCP::IMessage
+struct FAbortPackagesMessage : public IMPCollectorMessage
 {
 public:
 	FAbortPackagesMessage() = default;
@@ -244,6 +247,7 @@ public:
 	virtual void Write(FCbWriter& Writer) const override;
 	virtual bool TryRead(FCbObjectView Object) override;
 	virtual FGuid GetMessageType() const override { return MessageType; }
+	virtual const TCHAR* GetDebugName() const override { return TEXT("AbortPackagesMessage"); }
 
 public:
 	TArray<FName> PackageNames;
@@ -255,7 +259,7 @@ public:
  * If from Server, request that Client shutdown.
  * If from Client, notify Server it is shutting down.
  */
-struct FAbortWorkerMessage : public UE::CompactBinaryTCP::IMessage
+struct FAbortWorkerMessage : public IMPCollectorMessage
 {
 public:
 	enum EType
@@ -268,6 +272,7 @@ public:
 	virtual void Write(FCbWriter& Writer) const override;
 	virtual bool TryRead(FCbObjectView Object) override;
 	virtual FGuid GetMessageType() const override { return MessageType; }
+	virtual const TCHAR* GetDebugName() const override { return TEXT("AbortWorkerMessage"); }
 
 public:
 	EType Type;
@@ -275,12 +280,13 @@ public:
 };
 
 /** Message From Server to Client giving all of the COTFS settings the client needs. */
-struct FInitialConfigMessage : public UE::CompactBinaryTCP::IMessage
+struct FInitialConfigMessage : public IMPCollectorMessage
 {
 public:
 	virtual void Write(FCbWriter& Writer) const override;
 	virtual bool TryRead(FCbObjectView Object) override;
 	virtual FGuid GetMessageType() const override { return MessageType; }
+	virtual const TCHAR* GetDebugName() const override { return TEXT("InitialConfigMessage"); }
 
 	void ReadFromLocal(const UCookOnTheFlyServer& COTFS, const TArray<ITargetPlatform*>& InOrderedSessionPlatforms,
 		const FCookByTheBookOptions& InCookByTheBookOptions, const FCookOnTheFlyOptions& InCookOnTheFlyOptions,
@@ -328,12 +334,13 @@ bool LoadFromCompactBinary(FCbFieldView Field, FDiscoveredPackageReplication& Ou
  * Message from CookWorker to Director that reports dependency packages discovered during load/save of
  * a package that were not found in the earlier traversal of the packages dependencies.
  */
-struct FDiscoveredPackagesMessage : public UE::CompactBinaryTCP::IMessage
+struct FDiscoveredPackagesMessage : public IMPCollectorMessage
 {
 public:
 	virtual void Write(FCbWriter& Writer) const override;
 	virtual bool TryRead(FCbObjectView Object) override;
 	virtual FGuid GetMessageType() const override { return MessageType; }
+	virtual const TCHAR* GetDebugName() const override { return TEXT("DiscoveredPackagesMessage"); }
 
 public:
 	TArray<FDiscoveredPackageReplication> Packages;
@@ -388,13 +395,14 @@ private:
  * The Director intiates a heartbeat message; the CookWorker always responds to a heartbeat message with its own heartbeat message
  * in reply, with the same number.
  */
-struct FHeartbeatMessage : public UE::CompactBinaryTCP::IMessage
+struct FHeartbeatMessage : public IMPCollectorMessage
 {
 public:
 	FHeartbeatMessage(int32 InHeartbeatNumber=-1);
 	virtual void Write(FCbWriter& Writer) const override;
 	virtual bool TryRead(FCbObjectView Object) override;
 	virtual FGuid GetMessageType() const override { return MessageType; }
+	virtual const TCHAR* GetDebugName() const override { return TEXT("HeartbeatMessage"); }
 
 public:
 	int32 HeartbeatNumber;
diff --git a/Engine/Source/Editor/UnrealEd/Private/Cooker/PackageResultsMessage.h b/Engine/Source/Editor/UnrealEd/Private/Cooker/PackageResultsMessage.h
--- a/Engine/Source/Editor/UnrealEd/Private/Cooker/PackageResultsMessage.h
+++ b/Engine/Source/Editor/UnrealEd/Private/Cooker/PackageResultsMessage.h
@@ -5,6 +5,7 @@
 #include "Async/Future.h"
 #include "CompactBinaryTCP.h"
 #include "Containers/Array.h"
+#include "CookMPCollector.h"
 #include "CookTypes.h"
 #include "HAL/CriticalSection.h"
 #include "Templates/UniquePtr.h"
@@ -141,12 +142,13 @@ private:
 };
 
 /** Message from Client to Server giving the results for saved or refused-to-cook packages. */
-struct FPackageResultsMessage : public UE::CompactBinaryTCP::IMessage
+struct FPackageResultsMessage : public IMPCollectorMessage
 {
 public:
 	virtual void Write(FCbWriter& Writer) const override;
 	virtual bool TryRead(FCbObjectView Object) override;
 	virtual FGuid GetMessageType() const override { return MessageType; }
+	virtual const TCHAR* GetDebugName() const override { return TEXT("PackageResultsMessage"); }
 
 public:
 	TArray<FPackageRemoteResult> Results;
diff --git a/Engine/Source/Editor/UnrealEd/Private/ShaderStatsCollector.cpp b/Engine/Source/Editor/UnrealEd/Private/ShaderStatsCollector.cpp
--- a/Engine/Source/Editor/UnrealEd/Private/ShaderStatsCollector.cpp
+++ b/Engine/Source/Editor/UnrealEd/Private/ShaderStatsCollector.cpp
@@ -3,6 +3,7 @@
 #include "ShaderStatsCollector.h"
 #include "ShaderStats.h"
 #include "AnalyticsEventAttribute.h"
+#include "Serialization/CompactBinaryWriter.h"
 
 static TAutoConsoleVariable<int> CVarShaderCompilerStatsPrintoutInterval(
 	TEXT("r.ShaderCompiler.StatsPrintoutInterval"),
diff --git a/Engine/Source/Editor/UnrealEd/Public/Commandlets/AssetRegistryGenerator.h b/Engine/Source/Editor/UnrealEd/Public/Commandlets/AssetRegistryGenerator.h
--- a/Engine/Source/Editor/UnrealEd/Public/Commandlets/AssetRegistryGenerator.h
+++ b/Engine/Source/Editor/UnrealEd/Public/Commandlets/AssetRegistryGenerator.h
@@ -540,12 +540,13 @@ private:
 	friend FAssetRegistryMPCollector;
 };
 
-class FAssetRegistryPackageMessage : public UE::CompactBinaryTCP::IMessage
+class FAssetRegistryPackageMessage : public IMPCollectorMessage
 {
 public:
 	virtual void Write(FCbWriter& Writer) const override;
 	virtual bool TryRead(FCbObjectView Object) override;
 	virtual FGuid GetMessageType() const override { return MessageType; }
+	virtual const TCHAR* GetDebugName() const override { return TEXT("AssetRegistryPackageMessage"); }
 
 	FName PackageName;
 	const ITargetPlatform* TargetPlatform;
