Ensure that the editor's PlayWorld context is valid prior to handling a request to end the current PIE session.

diff --git a/Engine/Source/Editor/UnrealEd/Private/EditorEngine.cpp b/Engine/Source/Editor/UnrealEd/Private/EditorEngine.cpp
--- a/Engine/Source/Editor/UnrealEd/Private/EditorEngine.cpp
+++ b/Engine/Source/Editor/UnrealEd/Private/EditorEngine.cpp
@@ -2196,6 +2196,10 @@ void UEditorEngine::Tick( float DeltaSeconds, bool bIdleMode )
 	// Update resource streaming after both regular Editor viewports and PIE had a chance to add viewers.
 	IStreamingManager::Get().Tick(DeltaSeconds);
 
+	// Determine whether or not we should end the current PIE session. In some cases the client context may not be fully
+	// initialized until player login is complete, so make sure we have a valid world before actually handling the request.
+	const bool bEndPlayMapThisFrame = PlayWorld && bRequestEndPlayMapQueued;
+
 	// Update Audio. This needs to occur after rendering as the rendering code updates the listener position.
 	if (AudioDeviceManager)
 	{
@@ -2208,7 +2212,7 @@ void UEditorEngine::Tick( float DeltaSeconds, bool bIdleMode )
 
 		// Update audio device.
 		AudioDeviceManager->UpdateActiveAudioDevices((!PlayWorld && bAudioIsRealtime) || (PlayWorld && !PlayWorld->IsPaused()));
-		if (bRequestEndPlayMapQueued)
+		if (bEndPlayMapThisFrame)
 		{
 			// Shutdown all audio devices if we've requested end playmap now to avoid issues with GC running
 			TArray<FAudioDevice*> AudioDevices = AudioDeviceManager->GetAudioDevices();
@@ -2248,7 +2252,7 @@ void UEditorEngine::Tick( float DeltaSeconds, bool bIdleMode )
 	}
 
 	// After the play world has ticked, see if a request was made to end pie
-	if( bRequestEndPlayMapQueued )
+	if (bEndPlayMapThisFrame)
 	{
 		EndPlayMap();
 	}
