Deleting an external object package (Actor Folder) was clearing the undo buffer. Do the same logic as external actor packages.

diff --git a/Engine/Source/Editor/UnrealEd/Private/ObjectTools.cpp b/Engine/Source/Editor/UnrealEd/Private/ObjectTools.cpp
--- a/Engine/Source/Editor/UnrealEd/Private/ObjectTools.cpp
+++ b/Engine/Source/Editor/UnrealEd/Private/ObjectTools.cpp
@@ -2407,9 +2407,9 @@ namespace ObjectTools
 
 			bool bIsReferenced = false;
 
-			// Skip external actor packages when considering whether to clear the transaction buffer, as you should be able to undo deleting an actor (but not an asset)
-			// If an external actor package is kept alive by the transaction buffer then it will be re-marked as "newly created" further down this function
-			if ( Package != nullptr && bPerformReferenceCheck && !Package->GetName().Contains(FPackagePath::GetExternalActorsFolderName()))
+			// Skip external object packages when considering whether to clear the transaction buffer, as you should be able to undo deleting an actor, an actor folder, etc. (but not an asset)
+			// If an external object package is kept alive by the transaction buffer then it will be re-marked as "newly created" further down this function
+			if ( Package != nullptr && bPerformReferenceCheck && !Package->GetName().Contains(FPackagePath::GetExternalActorsFolderName()) && !Package->GetName().Contains(FPackagePath::GetExternalObjectsFolderName()))
 			{
 				bool bIsReferencedByUndo = false;
 				GatherObjectReferencersForDeletion(Package, bIsReferenced, bIsReferencedByUndo);
