World Partition - Mininmap: Updated minimap builder to wrap texture modifications between Pre/PostEditChange() calls
* Avoid a possible Assertion failed: bIsRoutingPostCompilation == false

diff --git a/Engine/Source/Editor/UnrealEd/Private/WorldPartition/WorldPartitionMiniMapBuilder.cpp b/Engine/Source/Editor/UnrealEd/Private/WorldPartition/WorldPartitionMiniMapBuilder.cpp
--- a/Engine/Source/Editor/UnrealEd/Private/WorldPartition/WorldPartitionMiniMapBuilder.cpp
+++ b/Engine/Source/Editor/UnrealEd/Private/WorldPartition/WorldPartitionMiniMapBuilder.cpp
@@ -64,6 +64,8 @@ bool UWorldPartitionMiniMapBuilder::PreRun(UWorld* World, FPackageSourceControlH
 	{
 		TStrongObjectPtr<UTextureFactory> Factory(NewObject<UTextureFactory>());
 		WorldMiniMap->MiniMapTexture = Factory->CreateTexture2D(WorldMiniMap, TEXT("MinimapTexture"), RF_TextExportTransient);
+		WorldMiniMap->MiniMapTexture->PreEditChange(nullptr);
+		
 		WorldMiniMap->MiniMapTexture->Source.Init(MinimapImageSizeX, MinimapImageSizeY, 1, 1, TSF_BGRA8);
 		WorldMiniMap->MiniMapTexture->MipGenSettings = TMGS_SimpleAverage;
 		WorldMiniMap->MiniMapWorldBounds = IterativeWorldBounds;
@@ -201,6 +203,7 @@ bool UWorldPartitionMiniMapBuilder::PostRun(UWorld* World, FPackageSourceControl
 		WorldMiniMap->MiniMapTexture->LODGroup = TEXTUREGROUP_UI;
 		WorldMiniMap->MiniMapTexture->VirtualTextureStreaming = true;
 		WorldMiniMap->MiniMapTexture->UpdateResource();
+		WorldMiniMap->MiniMapTexture->PostEditChange();
 	}
 
 	// Compute relevant UV space for the minimap
