Cooker: MPCook: Add a registration API for IMPCollector.
#rnx

diff --git a/Engine/Source/Editor/UnrealEd/Classes/CookOnTheSide/CookOnTheFlyServer.h b/Engine/Source/Editor/UnrealEd/Classes/CookOnTheSide/CookOnTheFlyServer.h
--- a/Engine/Source/Editor/UnrealEd/Classes/CookOnTheSide/CookOnTheFlyServer.h
+++ b/Engine/Source/Editor/UnrealEd/Classes/CookOnTheSide/CookOnTheFlyServer.h
@@ -624,6 +624,13 @@ public:
 	// ICookInfo interface
 	UNREALED_API virtual UE::Cook::FInstigator GetInstigator(FName PackageName) override;
 	UNREALED_API virtual TArray<UE::Cook::FInstigator> GetInstigatorChain(FName PackageName) override;
+	UNREALED_API virtual UE::Cook::ECookType GetCookType() override;
+	UNREALED_API virtual UE::Cook::ECookingDLC GetCookingDLC() override;
+	UNREALED_API virtual UE::Cook::EProcessType GetProcessType() override;
+	UNREALED_API virtual void RegisterCollector(UE::Cook::IMPCollector* Collector,
+		UE::Cook::EProcessType ProcessType = UE::Cook::EProcessType::AllMPCook) override;
+	UNREALED_API virtual void UnregisterCollector(UE::Cook::IMPCollector* Collector) override;
+
 
 	/** Dumps cooking stats to the log. Run from the exec command "Cook stats". */
 	UNREALED_API void DumpStats();
@@ -836,9 +843,11 @@ public:
 	UNREALED_API void OnObjectSaved( UObject *ObjectSaved, FObjectPreSaveContext SaveContext );
 
 	DECLARE_MULTICAST_DELEGATE(FOnCookByTheBookStarted);
+	UE_DEPRECATED(5.4, "Use UE::Cook::FDelegates::CookByTheBookStarted (CoreUObject/Public/UObject/ICookInfo.h.")
 	static FOnCookByTheBookStarted& OnCookByTheBookStarted() { return CookByTheBookStartedEvent; };
 
 	DECLARE_MULTICAST_DELEGATE(FOnCookByTheBookFinished);
+	UE_DEPRECATED(5.4, "Use UE::Cook::FDelegates::CookByTheBookFinished (CoreUObject/Public/UObject/ICookInfo.h.")
 	static FOnCookByTheBookFinished& OnCookByTheBookFinished() { return CookByTheBookFinishedEvent; };
 	/**
 	* Marks a package as dirty for cook
diff --git a/Engine/Source/Editor/UnrealEd/Private/CookOnTheFlyServer.cpp b/Engine/Source/Editor/UnrealEd/Private/CookOnTheFlyServer.cpp
--- a/Engine/Source/Editor/UnrealEd/Private/CookOnTheFlyServer.cpp
+++ b/Engine/Source/Editor/UnrealEd/Private/CookOnTheFlyServer.cpp
@@ -7158,6 +7158,77 @@ TArray<UE::Cook::FInstigator> UCookOnTheFlyServer::GetInstigatorChain(FName Pack
 	return Result; // Unreachable
 }
 
+UE::Cook::ECookType UCookOnTheFlyServer::GetCookType()
+{
+	if (IsDirectorCookByTheBook())
+	{
+		return UE::Cook::ECookType::ByTheBook;
+	}
+	else
+	{
+		check(this->IsDirectorCookOnTheFly());
+		return UE::Cook::ECookType::OnTheFly;
+	}
+}
+
+UE::Cook::ECookingDLC UCookOnTheFlyServer::GetCookingDLC()
+{
+	using namespace UE::Cook;
+	return IsCookingDLC() ? ECookingDLC::Yes : ECookingDLC::No;
+}
+
+UE::Cook::EProcessType UCookOnTheFlyServer::GetProcessType()
+{
+	using namespace UE::Cook;
+	if (IsCookWorkerMode())
+	{
+		return EProcessType::Worker;
+	}
+	else if (CookDirector)
+	{
+		return EProcessType::Director;
+	}
+	else
+	{
+		return EProcessType::SingleProcess;
+	}
+}
+
+void UCookOnTheFlyServer::RegisterCollector(UE::Cook::IMPCollector* Collector, UE::Cook::EProcessType ProcessType)
+{
+	using namespace UE::Cook;
+
+	TRefCountPtr<IMPCollector> DeleteIfUnusedAndCallerHasNoReference(Collector);
+	if (CookDirector)
+	{
+		if (ProcessType == EProcessType::Director || ProcessType == EProcessType::AllMPCook)
+		{
+			CookDirector->Register(Collector);
+		}
+	}
+	else if (CookWorkerClient)
+	{
+		if (ProcessType == EProcessType::Worker || ProcessType == EProcessType::AllMPCook)
+		{
+			CookWorkerClient->Register(Collector);
+		}
+	}
+}
+
+void UCookOnTheFlyServer::UnregisterCollector(UE::Cook::IMPCollector* Collector)
+{
+	TRefCountPtr<UE::Cook::IMPCollector> DeleteIfCallerHasNoReference(Collector);
+	if (CookDirector)
+	{
+		CookDirector->Unregister(Collector);
+	}
+	else if (CookWorkerClient)
+	{
+		CookWorkerClient->Unregister(Collector);
+	}
+}
+
+
 void UCookOnTheFlyServer::DumpStats()
 {
 	UE_LOG(LogCook, Display, TEXT("IntStats:"));
@@ -9860,7 +9931,10 @@ void UCookOnTheFlyServer::CookByTheBookFinished()
 	}
 
 	ShutdownCookSession();
+	PRAGMA_DISABLE_DEPRECATION_WARNINGS;
 	CookByTheBookFinishedEvent.Broadcast();
+	PRAGMA_ENABLE_DEPRECATION_WARNINGS;
+	UE::Cook::FDelegates::CookByTheBookFinished.Broadcast(*this);
 	UE_LOG(LogCook, Display, TEXT("Done!"));
 }
 
@@ -11042,7 +11116,10 @@ void UCookOnTheFlyServer::StartCookByTheBook( const FCookByTheBookStartupOptions
 	RecordDLCPackagesFromBaseGame(BeginContext);
 	RegisterCookByTheBookDelegates();
 
+	PRAGMA_DISABLE_DEPRECATION_WARNINGS;
 	CookByTheBookStartedEvent.Broadcast();
+	PRAGMA_ENABLE_DEPRECATION_WARNINGS;
+	UE::Cook::FDelegates::CookByTheBookStarted.Broadcast(*this);
 }
 
 const UCookOnTheFlyServer::FCookByTheBookStartupOptions& UCookOnTheFlyServer::BlockOnPrebootCookGate(bool& bOutAbortCook,
@@ -11217,7 +11294,10 @@ void UCookOnTheFlyServer::StartCookAsCookWorker()
 	{
 		RegisterCookByTheBookDelegates();
 		BeginCookFinishShaderCodeLibrary(BeginContext);
+		PRAGMA_DISABLE_DEPRECATION_WARNINGS;
 		CookByTheBookStartedEvent.Broadcast();
+		PRAGMA_ENABLE_DEPRECATION_WARNINGS;
+		UE::Cook::FDelegates::CookByTheBookStarted.Broadcast(*this);
 	}
 }
 
@@ -11260,7 +11340,10 @@ void UCookOnTheFlyServer::CookAsCookWorkerFinished()
 	LogCookWorkerStats();
 	if (IsDirectorCookByTheBook())
 	{
+		PRAGMA_DISABLE_DEPRECATION_WARNINGS;
 		CookByTheBookFinishedEvent.Broadcast();
+		PRAGMA_ENABLE_DEPRECATION_WARNINGS;
+		UE::Cook::FDelegates::CookByTheBookFinished.Broadcast(*this);
 	}
 	CookWorkerClient->FlushLogs();
 }
diff --git a/Engine/Source/Runtime/CoreUObject/Private/UObject/ICookInfo.cpp b/Engine/Source/Runtime/CoreUObject/Private/UObject/ICookInfo.cpp
--- a/Engine/Source/Runtime/CoreUObject/Private/UObject/ICookInfo.cpp
+++ b/Engine/Source/Runtime/CoreUObject/Private/UObject/ICookInfo.cpp
@@ -45,6 +45,9 @@ FString FInstigator::ToString() const
 	return FString(Result);
 }
 
+FCookInfoEvent FDelegates::CookByTheBookStarted;
+FCookInfoEvent FDelegates::CookByTheBookFinished;
+
 }
 
 static thread_local ECookLoadType GCookLoadType = ECookLoadType::Unexpected;
diff --git a/Engine/Source/Runtime/CoreUObject/Public/UObject/CookEnums.h b/Engine/Source/Runtime/CoreUObject/Public/UObject/CookEnums.h
--- a/Engine/Source/Runtime/CoreUObject/Public/UObject/CookEnums.h
+++ b/Engine/Source/Runtime/CoreUObject/Public/UObject/CookEnums.h
@@ -21,6 +21,22 @@ enum class ECookingDLC
 	No,
 };
 
+/**
+ * The cookcommandlet can run in MPCook mode, where it creates worker processes to use more parallelism and
+ * reduce wall time of the cook. This enum is reported or received by some API functions on the cooker, and describes
+ * whether the current process the cooker and caller are in is the MPCook director, or an MPCook worker, or a cooker
+ * running as a singleprocess without MPCook.
+ */
+enum class EProcessType
+{
+	SingleProcess,
+	Director,
+	Worker,
+
+	/** Used only when passed as a filter: matches either Director or Worker. */
+	AllMPCook,
+};
+
 /* The Result of the Cook-Save of a package */
 enum class ECookResult : uint8
 {
diff --git a/Engine/Source/Runtime/CoreUObject/Public/UObject/ICookInfo.h b/Engine/Source/Runtime/CoreUObject/Public/UObject/ICookInfo.h
--- a/Engine/Source/Runtime/CoreUObject/Public/UObject/ICookInfo.h
+++ b/Engine/Source/Runtime/CoreUObject/Public/UObject/ICookInfo.h
@@ -2,13 +2,19 @@
 
 #pragma once
 
+#if WITH_EDITOR
 #include "Containers/Array.h"
 #include "Containers/UnrealString.h"
+#include "Delegates/Delegate.h"
 #include "HAL/Platform.h"
+#include "UObject/CookEnums.h"
 #include "UObject/NameTypes.h"
+#endif
 
 #if WITH_EDITOR
 
+namespace UE::Cook { class IMPCollector; }
+
 namespace UE::Cook
 {
 
@@ -97,9 +103,45 @@ public:
 	 * First element is the direct instigator of the package, last is the root instigator that started the chain.
 	 */
 	virtual TArray<FInstigator> GetInstigatorChain(FName PackageName) = 0;
+
+	/** The type (e.g. CookByTheBook) of the running cook. This function will not return ECookType::Unknown. */
+	virtual UE::Cook::ECookType GetCookType() = 0;
+	/** Whether dlc is being cooked (e.g. via "-dlcname=<PluginName>"). This function will not return ECookingDLC::Unknown. */
+	virtual UE::Cook::ECookingDLC GetCookingDLC() = 0;
+	/** The role the current process plays in its MPCook session, or EProcessType::SingleProcess if it is running standalone. */
+	virtual UE::Cook::EProcessType GetProcessType() = 0;
+
+	/**
+	 * MPCook: register in the current process a collector that replicates system-specific and package-specific
+	 * information between CookWorkers and the CookDirector. Registration will be skipped if the current cook is
+	 * singleprocess, or if the provided ProcessType does not match the current processtype. If registration is
+	 * skipped, the Collector will be referenced but then immediately released, which will delete the Collector if
+	 * the caller does not have their own TRefCountPtr to it.
+	 * 
+	 * @param ProcessType: The given collector will only be registered on the given process types, allowing you
+	 *                     to register a different class on Workers than on the Director if desired. 
+	 */
+	virtual void RegisterCollector(IMPCollector* Collector,
+		UE::Cook::EProcessType ProcessType=UE::Cook::EProcessType::AllMPCook) = 0;
+	/**
+	 * MPCook: Unregister in the current process a collector that was registered via RegisterCollector. Silently
+	 * returns if the collector is not registered. References to the Collector will be released, which will
+	 * delete the Collector if he caller does not have their own TRefCountPtr to it.
+	 */
+	virtual void UnregisterCollector(IMPCollector* Collector) = 0;
+};
+
+DECLARE_MULTICAST_DELEGATE_OneParam(FCookInfoEvent, ICookInfo&);
+
+/** UE::Cook::FDelegates: callbacks for cook events. */
+struct FDelegates
+{
+public:
+	static COREUOBJECT_API FCookInfoEvent CookByTheBookStarted;
+	static COREUOBJECT_API FCookInfoEvent CookByTheBookFinished;
 };
 
-}
+} // namespace UE::Cook
 
 /**
  * A scope around loads when cooking that indicates whether the loaded package is needed in game or not.
@@ -127,4 +169,4 @@ private:
 	ECookLoadType PreviousScope;
 };
 
-#endif
\ No newline at end of file
+#endif // WITH_EDITOR
\ No newline at end of file
