UEditorEngine::OnLevelRemovedFromWorld : Call ResetTransaction when Trans->ContainsPieObjects returns true to avoid stale reference failed check in World Travel (UEngine::LoadMap)

#rnx

diff --git a/Engine/Source/Editor/UnrealEd/Private/EditorEngine.cpp b/Engine/Source/Editor/UnrealEd/Private/EditorEngine.cpp
--- a/Engine/Source/Editor/UnrealEd/Private/EditorEngine.cpp
+++ b/Engine/Source/Editor/UnrealEd/Private/EditorEngine.cpp
@@ -6943,6 +6943,11 @@ void UEditorEngine::OnLevelRemovedFromWorld(ULevel* InLevel, UWorld* InWorld)
 		{
 			if (InWorld->IsPlayInEditor())
 			{
+				if (Trans->ContainsPieObjects())
+				{
+					ResetTransaction(NSLOCTEXT("UnrealEd", "LevelRemovedFromWorldEditorCallbackPIE", "Level removed from PIE/SIE world"));
+				}
+
 				// Each additional instance of PIE in a multiplayer game will add another barrier, so if the event is triggered then this is the case and we need to lift it
 				// Otherwise there will be an imbalance between barriers set and barriers removed and we won't be able to undo when we return.
 				Trans->RemoveUndoBarrier();
