Make sure importers do not allow nan or infinite float mesh data.
Also make sure the static mesh build fix those issue before building the mesh description.

#rnx

diff --git a/Engine/Plugins/Experimental/Enterprise/DatasmithInterchange/Source/DatasmithInterchange/Private/InterchangeDatasmithTranslator.cpp b/Engine/Plugins/Experimental/Enterprise/DatasmithInterchange/Source/DatasmithInterchange/Private/InterchangeDatasmithTranslator.cpp
--- a/Engine/Plugins/Experimental/Enterprise/DatasmithInterchange/Source/DatasmithInterchange/Private/InterchangeDatasmithTranslator.cpp
+++ b/Engine/Plugins/Experimental/Enterprise/DatasmithInterchange/Source/DatasmithInterchange/Private/InterchangeDatasmithTranslator.cpp
@@ -625,6 +625,12 @@ TFuture<TOptional<UE::Interchange::FMeshPayloadData>> UInterchangeDatasmithTrans
 				{
 					UE::Interchange::FMeshPayloadData StaticMeshPayloadData;
 					StaticMeshPayloadData.MeshDescription = MoveTemp(DatasmithMeshPayload.LodMeshes[0]);
+					if (!FStaticMeshOperations::ValidateAndFixData(StaticMeshPayloadData.MeshDescription, MeshElement->GetName()))
+					{
+						UInterchangeResultError_Generic* ErrorResult = AddMessage<UInterchangeResultError_Generic>();
+						ErrorResult->SourceAssetName = SourceData ? SourceData->GetFilename() : FString();
+						ErrorResult->Text = LOCTEXT("GetMeshPayloadData_ValidateMeshDescriptionFail", "Invalid mesh data (NAN) was found and fix to zero. Mesh render can be bad.");
+					}
 					// Bake the payload mesh, with the provided transform
 					if (!MeshGlobalTransform.Equals(FTransform::Identity))
 					{
diff --git a/Engine/Plugins/Interchange/Runtime/Source/Import/Private/Fbx/InterchangeFbxTranslator.cpp b/Engine/Plugins/Interchange/Runtime/Source/Import/Private/Fbx/InterchangeFbxTranslator.cpp
--- a/Engine/Plugins/Interchange/Runtime/Source/Import/Private/Fbx/InterchangeFbxTranslator.cpp
+++ b/Engine/Plugins/Interchange/Runtime/Source/Import/Private/Fbx/InterchangeFbxTranslator.cpp
@@ -20,6 +20,7 @@
 #include "Serialization/JsonWriter.h"
 #include "Serialization/JsonSerializer.h"
 #include "Serialization/LargeMemoryReader.h"
+#include "StaticMeshOperations.h"
 #include "Texture/InterchangeTexturePayloadData.h"
 #include "UObject/GCObjectScopeGuard.h"
 
@@ -294,7 +295,13 @@ TFuture<TOptional<UE::Interchange::FMeshPayloadData>> UInterchangeFbxTranslator:
 			break;
 		}
 
-		
+		if (!FStaticMeshOperations::ValidateAndFixData(MeshPayloadData.MeshDescription, PayLoadKey.UniqueId))
+		{
+			UInterchangeResultError_Generic* ErrorResult = AddMessage<UInterchangeResultError_Generic>();
+			ErrorResult->SourceAssetName = SourceData ? SourceData->GetFilename() : FString();
+			ErrorResult->Text = NSLOCTEXT("UInterchangeFbxTranslator", "GetMeshPayloadData_ValidateMeshDescriptionFail", "Invalid mesh data (NAN) was found and fix to zero. Mesh render can be bad.");
+		}
+
 		Promise->SetValue(MoveTemp(MeshPayloadData));
 	}));
 
diff --git a/Engine/Plugins/Interchange/Runtime/Source/Import/Private/Gltf/InterchangeGltfTranslator.cpp b/Engine/Plugins/Interchange/Runtime/Source/Import/Private/Gltf/InterchangeGltfTranslator.cpp
--- a/Engine/Plugins/Interchange/Runtime/Source/Import/Private/Gltf/InterchangeGltfTranslator.cpp
+++ b/Engine/Plugins/Interchange/Runtime/Source/Import/Private/Gltf/InterchangeGltfTranslator.cpp
@@ -1111,6 +1111,13 @@ TFuture< TOptional< UE::Interchange::FMeshPayloadData > > UInterchangeGLTFTransl
 			TOptional<UE::Interchange::FMeshPayloadData> Result;
 			if (bSuccessfullAcquisition)
 			{
+				if (!FStaticMeshOperations::ValidateAndFixData(MeshPayLoadData.MeshDescription, PayLoadKey.UniqueId))
+				{
+					UInterchangeResultError_Generic* ErrorResult = AddMessage<UInterchangeResultError_Generic>();
+					ErrorResult->SourceAssetName = SourceData ? SourceData->GetFilename() : FString();
+					ErrorResult->Text = NSLOCTEXT("UInterchangeGLTFTranslator", "GetMeshPayloadData_ValidateMeshDescriptionFail", "Invalid mesh data (NAN) was found and fix to zero. Mesh render can be bad.");
+				}
+
 				Result.Emplace(MeshPayLoadData);
 			}
 
diff --git a/Engine/Plugins/Interchange/Runtime/Source/Import/Private/Mesh/InterchangeOBJTranslator.cpp b/Engine/Plugins/Interchange/Runtime/Source/Import/Private/Mesh/InterchangeOBJTranslator.cpp
--- a/Engine/Plugins/Interchange/Runtime/Source/Import/Private/Mesh/InterchangeOBJTranslator.cpp
+++ b/Engine/Plugins/Interchange/Runtime/Source/Import/Private/Mesh/InterchangeOBJTranslator.cpp
@@ -1563,6 +1563,13 @@ TFuture<TOptional<UE::Interchange::FMeshPayloadData>> UInterchangeOBJTranslator:
 			FMeshPayloadData Payload;
 			Payload.MeshDescription = ObjDataPtr->MakeMeshDescriptionForGroup(PayLoadKey.UniqueId, MeshGlobalTransform);
 
+			if (!FStaticMeshOperations::ValidateAndFixData(Payload.MeshDescription, PayLoadKey.UniqueId))
+			{
+				UInterchangeResultError_Generic* ErrorResult = AddMessage<UInterchangeResultError_Generic>();
+				ErrorResult->SourceAssetName = SourceData ? SourceData->GetFilename() : FString();
+				ErrorResult->Text = NSLOCTEXT("UInterchangeOBJTranslator", "GetMeshPayloadData_ValidateMeshDescriptionFail", "Invalid mesh data (NAN) was found and fix to zero. Mesh render can be bad.");
+			}
+
 			return TOptional<FMeshPayloadData>(Payload);
 		}
 	);
diff --git a/Engine/Source/Developer/MeshBuilder/Private/MeshDescriptionHelper.cpp b/Engine/Source/Developer/MeshBuilder/Private/MeshDescriptionHelper.cpp
--- a/Engine/Source/Developer/MeshBuilder/Private/MeshDescriptionHelper.cpp
+++ b/Engine/Source/Developer/MeshBuilder/Private/MeshDescriptionHelper.cpp
@@ -46,6 +46,9 @@ void FMeshDescriptionHelper::SetupRenderMeshDescription(UObject* Owner, FMeshDes
 		RenderMeshDescription.Compact(Remappings);
 	}
 
+	//Make sure we do not have nan or infinite float in the mesh description data
+	FStaticMeshOperations::ValidateAndFixData(RenderMeshDescription, Owner->GetName());
+
 	//This function make sure the Polygon Normals Tangents Binormals are computed and also remove degenerated triangle from the render mesh description.
 	FStaticMeshOperations::ComputeTriangleTangentsAndNormals(RenderMeshDescription, ComparisonThreshold, *Owner->GetPathName());
 
diff --git a/Engine/Source/Editor/UnrealEd/Private/Fbx/FbxStaticMeshImport.cpp b/Engine/Source/Editor/UnrealEd/Private/Fbx/FbxStaticMeshImport.cpp
--- a/Engine/Source/Editor/UnrealEd/Private/Fbx/FbxStaticMeshImport.cpp
+++ b/Engine/Source/Editor/UnrealEd/Private/Fbx/FbxStaticMeshImport.cpp
@@ -732,7 +732,7 @@ bool UnFbx::FFbxImporter::BuildStaticMeshFromGeometry(FbxNode* Node, UStaticMesh
 			int32 RealVertexIndex = VertexOffset + VertexIndex;
 			FbxVector4 FbxPosition = Mesh->GetControlPoints()[VertexIndex];
 			FbxPosition = TotalMatrix.MultT(FbxPosition);
-			const FVector3f VertexPosition = (FVector3f)Converter.ConvertPos(FbxPosition);
+			FVector3f VertexPosition = (FVector3f)Converter.ConvertPos(FbxPosition);
 			
 			FVertexID AddedVertexId = MeshDescription->CreateVertex();
 			VertexPositions[AddedVertexId] = VertexPosition;
diff --git a/Engine/Source/Editor/UnrealEd/Private/Fbx/FbxUtilsImport.cpp b/Engine/Source/Editor/UnrealEd/Private/Fbx/FbxUtilsImport.cpp
--- a/Engine/Source/Editor/UnrealEd/Private/Fbx/FbxUtilsImport.cpp
+++ b/Engine/Source/Editor/UnrealEd/Private/Fbx/FbxUtilsImport.cpp
@@ -2,6 +2,7 @@
 
 #include "CoreMinimal.h"
 #include "FbxImporter.h"
+#include "Logging/TokenizedMessage.h"
 
 using namespace UnFbx;
 
@@ -10,6 +11,55 @@ FbxAMatrix FFbxDataConverter::JointPostConversionMatrix;
 FbxAMatrix FFbxDataConverter::AxisConversionMatrix;
 FbxAMatrix FFbxDataConverter::AxisConversionMatrixInv;
 
+namespace UE::Fbx::Private
+{
+	void LogFinitErrorOnce()
+	{
+		static bool bLogged = false;
+		if (!bLogged)
+		{
+			bLogged = true;
+			UE_LOG(LogFbx, Error, TEXT("Import mesh have some infinite value in the data."));
+		}
+	}
+
+	void VerifyFiniteValue(float& Value)
+	{
+		if (!FMath::IsFinite(Value))
+		{
+			Value = 0.0f;
+			LogFinitErrorOnce();
+		}
+	}
+
+	void VerifyFiniteValue(FVector& Value)
+	{
+		if (Value.ContainsNaN())
+		{
+			Value.Set(0.0, 0.0, 0.0);
+			LogFinitErrorOnce();
+		}
+	}
+
+	void VerifyFiniteValue(FQuat& Value)
+	{
+		if (Value.ContainsNaN())
+		{
+			Value = FQuat::Identity;
+			LogFinitErrorOnce();
+		}
+	}
+
+	void VerifyFiniteValue(FMatrix& Value)
+	{
+		if (Value.ContainsNaN())
+		{
+			Value.SetIdentity();
+			LogFinitErrorOnce();
+		}
+	}
+} //ns UE::Fbx::Private
+
 FVector FFbxDataConverter::ConvertPos(FbxVector4 Vector)
 {
 	FVector Out;
@@ -17,6 +67,7 @@ FVector FFbxDataConverter::ConvertPos(FbxVector4 Vector)
 	// flip Y, then the right-handed axis system is converted to LHS
 	Out[1] = -Vector[1];
 	Out[2] = Vector[2];
+	UE::Fbx::Private::VerifyFiniteValue(Out);
 	return Out;
 }
 
@@ -28,12 +79,15 @@ FVector FFbxDataConverter::ConvertDir(FbxVector4 Vector)
 	Out[0] = Vector[0];
 	Out[1] = -Vector[1];
 	Out[2] = Vector[2];
+	UE::Fbx::Private::VerifyFiniteValue(Out);
 	return Out;
 }
 
 FRotator FFbxDataConverter::ConvertEuler( FbxDouble3 Euler )
 {
-	return FRotator::MakeFromEuler( FVector( Euler[0], -Euler[1], Euler[2] ) );
+	FVector UnrealEuler(Euler[0], -Euler[1], Euler[2]);
+	UE::Fbx::Private::VerifyFiniteValue(UnrealEuler);
+	return FRotator::MakeFromEuler(UnrealEuler);
 }
 
 
@@ -43,6 +97,7 @@ FVector FFbxDataConverter::ConvertScale(FbxDouble3 Vector)
 	Out[0] = Vector[0];
 	Out[1] = Vector[1];
 	Out[2] = Vector[2];
+	UE::Fbx::Private::VerifyFiniteValue(Out);
 	return Out;
 }
 
@@ -53,6 +108,7 @@ FVector FFbxDataConverter::ConvertScale(FbxVector4 Vector)
 	Out[0] = Vector[0];
 	Out[1] = Vector[1];
 	Out[2] = Vector[2];
+	UE::Fbx::Private::VerifyFiniteValue(Out);
 	return Out;
 }
 
@@ -70,6 +126,7 @@ FVector FFbxDataConverter::ConvertRotationToFVect(FbxQuaternion Quaternion, bool
 	FQuat UnrealQuaternion = ConvertRotToQuat(Quaternion);
 	FVector Euler;
 	Euler = UnrealQuaternion.Euler();
+	UE::Fbx::Private::VerifyFiniteValue(Euler);
 	if (bInvertOrient)
 	{
 		Euler.Y = -Euler.Y;
@@ -88,6 +145,7 @@ FQuat FFbxDataConverter::ConvertRotToQuat(FbxQuaternion Quaternion)
 	UnrealQuat.Y = -Quaternion[1];
 	UnrealQuat.Z = Quaternion[2];
 	UnrealQuat.W = -Quaternion[3];
+	UE::Fbx::Private::VerifyFiniteValue(UnrealQuat);
 	
 	return UnrealQuat;
 }
@@ -99,6 +157,7 @@ float FFbxDataConverter::ConvertDist(FbxDouble Distance)
 {
 	float Out;
 	Out = (float)Distance;
+	UE::Fbx::Private::VerifyFiniteValue(Out);
 	return Out;
 }
 
@@ -139,7 +198,7 @@ FMatrix FFbxDataConverter::ConvertMatrix(const FbxAMatrix& Matrix)
 			UEMatrix.M[i][3] = Row[3];
 		}
 	}
-	
+	UE::Fbx::Private::VerifyFiniteValue(UEMatrix);
 	return UEMatrix;
 }
 
diff --git a/Engine/Source/Runtime/StaticMeshDescription/Private/StaticMeshOperations.cpp b/Engine/Source/Runtime/StaticMeshDescription/Private/StaticMeshOperations.cpp
--- a/Engine/Source/Runtime/StaticMeshDescription/Private/StaticMeshOperations.cpp
+++ b/Engine/Source/Runtime/StaticMeshDescription/Private/StaticMeshOperations.cpp
@@ -2945,4 +2945,76 @@ void FStaticMeshOperations::ReorderMeshDescriptionPolygonGroups(const FMeshDescr
 	}
 }
 
+bool FStaticMeshOperations::ValidateAndFixData(FMeshDescription& MeshDescription, const FString& DebugName)
+{
+	bool bHasInvalidPositions = false;
+	bool bHasInvalidTangentSpaces = false;
+	bool bHasInvalidUVs = false;
+	bool bHasInvalidVertexColors = false;
+
+	FStaticMeshAttributes Attributes(MeshDescription);
+
+	TVertexAttributesRef<FVector3f> VertexPositions = Attributes.GetVertexPositions();
+	for (const FVertexID VertexID : MeshDescription.Vertices().GetElementIDs())
+	{
+		if (VertexPositions[VertexID].ContainsNaN())
+		{
+			bHasInvalidPositions = true;
+			VertexPositions[VertexID] = FVector3f::ZeroVector;
+		}
+	}
+	TVertexInstanceAttributesRef<FVector3f> VertexInstanceNormals = Attributes.GetVertexInstanceNormals();
+	TVertexInstanceAttributesRef<FVector3f> VertexInstanceTangents = Attributes.GetVertexInstanceTangents();
+	TVertexInstanceAttributesRef<float> VertexInstanceBinormalSigns = Attributes.GetVertexInstanceBinormalSigns();
+	TVertexInstanceAttributesRef<FVector2f> VertexInstanceUVs = Attributes.GetVertexInstanceUVs();
+	TVertexInstanceAttributesRef<FVector4f> VertexInstanceColors = Attributes.GetVertexInstanceColors();
+	for (const FVertexInstanceID VertexInstanceID : MeshDescription.VertexInstances().GetElementIDs())
+	{
+		if (VertexInstanceNormals[VertexInstanceID].ContainsNaN())
+		{
+			bHasInvalidTangentSpaces = true;
+			VertexInstanceNormals[VertexInstanceID] = FVector3f::Zero();
+		}
+		if (VertexInstanceTangents[VertexInstanceID].ContainsNaN())
+		{
+			bHasInvalidTangentSpaces = true;
+			VertexInstanceTangents[VertexInstanceID] = FVector3f::Zero();
+		}
+		if (FMath::IsNaN(VertexInstanceBinormalSigns[VertexInstanceID]))
+		{
+			bHasInvalidTangentSpaces = true;
+			VertexInstanceBinormalSigns[VertexInstanceID] = 0.0f;
+		}
+		if (VertexInstanceUVs[VertexInstanceID].ContainsNaN())
+		{
+			bHasInvalidUVs = true;
+			VertexInstanceUVs[VertexInstanceID] = FVector2f::Zero();
+		}
+		if (VertexInstanceColors[VertexInstanceID].ContainsNaN())
+		{
+			bHasInvalidVertexColors = true;
+			VertexInstanceColors[VertexInstanceID] = FVector4f::One();
+		}
+	}
+
+	if (bHasInvalidPositions)
+	{
+		UE_LOG(LogStaticMeshOperations, Error, TEXT("Mesh %s has NaNs in it's vertex positions! Offending positions are set to zero."), *DebugName);
+	}
+	if (bHasInvalidTangentSpaces)
+	{
+		UE_LOG(LogStaticMeshOperations, Error, TEXT("Mesh %s has NaNs in it's vertex instance tangent space! Offending tangents are set to zero."), *DebugName);
+	}
+	if (bHasInvalidUVs)
+	{
+		UE_LOG(LogStaticMeshOperations, Error, TEXT("Mesh %s has NaNs in it's vertex instance uvs! Offending uvs are set to zero."), *DebugName);
+	}
+	if (bHasInvalidVertexColors)
+	{
+		UE_LOG(LogStaticMeshOperations, Error, TEXT("Mesh %s has NaNs in it's vertex instance colors! Offending colors are set to white."), *DebugName);
+	}
+
+	return !bHasInvalidPositions && !bHasInvalidTangentSpaces && !bHasInvalidUVs && !bHasInvalidVertexColors;
+}
+
 #undef LOCTEXT_NAMESPACE
diff --git a/Engine/Source/Runtime/StaticMeshDescription/Public/StaticMeshOperations.h b/Engine/Source/Runtime/StaticMeshDescription/Public/StaticMeshOperations.h
--- a/Engine/Source/Runtime/StaticMeshDescription/Public/StaticMeshOperations.h
+++ b/Engine/Source/Runtime/StaticMeshDescription/Public/StaticMeshOperations.h
@@ -201,4 +201,7 @@ public:
 		, FMeshDescription& DestinationMeshDescription
 		, TOptional<const FString> UnmatchMaterialNameWarning
 		, TOptional<const FString> DestinationPolygonGroupCountDifferFromSource_Msg);
+
+	/** Verify the mesh data does not contain any NAN or INF float value, if such a case happen the value are set to zero or identity for matrix or quat. */
+	static STATICMESHDESCRIPTION_API bool ValidateAndFixData(FMeshDescription& MeshDescription, const FString& DebugName);
 };
