Set  some unitialized variables in Thumbnail and SlateTexture

#rnx

diff --git a/Engine/Source/Editor/UnrealEd/Public/AssetThumbnail.h b/Engine/Source/Editor/UnrealEd/Public/AssetThumbnail.h
--- a/Engine/Source/Editor/UnrealEd/Public/AssetThumbnail.h
+++ b/Engine/Source/Editor/UnrealEd/Public/AssetThumbnail.h
@@ -274,17 +274,17 @@ private:
 		/** The object whose thumbnail is rendered */
 		FAssetData AssetData;
 		/** Rendering resource for slate */
-		FSlateTexture2DRHIRef* ThumbnailTexture;
+		FSlateTexture2DRHIRef* ThumbnailTexture = nullptr;
 		/** Render target for slate */
-		FSlateTextureRenderTarget2DResource* ThumbnailRenderTarget;
+		FSlateTextureRenderTarget2DResource* ThumbnailRenderTarget = nullptr;
 		/** The time since last access */
-		double LastAccessTime;
+		double LastAccessTime = 0.0;
 		/** The time since last update */
-		double LastUpdateTime;
+		double LastUpdateTime = 0.0;
 		/** Width of the thumbnail */
-		uint32 Width;
+		uint32 Width = 0;
 		/** Height of the thumbnail */
-		uint32 Height;
+		uint32 Height = 0;
 		~FThumbnailInfo();
 	};
 	/**
diff --git a/Engine/Source/Runtime/Engine/Private/Slate/SlateTextures.cpp b/Engine/Source/Runtime/Engine/Private/Slate/SlateTextures.cpp
--- a/Engine/Source/Runtime/Engine/Private/Slate/SlateTextures.cpp
+++ b/Engine/Source/Runtime/Engine/Private/Slate/SlateTextures.cpp
@@ -11,6 +11,7 @@ FSlateTexture2DRHIRef::FSlateTexture2DRHIRef( FTexture2DRHIRef InRef, uint32 InW
 	: TSlateTexture( InRef )
 	, Width( InWidth )
 	, Height( InHeight )
+	, TexCreateFlags( TexCreate_None )
 	, PixelFormat( PF_Unknown )
 	, bCreateEmptyTexture( false )
 {
