UEJpegComp: changed around exporter so that it exports original raw data and doesn't recompress anything into UEJ format.

diff --git a/Engine/Source/Editor/UnrealEd/Classes/Exporters/TextureExporterPNG.h b/Engine/Source/Editor/UnrealEd/Classes/Exporters/TextureExporterPNG.h
--- a/Engine/Source/Editor/UnrealEd/Classes/Exporters/TextureExporterPNG.h
+++ b/Engine/Source/Editor/UnrealEd/Classes/Exporters/TextureExporterPNG.h
@@ -44,11 +44,14 @@ class UTextureExporterJPEG : public UExporter
 
 
 UCLASS()
-class UTextureExporterUEJPEG : public UTextureExporterGeneric
+class UTextureExporterUEJPEG : public UExporter
 {
 	GENERATED_UCLASS_BODY()
 	
-	virtual bool SupportsTexture(UTexture* Texture) const override;
+	//~ Begin UExporter Interface
+	UNREALED_API virtual bool SupportsObject(UObject* Object) const override;
+	UNREALED_API virtual bool ExportBinary( UObject* Object, const TCHAR* Type, FArchive& Ar, FFeedbackContext* Warn, int32 FileIndex = 0, uint32 PortFlags=0 ) override;
+	//~ End UExporter Interface
 };
 
 
diff --git a/Engine/Source/Editor/UnrealEd/Private/Factories/EditorFactories.cpp b/Engine/Source/Editor/UnrealEd/Private/Factories/EditorFactories.cpp
--- a/Engine/Source/Editor/UnrealEd/Private/Factories/EditorFactories.cpp
+++ b/Engine/Source/Editor/UnrealEd/Private/Factories/EditorFactories.cpp
@@ -4999,23 +4999,6 @@ UVirtualTextureBuilderExporterPNG::UVirtualTextureBuilderExporterPNG(const FObje
 
 //========================================================
 
-UTextureExporterUEJPEG::UTextureExporterUEJPEG(const FObjectInitializer& ObjectInitializer)
-	: UTextureExporterGeneric(ObjectInitializer)
-{
-	FormatExtension.Add(TEXT("UEJ"));
-	FormatDescription.Add(TEXT("UE JPEG"));
-}
-
-bool UTextureExporterUEJPEG::SupportsTexture(UTexture* Texture) const
-{
-	ETextureSourceFormat TSF = Texture->Source.GetFormat();
-	ERawImageFormat::Type RawFormat = FImageCoreUtils::ConvertToRawImageFormat(TSF);
-	// supports all non-HDR formats :
-	return ! ERawImageFormat::IsHDR(RawFormat);
-}
-
-//========================================================
-
 UTextureExporterEXR::UTextureExporterEXR(const FObjectInitializer& ObjectInitializer)
 	: UTextureExporterGeneric(ObjectInitializer)
 {
@@ -5257,6 +5240,67 @@ bool UTextureExporterJPEG::ExportBinary(UObject* Object, const TCHAR* Type, FArc
 	return true;
 }
 
+//-------------
+// UTextureExporterUEJPEG does not let you compress data to UEJPEG
+//	it only writes out existing UEJPEG data
+// do NOT use UTextureExporterGeneric here, that would go to ImageWrapper and compress to UEJPEG
+// note: does not support UDIM blocks like UTextureExporterGeneric
+
+UTextureExporterUEJPEG::UTextureExporterUEJPEG(const FObjectInitializer& ObjectInitializer)
+	: Super(ObjectInitializer)
+{
+	SupportedClass = UTexture2D::StaticClass();
+	PreferredFormatIndex = 0;
+	FormatExtension.Add(TEXT("UEJ"));
+	FormatDescription.Add(TEXT("UE-JPEG original imported into uasset"));
+}
+
+bool UTextureExporterUEJPEG::SupportsObject(UObject* Object) const
+{
+	if (Super::SupportsObject(Object))
+	{
+		UTexture2D* Texture = Cast<UTexture2D>(Object);
+
+		if (Texture)
+		{
+			// we do NOT do lossy recompression
+			
+			if (Texture->Source.GetNumBlocks() > 1 )
+			{
+				// does not support UDIM
+				return false;
+			}
+
+			// Check it has JPEG BulkData :
+			if ( Texture->Source.GetSourceCompression() == TSCF_UEJPEG && Texture->Source.GetSizeOnDisk() > 0 )
+			{
+				ETextureSourceFormat TSF = Texture->Source.GetFormat();
+				ERawImageFormat::Type RawFormat = FImageCoreUtils::ConvertToRawImageFormat(TSF);
+				check( RawFormat == ERawImageFormat::G8 || RawFormat == ERawImageFormat::BGRA8 );
+
+				return true;
+			}
+		}
+	}
+	return false;
+}
+
+bool UTextureExporterUEJPEG::ExportBinary(UObject* Object, const TCHAR* Type, FArchive& Ar, FFeedbackContext* Warn, int32 FileIndex, uint32 PortFlags)
+{
+	UTexture2D* Texture = Cast<UTexture2D>(Object);
+	check(Texture != nullptr);
+
+	check(Texture->Source.GetSourceCompression() == TSCF_UEJPEG && Texture->Source.GetSizeOnDisk() > 0 );
+
+	// just write the JPEG data we already have :
+	
+	UE_LOG(LogEditorFactories, Display, TEXT("Exporting Texture as UE-JPEG stored bits (no lossy decompress or recompress)."));
+
+	Texture->Source.OperateOnLoadedBulkData( [&](const FSharedBuffer& BulkDataBuffer) {
+		Ar.Serialize(const_cast<void*>(BulkDataBuffer.GetData()), BulkDataBuffer.GetSize());
+	} );
+	return true;
+}
 /*------------------------------------------------------------------------------
 	UTextureExporterTGA implementation.
 ------------------------------------------------------------------------------*/
