Change DLCRecook to be game based so anyone can implement it.
Remove engine implementations for DLCRecook to move into more game specific choices. The overall concept is the same. Check what has changed and only cook those assets.
Remove the assets that are in the base game you want to recook.

diff --git a/Engine/Source/Editor/UnrealEd/Classes/CookOnTheSide/CookOnTheFlyServer.h b/Engine/Source/Editor/UnrealEd/Classes/CookOnTheSide/CookOnTheFlyServer.h
--- a/Engine/Source/Editor/UnrealEd/Classes/CookOnTheSide/CookOnTheFlyServer.h
+++ b/Engine/Source/Editor/UnrealEd/Classes/CookOnTheSide/CookOnTheFlyServer.h
@@ -93,7 +93,6 @@ enum class ECookByTheBookOptions
 	DlcLoadMainAssetRegistry =			0x00020000, // If cooking DLC, populate the main game asset registry
 	ZenStore =							0x00040000, // Store cooked data in Zen Store
 	DlcReevaluateUncookedAssets =		0x00080000, // If cooking DLC, ignore assets in the base asset registry that were not cooked, so that this cook has an opportunity to cook the assets
-	DlcRecook =							0x00100000,	// If cooking DLC, force all the assets in the DLC to be cooked even if they are marked as already cooked.
 };
 ENUM_CLASS_FLAGS(ECookByTheBookOptions);
 
@@ -1273,12 +1272,6 @@ private:
 
 	FString GetMetadataDirectory() const;
 
-	/**
-	 * In a recook build, gets the list of plugins that need to be recooked, which includes DLCName if cooking a DLC
-	 * and all the plugins specified by the -CookPlugins command line parameter.
-	 */
-	void GetPluginsToRecook(TSet<FString>& OutPlugins) const;
-
 	/**
 	 * Is the local CookOnTheFlyServer cooking a Project+Engine+EmbeddedPlugin Release that can be used as a
 	 * base for future DLC or Patches?
diff --git a/Engine/Source/Editor/UnrealEd/Private/Commandlets/CookCommandlet.cpp b/Engine/Source/Editor/UnrealEd/Private/Commandlets/CookCommandlet.cpp
--- a/Engine/Source/Editor/UnrealEd/Private/Commandlets/CookCommandlet.cpp
+++ b/Engine/Source/Editor/UnrealEd/Private/Commandlets/CookCommandlet.cpp
@@ -405,7 +405,6 @@ bool UCookCommandlet::CookByTheBook(const TArray<ITargetPlatform*>& Platforms)
 	CookOptions |= Switches.Contains(TEXT("CookAgainstFixedBase")) ? ECookByTheBookOptions::CookAgainstFixedBase : ECookByTheBookOptions::None;
 	CookOptions |= (Switches.Contains(TEXT("DlcLoadMainAssetRegistry")) || !bErrorOnEngineContentUse) ? ECookByTheBookOptions::DlcLoadMainAssetRegistry : ECookByTheBookOptions::None;
 	CookOptions |= Switches.Contains(TEXT("DlcReevaluateUncookedAssets")) ? ECookByTheBookOptions::DlcReevaluateUncookedAssets : ECookByTheBookOptions::None;
-	CookOptions |= Switches.Contains(TEXT("DlcRecook")) ? ECookByTheBookOptions::DlcRecook : ECookByTheBookOptions::None;
 	bool bCookList = Switches.Contains(TEXT("CookList"));
 
 	if (bCookSinglePackage)
diff --git a/Engine/Source/Editor/UnrealEd/Private/CookOnTheFlyServer.cpp b/Engine/Source/Editor/UnrealEd/Private/CookOnTheFlyServer.cpp
--- a/Engine/Source/Editor/UnrealEd/Private/CookOnTheFlyServer.cpp
+++ b/Engine/Source/Editor/UnrealEd/Private/CookOnTheFlyServer.cpp
@@ -1237,26 +1237,6 @@ FString UCookOnTheFlyServer::GetMetadataDirectory() const
 	return ProjectOrPluginRoot / TEXT("Metadata");
 }
 
-void UCookOnTheFlyServer::GetPluginsToRecook(TSet<FString>& OutPlugins) const
-{
-	OutPlugins.Empty();
-
-	// DLCName, if cooking a DLC.
-	if (IsCookingDLC())
-	{
-		OutPlugins.Add(CookByTheBookOptions->DlcName);
-	}
-
-	// Command line parameter -CookPlugins.
-	TArray<FString> PluginsList;
-	FString CookPluginsStr;
-	if (FParse::Value(FCommandLine::Get(), TEXT("CookPlugins="), CookPluginsStr, false))
-	{
-		CookPluginsStr.ParseIntoArray(PluginsList, TEXT(","));
-	}
-	OutPlugins.Append(PluginsList);
-}
-
 // allow for a command line to start async preloading a Development AssetRegistry if requested
 static FEventRef GPreloadAREvent(EEventMode::ManualReset);
 static FEventRef GPreloadARInfoEvent(EEventMode::ManualReset);
@@ -11740,39 +11720,18 @@ void UCookOnTheFlyServer::RecordDLCPackagesFromBaseGame(FBeginCookContext& Begin
 
 		TArray<FName>& PlatformBasedPackages = CookByTheBookOptions->BasedOnReleaseCookedPackages.FindOrAdd(PlatformName);
 		PlatformBasedPackages.Reset(ActivePackageList.Num());
-
-		// Packages that are present in PlatformBasedPackages will be stripped from the AssetRegistry and pak files generated for this cook.
-		// If we are recooking plugins, make sure packages that live in those plugins are *not* added to PlatformBasedPackages to prevent them from being stripped.
-		TArray<FString> PathsToSkip;
-		if (!!(CookOptions & ECookByTheBookOptions::DlcRecook))
+		for (UE::Cook::FConstructPackageData& PackageData : ActivePackageList)
 		{
-			TSet<FString> Plugins;
-			GetPluginsToRecook(Plugins);
-			for (const FString& Plugin : Plugins)
-			{
-				PathsToSkip.Add(GetMountedAssetPathForPlugin(Plugin));
-			}
+			PlatformBasedPackages.Add(PackageData.NormalizedFileName);
 		}
 
-		for (UE::Cook::FConstructPackageData& PackageData : ActivePackageList)
 		{
-			bool bShouldSkip = false;
-			if (PathsToSkip.Num() > 0)
-			{
-				const FString PackageNameString = PackageData.PackageName.ToString();
-				for (const FString& Path : PathsToSkip)
-				{
-					if (PackageNameString.StartsWith(*Path))
-					{
-						bShouldSkip = true;
-						break;
-					}
-				}
-			}
-
-			if (!bShouldSkip)
+			// allow game or plugins to modify if certain packages from the base game should be recooked.
+			TSet<FName> PackagesToClearCookResults;
+			UAssetManager::Get().ModifyDLCBasePackages(TargetPlatform, PlatformBasedPackages, PackagesToClearCookResults);
+			if (PackagesToClearCookResults.Num())
 			{
-				PlatformBasedPackages.Add(PackageData.NormalizedFileName);
+				PackageDatas->ClearCookResultsForPackages(PackagesToClearCookResults);
 			}
 		}
 		bFirstAddExistingPackageDatas = false;
@@ -11798,17 +11757,6 @@ void UCookOnTheFlyServer::RecordDLCPackagesFromBaseGame(FBeginCookContext& Begin
 			}
 		}
 	}
-
-	if (!!(CookOptions & ECookByTheBookOptions::DlcRecook))
-	{
-		TSet<FString> Plugins;
-		GetPluginsToRecook(Plugins);
-		for (const FString& Plugin : Plugins)
-		{
-			// Mark all the packages in the plugin as not cooked to force them to be cooked again.
-			PackageDatas->ClearCookResultsForPlugin(Plugin);
-		}
-	}
 }
 
 void UCookOnTheFlyServer::BeginCookPackageWriters(FBeginCookContext& BeginContext)
diff --git a/Engine/Source/Editor/UnrealEd/Private/Cooker/CookPackageData.cpp b/Engine/Source/Editor/UnrealEd/Private/Cooker/CookPackageData.cpp
--- a/Engine/Source/Editor/UnrealEd/Private/Cooker/CookPackageData.cpp
+++ b/Engine/Source/Editor/UnrealEd/Private/Cooker/CookPackageData.cpp
@@ -3152,27 +3152,20 @@ void FPackageDatas::ClearCookedPlatforms()
 	});
 }
 
-void FPackageDatas::ClearCookResultsForPlugin(const FString& InPluginName)
+void FPackageDatas::ClearCookResultsForPackages(const TSet<FName>& InPackages)
 {
-	TSharedPtr<IPlugin> Plugin = IPluginManager::Get().FindPlugin(InPluginName);
-	if (Plugin.IsValid())
-	{
-		FString PluginPath = Plugin->GetMountedAssetPath();
-		int AffectedPackagesCount = 0;
-
-		LockAndEnumeratePackageDatas([PluginPath, &AffectedPackagesCount](FPackageData* PackageData)
+	int32 AffectedPackagesCount = 0;
+	LockAndEnumeratePackageDatas([InPackages, &AffectedPackagesCount](FPackageData* PackageData)
+		{
+			const FName& PackageName = PackageData->GetPackageName();
+			if (InPackages.Contains(PackageName))
 			{
-				// Does this PackageData belong to the plugin?
-				const FName& PackageName = PackageData->GetPackageName();
-				if (PackageName.ToString().StartsWith(PluginPath))
-				{
-					PackageData->ClearCookResults();
-					AffectedPackagesCount++;
-				}
-			});
+				PackageData->ClearCookResults();
+				AffectedPackagesCount++;
+			}
+		});
 
-		UE_LOG(LogCook, Display, TEXT("Cleared the cook results of %d packages because plugin '%s' was requested to be recooked."), AffectedPackagesCount, *InPluginName);
-	}
+	UE_LOG(LogCook, Display, TEXT("Cleared the cook results of %d packages because ClearCookResultsForPackages requested them to be recooked."), AffectedPackagesCount);
 }
 
 void FPackageDatas::OnRemoveSessionPlatform(const ITargetPlatform* TargetPlatform)
diff --git a/Engine/Source/Editor/UnrealEd/Private/Cooker/CookPackageData.h b/Engine/Source/Editor/UnrealEd/Private/Cooker/CookPackageData.h
--- a/Engine/Source/Editor/UnrealEd/Private/Cooker/CookPackageData.h
+++ b/Engine/Source/Editor/UnrealEd/Private/Cooker/CookPackageData.h
@@ -1415,8 +1415,8 @@ public:
 	void Clear();
 	/** Set all platforms to not cooked in all PackageDatas. Used to e.g. invalidate previous cooks. */
 	void ClearCookedPlatforms();
-	/** Set all platforms to not cooked in all PackageDatas that belong to the given plugin. */
-	void ClearCookResultsForPlugin(const FString& InPluginName);
+	/** Set all platforms to not cooked in all PackageDatas that are in the given set */
+	void ClearCookResultsForPackages(const TSet<FName>& InPackages);
 	/** Remove all data about the given platform from all PackageDatas and other memory used by *this. */
 	void OnRemoveSessionPlatform(const ITargetPlatform* TargetPlatform);
 
diff --git a/Engine/Source/Runtime/Engine/Classes/Engine/AssetManager.h b/Engine/Source/Runtime/Engine/Classes/Engine/AssetManager.h
--- a/Engine/Source/Runtime/Engine/Classes/Engine/AssetManager.h
+++ b/Engine/Source/Runtime/Engine/Classes/Engine/AssetManager.h
@@ -527,6 +527,13 @@ public:
 	ENGINE_API virtual void ModifyDLCCook(const FString& DLCName, TConstArrayView<const ITargetPlatform*> TargetPlatforms,
 		TArray<FName>& PackagesToCook, TArray<FName>& PackagesToNeverCook);
 
+	/**
+	 * Allows for game code to modify the base packages that have been read in from the DevelopmentAssetRegistry when performing a DLC cook.
+	 * Can be used to modify which packages should be considered to be already cooked.
+	 * Any packages within the PackagesToClearResults will have their cook results cleared and be cooked again if requested by the cooker.
+	 */
+	ENGINE_API virtual void ModifyDLCBasePackages(const ITargetPlatform* TargetPlatform, TArray<FName>& PlatformBasedPackages, TSet<FName>& PackagesToClearResults) const {};
+
 	/**
 	 * If the given package contains a primary asset, get the packages referenced by its AssetBundleEntries.
 	 * Used to inform the cook of should-be-cooked dependencies of PrimaryAssets for PrimaryAssets that
