Asset Action to resize texture sources
shrinks sources above threshold size to make uassets smaller
tries to change LODBias to keep output size the same, but that is in no way guaranteed

#rnx

diff --git a/Engine/Plugins/Editor/EngineAssetDefinitions/Source/Private/AssetDefinition_Texture.cpp b/Engine/Plugins/Editor/EngineAssetDefinitions/Source/Private/AssetDefinition_Texture.cpp
--- a/Engine/Plugins/Editor/EngineAssetDefinitions/Source/Private/AssetDefinition_Texture.cpp
+++ b/Engine/Plugins/Editor/EngineAssetDefinitions/Source/Private/AssetDefinition_Texture.cpp
@@ -17,6 +17,7 @@
 #include "VirtualTexturingEditorModule.h"
 #include "Algo/AnyOf.h"
 #include "AssetRegistry/AssetRegistryHelpers.h"
+#include "TextureAssetActions.h"
 
 #define LOCTEXT_NAMESPACE "UAssetDefinition_Texture"
 
@@ -55,6 +56,14 @@ namespace MenuExtension_Texture
 		);
 	}
 
+	static void ExecuteResizeTextureSource(const FToolMenuContext& InContext)
+	{
+		if (const UContentBrowserAssetContextMenuContext* CBContext = UContentBrowserAssetContextMenuContext::FindContextWithAssets(InContext))
+		{
+			UE::TextureAssetActions::ResizeTextureSource_WithDialog(CBContext->LoadSelectedObjects<UTexture>());
+		}
+	}
+
 	static void ExecuteConvertToVirtualTexture(const FToolMenuContext& InContext)
 	{
 		if (const UContentBrowserAssetContextMenuContext* CBContext = UContentBrowserAssetContextMenuContext::FindContextWithAssets(InContext))
@@ -97,9 +106,11 @@ namespace MenuExtension_Texture
 		UToolMenus::RegisterStartupCallback(FSimpleMulticastDelegate::FDelegate::CreateLambda([]()
 		{
 			FToolMenuOwnerScoped OwnerScoped(UE_MODULE_NAME);
-			UToolMenu* Menu = UE::ContentBrowser::ExtendToolMenu_AssetContextMenu(UTexture::StaticClass());
+
+			{
+				UToolMenu* Menu = UE::ContentBrowser::ExtendToolMenu_AssetContextMenu(UTexture::StaticClass());
 		
-			FToolMenuSection& Section = Menu->FindOrAddSection("GetAssetActions");
+				FToolMenuSection& Section = Menu->FindOrAddSection("GetAssetActions");
 				Section.AddDynamicEntry(NAME_None, FNewToolMenuSectionDelegate::CreateLambda([](FToolMenuSection& InSection)
 				{
 					if (const UContentBrowserAssetContextMenuContext* Context = UContentBrowserAssetContextMenuContext::FindContextWithAssets(InSection))
@@ -111,7 +122,36 @@ namespace MenuExtension_Texture
 							const FToolMenuExecuteAction UIAction = FToolMenuExecuteAction::CreateStatic(&ExecuteCreateMaterial);
 							InSection.AddMenuEntry("Texture_CreateMaterial", Label, ToolTip, Icon, UIAction);
 						}
-											
+															
+						if ( Context->SelectedAssets.Num() == 1 )
+						{
+							const TAttribute<FText> Label = LOCTEXT("Texture_FindMaterials", "Find Materials Using This");
+							const TAttribute<FText> ToolTip = LOCTEXT("Texture_FindMaterialsTooltip", "Finds all materials that use this material in the content browser.");
+							const FSlateIcon Icon = FSlateIcon(FAppStyle::GetAppStyleSetName(), "Icons.Find");
+							const FToolMenuExecuteAction UIAction = FToolMenuExecuteAction::CreateStatic(&ExecuteFindMaterials);
+							InSection.AddMenuEntry("Texture_FindMaterials", Label, ToolTip, Icon, UIAction);
+						}
+						
+						{
+							const TAttribute<FText> Label = LOCTEXT("Texture_ResizeSource", "Resize Texture Source");
+							const TAttribute<FText> ToolTip = LOCTEXT("Texture_ResizeSourceTooltip", "Reduce texture asset size by shrinking the texture source dimensions.");
+							const FSlateIcon Icon = FSlateIcon(FAppStyle::GetAppStyleSetName(), "ClassIcon.Texture2D");
+							const FToolMenuExecuteAction UIAction = FToolMenuExecuteAction::CreateStatic(&ExecuteResizeTextureSource);
+							InSection.AddMenuEntry("Texture_ResizeSource", Label, ToolTip, Icon, UIAction);
+						}
+					}
+				}));
+			}
+
+			{
+				// VT actions should only be on Texture2D , not all UTexture
+				UToolMenu* Menu = UE::ContentBrowser::ExtendToolMenu_AssetContextMenu(UTexture2D::StaticClass());
+		
+				FToolMenuSection& Section = Menu->FindOrAddSection("GetAssetActions");
+				Section.AddDynamicEntry(NAME_None, FNewToolMenuSectionDelegate::CreateLambda([](FToolMenuSection& InSection)
+				{
+					if (const UContentBrowserAssetContextMenuContext* Context = UContentBrowserAssetContextMenuContext::FindContextWithAssets(InSection))
+					{											
 						static const auto CVarVirtualTexturesEnabled = IConsoleManager::Get().FindTConsoleVariableDataInt(TEXT("r.VirtualTextures"));
 						check(CVarVirtualTexturesEnabled);
 						
@@ -130,36 +170,35 @@ namespace MenuExtension_Texture
 									AssetData.GetTagValue<bool>("VirtualTextureStreaming", VirtualTextured);
 									return VirtualTextured;
 								});
+								
+							const bool bHasNonVirtualTextures =
+								Algo::AnyOf(Context->SelectedAssets, [](const FAssetData& AssetData){ 
+									bool VirtualTextured = false;
+									AssetData.GetTagValue<bool>("VirtualTextureStreaming", VirtualTextured);
+									return !VirtualTextured;
+								});
 
 							if (bHasVirtualTextures)
 							{
-								const TAttribute<FText> Label = LOCTEXT("Texture_ConvertToRegular", "Convert to Regular Texture");
+								const TAttribute<FText> Label = LOCTEXT("Texture_ConvertToRegular", "Convert VT to Regular Texture");
 								const TAttribute<FText> ToolTip = LOCTEXT("Texture_ConvertToRegularTooltip", "Converts this texture to a regular 2D texture if it is a virtual texture.");
 								const FSlateIcon Icon = FSlateIcon(FAppStyle::GetAppStyleSetName(), "ClassIcon.Texture2D");
 								const FToolMenuExecuteAction UIAction = FToolMenuExecuteAction::CreateStatic(&ExecuteConvertToRegularTexture);
-								InSection.AddMenuEntry("Texture_ConvertToVT", Label, ToolTip, Icon, UIAction);
+								InSection.AddMenuEntry("Texture_ConvertToRegular", Label, ToolTip, Icon, UIAction);
 							}
 						
-							if (!bHasVirtualTextures)
+							if (bHasNonVirtualTextures)
 							{
 								const TAttribute<FText> Label = LOCTEXT("Texture_ConvertToVT", "Convert to Virtual Texture");
-								const TAttribute<FText> ToolTip = LOCTEXT("Texture_ConvertToVTTooltip", "Converts this texture to a virtual texture if it fits the size limit imposed in the texture importer settings.");
+								const TAttribute<FText> ToolTip = LOCTEXT("Texture_ConvertToVTTooltip", "Converts this texture to a virtual texture if it exceeds the specified size.");
 								const FSlateIcon Icon = FSlateIcon(FAppStyle::GetAppStyleSetName(), "ClassIcon.Texture2D");
 								const FToolMenuExecuteAction UIAction = FToolMenuExecuteAction::CreateStatic(&ExecuteConvertToVirtualTexture);
 								InSection.AddMenuEntry("Texture_ConvertToVT", Label, ToolTip, Icon, UIAction);
 							}
 						}
-						
-						if ( Context->SelectedAssets.Num() == 1 )
-						{
-							const TAttribute<FText> Label = LOCTEXT("Texture_FindMaterials", "Find Materials Using This");
-							const TAttribute<FText> ToolTip = LOCTEXT("Texture_FindMaterialsTooltip", "Finds all materials that use this material in the content browser.");
-							const FSlateIcon Icon = FSlateIcon(FAppStyle::GetAppStyleSetName(), "Icons.Find");
-							const FToolMenuExecuteAction UIAction = FToolMenuExecuteAction::CreateStatic(&ExecuteFindMaterials);
-							InSection.AddMenuEntry("Texture_FindMaterials", Label, ToolTip, Icon, UIAction);
-						}
 					}
 				}));
+			}
 		}));
 	});
 }
diff --git a/Engine/Source/Editor/UnrealEd/Private/TextureAssetActions.cpp b/Engine/Source/Editor/UnrealEd/Private/TextureAssetActions.cpp
new file mode 100644
--- /dev/null
+++ b/Engine/Source/Editor/UnrealEd/Private/TextureAssetActions.cpp
@@ -0,0 +1,707 @@
+// Copyright Epic Games, Inc. All Rights Reserved.
+
+#include "TextureAssetActions.h"
+
+#include "ContentBrowserMenuContexts.h"
+#include "Styling/AppStyle.h"
+#include "IContentBrowserSingleton.h"
+#include "ContentBrowserModule.h"
+
+#include "SlateFwd.h"
+#include "UObject/Object.h"
+#include "Layout/Visibility.h"
+#include "Editor.h"
+#include "Misc/FeedbackContext.h"
+#include "Framework/Application/SlateApplication.h"
+#include "IAssetTools.h"
+#include "AssetRegistry/AssetRegistryHelpers.h"
+
+#include "Widgets/Text/STextBlock.h"
+#include "Widgets/Images/SImage.h"
+#include "Widgets/Layout/SBorder.h"
+#include "Widgets/Layout/SSeparator.h"
+#include "Widgets/Layout/SUniformGridPanel.h"
+#include "Widgets/Input/SButton.h"
+#include "Widgets/Layout/SScrollBox.h"
+#include "Widgets/Layout/SBox.h"
+#include "Widgets/Input/SComboBox.h"
+#include "Widgets/Input/SNumericEntryBox.h"
+
+#include "Engine/Texture.h"
+#include "Interfaces/ITargetPlatformManagerModule.h"
+#include "TextureSourceDataUtils.h"
+
+#define LOCTEXT_NAMESPACE "TextureAssetActions"
+
+struct FTextureAssetActionStatus
+{
+	bool Enabled = true;
+	bool UnderSized = false;
+	bool WrongType = false;
+};
+
+/**
+* STextureActionDlg
+*
+* This class creates and launches a dialog then awaits the result to return to the user.
+*/
+class STextureActionDlg
+{
+public:
+	enum EResult
+	{
+		Cancel = 0,			// No/Cancel, normal usage would stop the current action
+		Confirm = 1,		// Yes/Ok/Etc, normal usage would continue with action
+	};
+
+	STextureActionDlg(const TArray<UTexture*>& Textures);
+
+	/**  Shows the dialog box and waits for the user to respond. */
+	EResult ShowModal();
+
+private:
+	TSharedPtr<SWindow> DialogWindow;
+	TSharedPtr<class STextureAssetList> DialogWidget;
+};
+
+class STextureAssetList : public SCompoundWidget
+{
+public:
+	
+public:
+
+	SLATE_BEGIN_ARGS(STextureAssetList)
+	{}
+	/** Window in which this widget resides */
+	SLATE_ATTRIBUTE(TSharedPtr<SWindow>, ParentWindow)
+	SLATE_END_ARGS()
+
+	void Construct(const FArguments& InArgs);
+	void Init(const TArray<UTexture *> &Textures);
+
+	STextureActionDlg::EResult GetUserResponse() const;
+
+private:
+
+	void UpdateList();
+	void DoAction();
+
+	/**
+	* Creates a single line showing an asset and it's status related to VT conversion
+	*/
+	TSharedRef<SWidget> CreateAssetLine(int index, const FAssetData &Asset, const FTextureAssetActionStatus &Status);
+
+	FReply OnButtonClick(STextureActionDlg::EResult ButtonID);
+
+	void OnThresholdChanged(TSharedPtr<int32> InSelectedItem, ESelectInfo::Type SelectInfo);
+
+	FText GetThresholdText() const;
+
+	TSharedRef<SWidget> OnGenerateThresholdWidget(TSharedPtr<int32> InItem);
+
+	FReply OnExpanderClicked(int index);
+
+	bool GetOkButtonEnabled() const;
+	
+	EVisibility GetDetailVisibility(int index) const;
+
+	const FSlateBrush* GetExpanderImage(int index) const;
+
+	EVisibility GetIntroMessageVisibility() const;
+
+	EVisibility GetErrorMessageVisibility() const;
+
+	EVisibility GetThresholdVisibility() const;
+
+	FText GetIntroMessage() const;
+
+	FText GetErrorMessage() const;
+	
+	TArray<int> ExpandedIndexes;
+
+	FText IntroMessage;
+	FText ErrorMessage;
+	STextureActionDlg::EResult	 UserResponse;
+
+	TSharedPtr<SVerticalBox>	 AssetListContainer;
+	//TSharedPtr<STextBlock>		 MessageTextBlock;
+	//TSharedPtr<SHorizontalBox>	 ThresholdContainer;
+	//TSharedPtr<SHorizontalBox>	 ErrorContainer;
+	//TSharedPtr<SButton>			 OkButton;
+
+	/** Pointer to the window which holds this Widget, required for modal control */
+	TSharedPtr<SWindow>			 ParentWindow;
+
+	struct TextureListEntry
+	{
+		UTexture * Texture;
+		bool Enabled;
+	};
+
+	TArray<TextureListEntry> TextureList;
+	TArray<FAssetData> AssetList;
+	TArray<FTextureAssetActionStatus> AssetStatus;
+
+	int ThresholdValue;
+	bool bThresholdVisible;
+
+	TArray<TSharedPtr<int32>> TextureSizes;
+};
+
+
+void STextureAssetList::Construct(const FArguments& InArgs)
+{
+	IntroMessage = LOCTEXT("TAA_Intro", "Reduce size of Texture Source to compact uassets.  Resizing is done using mip filter.  LODBias is adjusted but platform built size may change.");
+
+	UserResponse = STextureActionDlg::Cancel;
+	ParentWindow = InArgs._ParentWindow.Get();
+	static FName ErrorIcon = "MessageLog.Error";
+
+
+	for (int i = 0; i < 16; i++)
+	{
+		TextureSizes.Add(MakeShareable(new int32(1 << i)));
+	}
+	ThresholdValue = *TextureSizes[10];
+
+	this->ChildSlot[
+		SNew(SVerticalBox)
+		// Textbox at the top giving an introductory message
+		+ SVerticalBox::Slot()
+			.AutoHeight()
+			.Padding(8.0f, 4.0f, 8.0f, 4.0f)
+			[
+				SNew(STextBlock)
+				.AutoWrapText(true)
+				.Visibility(this, &STextureAssetList::GetIntroMessageVisibility)
+				.Text(this, &STextureAssetList::GetIntroMessage)
+			]
+		// Error message at the top giving a common error message
+		+ SVerticalBox::Slot()
+			.AutoHeight()
+			.Padding(8.0f, 4.0f, 8.0f, 4.0f)
+			[
+				SNew(SHorizontalBox)
+				.Visibility(this, &STextureAssetList::GetErrorMessageVisibility)
+				+ SHorizontalBox::Slot()
+				.AutoWidth()
+				.VAlign(VAlign_Center)
+				.HAlign(HAlign_Center)
+				[
+					SNew(SImage)
+					.Image(FAppStyle::GetBrush(ErrorIcon))
+				]
+				+ SHorizontalBox::Slot()
+				.AutoWidth()
+				.VAlign(VAlign_Center)
+				.HAlign(HAlign_Center)
+				[
+					SNew(STextBlock)
+					.AutoWrapText(true)
+					.Text(this, &STextureAssetList::GetErrorMessage)
+				]
+			]
+		+ SVerticalBox::Slot()
+			.AutoHeight()
+			.Padding(8.0f, 4.0f, 8.0f, 4.0f)
+			[
+				SNew(SSeparator)
+			]
+		// The actual list of assets
+		+ SVerticalBox::Slot()
+			.Padding(8.0f, 4.0f, 8.0f, 4.0f)
+			[
+				SNew(SBorder)
+				[
+					SNew(SScrollBox)
+					+ SScrollBox::Slot()
+					[
+						SAssignNew(AssetListContainer, SVerticalBox)
+					]
+				]
+			]
+		// The bottom row of widgets: texture size selector
+		+ SVerticalBox::Slot()
+			.AutoHeight()
+			.HAlign(HAlign_Left)
+			.Padding(8.0f, 4.0f, 8.0f, 4.0f)
+			[
+				SNew(SHorizontalBox)
+				+SHorizontalBox::Slot()
+				.AutoWidth()
+				.VAlign(VAlign_Center)
+				[
+					SNew(STextBlock)
+					.Text(LOCTEXT("TAA_Size", "Texture size threshold: "))
+					.AutoWrapText(true)
+				]
+				+ SHorizontalBox::Slot()
+				.AutoWidth()
+				.VAlign(VAlign_Center)
+				[
+					SNew(SComboBox<TSharedPtr<int32>>)
+					.OptionsSource(&TextureSizes)
+					.OnSelectionChanged(this, &STextureAssetList::OnThresholdChanged)
+					.OnGenerateWidget(this, &STextureAssetList::OnGenerateThresholdWidget)
+					.InitiallySelectedItem(TextureSizes[10])
+					[
+						SNew(STextBlock)
+						.Text(this, &STextureAssetList::GetThresholdText)
+					]
+				]
+			]
+		// Separator
+		+ SVerticalBox::Slot()
+			.AutoHeight()
+			.Padding(8.0f, 4.0f, 8.0f, 4.0f)
+			[
+				SNew(SSeparator)
+			]
+		// Dialog ok/cancel buttons
+		+ SVerticalBox::Slot()
+			.AutoHeight()
+			.HAlign(HAlign_Right)
+			.Padding(8.0f, 4.0f, 8.0f, 4.0f)
+			[
+				SNew(SUniformGridPanel)
+				.SlotPadding(FAppStyle::GetMargin("StandardDialog.SlotPadding"))
+				.MinDesiredSlotWidth(FAppStyle::GetFloat("StandardDialog.MinDesiredSlotWidth"))
+				.MinDesiredSlotHeight(FAppStyle::GetFloat("StandardDialog.MinDesiredSlotHeight"))
+				+ SUniformGridPanel::Slot(0, 0)
+				[
+					SNew(SButton)
+					.HAlign(HAlign_Center)
+					.ContentPadding(FAppStyle::GetMargin("StandardDialog.ContentPadding"))
+					.OnClicked(this, &STextureAssetList::OnButtonClick, STextureActionDlg::Confirm)
+					.IsEnabled(this, &STextureAssetList::GetOkButtonEnabled)
+					.Text(LOCTEXT("TAA_OK", "OK"))
+				]
+				+ SUniformGridPanel::Slot(1, 0)
+				[
+					SNew(SButton)
+					.HAlign(HAlign_Center)
+					.ContentPadding(FAppStyle::GetMargin("StandardDialog.ContentPadding"))
+					.OnClicked(this, &STextureAssetList::OnButtonClick, STextureActionDlg::Cancel)
+					.Text(LOCTEXT("TAA_Cancel", "Cancel"))
+				]
+			]
+	];
+
+	// will be done by Init :
+	//UpdateList();
+}
+
+void STextureAssetList::Init(const TArray<UTexture *> &Textures)
+{
+	TextureList.SetNum(Textures.Num());
+	for(int i=0;i<Textures.Num();i++)
+	{
+		TextureList[i].Texture = Textures[i];
+		TextureList[i].Enabled = true;
+	}
+
+	UpdateList();
+}
+
+STextureActionDlg::EResult STextureAssetList::GetUserResponse() const
+{
+	return UserResponse;
+}
+
+TSharedRef<SWidget> STextureAssetList::CreateAssetLine(int index, const FAssetData &Asset, const FTextureAssetActionStatus &Status)
+{
+	//const bool bEngineAsset = Asset.PackagePath.ToString().StartsWith(TEXT("/Engine/")); // use FPackageName::SplitPackageNameRoot
+
+	FName SeverityIcon = NAME_None;
+	FText DetailedInfoText;
+
+	if (Status.WrongType)
+	{
+		SeverityIcon = "MessageLog.Error";
+		DetailedInfoText = LOCTEXT("TAA_WrongType", "The texture is not a supported type.");
+	}
+	else if (Status.UnderSized)
+	{
+		SeverityIcon = "MessageLog.Note";
+		DetailedInfoText = LOCTEXT("TAA_UnderSized", "The texture was under the threshold size.");
+	}
+	/*
+	// ConvertToVT does this, currently we do not
+	else if (bEngineAsset) // @@??
+	{
+		SeverityIcon = "MessageLog.Note";
+		DetailedInfoText = LOCTEXT("TAA_EngineAsset", "The texture is an engine asset, a copy will be created in the current project.");
+	}
+	*/
+
+	auto Result =
+		SNew(SVerticalBox)
+		+ SVerticalBox::Slot()
+		.AutoHeight()
+		.VAlign(VAlign_Center)
+		.HAlign(HAlign_Left)
+		[
+			// Class icon
+			SNew(SHorizontalBox)
+			+ SHorizontalBox::Slot()
+			.AutoWidth()
+			.VAlign(VAlign_Center)
+			.HAlign(HAlign_Center)
+			[
+				SNew(SBox)
+				.Padding(2.0f)
+				[
+					SNew(SImage)
+					.Image(FSlateIcon(FAppStyle::GetAppStyleSetName(),
+						"ClassIcon.Texture2D" ).GetIcon())
+				]
+			]
+			// Error/warning icon
+			+ SHorizontalBox::Slot()
+			.AutoWidth()
+			.HAlign(HAlign_Center)
+			.VAlign(VAlign_Center)
+			[
+				SNew(SBox)
+				.Padding(2.0f)
+				[
+					(SeverityIcon == NAME_None) ? SNullWidget::NullWidget :
+					static_cast<TSharedRef<SWidget>>(SNew(SImage).Image(FAppStyle::GetBrush(SeverityIcon)))
+				]
+			]
+			// Fold out button with asset name
+			+ SHorizontalBox::Slot()
+			.AutoWidth()
+			.HAlign(HAlign_Center)
+			.VAlign(VAlign_Center)
+			[
+					(DetailedInfoText.IsEmpty())
+					? static_cast<TSharedRef<SWidget>>(SNew(STextBlock)
+						.Text(FText::FromString(Asset.GetObjectPathString())))
+				: SNew(SButton)
+				.ButtonStyle(FCoreStyle::Get(), "NoBorder")
+				.VAlign(VAlign_Center)
+				.HAlign(HAlign_Center)
+				.ClickMethod(EButtonClickMethod::MouseDown)
+				.OnClicked(this, &STextureAssetList::OnExpanderClicked, index)
+				.ContentPadding(0.f)
+				.ForegroundColor(FSlateColor::UseForeground())
+				.IsFocusable(false)
+				//.Text(FText::FromName(Asset.ObjectPath))
+				[
+					// Fold out icon
+					SNew(SHorizontalBox)
+					+ SHorizontalBox::Slot()
+					.AutoWidth()
+					.VAlign(VAlign_Center)
+					.HAlign(HAlign_Center)
+					[
+						SNew(SImage)
+						.Image(this, &STextureAssetList::GetExpanderImage, index)
+						.ColorAndOpacity(FSlateColor::UseForeground())
+					]
+					// Fold out text
+					+ SHorizontalBox::Slot()
+					.AutoWidth()
+					.VAlign(VAlign_Center)
+					[
+						SNew(STextBlock)
+						.Text(FText::FromString(Asset.GetObjectPathString()))
+						.ColorAndOpacity( Status.Enabled ? FSlateColor::UseForeground() : FSlateColor::UseSubduedForeground())
+					] // change color for enabled/not
+				]
+			]
+		]
+		+ SVerticalBox::Slot()
+		.AutoHeight()
+		.VAlign(VAlign_Center)
+		.HAlign(HAlign_Left)
+		[
+			SNew(STextBlock)
+			.Text(DetailedInfoText)
+			.Visibility(this, &STextureAssetList::GetDetailVisibility, index)
+		]
+		;
+	return Result;
+}
+
+static void DoResizeTextureSource(UTexture * Texture,int TargetSize)
+{
+	// we do the resizing considering only mip/LOD/build settings for the running Editor platform (eg. Windows)
+	const ITargetPlatform* RunningPlatform = GetTargetPlatformManagerRef().GetRunningTargetPlatform();
+	
+	int32 BeforeSizeX;
+	int32 BeforeSizeY;
+	Texture->GetBuiltTextureSize(RunningPlatform, BeforeSizeX, BeforeSizeY);
+
+	if ( ! UE::TextureUtilitiesCommon::Experimental::DownsizeTextureSourceData(Texture, TargetSize, RunningPlatform))
+	{
+		UE_LOG(LogTexture, Display, TEXT("Texture (%s) did not resize."), *Texture->GetName());
+
+		// did not resize, but may have done PreEditChange
+		return;
+	}
+
+	UE_LOG(LogTexture, Display, TEXT("Texture (%s) did resize."), *Texture->GetName());
+
+	Texture->LODBias = 0;
+	
+	int32 AfterSizeX;
+	int32 AfterSizeY;
+	Texture->GetBuiltTextureSize(RunningPlatform, AfterSizeX, AfterSizeY);
+
+	// if AfterSize > BeforeSize , kick up LODBias
+	//	to try to preserve GetBuiltTextureSize
+	while( AfterSizeX > BeforeSizeX || AfterSizeY > BeforeSizeY )
+	{
+		Texture->LODBias ++;
+		// just shifting down AfterSize is not exactly right
+		//	but ensures that our loop terminates
+		AfterSizeX = (AfterSizeX+1)>>1;
+		AfterSizeY = (AfterSizeY+1)>>1;
+	}
+	
+	// recompute AfterSize if we changed LODBias
+	if ( Texture->LODBias != 0 )
+	{
+		Texture->GetBuiltTextureSize(RunningPlatform, AfterSizeX, AfterSizeY);
+	}
+	
+	if ( BeforeSizeX != AfterSizeX || BeforeSizeY != AfterSizeY )
+	{
+		// not a warning, just FYI
+		// changing built size is totally possible and expected to happen sometimes
+		UE_LOG(LogTexture,Display,TEXT("DoResizeTextureSource failed to preserve built size; was: %dx%d now: %dx%d on [%s]"),
+			BeforeSizeX,BeforeSizeY,
+			AfterSizeX,AfterSizeY,
+			*Texture->GetFullName());
+	}
+
+	// DownsizeTextureSourceData did the PreEditChange
+	Texture->PostEditChange();
+}
+
+void STextureAssetList::DoAction()
+{
+	int NumEnabled = 0;
+	for (TextureListEntry & Entry : TextureList)
+	{
+		UTexture * Texture = Entry.Texture;
+
+		UE_LOG(LogTexture, Display, TEXT("Texture (%s) Enabled=%d"), *Texture->GetName() , (int)Entry.Enabled);
+
+		if ( Entry.Enabled )
+		{
+			NumEnabled ++;	
+		}
+	}
+
+	if ( NumEnabled == 0 )
+	{
+		return;
+	}
+
+	FScopedSlowTask Progress(NumEnabled, LOCTEXT("ResizingTextures", "Resizing Textures ..."));
+	Progress.MakeDialog(/*ShowCancelButton*/true);
+
+	for (TextureListEntry & Entry : TextureList)
+	{
+		if ( ! Entry.Enabled ) continue;
+
+		UTexture * Texture = Entry.Texture;
+
+		UE_LOG(LogTexture, Display, TEXT("Texture (%s) Resizing to <= %d"), *Texture->GetName() , ThresholdValue);
+				
+		Progress.EnterProgressFrame(1.f);
+		if (Progress.ShouldCancel())
+		{
+			break;
+		}
+
+		DoResizeTextureSource(Texture,ThresholdValue);
+	}
+}
+
+void STextureAssetList::UpdateList()
+{
+	AssetList.Empty();
+	AssetStatus.Empty();
+
+	// filter select textures to see if they should be acted on
+	for (TextureListEntry & Entry : TextureList)
+	{
+		UTexture * Texture = Entry.Texture;
+
+		AssetList.Add(Texture);
+		FTextureAssetActionStatus* Status = new(AssetStatus) FTextureAssetActionStatus();
+		
+		Entry.Enabled = true;
+
+		ETextureClass Class = Texture->GetTextureClass();
+		if ( Class != ETextureClass::TwoD ) // @@ also ::Cube , Array ?
+		{
+			Status->WrongType = true;
+			Entry.Enabled = false;
+		}
+		
+		FIntPoint SourceSize = Texture->Source.GetLogicalSize();
+		int MaxSize = FMath::Max(SourceSize.X,SourceSize.Y);
+		if ( MaxSize <= ThresholdValue )
+		{
+			Entry.Enabled = false;
+			Status->UnderSized = true;
+		}
+
+		Status->Enabled = Entry.Enabled;
+	}
+
+	check(AssetList.Num() == AssetStatus.Num());
+	AssetListContainer->ClearChildren();
+
+	for (int Id = 0; Id < AssetList.Num(); Id++)
+	{
+		AssetListContainer->AddSlot()
+			.AutoHeight()
+			[
+				CreateAssetLine(Id, AssetList[Id], AssetStatus[Id])
+			];
+	}
+
+	ErrorMessage = FText();
+}
+
+FReply STextureAssetList::OnButtonClick(STextureActionDlg::EResult ButtonID)
+{
+	ParentWindow->RequestDestroyWindow();
+	UserResponse = ButtonID;
+
+	if (ButtonID == STextureActionDlg::EResult::Confirm)
+	{
+		DoAction();
+	}
+
+	return FReply::Handled();
+}
+
+void STextureAssetList::OnThresholdChanged(TSharedPtr<int32> InSelectedItem, ESelectInfo::Type SelectInfo)
+{
+	ThresholdValue = *InSelectedItem;
+	UpdateList();
+}
+
+FText STextureAssetList::GetThresholdText() const
+{
+	return FText::FromString(FString::Format(TEXT("{0}"), TArray<FStringFormatArg>({ ThresholdValue })));
+}
+
+TSharedRef<SWidget> STextureAssetList::OnGenerateThresholdWidget(TSharedPtr<int32> InItem)
+{
+	return SNew(STextBlock)
+		.Text(FText::FromString(FString::Format(TEXT("{0}"), TArray<FStringFormatArg>({ *InItem }))));
+}
+
+FReply STextureAssetList::OnExpanderClicked(int index)
+{
+	if (ExpandedIndexes.Contains(index))
+	{
+		ExpandedIndexes.Remove(index);
+	}
+	else
+	{
+		ExpandedIndexes.Add(index);
+	}
+	return FReply::Handled();
+}
+
+bool STextureAssetList::GetOkButtonEnabled() const
+{
+	return ErrorMessage.IsEmpty();
+}
+
+EVisibility STextureAssetList::GetDetailVisibility(int index) const
+{
+	return (ExpandedIndexes.Contains(index)) ? EVisibility::Visible : EVisibility::Collapsed;
+}
+
+const FSlateBrush* STextureAssetList::GetExpanderImage(int index) const
+{
+	FName ResourceName;
+	if (GetDetailVisibility(index) == EVisibility::Visible)
+	{
+		static FName ExpandedName = "TreeArrow_Expanded";
+		ResourceName = ExpandedName;
+	}
+	else
+	{
+		static FName CollapsedName = "TreeArrow_Collapsed";
+		ResourceName = CollapsedName;
+	}
+	return FCoreStyle::Get().GetBrush(ResourceName);
+}
+
+EVisibility STextureAssetList::GetIntroMessageVisibility() const
+{
+	return (IntroMessage.IsEmpty()) ? EVisibility::Collapsed : EVisibility::Visible;
+}
+
+EVisibility STextureAssetList::GetErrorMessageVisibility() const
+{
+	return (ErrorMessage.IsEmpty()) ? EVisibility::Collapsed : EVisibility::Visible;
+}
+
+EVisibility STextureAssetList::GetThresholdVisibility() const
+{
+	return (bThresholdVisible) ? EVisibility::Collapsed : EVisibility::Visible;
+}
+
+FText STextureAssetList::GetIntroMessage() const
+{
+	return IntroMessage;
+}
+
+FText STextureAssetList::GetErrorMessage() const
+{
+	return ErrorMessage;
+}
+
+STextureActionDlg::STextureActionDlg(const TArray<UTexture *> &Textures)
+{
+	if (FSlateApplication::IsInitialized())
+	{
+		DialogWindow = SNew(SWindow)
+			.Title(LOCTEXT("TAA_Title", "Texture Asset : Resize Source"))
+			.SupportsMinimize(false).SupportsMaximize(false)
+			.ClientSize(FVector2D(500, 500));
+
+		TSharedPtr<SBorder> DialogWrapper =
+			SNew(SBorder)
+			.BorderImage(FAppStyle::GetBrush("ToolPanel.GroupBorder"))
+			.Padding(4.0f)
+			[
+				SAssignNew(DialogWidget, STextureAssetList)
+				.ParentWindow(DialogWindow)
+			];
+
+		// Init will do UpdateList :
+		DialogWidget->Init(Textures);
+		DialogWindow->SetContent(DialogWrapper.ToSharedRef());
+	}
+}
+
+STextureActionDlg::EResult STextureActionDlg::ShowModal()
+{
+	//Show Dialog
+	GEditor->EditorAddModalWindow(DialogWindow.ToSharedRef());
+	EResult UserResponse = (EResult)DialogWidget->GetUserResponse();
+	DialogWindow->GetParentWindow()->RemoveDescendantWindow(DialogWindow.ToSharedRef());
+	return UserResponse;
+}
+
+
+void UE::TextureAssetActions::ResizeTextureSource_WithDialog(const TArray<UTexture*> & InTextures)
+{
+	STextureActionDlg Dlg(InTextures);
+	Dlg.ShowModal();
+}
+
+#undef LOCTEXT_NAMESPACE
diff --git a/Engine/Source/Editor/UnrealEd/Public/TextureAssetActions.h b/Engine/Source/Editor/UnrealEd/Public/TextureAssetActions.h
new file mode 100644
--- /dev/null
+++ b/Engine/Source/Editor/UnrealEd/Public/TextureAssetActions.h
@@ -0,0 +1,17 @@
+// Copyright Epic Games, Inc. All Rights Reserved.
+
+#pragma once
+
+#include "CoreMinimal.h"
+
+class UTexture;
+
+namespace UE
+{
+namespace TextureAssetActions
+{
+
+UNREALED_API void ResizeTextureSource_WithDialog(const TArray<UTexture*> & InTextures);
+
+}
+};
diff --git a/Engine/Source/Editor/VirtualTexturingEditor/Private/SConvertToVirtualTexture.cpp b/Engine/Source/Editor/VirtualTexturingEditor/Private/SConvertToVirtualTexture.cpp
--- a/Engine/Source/Editor/VirtualTexturingEditor/Private/SConvertToVirtualTexture.cpp
+++ b/Engine/Source/Editor/VirtualTexturingEditor/Private/SConvertToVirtualTexture.cpp
@@ -436,12 +436,31 @@ void SConvertToVirtualTexture::UpdateList()
 	AssetList.Empty();
 	AssetStatus.Empty();
 
+	auto CheckUnderSized = [] (const UTexture2D * Texture, int SizeThreshold, bool bConvertBackward)
+	{
+		// code dupe from FVirtualTextureConversionWorker::FilterList
+		
+		bool DoInclude;
+
+		// don't use Texture->GetSizeX() as it can be Default texture
+		FIntPoint Size = Texture->Source.GetLogicalSize();
+		uint64 TexturePixelCount = (uint64) Size.X * Size.Y;
+
+		DoInclude = ( TexturePixelCount >= (uint64)SizeThreshold * SizeThreshold );
+			
+		// for Backwards this should be the other way around
+		//	we want to filter textures *Smaller* than SizeThreshold
+		if ( bConvertBackward ) DoInclude = ! DoInclude;
+
+		return ! DoInclude;
+	};
+
 	for (UTexture2D *Texture : Worker.Textures)
 	{
 		AssetList.Add(Texture);
 		FConversionStatus* Status = new(AssetStatus) FConversionStatus();
 		Status->UserSelected = Worker.UserTextures.Contains(Texture);
-		Status->UnderSized = (Texture->GetSizeX()*Texture->GetSizeY() < ThresholdValue*ThresholdValue);
+		Status->UnderSized = CheckUnderSized(Texture,ThresholdValue,bBackwards);
 		Status->NonPowerOf2 = !Texture->Source.IsPowerOfTwo() && Texture->PowerOfTwoMode == ETexturePowerOfTwoSetting::None;
 	}
 
@@ -450,7 +469,7 @@ void SConvertToVirtualTexture::UpdateList()
 		AssetList.Add(Texture);
 		FConversionStatus* Status = new(AssetStatus) FConversionStatus();
 		Status->UserSelected = Worker.UserTextures.Contains(Texture);
-		Status->UnderSized = (Texture->GetSizeX()*Texture->GetSizeY() < ThresholdValue*ThresholdValue);
+		Status->UnderSized = CheckUnderSized(Texture,ThresholdValue,bBackwards);
 		Status->NonPowerOf2 = !Texture->Source.IsPowerOfTwo() && Texture->PowerOfTwoMode == ETexturePowerOfTwoSetting::None;
 		Status->InvalidMaterialUsage = true;
 	}
@@ -603,7 +622,9 @@ FConvertToVTDlg::FConvertToVTDlg(const TArray<UTexture2D *> &Textures, bool bBac
 	if (FSlateApplication::IsInitialized())
 	{
 		DialogWindow = SNew(SWindow)
-			.Title((bBackwards) ? LOCTEXT("ConvertToVTDlgTitle_Backwards", "Convert From VT") : LOCTEXT("ConvertToVTDlgTitle", "Convert To VT"))
+			.Title((bBackwards) ? 
+				LOCTEXT("ConvertToVTDlgTitle_Backwards", "Convert VT to Regular if < Threshold") : 
+				LOCTEXT("ConvertToVTDlgTitle", "Convert To VT Textures >= Threshold"))
 			.SupportsMinimize(false).SupportsMaximize(false)
 			.ClientSize(FVector2D(500, 500));
 
diff --git a/Engine/Source/Editor/VirtualTexturingEditor/Private/VirtualTextureConversionWorker.cpp b/Engine/Source/Editor/VirtualTexturingEditor/Private/VirtualTextureConversionWorker.cpp
--- a/Engine/Source/Editor/VirtualTexturingEditor/Private/VirtualTextureConversionWorker.cpp
+++ b/Engine/Source/Editor/VirtualTexturingEditor/Private/VirtualTextureConversionWorker.cpp
@@ -97,7 +97,6 @@ void FVirtualTextureConversionWorker::AddReferencedObjects( FReferenceCollector&
 	Collector.AddReferencedObjects( Materials );
 	Collector.AddReferencedObjects( Functions );
 	Collector.AddReferencedObjects( SizeRejectedTextures );
-	Collector.AddReferencedObjects( SizeRejectedTextures );
 	Collector.AddReferencedObjects( MaterialRejectedTextures );
 
 	for (auto& KV : AuditTrail)
@@ -482,7 +481,27 @@ void FVirtualTextureConversionWorker::FilterList(int32 SizeThreshold)
 
 	for (UTexture2D *Texture : UserTextures)
 	{
-		if (Texture->GetSizeX()*Texture->GetSizeY() >= SizeThreshold * SizeThreshold)
+		bool DoInclude;
+
+		if ( !Texture->Source.IsPowerOfTwo() && Texture->PowerOfTwoMode == ETexturePowerOfTwoSetting::None )
+		{
+			// not pow2, reject
+			DoInclude = false;
+		}
+		else 
+		{
+			// don't use Texture->GetSizeX() as it can be Default texture
+			FIntPoint Size = Texture->Source.GetLogicalSize();
+			uint64 TexturePixelCount = (uint64) Size.X * Size.Y;
+
+			DoInclude = ( TexturePixelCount >= (uint64)SizeThreshold * SizeThreshold );
+			
+			// for Backwards this should be the other way around
+			//	we want to filter textures *Smaller* than SizeThreshold
+			if ( bConvertBackward ) DoInclude = ! DoInclude;
+		}
+
+		if ( DoInclude )
 		{
 			Textures.Add(Texture);
 		}
diff --git a/Engine/Source/Runtime/Engine/Private/TextureDerivedData.cpp b/Engine/Source/Runtime/Engine/Private/TextureDerivedData.cpp
--- a/Engine/Source/Runtime/Engine/Private/TextureDerivedData.cpp
+++ b/Engine/Source/Runtime/Engine/Private/TextureDerivedData.cpp
@@ -3937,7 +3937,7 @@ bool UTexture::DownsizeImageUsingTextureSettings(const ITargetPlatform* TargetPl
 {
 	// resize so that the largest dimension is <= TargetSize
 
-	if (TargetSize < 1 || LayerIndex < 0 || InOutImage.IsImageInfoValid() == false)
+	if (TargetSize <= 1 || LayerIndex < 0 || InOutImage.IsImageInfoValid() == false)
 	{
 		UE_LOG(LogTexture, Error, TEXT("Invalid parameter supplied to DownsizeImageUsingTextureSettings target size = %d layer index = %d image valid: %s"),
 			TargetSize, LayerIndex, InOutImage.IsImageInfoValid() ? TEXT("true") : TEXT("false"));
@@ -3947,6 +3947,7 @@ bool UTexture::DownsizeImageUsingTextureSettings(const ITargetPlatform* TargetPl
 	if (TargetSize >= InOutImage.SizeX && TargetSize >= InOutImage.SizeY)
 	{
 		// both dimensions already small enough, early out
+		// InOutImage is not changed
 		return true;
 	}
 
diff --git a/Engine/Source/Runtime/TextureUtilitiesCommon/Private/TextureSourceDataUtils.cpp b/Engine/Source/Runtime/TextureUtilitiesCommon/Private/TextureSourceDataUtils.cpp
--- a/Engine/Source/Runtime/TextureUtilitiesCommon/Private/TextureSourceDataUtils.cpp
+++ b/Engine/Source/Runtime/TextureUtilitiesCommon/Private/TextureSourceDataUtils.cpp
@@ -14,12 +14,16 @@ namespace UE::TextureUtilitiesCommon::Experimental
 
 namespace Private
 {
+
+	// resize so that the largest dimension is <= MaxSize
 	bool ResizeTexture2D(UTexture* Texture, int32 MaxSize, const ITargetPlatform* TargetPlatform)
 	{
 		// Protect the code from an async build of the texture
+		// @@ I think that this PreEditChange can (and should) be moved to after DownsizeImageUsingTextureSettings
 		Texture->PreEditChange(nullptr);
 
 		// We want to reduce the asset size so ignore the imported mip(s)
+		//	@@ should not do this if MipGen == LeaveExisting
 		const int32 MipIndex = 0;
 		FImage SourceMip0;
 		if (!Texture->Source.GetMipImage(SourceMip0, MipIndex))
@@ -28,6 +32,8 @@ namespace Private
 		}
 
 		const int32 LayerIndex = 0;
+		// DownsizeImageUsingTextureSettings is not const but I think it can and should be?
+		//	(the Texture itself is not changed, only the FImage argument is changed)
 		if ( ! Texture->DownsizeImageUsingTextureSettings(TargetPlatform, SourceMip0, MaxSize, LayerIndex) )
 		{
 			return false;
@@ -54,6 +60,10 @@ namespace Private
 
 	bool ResizeTextureSlicedBy2DLayers(UTexture* Texture, int32 MaxSize, const ITargetPlatform* TargetPlatform)
 	{
+		// should check Source size vs MaxSize and early return here
+		if ( Texture->Source.GetSizeX() <= MaxSize && Texture->Source.GetSizeY() <= MaxSize )
+			return false;
+				
 		// Protect the code from an async build of the texture
 		Texture->PreEditChange(nullptr);
 
@@ -67,6 +77,7 @@ namespace Private
 
 			{
 				// We want to reduce the asset size so ignore the imported mip(s)
+				//	@@ should not do this if MipGen == LeaveExisting
 				TArray<FImage> Slices;
 				Slices.Reserve(Texture->Source.GetNumSlices());
 
@@ -95,7 +106,10 @@ namespace Private
 					FImage& ResizedSlice = ResizedSlices.AddDefaulted_GetRef();
 
 					const int32 LayerIndex = 0;
-					Texture->DownsizeImageUsingTextureSettings(TargetPlatform, Slice, MaxSize, LayerIndex);
+					if ( ! Texture->DownsizeImageUsingTextureSettings(TargetPlatform, Slice, MaxSize, LayerIndex) )
+					{
+						return false;
+					}
 					ResizedSlice = MoveTemp(Slice);
 				}
 			}
@@ -128,10 +142,18 @@ namespace Private
 		// Protect the code from an async build of the texture
 		Texture->PreEditChange(nullptr);
 
+		/*
 		FIntPoint LogicalSourceSize = Texture->Source.GetLogicalSize();
 		double RatioX = double(MaxSize) / LogicalSourceSize.X;
 		double RatioY = double(MaxSize) / LogicalSourceSize.Y;
 
+		// early return if we're not shrinking
+		if ( RatioX >= 1.0 && RatioY >= 1.0 )
+		{
+			return false;
+		}
+		*/
+
 		TArray<FTextureSourceBlock > ResizedSourceBlocks;
 		ResizedSourceBlocks.Reserve(Texture->Source.GetNumBlocks());
 	
@@ -154,9 +176,13 @@ namespace Private
 		
 			FImage& ResizedBlock = ResizedBlocks.AddDefaulted_GetRef();
 
-			int32 BlockMaxSize = FMath::RoundToInt32(FMath::Min(ResizedSourceBlock.SizeX * RatioX, ResizedSourceBlock.SizeY * RatioY));
+			//int32 BlockMaxSize = FMath::RoundToInt32(FMath::Min(ResizedSourceBlock.SizeX * RatioX, ResizedSourceBlock.SizeY * RatioY));
 
-			Texture->DownsizeImageUsingTextureSettings(TargetPlatform, SourceMip0, MaxSize, LayerIndex);
+			// each block is resized to MaxSize
+			if ( ! Texture->DownsizeImageUsingTextureSettings(TargetPlatform, SourceMip0, MaxSize, LayerIndex) )
+			{
+				return false;
+			}
 			ResizedBlock = MoveTemp(SourceMip0);
 			
 			ResizedSourceBlock.SizeX = ResizedBlock.SizeX;
@@ -267,6 +293,8 @@ bool DownsizeTextureSourceData(UTexture* Texture, int32 TargetSizeInGame, const
 			return Private::ResizeTextureSlicedBy2DLayers(Texture, TargetSourceSize, TargetPlatform);
 		}
 	}
+	// could do GetTextureClass == Array ?
+	// other classes unsupported
 	
 	return false;
 }
@@ -304,6 +332,9 @@ bool DownsizeTexureSourceDataNearRenderingSize(UTexture* Texture, const ITargetP
 
 		return true;
 	}
+	//	PreEditChange may have been called even if DownsizeTextureSourceData return false
+	//	we don't PostEditChange here
+	//	that's okay but not great
 
 	return false;
 }
