Sequencer:
Updated Fbx Exporter for Anim Tracks to use double precision

diff --git a/Engine/Source/Editor/UnrealEd/Private/Fbx/FbxAnimationExport.cpp b/Engine/Source/Editor/UnrealEd/Private/Fbx/FbxAnimationExport.cpp
--- a/Engine/Source/Editor/UnrealEd/Private/Fbx/FbxAnimationExport.cpp
+++ b/Engine/Source/Editor/UnrealEd/Private/Fbx/FbxAnimationExport.cpp
@@ -548,11 +548,11 @@ FbxNode* FFbxExporter::ExportAnimSequence( const UAnimSequence* AnimSeq, const U
 }
 
 
-void FFbxExporter::ExportAnimTrack(IAnimTrackAdapter& AnimTrackAdapter, AActor* Actor, USkeletalMeshComponent* InSkeletalMeshComponent, float SamplingRate)
+void FFbxExporter::ExportAnimTrack(IAnimTrackAdapter& AnimTrackAdapter, AActor* Actor, USkeletalMeshComponent* InSkeletalMeshComponent, double SamplingRate)
 {
 	// show a status update every 1 second worth of samples
-	const float UpdateFrequency = 1.0f;
-	float NextUpdateTime = UpdateFrequency;
+	const double UpdateFrequency = 1.0;
+	double NextUpdateTime = UpdateFrequency;
 
 	// find root and find the bone array
 	TArray<FbxNode*> BoneNodes;
@@ -580,12 +580,12 @@ void FFbxExporter::ExportAnimTrack(IAnimTrackAdapter& AnimTrackAdapter, AActor*
 	int32 LocalStartFrame = AnimTrackAdapter.GetLocalStartFrame();
 	int32 StartFrame = AnimTrackAdapter.GetStartFrame();
 	int32 AnimationLength = AnimTrackAdapter.GetLength();
-	float FrameRate = AnimTrackAdapter.GetFrameRate();
+	double FrameRate = AnimTrackAdapter.GetFrameRate();
 
 	TArray<USkeletalMeshComponent*> SkeletalMeshComponents;
 	Actor->GetComponents(SkeletalMeshComponents);
 
-	const float TickRate = 1.0f/FrameRate;
+	const double TickRate = 1.0/FrameRate;
 
 	FScopedSlowTask SlowTask(AnimationLength + 1, NSLOCTEXT("UnrealEd", "ExportAnimationProgress", "Exporting Animation"));
 	SlowTask.MakeDialog(true);
@@ -595,7 +595,7 @@ void FFbxExporter::ExportAnimTrack(IAnimTrackAdapter& AnimTrackAdapter, AActor*
 		SlowTask.EnterProgressFrame();
 		
 		int32 LocalFrame = LocalStartFrame + FrameCount;
-		float SampleTime = (StartFrame + FrameCount) / FrameRate;
+		double SampleTime = (StartFrame + FrameCount) / FrameRate;
 
 		// This will call UpdateSkelPose on the skeletal mesh component to move bones based on animations in the sequence
 		AnimTrackAdapter.UpdateAnimation(LocalFrame);
@@ -642,7 +642,7 @@ void FFbxExporter::ExportAnimTrack(IAnimTrackAdapter& AnimTrackAdapter, AActor*
 
 		NextUpdateTime -= SamplingRate;
 
-		if( NextUpdateTime <= 0.0f )
+		if( NextUpdateTime <= 0.0 )
 		{
 			NextUpdateTime = UpdateFrequency;
 			GWarn->StatusUpdate( FMath::RoundToInt( SampleTime ), AnimationLength, NSLOCTEXT("FbxExporter", "ExportingToFbxStatus", "Exporting to FBX") );
diff --git a/Engine/Source/Editor/UnrealEd/Private/Fbx/FbxMainExport.cpp b/Engine/Source/Editor/UnrealEd/Private/Fbx/FbxMainExport.cpp
--- a/Engine/Source/Editor/UnrealEd/Private/Fbx/FbxMainExport.cpp
+++ b/Engine/Source/Editor/UnrealEd/Private/Fbx/FbxMainExport.cpp
@@ -1736,7 +1736,7 @@ void FLevelSequenceAnimTrackAdapter::UpdateAnimation( int32 LocalFrame )
 	MovieScenePlayer->GetEvaluationTemplate().EvaluateSynchronousBlocking( Context );
 }
 
-float FLevelSequenceAnimTrackAdapter::GetFrameRate() const
+double FLevelSequenceAnimTrackAdapter::GetFrameRate() const
 {
 	return MovieScene->GetDisplayRate().AsDecimal();
 }
@@ -3249,7 +3249,7 @@ void FFbxExporter::ExportLevelSequenceBaked3DTransformTrack(IAnimTrackAdapter& A
 	int32 LocalStartFrame = FFrameRate::TransformTime(FFrameTime(DiscreteInclusiveLower(InPlaybackRange)), TickResolution, DisplayRate).RoundToFrame().Value;
 	int32 AnimationLength = FFrameRate::TransformTime(FFrameTime(FFrameNumber(DiscreteSize(InPlaybackRange))), TickResolution, DisplayRate).RoundToFrame().Value + 1; // Add one so that we export a key for the end frame
 
-	const float SampleRate = 1.0f/DisplayRate.AsDecimal();
+	const double SampleRate = 1.0/DisplayRate.AsDecimal();
 
 	for (int32 FrameNumber = LocalStartFrame; FrameNumber < LocalStartFrame + AnimationLength; ++FrameNumber)
 	{
diff --git a/Engine/Source/Editor/UnrealEd/Private/FbxExporter.h b/Engine/Source/Editor/UnrealEd/Private/FbxExporter.h
--- a/Engine/Source/Editor/UnrealEd/Private/FbxExporter.h
+++ b/Engine/Source/Editor/UnrealEd/Private/FbxExporter.h
@@ -57,7 +57,7 @@ namespace UnFbx
 		virtual int32 GetLength() const = 0;
 		/** Updates the runtime state of the animation track to the specified frame. */
 		virtual void UpdateAnimation(int32 LocalFrame) = 0;
-		virtual float GetFrameRate() const { return 1.f / DEFAULT_SAMPLERATE; }
+		virtual double GetFrameRate() const { return 1.0 / static_cast<double>(DEFAULT_SAMPLERATE); }
 		/** The anim sequence that drives this anim track */
 		virtual UAnimSequence* GetAnimSequence(int32 LocalFrame) const { return nullptr; }
 		/** The time into the anim sequence for the given LocalFrame */
@@ -73,7 +73,7 @@ namespace UnFbx
 		UNREALED_API virtual int32 GetStartFrame() const override;
 		UNREALED_API virtual int32 GetLength() const override;
 		UNREALED_API virtual void UpdateAnimation(int32 LocalFrame) override;
-		UNREALED_API virtual float GetFrameRate() const override;
+		UNREALED_API virtual double GetFrameRate() const override;
 		UNREALED_API virtual UAnimSequence* GetAnimSequence(int32 LocalFrame) const override;
 		UNREALED_API virtual float GetAnimTime(int32 LocalFrame) const override;
 
@@ -280,7 +280,7 @@ private:
 	/**
 	* Export Anim Track of the given SkeletalMeshComponent
 	*/
-	UNREALED_API void ExportAnimTrack( IAnimTrackAdapter& AnimTrackAdapter, AActor* Actor, USkeletalMeshComponent* SkeletalMeshComponent, float SamplingRate );
+	UNREALED_API void ExportAnimTrack( IAnimTrackAdapter& AnimTrackAdapter, AActor* Actor, USkeletalMeshComponent* SkeletalMeshComponent, double SamplingRate );
 
 	UNREALED_API void ExportModel(UModel* Model, FbxNode* Node, const char* Name);
 
