Make sure Fbx importer do not create multiple skeletal mesh in the same package.

#rnx

diff --git a/Engine/Source/Editor/UnrealEd/Private/Fbx/FbxFactory.cpp b/Engine/Source/Editor/UnrealEd/Private/Fbx/FbxFactory.cpp
--- a/Engine/Source/Editor/UnrealEd/Private/Fbx/FbxFactory.cpp
+++ b/Engine/Source/Editor/UnrealEd/Private/Fbx/FbxFactory.cpp
@@ -36,6 +36,7 @@
 #include "Misc/Paths.h"
 #include "Misc/ScopedSlowTask.h"
 #include "ObjectTools.h"
+#include "PackageTools.h"
 #include "PhysicsEngine/PhysicsAsset.h"
 #include "StaticMeshResources.h"
 #include "UObject/UnrealType.h"
@@ -641,9 +642,9 @@ UObject* UFbxFactory::FactoryCreateFile
 				else if ( bCanImportSkeletalMesh && ImportUI->MeshTypeToImport == FBXIT_SkeletalMesh )// skeletal mesh
 				{
 					int32 TotalNumNodes = 0;
-
 					for (int32 i = 0; i < SkelMeshArray.Num() && !bOperationCanceled; i++)
 					{
+						UPackage* Package = (SkelMeshArray.Num() == 1 ? Cast<UPackage>(InParent) : nullptr);
 						USkeletalMesh* BaseSkeletalMesh = nullptr;
 						TArray<FbxNode*> NodeArray = *SkelMeshArray[i];
 					
@@ -707,10 +708,31 @@ UObject* UFbxFactory::FactoryCreateFile
 							FSkeletalMeshImportData OutData;
 							if (LODIndex == 0 && SkelMeshNodeArray.Num() != 0)
 							{
-								FName OutputName = FbxImporter->MakeNameForMesh(Name.ToString(), SkelMeshNodeArray[0]);
+								FName OutputName = NAME_None;
+								if (Package == nullptr)
+								{
+									FString NewPackageName;
+									OutputName = FbxImporter->MakeNameForMesh(TEXT("None"), SkelMeshNodeArray[0]);
+									if (InParent != nullptr && InParent->GetOutermost() != nullptr)
+									{
+										NewPackageName = FPackageName::GetLongPackagePath(InParent->GetOutermost()->GetName()) + TEXT("/") + OutputName.ToString();
+									}
+									else
+									{
+										FbxImporter->AddTokenizedErrorMessage(FTokenizedMessage::Create(EMessageSeverity::Error, FText::Format(LOCTEXT("ImportSkeletalMesh", "Invalid Parent package when importing {0}.\nThe asset will not be imported."), FText::FromName(OutputName))), FFbxErrors::Generic_ImportingNewObjectFailed);
+										NewPackageName = OutputName.ToString();
+									}
+									NewPackageName = UPackageTools::SanitizePackageName(NewPackageName);
+									Package = CreatePackage(*NewPackageName);
+									Package->FullyLoad();
+								}
+								else
+								{
+									OutputName = FbxImporter->MakeNameForMesh(Name.ToString(), SkelMeshNodeArray[0]);
+								}
 
 								UnFbx::FFbxImporter::FImportSkeletalMeshArgs ImportSkeletalMeshArgs;
-								ImportSkeletalMeshArgs.InParent = InParent;
+								ImportSkeletalMeshArgs.InParent = Package;
 								ImportSkeletalMeshArgs.NodeArray = SkelMeshNodeArray;
 								ImportSkeletalMeshArgs.Name = OutputName;
 								ImportSkeletalMeshArgs.Flags = Flags;
@@ -731,7 +753,7 @@ UObject* UFbxFactory::FactoryCreateFile
 										// We need to remove all scaling from the root node before we set up animation data.
 										// Othewise some of the global transform calculations will be incorrect.
 										FbxImporter->RemoveTransformSettingsFromFbxNode(RootNodeToImport, ImportUI->SkeletalMeshImportData);
-										FbxImporter->SetupAnimationDataFromMesh(BaseSkeletalMesh, InParent, SkelMeshNodeArray, ImportUI->AnimSequenceImportData, OutputName.ToString());
+										FbxImporter->SetupAnimationDataFromMesh(BaseSkeletalMesh, Package, SkelMeshNodeArray, ImportUI->AnimSequenceImportData, OutputName.ToString());
 
 										// Reapply the transforms for the rest of the import
 										FbxImporter->ApplyTransformSettingsToFbxNode(RootNodeToImport, ImportUI->SkeletalMeshImportData);
