- Fix wrong vertical offset on distance field fonts imported as offline on Linux/Mac: the vertical offset was set using a size with the distance field size factor applied on it (16 by default). Now, if the distance field is activated, the vertical offset will be converted to the actual desired size.
- Fix compilation of TTFontImport with Freetype on Windows (this path is deactivated by default on Windows however)

Result between Freetype path (used on Linux/Mac) and Windows are now more similar, even if they are not still perfectly identical. It would need some rework in a subsequent task.

#tests On Windows, activated the Freetype path import path in TTFontImport (set #define USE_FREETYPE to 1), activate distance field, then reimport the font with USE_FREETYPE set to 0 or set to 1: both path gives very similar results (due to some imprecision that would need another task to fix the issue) instead of having big vertical offset.

diff --git a/Engine/Source/Editor/UnrealEd/Private/Factories/TTFontImport.cpp b/Engine/Source/Editor/UnrealEd/Private/Factories/TTFontImport.cpp
--- a/Engine/Source/Editor/UnrealEd/Private/Factories/TTFontImport.cpp
+++ b/Engine/Source/Editor/UnrealEd/Private/Factories/TTFontImport.cpp
@@ -1513,7 +1513,7 @@ void* UTrueTypeFontFactory::LoadFontFace( void* FTLibrary, int32 Height, FFeedba
 
 	// Create the Windows font
 	HFONT FontHandle =
-		CreateFont(
+		CreateFontW(
 			-Height,
 			0,
 			0,
@@ -1533,7 +1533,7 @@ void* UTrueTypeFontFactory::LoadFontFace( void* FTLibrary, int32 Height, FFeedba
 	{
 		TCHAR ErrorBuffer[1024];
 		Warn->Logf(ELogVerbosity::Error, TEXT("CreateFont failed: %s"), FPlatformMisc::GetSystemErrorMessage(ErrorBuffer, 1024, 0) );
-		return false;
+		return nullptr;
 	}
 
 	// Create DC
@@ -1542,7 +1542,7 @@ void* UTrueTypeFontFactory::LoadFontFace( void* FTLibrary, int32 Height, FFeedba
 	{
 		TCHAR ErrorBuffer[1024];
 		Warn->Logf(ELogVerbosity::Error, TEXT("GetDC failed: %s"), FPlatformMisc::GetSystemErrorMessage(ErrorBuffer, 1024, 0) );
-		return false;
+		return nullptr;
 	}
 
 	HDC DCHandle = CreateCompatibleDC( DeviceDCHandle );
@@ -1550,7 +1550,7 @@ void* UTrueTypeFontFactory::LoadFontFace( void* FTLibrary, int32 Height, FFeedba
 	{
 		TCHAR ErrorBuffer[1024];
 		Warn->Logf(ELogVerbosity::Error, TEXT("CreateDC failed: %s"), FPlatformMisc::GetSystemErrorMessage(ErrorBuffer, 1024, 0) );
-		return false;
+		return nullptr;
 	}
 
 	SelectObject( DCHandle, FontHandle );
@@ -2052,6 +2052,7 @@ bool UTrueTypeFontFactory::CreateFontTexture(
 				FontY = FMath::RoundToInt(FontY / (float)ImportOptions->Data.DistanceFieldScaleFactor);
 				FontWidth = FMath::RoundToInt(FontWidth / (float)ImportOptions->Data.DistanceFieldScaleFactor);
 				FontHeight = FMath::RoundToInt(FontHeight / (float)ImportOptions->Data.DistanceFieldScaleFactor);
+				VerticalOffset = FMath::RoundToInt(VerticalOffset / (float)ImportOptions->Data.DistanceFieldScaleFactor);
 			}
 			NewCharacterRef.StartU =
 				FMath::Clamp<int32>( FontX - ImportOptions->Data.ExtendBoxLeft,
@@ -2071,9 +2072,9 @@ bool UTrueTypeFontFactory::CreateFontTexture(
 			const int32 DestDataPitch = BitmapWidth * BitmapBytesPerPixel;
 
 			// Draw character into font and advance.
-			for( int32 SourceY = 0; SourceY < Glyph->bitmap.rows; ++SourceY )
+			for( uint32 SourceY = 0; SourceY < Glyph->bitmap.rows; ++SourceY )
 			{
-				for( int32 SourceX = 0; SourceX < Glyph->bitmap.width; ++SourceX )
+				for( uint32 SourceX = 0; SourceX < Glyph->bitmap.width; ++SourceX )
 				{
 					uint8 Opacity = Glyph->bitmap.buffer[ SourceY * Glyph->bitmap.width + SourceX ];
 
