[Behavior Tree] Fix re-opening BT assets on project start-up opens up the BB instead

diff --git a/Engine/Source/Editor/BehaviorTreeEditor/Private/BehaviorTreeEditor.cpp b/Engine/Source/Editor/BehaviorTreeEditor/Private/BehaviorTreeEditor.cpp
--- a/Engine/Source/Editor/BehaviorTreeEditor/Private/BehaviorTreeEditor.cpp
+++ b/Engine/Source/Editor/BehaviorTreeEditor/Private/BehaviorTreeEditor.cpp
@@ -1640,6 +1640,12 @@ void FBehaviorTreeEditor::FocusWindow(UObject* ObjectToFocusOn)
 	FWorkflowCentricApplication::FocusWindow(ObjectToFocusOn);
 }
 
+bool FBehaviorTreeEditor::IncludeAssetInRestoreOpenAssetsPrompt(UObject* Asset) const
+{
+	// If we're in a BT editor which has a valid BehaviorTree, then don't reopen the BB during restore, only the BT
+	return Asset && (!Asset->IsA(UBlackboardData::StaticClass()) || !BehaviorTree);
+}
+
 void FBehaviorTreeEditor::OnNodeTitleCommitted(const FText& NewText, ETextCommit::Type CommitInfo, UEdGraphNode* NodeBeingChanged)
 {
 	if (NodeBeingChanged)
diff --git a/Engine/Source/Editor/BehaviorTreeEditor/Public/BehaviorTreeEditor.h b/Engine/Source/Editor/BehaviorTreeEditor/Public/BehaviorTreeEditor.h
--- a/Engine/Source/Editor/BehaviorTreeEditor/Public/BehaviorTreeEditor.h
+++ b/Engine/Source/Editor/BehaviorTreeEditor/Public/BehaviorTreeEditor.h
@@ -65,6 +65,10 @@ public:
 	virtual void FocusWindow(UObject* ObjectToFocusOn = NULL) override;
 	//~ End IBehaviorTreeEditor Interface
 
+	//~ Begin IAssetEditorInstance Interface
+	virtual bool IncludeAssetInRestoreOpenAssetsPrompt(UObject* Asset) const override;
+	//~ End IAssetEditorInstance Interface
+
 	//~ Begin FEditorUndoClient Interface
 	virtual void PostUndo(bool bSuccess) override;
 	virtual void PostRedo(bool bSuccess) override;
diff --git a/Engine/Source/Editor/PropertyEditor/Private/PropertyEditorToolkit.h b/Engine/Source/Editor/PropertyEditor/Private/PropertyEditorToolkit.h
--- a/Engine/Source/Editor/PropertyEditor/Private/PropertyEditorToolkit.h
+++ b/Engine/Source/Editor/PropertyEditor/Private/PropertyEditorToolkit.h
@@ -38,7 +38,7 @@ public:
 
 	virtual FLinearColor GetWorldCentricTabColorScale() const override;
 
-	virtual bool IncludeAssetInRestoreOpenAssetsPrompt() const override { return false; }
+	virtual bool IncludeAssetInRestoreOpenAssetsPrompt(UObject* Asset) const override { return false; }
 
 	bool IsExposedAsColumn( const TWeakPtr< IPropertyTreeRow >& Row ) const;
 
diff --git a/Engine/Source/Editor/UnrealEd/Private/Subsystems/AssetEditorSubsystem.cpp b/Engine/Source/Editor/UnrealEd/Private/Subsystems/AssetEditorSubsystem.cpp
--- a/Engine/Source/Editor/UnrealEd/Private/Subsystems/AssetEditorSubsystem.cpp
+++ b/Engine/Source/Editor/UnrealEd/Private/Subsystems/AssetEditorSubsystem.cpp
@@ -392,7 +392,7 @@ void UAssetEditorSubsystem::NotifyAssetOpened(UObject* Asset, IAssetEditorInstan
 
 	AssetOpenedInEditorEvent.Broadcast(Asset, InInstance);
 
-	if(InInstance->IncludeAssetInRestoreOpenAssetsPrompt())
+	if(InInstance->IncludeAssetInRestoreOpenAssetsPrompt(Asset))
 	{
 		SaveOpenAssetEditors(false);
 	}
@@ -1364,10 +1364,10 @@ void UAssetEditorSubsystem::SaveOpenAssetEditors(const bool bOnShutdown)
 		for (const TPair<IAssetEditorInstance*, FAssetEntry>& EditorPair : OpenedEditors)
 		{
 			IAssetEditorInstance* Editor = EditorPair.Key;
-			if (Editor != nullptr && Editor->IncludeAssetInRestoreOpenAssetsPrompt())
+			if (Editor != nullptr)
 			{
 				UObject* EditedObject = EditorPair.Value.ObjectPtr.Get();
-				if (EditedObject != nullptr)
+				if (EditedObject != nullptr && Editor->IncludeAssetInRestoreOpenAssetsPrompt(EditedObject))
 				{
 					// only record assets that have a valid saved package
 					UPackage* Package = EditedObject->GetOutermost();
diff --git a/Engine/Source/Editor/UnrealEd/Public/Subsystems/AssetEditorSubsystem.h b/Engine/Source/Editor/UnrealEd/Public/Subsystems/AssetEditorSubsystem.h
--- a/Engine/Source/Editor/UnrealEd/Public/Subsystems/AssetEditorSubsystem.h
+++ b/Engine/Source/Editor/UnrealEd/Public/Subsystems/AssetEditorSubsystem.h
@@ -70,7 +70,16 @@ public:
 	}
 
 	/** If false, the asset being edited will not be included in reopen assets prompt on restart */
+	UE_DEPRECATED(5.4, "Use the override that takes in an UObject instead")
 	virtual bool IncludeAssetInRestoreOpenAssetsPrompt() const { return true; }
+
+	virtual bool IncludeAssetInRestoreOpenAssetsPrompt(UObject* Asset) const
+	{
+		PRAGMA_DISABLE_DEPRECATION_WARNINGS
+		return IncludeAssetInRestoreOpenAssetsPrompt();
+		PRAGMA_ENABLE_DEPRECATION_WARNINGS
+	}
+
 	virtual bool IsPrimaryEditor() const = 0;
 	virtual void InvokeTab(const struct FTabId& TabId) = 0;
 	UE_DEPRECATED(5.0, "Toolbar tab no longer exists and tab ID will return None; do not add it to layouts")
