Harden against crash when slate progress dialog runs FBlueprintEditorUtils::GetIconFromPin during save

diff --git a/Engine/Source/Editor/UnrealEd/Private/Kismet2/BlueprintEditorUtils.cpp b/Engine/Source/Editor/UnrealEd/Private/Kismet2/BlueprintEditorUtils.cpp
--- a/Engine/Source/Editor/UnrealEd/Private/Kismet2/BlueprintEditorUtils.cpp
+++ b/Engine/Source/Editor/UnrealEd/Private/Kismet2/BlueprintEditorUtils.cpp
@@ -9753,7 +9753,9 @@ const FSlateBrush* FBlueprintEditorUtils::GetIconFromPin( const FEdGraphPinType&
 	{
 		IconBrush = FAppStyle::GetBrush(TEXT("GraphEditor.Delegate_16x"));
 	}
-	else if( PinSubObject )
+	// FindObject will crash if called during save - and we have reported crashes here
+	// due to the save progress dialog invoking this function somehow
+	else if( PinSubObject && !UE::IsSavingPackage(nullptr)) 
 	{
 		UClass* VarClass = FindObject<UClass>(nullptr, *PinSubObject->GetFullName());
 		if( VarClass )
