Allow building without Turnkey Support
#tests Opened editor and confirmed it loads with no issues

diff --git a/Engine/Source/Developer/IOS/IOSPlatformEditor/Private/IOSCustomIconProjectBuildMutatorFeature.cpp b/Engine/Source/Developer/IOS/IOSPlatformEditor/Private/IOSCustomIconProjectBuildMutatorFeature.cpp
--- a/Engine/Source/Developer/IOS/IOSPlatformEditor/Private/IOSCustomIconProjectBuildMutatorFeature.cpp
+++ b/Engine/Source/Developer/IOS/IOSPlatformEditor/Private/IOSCustomIconProjectBuildMutatorFeature.cpp
@@ -48,6 +48,7 @@ static bool RequiresBuild()
 
 bool FIOSCustomIconProjectBuildMutatorFeature ::RequiresProjectBuild(const FName& InPlatformInfoName, FText& OutReason) const
 {
+#if UE_WITH_TURNKEY_SUPPORT
 	const PlatformInfo::FTargetPlatformInfo* const PlatInfo = PlatformInfo::FindPlatformInfo(InPlatformInfoName);
 	check(PlatInfo);
 
@@ -67,5 +68,6 @@ bool FIOSCustomIconProjectBuildMutatorFeature ::RequiresProjectBuild(const FName
 			}
 		}
 	}
+#endif // UE_WITH_TURNKEY_SUPPORT
 	return false;
 }
\ No newline at end of file
diff --git a/Engine/Source/Developer/TargetPlatform/Private/TargetPlatformManagerModule.cpp b/Engine/Source/Developer/TargetPlatform/Private/TargetPlatformManagerModule.cpp
--- a/Engine/Source/Developer/TargetPlatform/Private/TargetPlatformManagerModule.cpp
+++ b/Engine/Source/Developer/TargetPlatform/Private/TargetPlatformManagerModule.cpp
@@ -156,7 +156,7 @@ public:
 		, bIgnoreFirstDelegateCall(true)
 		, bSkipOneTextureFormatManagerInvalidate(false)
 	{
-#if WITH_EDITOR
+#if WITH_EDITOR && UE_WITH_TURNKEY_SUPPORT
 
 		ITurnkeySupportModule::Get().UpdateSdkInfo();
 #endif
@@ -1197,12 +1197,14 @@ protected:
 			bTPInitialized = TargetPlatform->InitializeHostPlatform();
 		}
 
+#if UE_WITH_TURNKEY_SUPPORT
 		if (bTPInitialized)
 		{
 			ITurnkeySupportModule::Get().UpdateSdkInfo();
 		}
 
 		ITurnkeySupportModule::Get().ClearDeviceStatus(PlatformName);
+#endif
 		Invalidate();
 
 		return bTPInitialized;
diff --git a/Engine/Source/Editor/TurnkeySupport/Private/TurnkeyEditorSupport.cpp b/Engine/Source/Editor/TurnkeySupport/Private/TurnkeyEditorSupport.cpp
--- a/Engine/Source/Editor/TurnkeySupport/Private/TurnkeyEditorSupport.cpp
+++ b/Engine/Source/Editor/TurnkeySupport/Private/TurnkeyEditorSupport.cpp
@@ -1,6 +1,9 @@
 // Copyright Epic Games, Inc. All Rights Reserved.
 
 #include "TurnkeyEditorSupport.h"
+
+#if UE_WITH_TURNKEY_SUPPORT
+
 #include "Internationalization/Text.h"
 #include "ITurnkeyIOModule.h"
 #include "Misc/AssertionMacros.h"
@@ -288,3 +291,5 @@ bool FTurnkeyEditorSupport::IsPIERunning()
 
 
 #undef LOCTEXT_NAMESPACE
+
+#endif // UE_WITH_TURNKEY_SUPPORT
\ No newline at end of file
diff --git a/Engine/Source/Editor/TurnkeySupport/Private/TurnkeyEditorSupport.h b/Engine/Source/Editor/TurnkeySupport/Private/TurnkeyEditorSupport.h
--- a/Engine/Source/Editor/TurnkeySupport/Private/TurnkeyEditorSupport.h
+++ b/Engine/Source/Editor/TurnkeySupport/Private/TurnkeyEditorSupport.h
@@ -3,7 +3,9 @@
 #pragma once
 
 #include "CoreMinimal.h"
+#include "Interfaces/ITurnkeySupportModule.h"
 
+#if UE_WITH_TURNKEY_SUPPORT
 
 class FTurnkeyEditorSupport
 {
@@ -25,3 +27,5 @@ public:
 	static void ShowInstallationHelp(FName IniPlatformName, FString DocLink);
 	static bool IsPIERunning();
 };
+
+#endif // UE_WITH_TURNKEY_SUPPORT
\ No newline at end of file
diff --git a/Engine/Source/Editor/TurnkeySupport/Private/TurnkeySupportModule.cpp b/Engine/Source/Editor/TurnkeySupport/Private/TurnkeySupportModule.cpp
--- a/Engine/Source/Editor/TurnkeySupport/Private/TurnkeySupportModule.cpp
+++ b/Engine/Source/Editor/TurnkeySupport/Private/TurnkeySupportModule.cpp
@@ -1,6 +1,12 @@
 // Copyright Epic Games, Inc. All Rights Reserved.
 
 #include "TurnkeySupportModule.h"
+
+#if !UE_WITH_TURNKEY_SUPPORT
+class FTurnkeySupportModuleEmpty : public IModuleInterface {};
+IMPLEMENT_MODULE(FTurnkeySupportModuleEmpty, TurnkeySupport)
+#else // UE_WITH_TURNKEY_SUPPORT
+
 #include "TurnkeySupport.h"
 
 #include "SlateOptMacros.h"
@@ -2677,3 +2683,5 @@ void FTurnkeySupportModule::ShutdownModule( )
 IMPLEMENT_MODULE(FTurnkeySupportModule, TurnkeySupport);
 
 #undef LOCTEXT_NAMESPACE
+
+#endif // UE_WITH_TURNKEY_SUPPORT
\ No newline at end of file
diff --git a/Engine/Source/Editor/TurnkeySupport/Private/TurnkeySupportModule.h b/Engine/Source/Editor/TurnkeySupport/Private/TurnkeySupportModule.h
--- a/Engine/Source/Editor/TurnkeySupport/Private/TurnkeySupportModule.h
+++ b/Engine/Source/Editor/TurnkeySupport/Private/TurnkeySupportModule.h
@@ -4,6 +4,9 @@
 
 #include "CoreMinimal.h"
 #include "Interfaces/ITurnkeySupportModule.h"
+
+#if UE_WITH_TURNKEY_SUPPORT
+
 #if WITH_EDITOR
 	#include "Misc/CoreMisc.h"
 #endif
@@ -62,3 +65,5 @@ private:
 #endif
 
 };
+
+#endif // UE_WITH_TURNKEY_SUPPORT
diff --git a/Engine/Source/Editor/TurnkeySupport/Private/TurnkeyWrapper.cpp b/Engine/Source/Editor/TurnkeySupport/Private/TurnkeyWrapper.cpp
--- a/Engine/Source/Editor/TurnkeySupport/Private/TurnkeyWrapper.cpp
+++ b/Engine/Source/Editor/TurnkeySupport/Private/TurnkeyWrapper.cpp
@@ -1,5 +1,9 @@
 // Copyright Epic Games, Inc. All Rights Reserved.
 
+#include "Interfaces/ITurnkeySupportModule.h"
+
+#if UE_WITH_TURNKEY_SUPPORT
+
 #include "Containers/Array.h"
 #include "Containers/UnrealString.h"
 #include "HAL/Platform.h"
@@ -58,3 +62,4 @@ FString ConvertToDDPIDeviceId(const FString& DeviceId)
 	return FString::Printf(TEXT("%s@%s"), *ConvertToDDPIPlatform(PlatformAndDevice[0]), *PlatformAndDevice[1]);
 }
 
+#endif // UE_WITH_TURNKEY_SUPPORT
\ No newline at end of file
diff --git a/Engine/Source/Editor/TurnkeySupport/Public/Interfaces/ITurnkeySupportModule.h b/Engine/Source/Editor/TurnkeySupport/Public/Interfaces/ITurnkeySupportModule.h
--- a/Engine/Source/Editor/TurnkeySupport/Public/Interfaces/ITurnkeySupportModule.h
+++ b/Engine/Source/Editor/TurnkeySupport/Public/Interfaces/ITurnkeySupportModule.h
@@ -7,6 +7,12 @@
 #include "Modules/ModuleManager.h"
 #include "Modules/ModuleManager.h"
 
+#ifndef UE_WITH_TURNKEY_SUPPORT
+#define UE_WITH_TURNKEY_SUPPORT 1
+#endif
+
+#if UE_WITH_TURNKEY_SUPPORT
+
 class SWidget;
 
 enum class ETurnkeyPlatformSdkStatus : uint8
@@ -113,3 +119,5 @@ public:
 	 */
 	virtual ~ITurnkeySupportModule( ) { }
 };
+
+#endif // UE_WITH_TURNKEY_SUPPORT
diff --git a/Engine/Source/Editor/TurnkeySupport/Public/TurnkeySupport.h b/Engine/Source/Editor/TurnkeySupport/Public/TurnkeySupport.h
--- a/Engine/Source/Editor/TurnkeySupport/Public/TurnkeySupport.h
+++ b/Engine/Source/Editor/TurnkeySupport/Public/TurnkeySupport.h
@@ -3,6 +3,9 @@
 #pragma once
 
 
+#include "Interfaces/ITurnkeySupportModule.h"
+
+#if UE_WITH_TURNKEY_SUPPORT
 
 /* Public Dependencies
  *****************************************************************************/
@@ -17,7 +20,6 @@
 /* Public Includes
  *****************************************************************************/
 
-#include "Interfaces/ITurnkeySupportModule.h"
 
 
 
@@ -28,3 +30,5 @@ FName ConvertToDDPIPlatform(const FName& Platform);
 FString ConvertToUATPlatform(const FString& Platform);
 FString ConvertToUATDeviceId(const FString& DeviceId);
 FString ConvertToDDPIDeviceId(const FString& DeviceId);
+
+#endif // UE_WITH_TURNKEY_SUPPORT
\ No newline at end of file
diff --git a/Engine/Source/Editor/UnrealEd/Private/Kismet2/DebuggerCommands.cpp b/Engine/Source/Editor/UnrealEd/Private/Kismet2/DebuggerCommands.cpp
--- a/Engine/Source/Editor/UnrealEd/Private/Kismet2/DebuggerCommands.cpp
+++ b/Engine/Source/Editor/UnrealEd/Private/Kismet2/DebuggerCommands.cpp
@@ -708,10 +708,12 @@ void FPlayWorldCommands::BuildToolbar(FToolMenuSection& InSection, bool bInclude
 	InSection.AddEntry(StepOverEntry);	
 	InSection.AddEntry(StepOutEntry);
 
+#if UE_WITH_TURNKEY_SUPPORT
 	if (bIncludeLaunchButtonAndOptions)
 	{
 		ITurnkeySupportModule::Get().MakeTurnkeyMenu(InSection);
 	}
+#endif
 }
 
 static void MakePreviewDeviceMenu(FMenuBuilder& MenuBuilder)
@@ -857,8 +859,10 @@ TSharedRef< SWidget > FPlayWorldCommands::GeneratePlayMenuContent(TSharedRef<FUI
 			FLocal::AddPlayModeMenuEntry(Section, PlayMode_Simulate);
 		}
 
+#if UE_WITH_TURNKEY_SUPPORT
 		// quick launch on devices
 		ITurnkeySupportModule::Get().MakeQuickLaunchItems(Menu, FOnQuickLaunchSelected::CreateStatic(&RememberQuickLaunch));
+#endif
 
 		// tip section
 		{
@@ -1368,7 +1372,9 @@ void FInternalPlayWorldCommandCallbacks::RepeatLastPlay_Clicked()
 	// hand over to Turnkey module
 	if (PlaySettings->LastExecutedPlayModeType == EPlayModeType::PlayMode_QuickLaunch)
 	{
+#if UE_WITH_TURNKEY_SUPPORT
 		ITurnkeySupportModule::Get().RepeatQuickLaunch(PlaySettings->LastExecutedLaunchName);
+#endif
 	}
 	else
 	{
