Fixed debug view mode in material preview window.

diff --git a/Engine/Source/Editor/UnrealEd/Private/PreviewMaterial.cpp b/Engine/Source/Editor/UnrealEd/Private/PreviewMaterial.cpp
--- a/Engine/Source/Editor/UnrealEd/Private/PreviewMaterial.cpp
+++ b/Engine/Source/Editor/UnrealEd/Private/PreviewMaterial.cpp
@@ -197,6 +197,10 @@ public:
 			{
 				bShaderTypeMatches = true;
 			}
+			else if (FCString::Stristr(ShaderType->GetName(), TEXT("FDebugViewModePS")))
+			{
+				bShaderTypeMatches = true;
+			}
 			else if (FCString::Stristr(ShaderType->GetName(), TEXT("FVelocity")))
 			{
 				bShaderTypeMatches = true;
