PR #11030: Fix scroll wheel not adjusting camera speed in asset toolkit viewports

diff --git a/Engine/Source/Editor/UnrealEd/Private/EditorViewportClient.cpp b/Engine/Source/Editor/UnrealEd/Private/EditorViewportClient.cpp
--- a/Engine/Source/Editor/UnrealEd/Private/EditorViewportClient.cpp
+++ b/Engine/Source/Editor/UnrealEd/Private/EditorViewportClient.cpp
@@ -3747,14 +3747,9 @@ void FEditorViewportClient::OnChangeCameraSpeed( const struct FInputEventState&
 
 	if (GetDefault<ULevelEditorViewportSettings>()->FlightCameraControlExperimentalNavigation)
 	{
-		if( Key == EKeys::MouseScrollUp )
-		{
-			GetMutableDefault<ULevelEditorViewportSettings>()->CameraSpeed = FMath::Clamp<int32>(GetDefault<ULevelEditorViewportSettings>()->CameraSpeed + 1, 1, MaxCameraSpeeds);
-		}
-		else
-		{
-			GetMutableDefault<ULevelEditorViewportSettings>()->CameraSpeed = FMath::Clamp<int32>(GetDefault<ULevelEditorViewportSettings>()->CameraSpeed - 1, 1, MaxCameraSpeeds);
-		}
+		const int32 SpeedOffset = Key == EKeys::MouseScrollUp ? 1 : -1;
+		const int32 NewSpeed = FMath::Clamp<int32>(GetCameraSpeedSetting() + SpeedOffset, 1, MaxCameraSpeeds);;
+		SetCameraSpeedSetting(NewSpeed);
 	}
 	else
 	{
