Improve source code lookup feature, using dsym if exists, and fix file path obtained from build machine

diff --git a/Engine/Build/BatchFiles/Mac/SourceCodeLookup.sh b/Engine/Build/BatchFiles/Mac/SourceCodeLookup.sh
--- a/Engine/Build/BatchFiles/Mac/SourceCodeLookup.sh
+++ b/Engine/Build/BatchFiles/Mac/SourceCodeLookup.sh
@@ -3,9 +3,15 @@
 
 FunctionName=$1
 DylibFullPath=$2
+DsymFullPath=$3
 
 if [ -f "$DylibFullPath" ]; then
-    nm -Cg "$DylibFullPath" | grep "\b$FunctionName\b" | cut -d' ' -f1 | atos --fullPath -o "$DylibFullPath"
+# Use .dsym file if provided and exists
+    if [ ! -z "$DsymFullPath" ] && [ -f "$DsymFullPath" ]; then
+        nm -Cg "$DylibFullPath" | grep "\b$FunctionName\b" | cut -d' ' -f1 | atos --fullPath -o "$DsymFullPath"
+    else
+        nm -Cg "$DylibFullPath" | grep "\b$FunctionName\b" | cut -d' ' -f1 | atos --fullPath -o "$DylibFullPath"
+    fi
 else
     echo "$DylibFullPath does not exist"
     exit 1
diff --git a/Engine/Plugins/Developer/XCodeSourceCodeAccess/Source/XCodeSourceCodeAccess/Private/XCodeSourceCodeAccessor.cpp b/Engine/Plugins/Developer/XCodeSourceCodeAccess/Source/XCodeSourceCodeAccess/Private/XCodeSourceCodeAccessor.cpp
--- a/Engine/Plugins/Developer/XCodeSourceCodeAccess/Source/XCodeSourceCodeAccess/Private/XCodeSourceCodeAccessor.cpp
+++ b/Engine/Plugins/Developer/XCodeSourceCodeAccess/Source/XCodeSourceCodeAccess/Private/XCodeSourceCodeAccessor.cpp
@@ -170,6 +170,12 @@ bool FXCodeSourceCodeAccessor::DoesSolutionExist() const
 // ColumnNumber is not supported
 bool FXCodeSourceCodeAccessor::OpenFileAtLine(const FString& FullPath, int32 LineNumber, int32 ColumnNumber)
 {
+    if (!FPaths::FileExists(FullPath))
+    {
+        UE_LOG(LogXcodeAccessor, Warning, TEXT("FXCodeSourceCodeAccessor::OpenFileAtLine failed to open non-existent file: %s"), *FullPath);
+        return false;
+    }
+    
     // We use xed to open a file at specified line number through Xcode, it also reuses your currently opened workspace if file exists in project
     int32 ReturnCode = 0;
     FString Errors;
diff --git a/Engine/Source/Editor/UnrealEd/Private/SourceCodeNavigation.cpp b/Engine/Source/Editor/UnrealEd/Private/SourceCodeNavigation.cpp
--- a/Engine/Source/Editor/UnrealEd/Private/SourceCodeNavigation.cpp
+++ b/Engine/Source/Editor/UnrealEd/Private/SourceCodeNavigation.cpp
@@ -562,7 +562,8 @@ void FSourceCodeNavigationImpl::NavigateToFunctionSource( const FString& Functio
 		FString Name = FPaths::GetBaseFilename(FullModulePath);
 		if(Name == FunctionModuleName)
 		{
-            const FString SourceCodeLookupCommand = FString::Printf(TEXT("%sBuild/BatchFiles/Mac/SourceCodeLookup.sh \"%s\" \"%s\""), *FPaths::EngineDir(), *FunctionSymbolName, *FullModulePath);
+            FString FullDsymPath = FPaths::ChangeExtension(FullModulePath, TEXT("dsym"));
+            const FString SourceCodeLookupCommand = FString::Printf(TEXT("%sBuild/BatchFiles/Mac/SourceCodeLookup.sh \"%s\" \"%s\" \"%s\""), *FPaths::EngineDir(), *FunctionSymbolName, *FullModulePath, *FullDsymPath);
             FPlatformProcess::ExecProcess( TEXT("/bin/sh"), *SourceCodeLookupCommand, &ReturnCode, &Results, &Errors );
             if(ReturnCode == 0 && !Results.IsEmpty())
             {
@@ -579,6 +580,19 @@ void FSourceCodeNavigationImpl::NavigateToFunctionSource( const FString& Functio
                     int32 LineNumberPos = ColonIndex + 1;
                     int32 LineNumberLen = CloseIndex - LineNumberPos;
                     FString LineNumber = Results.Mid(LineNumberPos, LineNumberLen);
+                    
+                    if (!FPaths::FileExists(FileName))
+                    {
+                        // could be using dsym which has fullpath of source file on Horde build machine
+                        // fix it from [REDACTED_USER_HOME]/Build/++UE5/Sync/Engine/Source/** to FPaths::EngineDir()/Source/**
+                        int pos = FileName.Find(TEXT("/Engine/Source/"));
+                        if (pos != INDEX_NONE)
+                        {
+                            FileName.RemoveAt(0, pos + 8 /* length of "/Engine/" */ );
+                            FileName.InsertAt(0, FPaths::EngineDir());
+                        }
+                    }
+                    
                     if (SourceCodeAccessor.OpenFileAtLine( FileName, FCString::Atoi(*LineNumber), 0 ))
                     {
                         return;
@@ -590,12 +604,19 @@ void FSourceCodeNavigationImpl::NavigateToFunctionSource( const FString& Functio
                 }
                 else
                 {
-                    UE_LOG(LogSelectionDetails, Warning, TEXT("NavigateToFunctionSource: Unexpected SourceCodeLookup.sh output: %s"), *Results);
+                    if (!FPaths::FileExists(FullDsymPath))
+                    {
+                        UE_LOG(LogSelectionDetails, Warning, TEXT("NavigateToFunctionSource: Missing dSYM files, please install Editor symbols for debugging"), *Results);
+                    }
+                    else
+                    {
+                        UE_LOG(LogSelectionDetails, Warning, TEXT("NavigateToFunctionSource: Unexpected SourceCodeLookup.sh output: %s"), *Results);
+                    }
                 }
             }
-            else
+            else if (ReturnCode != 0 || !Errors.IsEmpty())
             {
-                UE_LOG(LogSelectionDetails, Warning, TEXT("NavigateToFunctionSource: SourceCodeLookup.sh error: %s"), *Errors);
+                UE_LOG(LogSelectionDetails, Warning, TEXT("NavigateToFunctionSource: SourceCodeLookup.sh failed with code: %d\n%s"), ReturnCode, *Errors);
             }
 		}
 	}
