Fixed crash with the SCurveEditor widget when it has null curves. Happens with an empty data table where its row struct has a curve member.

diff --git a/Engine/Source/Editor/UnrealEd/Private/SCurveEditor.cpp b/Engine/Source/Editor/UnrealEd/Private/SCurveEditor.cpp
--- a/Engine/Source/Editor/UnrealEd/Private/SCurveEditor.cpp
+++ b/Engine/Source/Editor/UnrealEd/Private/SCurveEditor.cpp
@@ -2555,9 +2555,9 @@ TArray<FRealCurve*> SCurveEditor::GetCurvesToFit() const
 {
 	TArray<FRealCurve*> FitCurves;
 
-	for(auto CurveViewModel : CurveViewModels)
+	for(const TSharedPtr<FCurveViewModel>& CurveViewModel : CurveViewModels)
 	{
-		if (CurveViewModel->bIsVisible)
+		if (CurveViewModel->bIsVisible && CurveViewModel->CurveInfo.CurveToEdit != nullptr)
 		{
 			FitCurves.Add(CurveViewModel->CurveInfo.CurveToEdit);
 		}
