Remove the use of LoadPackageLinker as that function is being deprecated.
LoadPackageLinker can cause an infinite loop due to not correctly flushing the async loader or correctly managing the LoadContext

diff --git a/Engine/Source/Editor/UnrealEd/Private/Commandlets/ContentCommandlets.cpp b/Engine/Source/Editor/UnrealEd/Private/Commandlets/ContentCommandlets.cpp
--- a/Engine/Source/Editor/UnrealEd/Private/Commandlets/ContentCommandlets.cpp
+++ b/Engine/Source/Editor/UnrealEd/Private/Commandlets/ContentCommandlets.cpp
@@ -737,8 +737,8 @@ void UResavePackagesCommandlet::LoadAndSaveOnePackage(const FString& Filename)
 		// Get the package linker.
 		VerboseMessage(TEXT("Pre GetPackageLinker"));
 
-		FLinkerLoad* Linker = LoadPackageLinker(nullptr, PackagePath, LOAD_NoVerify);
-	
+		FLinkerLoad* Linker = GetPackageLinker(nullptr, PackagePath, LOAD_NoVerify, nullptr);
+
 		// Bail early if we don't have a valid linker (package was out of date, etc)
 		if( !Linker )
 		{
@@ -2783,7 +2783,7 @@ int32 UWrangleContentCommandlet::Main( const FString& Params )
 				// This should help capture more references.
 				GRedirectCollector.ResolveAllSoftObjectPaths(NAME_None);
 				
-				FLinkerLoad* Linker = LoadPackageLinker(nullptr, PackagePath, LOAD_Quiet | LOAD_NoWarn | LOAD_NoVerify);
+				FLinkerLoad* Linker = GetPackageLinker(nullptr, PackagePath, LOAD_Quiet | LOAD_NoWarn | LOAD_NoVerify, nullptr);
 
 				UWorld* World = UWorld::FindWorldInPackage(Package);
 				if (World && World->IsPartitionedWorld())
@@ -2987,7 +2987,7 @@ int32 UWrangleContentCommandlet::Main( const FString& Params )
 				}
 			}
 
-			FLinkerLoad* Linker = LoadPackageLinker(nullptr, PackagePath, LOAD_Quiet | LOAD_NoWarn | LOAD_NoVerify);
+			FLinkerLoad* Linker = GetPackageLinker(nullptr, PackagePath, LOAD_Quiet | LOAD_NoWarn | LOAD_NoVerify, nullptr);
 
 			// go through the exports in the package, looking for public objects
 			for (int32 ExportIndex = 0; ExportIndex < Linker->ExportMap.Num(); ExportIndex++)
@@ -3925,4 +3925,4 @@ int32 ULandscapeGrassTypeCommandlet::Main(const FString& Params)
 	FEditorFileUtils::PromptForCheckoutAndSave(PackagesToSave, false, false, nullptr, true, false);
 	
 	return 0;
-}
+}
\ No newline at end of file
