[CrashFix] DiffUtils identical crashes when diffing maps of different types

diff --git a/Engine/Source/Editor/UnrealEd/Private/DiffUtils.cpp b/Engine/Source/Editor/UnrealEd/Private/DiffUtils.cpp
--- a/Engine/Source/Editor/UnrealEd/Private/DiffUtils.cpp
+++ b/Engine/Source/Editor/UnrealEd/Private/DiffUtils.cpp
@@ -967,9 +967,9 @@ static void IdenticalHelper(const FProperty* AProperty, const FProperty* BProper
 			
 			for (int32 VirtualIndex = 0; VirtualIndex < MapSizeA && VirtualIndex < MapSizeB; ++VirtualIndex)
 			{
-				IdenticalHelper(APropAsMap->KeyProp, APropAsMap->KeyProp, MapHelperA.GetKeyPtr(MapIndexA), MapHelperB.GetKeyPtr(MapIndexB),
+				IdenticalHelper(APropAsMap->KeyProp, BPropAsMap->KeyProp, MapHelperA.GetKeyPtr(MapIndexA), MapHelperB.GetKeyPtr(MapIndexB),
 					OwningOuterA, OwningOuterB, FPropertySoftPath(RootPath, VirtualIndex), DifferingSubProperties);
-				IdenticalHelper(APropAsMap->ValueProp, APropAsMap->ValueProp, MapHelperA.GetValuePtr(MapIndexA), MapHelperB.GetValuePtr(MapIndexB),
+				IdenticalHelper(APropAsMap->ValueProp, BPropAsMap->ValueProp, MapHelperA.GetValuePtr(MapIndexA), MapHelperB.GetValuePtr(MapIndexB),
 					OwningOuterA, OwningOuterB, FPropertySoftPath(RootPath, VirtualIndex), DifferingSubProperties);
 
 				AdvanceMapIterator(MapHelperA, MapIndexA);
