Preview Platform message displayed only when Preview is Active displayed

diff --git a/Engine/Source/Editor/UnrealEd/Private/SEditorViewport.cpp b/Engine/Source/Editor/UnrealEd/Private/SEditorViewport.cpp
--- a/Engine/Source/Editor/UnrealEd/Private/SEditorViewport.cpp
+++ b/Engine/Source/Editor/UnrealEd/Private/SEditorViewport.cpp
@@ -879,9 +879,7 @@ EVisibility SEditorViewport::GetCurrentFeatureLevelPreviewTextVisibility() const
 {
 	if (Client->GetWorld())
 	{
-		UMaterialShaderQualitySettings* MaterialShaderQualitySettings = UMaterialShaderQualitySettings::Get();
-		const FName& PreviewPlatform = MaterialShaderQualitySettings->GetPreviewPlatform();
-		return (PreviewPlatform != NAME_None) ? EVisibility::SelfHitTestInvisible : EVisibility::Collapsed;
+		return (GEditor && GEditor->IsFeatureLevelPreviewActive()) ? EVisibility::SelfHitTestInvisible : EVisibility::Collapsed;
 	}
 	else
 	{
