Fixed double click on a camera in the outliner would physically move it in the level.
This is a fallout from cl 27919107.

Actor preview viewports are locked on their actors, so by removing the IsAnyActorLocked condition, we were now moving those too.
This change switches it so that only the active viewport can be moved when locked on an actor. This allows piloting and zooming in the viewport and doesn't interfere with the actor previews.

diff --git a/Engine/Source/Editor/UnrealEd/Private/EditorServer.cpp b/Engine/Source/Editor/UnrealEd/Private/EditorServer.cpp
--- a/Engine/Source/Editor/UnrealEd/Private/EditorServer.cpp
+++ b/Engine/Source/Editor/UnrealEd/Private/EditorServer.cpp
@@ -4931,7 +4931,11 @@ void UEditorEngine::MoveViewportCamerasToBox(const FBox& BoundingBox, bool bActi
 			// Update all viewports.
 			for (FLevelEditorViewportClient* LinkedViewportClient : GetLevelViewportClients())
 			{
-				LinkedViewportClient->FocusViewportOnBox(BoundingBox);
+				// Skip viewports that are locked on an actor, like actor previews, except the current viewport as we want to allow focusing while piloting
+				if (!LinkedViewportClient->IsAnyActorLocked() || LinkedViewportClient == GCurrentLevelEditingViewportClient)
+				{
+					LinkedViewportClient->FocusViewportOnBox(BoundingBox);
+				}
 			}
 		}
 
