Change format specifier used in log for int64.

diff --git a/Engine/Source/Editor/UnrealEd/Private/WorldPartition/WorldPartitionBuilder.cpp b/Engine/Source/Editor/UnrealEd/Private/WorldPartition/WorldPartitionBuilder.cpp
--- a/Engine/Source/Editor/UnrealEd/Private/WorldPartition/WorldPartitionBuilder.cpp
+++ b/Engine/Source/Editor/UnrealEd/Private/WorldPartition/WorldPartitionBuilder.cpp
@@ -205,13 +205,13 @@ bool UWorldPartitionBuilder::Run(UWorld* World, FPackageSourceControlHelper& Pac
 		};
 
 		const int64 IterationCount = ((LoadingMode == ELoadingMode::IterativeCells2D) ? 1 : NumCellsIterations.Z) * NumCellsIterations.Y * NumCellsIterations.X;
-		int32 IterationIndex = 0;
+		int64 IterationIndex = 0;
 
 		UE_LOG(LogWorldPartitionBuilder, Display, TEXT("Iterative Cell Mode"));
-		UE_LOG(LogWorldPartitionBuilder, Display, TEXT("Cell Size %d"), IterativeCellSize);
-		UE_LOG(LogWorldPartitionBuilder, Display, TEXT("Cell Overlap %d"), IterativeCellOverlapSize);
-		UE_LOG(LogWorldPartitionBuilder, Display, TEXT("WorldBounds: Min %s, Max %s"), *CellInfo.EditorBounds.Min.ToString(), *CellInfo.EditorBounds.Max.ToString());
-		UE_LOG(LogWorldPartitionBuilder, Display, TEXT("Iteration Count: %d"), IterationCount);
+		UE_LOG(LogWorldPartitionBuilder, Display, TEXT("Cell Size:       %d"), IterativeCellSize);
+		UE_LOG(LogWorldPartitionBuilder, Display, TEXT("Cell Overlap:    %d"), IterativeCellOverlapSize);
+		UE_LOG(LogWorldPartitionBuilder, Display, TEXT("WorldBounds:     Min %s, Max %s"), *CellInfo.EditorBounds.Min.ToString(), *CellInfo.EditorBounds.Max.ToString());
+		UE_LOG(LogWorldPartitionBuilder, Display, TEXT("Iteration Count: %lld"), IterationCount);
 		
 		TArray<TUniquePtr<IWorldPartitionActorLoaderInterface::ILoaderAdapter>> LoaderAdapters;
 
@@ -225,11 +225,11 @@ bool UWorldPartitionBuilder::Run(UWorld* World, FPackageSourceControlHelper& Pac
 					
 					if (ShouldSkipCell(FWorldBuilderCellCoord(x, y, z)))
 					{
-						UE_LOG(LogWorldPartitionBuilder, Display, TEXT("[%d / %d] Processing cells... (skipped)"), IterationIndex, IterationCount);
+						UE_LOG(LogWorldPartitionBuilder, Display, TEXT("[%lld / %lld] Processing cells... (skipped)"), IterationIndex, IterationCount);
 						continue;
 					}
 					
-					UE_LOG(LogWorldPartitionBuilder, Display, TEXT("[%d / %d] Processing cells..."), IterationIndex, IterationCount);
+					UE_LOG(LogWorldPartitionBuilder, Display, TEXT("[%lld / %lld] Processing cells..."), IterationIndex, IterationCount);
 
 					double MinX = static_cast<double>(x * IterativeCellSize);
 					double MinY = static_cast<double>(y * IterativeCellSize);
