Cooker: MPCook: When retracting packages from the local cookworker, retract them immediately rather than waiting for the next tick. The list of packages to retract might become invalid for retraction in between the current director tick and the next director tick, because e.g. they will finish saving or because they are generator packages and they start saving.
#rnx

diff --git a/Engine/Source/Editor/UnrealEd/Private/Cooker/CookDirector.cpp b/Engine/Source/Editor/UnrealEd/Private/Cooker/CookDirector.cpp
--- a/Engine/Source/Editor/UnrealEd/Private/Cooker/CookDirector.cpp
+++ b/Engine/Source/Editor/UnrealEd/Private/Cooker/CookDirector.cpp
@@ -78,6 +78,7 @@ private:
 		Idle,
 		WantToRetract,
 		WaitingForResponse,
+		Count,
 	};
 	enum class ERetractionResult : uint8
 	{
@@ -86,7 +87,7 @@ private:
 	};
 private:
 	/** Try to select a worker for retraction */
-	ERetractionState TickWantToRetract();
+	ERetractionState TickWantToRetract(bool& bOutAnyIdle, int32& OutBusiestNumAssignments);
 	/** Tick the asynchronous wait for the message to come in, and synchronously handle it when it does. */
 	ERetractionState TickWaitingForResponse();
 
@@ -99,7 +100,7 @@ private:
 	TArray<FWorkerId> CalculateWorkersToSplitOver(int32 NumPackages, const FWorkerId& FromWorker,
 		TConstArrayView<TRefCountPtr<FCookWorkerServer>> LocalRemoteWorkers);
 
-	void SetRetractionState(ERetractionState NewState);
+	void SetRetractionState(ERetractionState NewState, bool& bOutHadStateChange);
 	bool IsAvailableForRetraction(const FWorkerId& WorkerId);
 
 private:
@@ -1499,7 +1500,8 @@ FCookDirector::FRetractionHandler::FRetractionHandler(FCookDirector& InDirector)
 {
 }
 
-FCookDirector::FRetractionHandler::ERetractionState FCookDirector::FRetractionHandler::TickWantToRetract()
+FCookDirector::FRetractionHandler::ERetractionState FCookDirector::FRetractionHandler::TickWantToRetract(
+	bool& bOutAnyIdle, int32& OutBusiestNumAssignments)
 {
 	FWorkerId BusiestWorker;
 	TArray<FWorkerId> IdleWorkers;
@@ -1539,6 +1541,9 @@ FCookDirector::FRetractionHandler::ERetractionState FCookDirector::FRetractionHa
 			IdleWorkers.Add(RemoteWorker->GetWorkerId());
 		}
 	}
+	bOutAnyIdle = !IdleWorkers.IsEmpty();
+	OutBusiestNumAssignments = BusiestNumAssignments;
+
 	if (IdleWorkers.IsEmpty() || BusiestNumAssignments < RetractionMinimumNumAssignments)
 	{
 		// Worker loads changed after the point where we decided to initialize the RetractionHandler,
@@ -1592,49 +1597,58 @@ void FCookDirector::FRetractionHandler::InitializeForResultsMessage(const FWorke
 
 void FCookDirector::FRetractionHandler::TickFromSchedulerThread(bool bAllWorkersConnected, bool bAnyIdle, int32 BusiestNumAssignments)
 {
-	switch (RetractionState)
-	{
-	case ERetractionState::Idle:
+	bool bHadStateChange;
+	int32 NumTransitions = 0;
+	constexpr int32 MaxNumTransitions = static_cast<int32>(ERetractionState::Count);
+	do
 	{
-		if (!bAnyIdle || !bAllWorkersConnected || BusiestNumAssignments <= RetractionMinimumNumAssignments)
+		bHadStateChange = false;
+		switch (RetractionState)
+		{
+		case ERetractionState::Idle:
 		{
+			if (!bAnyIdle || !bAllWorkersConnected || BusiestNumAssignments <= RetractionMinimumNumAssignments)
+			{
+				break;
+			}
+			SetRetractionState(ERetractionState::WantToRetract, bHadStateChange);
 			break;
 		}
-		SetRetractionState(ERetractionState::WantToRetract);
-		ERetractionState NewState = TickWantToRetract();
-		SetRetractionState(NewState);
-		break;
-	}
-	case ERetractionState::WantToRetract:
-	{
-		if (!bAnyIdle || !bAllWorkersConnected || BusiestNumAssignments <= RetractionMinimumNumAssignments)
+		case ERetractionState::WantToRetract:
+		{
+			if (!bAnyIdle || !bAllWorkersConnected || BusiestNumAssignments <= RetractionMinimumNumAssignments)
+			{
+				SetRetractionState(ERetractionState::Idle, bHadStateChange);
+				break;
+			}
+			ERetractionState NewState = TickWantToRetract(bAnyIdle, BusiestNumAssignments);
+			SetRetractionState(NewState, bHadStateChange);
+			break;
+		}
+		case ERetractionState::WaitingForResponse:
 		{
-			SetRetractionState(ERetractionState::Idle);
+			ERetractionState NewState = TickWaitingForResponse();
+			SetRetractionState(NewState, bHadStateChange);
 			break;
 		}
-		ERetractionState NewState = TickWantToRetract();
-		SetRetractionState(NewState);
-		break;
-	}
-	case ERetractionState::WaitingForResponse:
-	{
-		ERetractionState NewState = TickWaitingForResponse();
-		SetRetractionState(NewState);
-		break;
-	}
-	default:
-		checkNoEntry();
-		break;
-	}
+		default:
+			checkNoEntry();
+			break;
+		}
+	} while (bHadStateChange
+		&& RetractionState != ERetractionState::Idle
+		&& ++NumTransitions <= MaxNumTransitions);
 }
 
-void FCookDirector::FRetractionHandler::SetRetractionState(ERetractionState NewState)
+void FCookDirector::FRetractionHandler::SetRetractionState(ERetractionState NewState, bool &bOutHadStateChange)
 {
 	if (RetractionState == NewState)
 	{
+		bOutHadStateChange = false;
 		return;
 	}
 
+	bOutHadStateChange = true;
 	RetractionState = NewState;
 	if (NewState == ERetractionState::Idle)
 	{
@@ -1794,7 +1808,8 @@ void FCookDirector::FRetractionHandler::HandleRetractionMessage(FMPCollectorServ
 		UE_LOG(LogCook, Warning, TEXT("Retractionmessage received from CookWorker %d when we were not expecting one."),
 			Context.GetProfileId());
 		InitializeForResultsMessage(Context.GetWorkerId());
-		SetRetractionState(ERetractionState::WaitingForResponse);
+		bool bUnusedHadStateChange;
+		SetRetractionState(ERetractionState::WaitingForResponse, bUnusedHadStateChange);
 	}
 
 	UE_CLOG(WorkerWithResults.IsValid(), LogCook, Error,
