[Deprecation Cleanup] Remove deprecated code in actor factory class

diff --git a/Engine/Plugins/2D/Paper2D/Source/Paper2DEditor/Classes/PaperFlipbookActorFactory.h b/Engine/Plugins/2D/Paper2D/Source/Paper2DEditor/Classes/PaperFlipbookActorFactory.h
--- a/Engine/Plugins/2D/Paper2D/Source/Paper2DEditor/Classes/PaperFlipbookActorFactory.h
+++ b/Engine/Plugins/2D/Paper2D/Source/Paper2DEditor/Classes/PaperFlipbookActorFactory.h
@@ -15,7 +15,6 @@ class UPaperFlipbookActorFactory : public UActorFactory
 
 	//~ Begin UActorFactory Interface
 	virtual void PostSpawnActor(UObject* Asset, AActor* NewActor) override;
-	virtual void PostCreateBlueprint(UObject* Asset, AActor* CDO) override;
 	virtual bool CanCreateActorFrom(const FAssetData& AssetData, FText& OutErrorMsg) override;
 	//~ End UActorFactory Interface
 };
diff --git a/Engine/Plugins/2D/Paper2D/Source/Paper2DEditor/Classes/PaperSpriteActorFactory.h b/Engine/Plugins/2D/Paper2D/Source/Paper2DEditor/Classes/PaperSpriteActorFactory.h
--- a/Engine/Plugins/2D/Paper2D/Source/Paper2DEditor/Classes/PaperSpriteActorFactory.h
+++ b/Engine/Plugins/2D/Paper2D/Source/Paper2DEditor/Classes/PaperSpriteActorFactory.h
@@ -15,7 +15,6 @@ class UPaperSpriteActorFactory : public UActorFactory
 
 	//~ Begin UActorFactory Interface
 	virtual void PostSpawnActor(UObject* Asset, AActor* NewActor) override;
-	virtual void PostCreateBlueprint(UObject* Asset, AActor* CDO) override;
 	virtual bool CanCreateActorFrom(const FAssetData& AssetData, FText& OutErrorMsg) override;
 	//~ End UActorFactory Interface
 };
diff --git a/Engine/Plugins/2D/Paper2D/Source/Paper2DEditor/Classes/TerrainSplineActorFactory.h b/Engine/Plugins/2D/Paper2D/Source/Paper2DEditor/Classes/TerrainSplineActorFactory.h
--- a/Engine/Plugins/2D/Paper2D/Source/Paper2DEditor/Classes/TerrainSplineActorFactory.h
+++ b/Engine/Plugins/2D/Paper2D/Source/Paper2DEditor/Classes/TerrainSplineActorFactory.h
@@ -15,7 +15,6 @@ class UTerrainSplineActorFactory : public UActorFactory
 
 	// UActorFactory interface
 	virtual void PostSpawnActor(UObject* Asset, AActor* NewActor) override;
-	virtual void PostCreateBlueprint(UObject* Asset, AActor* CDO) override;
 	virtual bool CanCreateActorFrom(const FAssetData& AssetData, FText& OutErrorMsg) override;
 	// End of UActorFactory interface
 };
diff --git a/Engine/Plugins/2D/Paper2D/Source/Paper2DEditor/Classes/TileMapActorFactory.h b/Engine/Plugins/2D/Paper2D/Source/Paper2DEditor/Classes/TileMapActorFactory.h
--- a/Engine/Plugins/2D/Paper2D/Source/Paper2DEditor/Classes/TileMapActorFactory.h
+++ b/Engine/Plugins/2D/Paper2D/Source/Paper2DEditor/Classes/TileMapActorFactory.h
@@ -15,7 +15,6 @@ class UTileMapActorFactory : public UActorFactory
 
 	// UActorFactory interface
 	virtual void PostSpawnActor(UObject* Asset, AActor* NewActor) override;
-	virtual void PostCreateBlueprint(UObject* Asset, AActor* CDO) override;
 	virtual bool CanCreateActorFrom(const FAssetData& AssetData, FText& OutErrorMsg) override;
 	// End of UActorFactory interface
 };
diff --git a/Engine/Plugins/2D/Paper2D/Source/Paper2DEditor/Private/PaperFlipbookActorFactory.cpp b/Engine/Plugins/2D/Paper2D/Source/Paper2DEditor/Private/PaperFlipbookActorFactory.cpp
--- a/Engine/Plugins/2D/Paper2D/Source/Paper2DEditor/Private/PaperFlipbookActorFactory.cpp
+++ b/Engine/Plugins/2D/Paper2D/Source/Paper2DEditor/Private/PaperFlipbookActorFactory.cpp
@@ -34,20 +34,6 @@ void UPaperFlipbookActorFactory::PostSpawnActor(UObject* Asset, AActor* NewActor
 	}
 }
 
-void UPaperFlipbookActorFactory::PostCreateBlueprint(UObject* Asset, AActor* CDO)
-{
-	if (UPaperFlipbook* Flipbook = Cast<UPaperFlipbook>(Asset))
-	{
-		if (APaperFlipbookActor* TypedActor = Cast<APaperFlipbookActor>(CDO))
-		{
-			UPaperFlipbookComponent* RenderComponent = TypedActor->GetRenderComponent();
-			check(RenderComponent);
-
-			RenderComponent->SetFlipbook(Flipbook);
-		}
-	}
-}
-
 bool UPaperFlipbookActorFactory::CanCreateActorFrom(const FAssetData& AssetData, FText& OutErrorMsg)
 {
 	if (AssetData.IsValid() && AssetData.IsInstanceOf(UPaperFlipbook::StaticClass()))
diff --git a/Engine/Plugins/2D/Paper2D/Source/Paper2DEditor/Private/PaperSpriteActorFactory.cpp b/Engine/Plugins/2D/Paper2D/Source/Paper2DEditor/Private/PaperSpriteActorFactory.cpp
--- a/Engine/Plugins/2D/Paper2D/Source/Paper2DEditor/Private/PaperSpriteActorFactory.cpp
+++ b/Engine/Plugins/2D/Paper2D/Source/Paper2DEditor/Private/PaperSpriteActorFactory.cpp
@@ -41,25 +41,6 @@ void UPaperSpriteActorFactory::PostSpawnActor(UObject* Asset, AActor* NewActor)
 	}
 }
 
-void UPaperSpriteActorFactory::PostCreateBlueprint(UObject* Asset, AActor* CDO)
-{
-	if (UPaperSprite* Sprite = Cast<UPaperSprite>(Asset))
-	{
-		if (APaperSpriteActor* TypedActor = Cast<APaperSpriteActor>(CDO))
-		{
-			UPaperSpriteComponent* RenderComponent = TypedActor->GetRenderComponent();
-			check(RenderComponent);
-
-			RenderComponent->SetSprite(Sprite);
-
-			if (Sprite->BodySetup != nullptr)
-			{
-				RenderComponent->BodyInstance.CopyBodyInstancePropertiesFrom(&(Sprite->BodySetup->DefaultInstance));
-			}
-		}
-	}
-}
-
 bool UPaperSpriteActorFactory::CanCreateActorFrom(const FAssetData& AssetData, FText& OutErrorMsg)
 {
 	if (AssetData.IsValid() && AssetData.IsInstanceOf(UPaperSprite::StaticClass()))
diff --git a/Engine/Plugins/2D/Paper2D/Source/Paper2DEditor/Private/TerrainSplineEditing/TerrainSplineActorFactory.cpp b/Engine/Plugins/2D/Paper2D/Source/Paper2DEditor/Private/TerrainSplineEditing/TerrainSplineActorFactory.cpp
--- a/Engine/Plugins/2D/Paper2D/Source/Paper2DEditor/Private/TerrainSplineEditing/TerrainSplineActorFactory.cpp
+++ b/Engine/Plugins/2D/Paper2D/Source/Paper2DEditor/Private/TerrainSplineEditing/TerrainSplineActorFactory.cpp
@@ -35,11 +35,6 @@ void UTerrainSplineActorFactory::PostSpawnActor(UObject* Asset, AActor* NewActor
 	}
 }
 
-void UTerrainSplineActorFactory::PostCreateBlueprint(UObject* Asset, AActor* CDO)
-{
-	checkf(false, TEXT("APaperTerrainActor isn't blueprintable; how did you get here?"));
-}
-
 bool UTerrainSplineActorFactory::CanCreateActorFrom(const FAssetData& AssetData, FText& OutErrorMsg)
 {
 	if (GetDefault<UPaperRuntimeSettings>()->bEnableTerrainSplineEditing)
diff --git a/Engine/Plugins/2D/Paper2D/Source/Paper2DEditor/Private/TileMapEditing/TileMapActorFactory.cpp b/Engine/Plugins/2D/Paper2D/Source/Paper2DEditor/Private/TileMapEditing/TileMapActorFactory.cpp
--- a/Engine/Plugins/2D/Paper2D/Source/Paper2DEditor/Private/TileMapEditing/TileMapActorFactory.cpp
+++ b/Engine/Plugins/2D/Paper2D/Source/Paper2DEditor/Private/TileMapEditing/TileMapActorFactory.cpp
@@ -47,27 +47,6 @@ void UTileMapActorFactory::PostSpawnActor(UObject* Asset, AActor* NewActor)
 	}
 }
 
-void UTileMapActorFactory::PostCreateBlueprint(UObject* Asset, AActor* CDO)
-{
-	if (APaperTileMapActor* TypedActor = Cast<APaperTileMapActor>(CDO))
-	{
-		UPaperTileMapComponent* RenderComponent = TypedActor->GetRenderComponent();
-		check(RenderComponent);
-
-		if (UPaperTileMap* TileMap = Cast<UPaperTileMap>(Asset))
-		{
-			RenderComponent->TileMap = TileMap;
-		}
-		else if (RenderComponent->OwnsTileMap())
-		{
-			UPaperTileMap* OwnedTileMap = RenderComponent->TileMap;
-			check(OwnedTileMap);
-
-			GetDefault<UPaperImporterSettings>()->ApplySettingsForTileMapInit(OwnedTileMap, Cast<UPaperTileSet>(Asset));
-		}
-	}
-}
-
 bool UTileMapActorFactory::CanCreateActorFrom(const FAssetData& AssetData, FText& OutErrorMsg)
 {
 	if (AssetData.IsValid())
diff --git a/Engine/Plugins/Enterprise/LidarPointCloud/Source/LidarPointCloudEditor/Private/ActorFactoryLidarPointCloud.cpp b/Engine/Plugins/Enterprise/LidarPointCloud/Source/LidarPointCloudEditor/Private/ActorFactoryLidarPointCloud.cpp
--- a/Engine/Plugins/Enterprise/LidarPointCloud/Source/LidarPointCloudEditor/Private/ActorFactoryLidarPointCloud.cpp
+++ b/Engine/Plugins/Enterprise/LidarPointCloud/Source/LidarPointCloudEditor/Private/ActorFactoryLidarPointCloud.cpp
@@ -48,17 +48,6 @@ UObject* UActorFactoryLidarPointCloud::GetAssetFromActorInstance(AActor* Instanc
 	return PCA->GetPointCloudComponent()->GetPointCloud();
 }
 
-void UActorFactoryLidarPointCloud::PostCreateBlueprint(UObject* Asset, AActor* CDO)
-{
-	if (Asset != nullptr && CDO != nullptr)
-	{
-		ULidarPointCloud* PointCloud = CastChecked<ULidarPointCloud>(Asset);
-		ALidarPointCloudActor* PointCloudActor = CastChecked<ALidarPointCloudActor>(CDO);
-
-		PointCloudActor->GetPointCloudComponent()->SetPointCloud(PointCloud);
-	}
-}
-
 FQuat UActorFactoryLidarPointCloud::AlignObjectToSurfaceNormal(const FVector& InSurfaceNormal, const FQuat& ActorRotation) const
 {
 	// Meshes align the Z (up) axis with the surface normal
diff --git a/Engine/Plugins/Enterprise/LidarPointCloud/Source/LidarPointCloudEditor/Private/ActorFactoryLidarPointCloud.h b/Engine/Plugins/Enterprise/LidarPointCloud/Source/LidarPointCloudEditor/Private/ActorFactoryLidarPointCloud.h
--- a/Engine/Plugins/Enterprise/LidarPointCloud/Source/LidarPointCloudEditor/Private/ActorFactoryLidarPointCloud.h
+++ b/Engine/Plugins/Enterprise/LidarPointCloud/Source/LidarPointCloudEditor/Private/ActorFactoryLidarPointCloud.h
@@ -18,7 +18,6 @@ class UActorFactoryLidarPointCloud : public UActorFactory
 	// Begin UActorFactory Interface
 	virtual bool CanCreateActorFrom(const FAssetData& AssetData, FText& OutErrorMsg) override;
 	virtual void PostSpawnActor(UObject* Asset, AActor* NewActor) override;
-	virtual void PostCreateBlueprint(UObject* Asset, AActor* CDO) override;
 	virtual UObject* GetAssetFromActorInstance(AActor* ActorInstance) override;
 	virtual FQuat AlignObjectToSurfaceNormal(const FVector& InSurfaceNormal, const FQuat& ActorRotation) const override;
 	// End UActorFactory Interface
diff --git a/Engine/Plugins/Experimental/ChaosFlesh/Source/ChaosFleshEditor/Private/ChaosFlesh/Asset/ActorFactoryFlesh.cpp b/Engine/Plugins/Experimental/ChaosFlesh/Source/ChaosFleshEditor/Private/ChaosFlesh/Asset/ActorFactoryFlesh.cpp
--- a/Engine/Plugins/Experimental/ChaosFlesh/Source/ChaosFleshEditor/Private/ChaosFlesh/Asset/ActorFactoryFlesh.cpp
+++ b/Engine/Plugins/Experimental/ChaosFlesh/Source/ChaosFleshEditor/Private/ChaosFlesh/Asset/ActorFactoryFlesh.cpp
@@ -50,15 +50,4 @@ void UActorFactoryFlesh::PostSpawnActor(UObject* Asset, AActor* NewActor)
 	NewFleshActor->GetFleshComponent()->RegisterComponent();
 }
 
-void UActorFactoryFlesh::PostCreateBlueprint(UObject* Asset, AActor* CDO)
-{
-	if (Asset != NULL && CDO != NULL)
-	{
-		UFleshAsset* Flesh = CastChecked<UFleshAsset>(Asset);
-		AFleshActor* FleshActor = CastChecked<AFleshActor>(CDO);
-
-		FleshActor->GetFleshComponent()->SetRestCollection(Flesh);
-	}
-}
-
 #undef LOCTEXT_NAMESPACE
diff --git a/Engine/Plugins/Experimental/ChaosFlesh/Source/ChaosFleshEditor/Private/ChaosFlesh/Asset/ActorFactoryFlesh.h b/Engine/Plugins/Experimental/ChaosFlesh/Source/ChaosFleshEditor/Private/ChaosFlesh/Asset/ActorFactoryFlesh.h
--- a/Engine/Plugins/Experimental/ChaosFlesh/Source/ChaosFleshEditor/Private/ChaosFlesh/Asset/ActorFactoryFlesh.h
+++ b/Engine/Plugins/Experimental/ChaosFlesh/Source/ChaosFleshEditor/Private/ChaosFlesh/Asset/ActorFactoryFlesh.h
@@ -16,7 +16,6 @@ class UActorFactoryFlesh : public UActorFactory
 
 	//~ Begin UActorFactory Interface
 	virtual void PostSpawnActor( UObject* Asset, AActor* NewActor ) override;
-	virtual void PostCreateBlueprint( UObject* Asset, AActor* CDO ) override;
 	virtual bool CanCreateActorFrom( const FAssetData& AssetData, FText& OutErrorMsg ) override;
 	//~ End UActorFactory Interface
 };
diff --git a/Engine/Plugins/Experimental/ChaosSolverPlugin/Source/ChaosSolverEditor/Private/Chaos/ActorFactoryChaosSolver.cpp b/Engine/Plugins/Experimental/ChaosSolverPlugin/Source/ChaosSolverEditor/Private/Chaos/ActorFactoryChaosSolver.cpp
--- a/Engine/Plugins/Experimental/ChaosSolverPlugin/Source/ChaosSolverEditor/Private/Chaos/ActorFactoryChaosSolver.cpp
+++ b/Engine/Plugins/Experimental/ChaosSolverPlugin/Source/ChaosSolverEditor/Private/Chaos/ActorFactoryChaosSolver.cpp
@@ -40,14 +40,5 @@ void UActorFactoryChaosSolver::PostSpawnActor(UObject* Asset, AActor* NewActor)
 	AChaosSolverActor* NewChaosSolverActor = CastChecked<AChaosSolverActor>(NewActor);
 }
 
-void UActorFactoryChaosSolver::PostCreateBlueprint(UObject* Asset, AActor* CDO)
-{
-	if (Asset != NULL && CDO != NULL)
-	{
-		UChaosSolver* ChaosSolver = CastChecked<UChaosSolver>(Asset);
-		AChaosSolverActor* ChaosSolverActor = CastChecked<AChaosSolverActor>(CDO);
-	}
-}
-
 #undef LOCTEXT_NAMESPACE
 
diff --git a/Engine/Plugins/Experimental/ChaosSolverPlugin/Source/ChaosSolverEditor/Private/Chaos/ActorFactoryChaosSolver.h b/Engine/Plugins/Experimental/ChaosSolverPlugin/Source/ChaosSolverEditor/Private/Chaos/ActorFactoryChaosSolver.h
--- a/Engine/Plugins/Experimental/ChaosSolverPlugin/Source/ChaosSolverEditor/Private/Chaos/ActorFactoryChaosSolver.h
+++ b/Engine/Plugins/Experimental/ChaosSolverPlugin/Source/ChaosSolverEditor/Private/Chaos/ActorFactoryChaosSolver.h
@@ -16,7 +16,6 @@ class UActorFactoryChaosSolver : public UActorFactory
 
 	//~ Begin UActorFactory Interface
 	virtual void PostSpawnActor( UObject* Asset, AActor* NewActor ) override;
-	virtual void PostCreateBlueprint( UObject* Asset, AActor* CDO ) override;
 	virtual bool CanCreateActorFrom( const FAssetData& AssetData, FText& OutErrorMsg ) override;
 	//~ End UActorFactory Interface
 };
diff --git a/Engine/Plugins/Experimental/FieldSystemPlugin/Source/FieldSyStemEditor/Private/Field/ActorFactoryFieldSystem.cpp b/Engine/Plugins/Experimental/FieldSystemPlugin/Source/FieldSyStemEditor/Private/Field/ActorFactoryFieldSystem.cpp
--- a/Engine/Plugins/Experimental/FieldSystemPlugin/Source/FieldSyStemEditor/Private/Field/ActorFactoryFieldSystem.cpp
+++ b/Engine/Plugins/Experimental/FieldSystemPlugin/Source/FieldSyStemEditor/Private/Field/ActorFactoryFieldSystem.cpp
@@ -47,15 +47,4 @@ void UActorFactoryFieldSystem::PostSpawnActor(UObject* Asset, AActor* NewActor)
 	NewFieldSystemActor->GetFieldSystemComponent()->RegisterComponent();
 }
 
-void UActorFactoryFieldSystem::PostCreateBlueprint(UObject* Asset, AActor* CDO)
-{
-	if (Asset != NULL && CDO != NULL)
-	{
-		UFieldSystem* FieldSystem = CastChecked<UFieldSystem>(Asset);
-		AFieldSystemActor* FieldSystemActor = CastChecked<AFieldSystemActor>(CDO);
-
-		FieldSystemActor->GetFieldSystemComponent()->FieldSystem = FieldSystem;
-	}
-}
-
 #undef LOCTEXT_NAMESPACE
diff --git a/Engine/Plugins/Experimental/FieldSystemPlugin/Source/FieldSyStemEditor/Private/Field/ActorFactoryFieldSystem.h b/Engine/Plugins/Experimental/FieldSystemPlugin/Source/FieldSyStemEditor/Private/Field/ActorFactoryFieldSystem.h
--- a/Engine/Plugins/Experimental/FieldSystemPlugin/Source/FieldSyStemEditor/Private/Field/ActorFactoryFieldSystem.h
+++ b/Engine/Plugins/Experimental/FieldSystemPlugin/Source/FieldSyStemEditor/Private/Field/ActorFactoryFieldSystem.h
@@ -16,7 +16,6 @@ class UActorFactoryFieldSystem : public UActorFactory
 
 	//~ Begin UActorFactory Interface
 	virtual void PostSpawnActor( UObject* Asset, AActor* NewActor ) override;
-	virtual void PostCreateBlueprint( UObject* Asset, AActor* CDO ) override;
 	virtual bool CanCreateActorFrom( const FAssetData& AssetData, FText& OutErrorMsg ) override;
 	//~ End UActorFactory Interface
 };
diff --git a/Engine/Plugins/Experimental/GeometryCollectionPlugin/Source/GeometryCollectionEditor/Private/GeometryCollection/ActorFactoryGeometryCollection.cpp b/Engine/Plugins/Experimental/GeometryCollectionPlugin/Source/GeometryCollectionEditor/Private/GeometryCollection/ActorFactoryGeometryCollection.cpp
--- a/Engine/Plugins/Experimental/GeometryCollectionPlugin/Source/GeometryCollectionEditor/Private/GeometryCollection/ActorFactoryGeometryCollection.cpp
+++ b/Engine/Plugins/Experimental/GeometryCollectionPlugin/Source/GeometryCollectionEditor/Private/GeometryCollection/ActorFactoryGeometryCollection.cpp
@@ -61,15 +61,4 @@ void UActorFactoryGeometryCollection::PostSpawnActor(UObject* Asset, AActor* New
 	NewGeometryCollectionActor->GetGeometryCollectionComponent()->RegisterComponent();
 }
 
-void UActorFactoryGeometryCollection::PostCreateBlueprint(UObject* Asset, AActor* CDO)
-{
-	if (Asset != NULL && CDO != NULL)
-	{
-		UGeometryCollection* GeometryCollection = CastChecked<UGeometryCollection>(Asset);
-		AGeometryCollectionActor* GeometryCollectionActor = CastChecked<AGeometryCollectionActor>(CDO);
-
-		GeometryCollectionActor->GetGeometryCollectionComponent()->SetRestCollection(GeometryCollection);
-	}
-}
-
 #undef LOCTEXT_NAMESPACE
diff --git a/Engine/Plugins/Experimental/GeometryCollectionPlugin/Source/GeometryCollectionEditor/Private/GeometryCollection/ActorFactoryGeometryCollection.h b/Engine/Plugins/Experimental/GeometryCollectionPlugin/Source/GeometryCollectionEditor/Private/GeometryCollection/ActorFactoryGeometryCollection.h
--- a/Engine/Plugins/Experimental/GeometryCollectionPlugin/Source/GeometryCollectionEditor/Private/GeometryCollection/ActorFactoryGeometryCollection.h
+++ b/Engine/Plugins/Experimental/GeometryCollectionPlugin/Source/GeometryCollectionEditor/Private/GeometryCollection/ActorFactoryGeometryCollection.h
@@ -18,7 +18,6 @@ class UActorFactoryGeometryCollection : public UActorFactory
 
 	//~ Begin UActorFactory Interface
 	virtual void PostSpawnActor( UObject* Asset, AActor* NewActor ) override;
-	virtual void PostCreateBlueprint( UObject* Asset, AActor* CDO ) override;
 	virtual bool CanCreateActorFrom( const FAssetData& AssetData, FText& OutErrorMsg ) override;
 	//~ End UActorFactory Interface
 };
diff --git a/Engine/Plugins/Experimental/Mutable/Source/CustomizableObjectEditor/Classes/MuCOE/CustomizableObjectInstanceFactory.h b/Engine/Plugins/Experimental/Mutable/Source/CustomizableObjectEditor/Classes/MuCOE/CustomizableObjectInstanceFactory.h
--- a/Engine/Plugins/Experimental/Mutable/Source/CustomizableObjectEditor/Classes/MuCOE/CustomizableObjectInstanceFactory.h
+++ b/Engine/Plugins/Experimental/Mutable/Source/CustomizableObjectEditor/Classes/MuCOE/CustomizableObjectInstanceFactory.h
@@ -19,7 +19,6 @@ class UCustomizableObjectInstanceFactory : public UActorFactory
 protected:
 	//~ Begin UActorFactory Interface
 	void PostSpawnActor(UObject* Asset, AActor* NewActor) override;
-	void PostCreateBlueprint(UObject* Asset, AActor* CDO) override;
 	UObject* GetAssetFromActorInstance(AActor* ActorInstance) override;
 	bool CanCreateActorFrom(const FAssetData& AssetData, FText& OutErrorMsg) override;
 	//~ End UActorFactory Interface
diff --git a/Engine/Plugins/Experimental/Mutable/Source/CustomizableObjectEditor/Private/MuCOE/CustomizableObjectInstanceFactory.cpp b/Engine/Plugins/Experimental/Mutable/Source/CustomizableObjectEditor/Private/MuCOE/CustomizableObjectInstanceFactory.cpp
--- a/Engine/Plugins/Experimental/Mutable/Source/CustomizableObjectEditor/Private/MuCOE/CustomizableObjectInstanceFactory.cpp
+++ b/Engine/Plugins/Experimental/Mutable/Source/CustomizableObjectEditor/Private/MuCOE/CustomizableObjectInstanceFactory.cpp
@@ -69,20 +69,6 @@ void UCustomizableObjectInstanceFactory::PostSpawnActor(UObject* Asset, AActor*
 	}
 }
 
-
-void UCustomizableObjectInstanceFactory::PostCreateBlueprint(UObject* Asset, AActor* CDO)
-{
-	//if (Asset != NULL && CDO != NULL)
-	//{
-	//	USkeletalMesh* SkeletalMesh = GetSkeletalMeshFromAsset(Asset);
-	//	UAnimBlueprint* AnimBlueprint = Cast<UAnimBlueprint>(Asset);
-
-	//	ASkeletalMeshActor* SkeletalMeshActor = CastChecked<ASkeletalMeshActor>(CDO);
-	//	SkeletalMeshActor->GetSkeletalMeshComponent()->SkeletalMesh = SkeletalMesh;
-	//	SkeletalMeshActor->GetSkeletalMeshComponent()->AnimClass = AnimBlueprint ? Cast<UAnimBlueprintGeneratedClass>(AnimBlueprint->GeneratedClass) : NULL;
-	//}
-}
-
 UObject* UCustomizableObjectInstanceFactory::GetAssetFromActorInstance(AActor* ActorInstance)
 {
 	if (ACustomizableSkeletalMeshActor* CSMActor = CastChecked<ACustomizableSkeletalMeshActor>(ActorInstance))
diff --git a/Engine/Plugins/FX/Niagara/Source/NiagaraEditor/Private/NiagaraEditorCommon.cpp b/Engine/Plugins/FX/Niagara/Source/NiagaraEditor/Private/NiagaraEditorCommon.cpp
--- a/Engine/Plugins/FX/Niagara/Source/NiagaraEditor/Private/NiagaraEditorCommon.cpp
+++ b/Engine/Plugins/FX/Niagara/Source/NiagaraEditor/Private/NiagaraEditorCommon.cpp
@@ -1932,16 +1932,6 @@ UObject* UActorFactoryNiagara::GetAssetFromActorInstance(AActor* Instance)
 	}
 }
 
-void UActorFactoryNiagara::PostCreateBlueprint(UObject* Asset, AActor* CDO)
-{
-	if (Asset != NULL && CDO != NULL)
-	{
-		UNiagaraSystem* System = CastChecked<UNiagaraSystem>(Asset);
-		ANiagaraActor* Actor = CastChecked<ANiagaraActor>(CDO);
-		Actor->GetNiagaraComponent()->SetAsset(System);
-	}
-}
-
 
 ///////////////////////////////////////////////////////////////////////////////
 /// INiagaraScriptGraphFocusInfo											///
diff --git a/Engine/Plugins/FX/Niagara/Source/NiagaraEditor/Public/ActorFactoryNiagara.h b/Engine/Plugins/FX/Niagara/Source/NiagaraEditor/Public/ActorFactoryNiagara.h
--- a/Engine/Plugins/FX/Niagara/Source/NiagaraEditor/Public/ActorFactoryNiagara.h
+++ b/Engine/Plugins/FX/Niagara/Source/NiagaraEditor/Public/ActorFactoryNiagara.h
@@ -18,7 +18,6 @@ class UActorFactoryNiagara : public UActorFactory
 
 	//~ Begin UActorFactory Interface
 	virtual void PostSpawnActor(UObject* Asset, AActor* NewActor) override;
-	virtual void PostCreateBlueprint(UObject* Asset, AActor* CDO) override;
 	virtual bool CanCreateActorFrom(const FAssetData& AssetData, FText& OutErrorMsg) override;
 	virtual UObject* GetAssetFromActorInstance(AActor* ActorInstance) override;
 	//~ End UActorFactory Interface
diff --git a/Engine/Plugins/Media/MediaPlate/Source/MediaPlateEditor/Private/ActorFactoryMediaPlate.cpp b/Engine/Plugins/Media/MediaPlate/Source/MediaPlateEditor/Private/ActorFactoryMediaPlate.cpp
--- a/Engine/Plugins/Media/MediaPlate/Source/MediaPlateEditor/Private/ActorFactoryMediaPlate.cpp
+++ b/Engine/Plugins/Media/MediaPlate/Source/MediaPlateEditor/Private/ActorFactoryMediaPlate.cpp
@@ -47,11 +47,6 @@ void UActorFactoryMediaPlate::PostSpawnActor(UObject* Asset, AActor* NewActor)
 	SetUpActor(Asset, NewActor);
 }
 
-void UActorFactoryMediaPlate::PostCreateBlueprint(UObject* Asset, AActor* CDO)
-{
-	SetUpActor(Asset, CDO);
-}
-
 void UActorFactoryMediaPlate::SetUpActor(UObject* Asset, AActor* Actor)
 {
 	if (Actor != nullptr)
diff --git a/Engine/Plugins/Media/MediaPlate/Source/MediaPlateEditor/Private/ActorFactoryMediaPlate.h b/Engine/Plugins/Media/MediaPlate/Source/MediaPlateEditor/Private/ActorFactoryMediaPlate.h
--- a/Engine/Plugins/Media/MediaPlate/Source/MediaPlateEditor/Private/ActorFactoryMediaPlate.h
+++ b/Engine/Plugins/Media/MediaPlate/Source/MediaPlateEditor/Private/ActorFactoryMediaPlate.h
@@ -15,7 +15,6 @@ class UActorFactoryMediaPlate : public UActorFactory
 	//~ Begin UActorFactory Interface
 	virtual bool CanCreateActorFrom(const FAssetData& AssetData, FText& OutErrorMsg) override;
 	virtual void PostSpawnActor(UObject* Asset, AActor* NewActor) override;
-	virtual void PostCreateBlueprint(UObject* Asset, AActor* CDO) override;
 	//~ End UActorFactory Interface
 
 private:
diff --git a/Engine/Plugins/Runtime/ApexDestruction/Source/ApexDestructionEditor/Private/ActorFactoryDestructible.cpp b/Engine/Plugins/Runtime/ApexDestruction/Source/ApexDestructionEditor/Private/ActorFactoryDestructible.cpp
--- a/Engine/Plugins/Runtime/ApexDestruction/Source/ApexDestructionEditor/Private/ActorFactoryDestructible.cpp
+++ b/Engine/Plugins/Runtime/ApexDestruction/Source/ApexDestructionEditor/Private/ActorFactoryDestructible.cpp
@@ -66,19 +66,6 @@ PRAGMA_DISABLE_DEPRECATION_WARNINGS
 PRAGMA_ENABLE_DEPRECATION_WARNINGS
 }
 
-void UActorFactoryDestructible::PostCreateBlueprint(UObject* Asset, AActor* CDO)
-{
-PRAGMA_DISABLE_DEPRECATION_WARNINGS
-	if (Asset != NULL && CDO != NULL)
-	{
-		UDestructibleMesh* DestructibleMesh = CastChecked<UDestructibleMesh>(Asset);
-		ADestructibleActor* DestructibleActor = CastChecked<ADestructibleActor>(CDO);
-
-		DestructibleActor->GetDestructibleComponent()->SetSkeletalMesh(DestructibleMesh);
-	}
-PRAGMA_ENABLE_DEPRECATION_WARNINGS
-}
-
 FQuat UActorFactoryDestructible::AlignObjectToSurfaceNormal(const FVector& InSurfaceNormal, const FQuat& ActorRotation) const
 {
 	// Meshes align the Z (up) axis with the surface normal
diff --git a/Engine/Plugins/Runtime/ApexDestruction/Source/ApexDestructionEditor/Private/ActorFactoryDestructible.h b/Engine/Plugins/Runtime/ApexDestruction/Source/ApexDestructionEditor/Private/ActorFactoryDestructible.h
--- a/Engine/Plugins/Runtime/ApexDestruction/Source/ApexDestructionEditor/Private/ActorFactoryDestructible.h
+++ b/Engine/Plugins/Runtime/ApexDestruction/Source/ApexDestructionEditor/Private/ActorFactoryDestructible.h
@@ -16,7 +16,6 @@ class UActorFactoryDestructible : public UActorFactory
 
 	//~ Begin UActorFactory Interface
 	virtual void PostSpawnActor( UObject* Asset, AActor* NewActor ) override;
-	virtual void PostCreateBlueprint( UObject* Asset, AActor* CDO ) override;
 	virtual bool CanCreateActorFrom( const FAssetData& AssetData, FText& OutErrorMsg ) override;
 	virtual UObject* GetAssetFromActorInstance(AActor* ActorInstance) override;
 	virtual FQuat AlignObjectToSurfaceNormal(const FVector& InSurfaceNormal, const FQuat& ActorRotation) const override;
diff --git a/Engine/Plugins/Runtime/GeometryCache/Source/GeometryCacheEd/Classes/ActorFactoryGeometryCache.h b/Engine/Plugins/Runtime/GeometryCache/Source/GeometryCacheEd/Classes/ActorFactoryGeometryCache.h
--- a/Engine/Plugins/Runtime/GeometryCache/Source/GeometryCacheEd/Classes/ActorFactoryGeometryCache.h
+++ b/Engine/Plugins/Runtime/GeometryCache/Source/GeometryCacheEd/Classes/ActorFactoryGeometryCache.h
@@ -17,7 +17,6 @@ class UActorFactoryGeometryCache : public UActorFactory
 	// Begin UActorFactory Interface
 	virtual bool CanCreateActorFrom(const FAssetData& AssetData, FText& OutErrorMsg) override;
 	virtual void PostSpawnActor(UObject* Asset, AActor* NewActor) override;
-	virtual void PostCreateBlueprint(UObject* Asset, AActor* CDO) override;	
 	// End UActorFactory Interface
 };
 
diff --git a/Engine/Plugins/Runtime/GeometryCache/Source/GeometryCacheEd/Private/ActorFactoryGeometryCache.cpp b/Engine/Plugins/Runtime/GeometryCache/Source/GeometryCacheEd/Private/ActorFactoryGeometryCache.cpp
--- a/Engine/Plugins/Runtime/GeometryCache/Source/GeometryCacheEd/Private/ActorFactoryGeometryCache.cpp
+++ b/Engine/Plugins/Runtime/GeometryCache/Source/GeometryCacheEd/Private/ActorFactoryGeometryCache.cpp
@@ -47,16 +47,3 @@ void UActorFactoryGeometryCache::PostSpawnActor(UObject* Asset, AActor* NewActor
 	GeometryCacheComponent->RegisterComponent();
 }
 
-void UActorFactoryGeometryCache::PostCreateBlueprint(UObject* Asset, AActor* CDO)
-{
-	if (Asset != NULL && CDO != NULL)
-	{
-		// Set GeometryCache (data) instance
-		UGeometryCache* GeometryCache = CastChecked<UGeometryCache>(Asset);
-		AGeometryCacheActor* GeometryCacheActor = CastChecked<AGeometryCacheActor>(CDO);
-		UGeometryCacheComponent* GeometryCacheComponent = GeometryCacheActor->GetGeometryCacheComponent();
-
-		GeometryCacheComponent->GeometryCache = GeometryCache;
-	}
-}
-
diff --git a/Engine/Plugins/Runtime/HairStrands/Source/HairStrandsEditor/Private/GroomActorFactory.cpp b/Engine/Plugins/Runtime/HairStrands/Source/HairStrandsEditor/Private/GroomActorFactory.cpp
--- a/Engine/Plugins/Runtime/HairStrands/Source/HairStrandsEditor/Private/GroomActorFactory.cpp
+++ b/Engine/Plugins/Runtime/HairStrands/Source/HairStrandsEditor/Private/GroomActorFactory.cpp
@@ -55,18 +55,6 @@ UObject* UGroomActorFactory::GetAssetFromActorInstance(AActor* Instance)
 	return SMA->GetGroomComponent()->GroomAsset;
 }
 
-void UGroomActorFactory::PostCreateBlueprint( UObject* Asset, AActor* CDO )
-{
-	if (Asset != NULL && CDO != NULL)
-	{
-		UGroomAsset* GroomAsset = CastChecked<UGroomAsset>(Asset);
-		AGroomActor* GroomActor = CastChecked<AGroomActor>(CDO);
-		UGroomComponent* GroomComponent = GroomActor->GetGroomComponent();
-
-		GroomComponent->SetGroomAsset(GroomAsset);
-	}
-}
-
 FQuat UGroomActorFactory::AlignObjectToSurfaceNormal(const FVector& InSurfaceNormal, const FQuat& ActorRotation) const
 {
 	return FindActorAlignmentRotation(ActorRotation, FVector(0.f, 0.f, 1.f), InSurfaceNormal);
diff --git a/Engine/Plugins/Runtime/HairStrands/Source/HairStrandsEditor/Public/GroomActorFactory.h b/Engine/Plugins/Runtime/HairStrands/Source/HairStrandsEditor/Public/GroomActorFactory.h
--- a/Engine/Plugins/Runtime/HairStrands/Source/HairStrandsEditor/Public/GroomActorFactory.h
+++ b/Engine/Plugins/Runtime/HairStrands/Source/HairStrandsEditor/Public/GroomActorFactory.h
@@ -19,7 +19,6 @@ class UGroomActorFactory : public UActorFactory
 	//~ Begin UActorFactory Interface
 	virtual bool CanCreateActorFrom( const FAssetData& AssetData, FText& OutErrorMsg ) override;
 	virtual void PostSpawnActor( UObject* Asset, AActor* NewActor) override;
-	virtual void PostCreateBlueprint( UObject* Asset, AActor* CDO ) override;
 	virtual UObject* GetAssetFromActorInstance(AActor* ActorInstance) override;
 	virtual FQuat AlignObjectToSurfaceNormal(const FVector& InSurfaceNormal, const FQuat& ActorRotation) const override;
 	//~ End UActorFactory Interface
diff --git a/Engine/Source/Editor/FoliageEdit/Private/ActorFactoryProceduralFoliage.cpp b/Engine/Source/Editor/FoliageEdit/Private/ActorFactoryProceduralFoliage.cpp
--- a/Engine/Source/Editor/FoliageEdit/Private/ActorFactoryProceduralFoliage.cpp
+++ b/Engine/Source/Editor/FoliageEdit/Private/ActorFactoryProceduralFoliage.cpp
@@ -85,15 +85,4 @@ UObject* UActorFactoryProceduralFoliage::GetAssetFromActorInstance(AActor* Insta
 	
 	return ProceduralComponent->FoliageSpawner;
 }
-
-void UActorFactoryProceduralFoliage::PostCreateBlueprint(UObject* Asset, AActor* CDO)
-{
-	if (Asset != nullptr && CDO != nullptr)
-	{
-		UProceduralFoliageSpawner* FoliageSpawner = CastChecked<UProceduralFoliageSpawner>(Asset);
-		AProceduralFoliageVolume* PFV = CastChecked<AProceduralFoliageVolume>(CDO);
-		UProceduralFoliageComponent* ProceduralComponent = PFV->ProceduralComponent;
-		ProceduralComponent->FoliageSpawner = FoliageSpawner;
-	}
-}
 #undef LOCTEXT_NAMESPACE
diff --git a/Engine/Source/Editor/FoliageEdit/Private/ActorFactoryProceduralFoliage.h b/Engine/Source/Editor/FoliageEdit/Private/ActorFactoryProceduralFoliage.h
--- a/Engine/Source/Editor/FoliageEdit/Private/ActorFactoryProceduralFoliage.h
+++ b/Engine/Source/Editor/FoliageEdit/Private/ActorFactoryProceduralFoliage.h
@@ -24,7 +24,6 @@ class UActorFactoryProceduralFoliage : public UActorFactoryBoxVolume
 	virtual bool PreSpawnActor(UObject* Asset, FTransform& InOutLocation);
 	virtual bool CanCreateActorFrom( const FAssetData& AssetData, FText& OutErrorMsg ) override;
 	virtual void PostSpawnActor( UObject* Asset, AActor* NewActor) override;
-	virtual void PostCreateBlueprint( UObject* Asset, AActor* CDO ) override;
 	virtual UObject* GetAssetFromActorInstance(AActor* ActorInstance) override;
 	//~ End UActorFactory Interface
 };
diff --git a/Engine/Source/Editor/UnrealEd/Classes/ActorFactories/ActorFactory.h b/Engine/Source/Editor/UnrealEd/Classes/ActorFactories/ActorFactory.h
--- a/Engine/Source/Editor/UnrealEd/Classes/ActorFactories/ActorFactory.h
+++ b/Engine/Source/Editor/UnrealEd/Classes/ActorFactories/ActorFactory.h
@@ -58,14 +58,8 @@ class UActorFactory : public UObject, public IAssetFactoryInterface
 	FVector SpawnPositionOffset;
 
 	/** Called to actual create an actor with the supplied transform (scale is ignored), using the properties in the ActorFactory */
-	UE_DEPRECATED(5.0, "This function has been deprecated in favor of the other version that takes a FActorSpawnParameters in parameter")
-	UNREALED_API AActor* CreateActor(UObject* Asset, ULevel* InLevel, const FTransform& Transform, EObjectFlags InObjectFlags, const FName InName = NAME_None);
 	UNREALED_API AActor* CreateActor(UObject* InAsset, ULevel* InLevel, const FTransform& InTransform, const FActorSpawnParameters& InSpawnParams = FActorSpawnParameters());
 
-	/** Called to create a blueprint class that can be used to spawn an actor from this factory */
-	UE_DEPRECATED(5.0, "This function is no longer used. See FKismetEditorUtilities::CreateBlueprint.")
-	UNREALED_API UBlueprint* CreateBlueprint( UObject* Instance, UObject* Outer, const FName Name, const FName CallingContext = NAME_None );
-
 	UNREALED_API virtual bool CanCreateActorFrom( const FAssetData& AssetData, FText& OutErrorMsg );
 
 	/** Name to put on context menu. */
@@ -107,19 +101,11 @@ protected:
 
 	UNREALED_API virtual bool PreSpawnActor( UObject* Asset, FTransform& InOutLocation);
 
-	UE_DEPRECATED(5.0, "This function has been deprecated in favor of the other version that takes a FActorSpawnParameters in parameter")
-	UNREALED_API virtual AActor* SpawnActor(UObject* Asset, ULevel* InLevel, const FTransform& Transform, EObjectFlags ObjectFlags, const FName Name ) final;
 	UNREALED_API virtual AActor* SpawnActor(UObject* InAsset, ULevel* InLevel, const FTransform& InTransform, const FActorSpawnParameters& InSpawnParams);
 
-	/** Subclasses may implement this to modify the actor after it has been spawned 
-	    IMPORTANT: If you override this, you should usually also override PostCreateBlueprint()! */
+	/** Subclasses may implement this to modify the actor after it has been spawned */
 	UNREALED_API virtual void PostSpawnActor( UObject* Asset, AActor* NewActor );
 
-	/** Override this in derived factory classes if needed.  This is called after a blueprint is created by this factory to
-	    update the blueprint's CDO properties with state from the asset for this factory.
-		IMPORTANT: If you override this, you should usually also override PostSpawnActor()! */
-	UNREALED_API virtual void PostCreateBlueprint( UObject* Asset, AActor* CDO );
-
 	/**
 	 * Get the default label that should be used for the actor spawned by the given asset (does not include any numeric suffix).
 	 * For classes or BPs that inherit from AActor this will defer to AActor::GetDefaultActorLabel, and for everything else it will use the asset name.
diff --git a/Engine/Source/Editor/UnrealEd/Classes/ActorFactories/ActorFactoryAmbientSound.h b/Engine/Source/Editor/UnrealEd/Classes/ActorFactories/ActorFactoryAmbientSound.h
--- a/Engine/Source/Editor/UnrealEd/Classes/ActorFactories/ActorFactoryAmbientSound.h
+++ b/Engine/Source/Editor/UnrealEd/Classes/ActorFactories/ActorFactoryAmbientSound.h
@@ -22,7 +22,6 @@ class UActorFactoryAmbientSound : public UActorFactory
 
 	//~ Begin UActorFactory Interface
 	virtual void PostSpawnActor( UObject* Asset, AActor* NewActor ) override;
-	virtual void PostCreateBlueprint( UObject* Asset, AActor* CDO ) override;
 	virtual bool CanCreateActorFrom( const FAssetData& AssetData, FText& OutErrorMsg ) override;
 	virtual UObject* GetAssetFromActorInstance(AActor* ActorInstance) override;
 	//~ End UActorFactory Interface
diff --git a/Engine/Source/Editor/UnrealEd/Classes/ActorFactories/ActorFactoryAnimationAsset.h b/Engine/Source/Editor/UnrealEd/Classes/ActorFactories/ActorFactoryAnimationAsset.h
--- a/Engine/Source/Editor/UnrealEd/Classes/ActorFactories/ActorFactoryAnimationAsset.h
+++ b/Engine/Source/Editor/UnrealEd/Classes/ActorFactories/ActorFactoryAnimationAsset.h
@@ -20,7 +20,6 @@ class UActorFactoryAnimationAsset : public UActorFactorySkeletalMesh
 protected:
 	//~ Begin UActorFactory Interface
 	virtual void PostSpawnActor( UObject* Asset, AActor* NewActor ) override;
-	virtual void PostCreateBlueprint( UObject* Asset, AActor* CDO ) override;
 	virtual bool CanCreateActorFrom(const FAssetData& AssetData, FText& OutErrorMsg) override;
 	virtual UObject* GetAssetFromActorInstance(AActor* ActorInstance);
 	//~ End UActorFactory Interface
diff --git a/Engine/Source/Editor/UnrealEd/Classes/ActorFactories/ActorFactoryDeferredDecal.h b/Engine/Source/Editor/UnrealEd/Classes/ActorFactories/ActorFactoryDeferredDecal.h
--- a/Engine/Source/Editor/UnrealEd/Classes/ActorFactories/ActorFactoryDeferredDecal.h
+++ b/Engine/Source/Editor/UnrealEd/Classes/ActorFactories/ActorFactoryDeferredDecal.h
@@ -21,7 +21,6 @@ public:
 protected:
 	//~ Begin UActorFactory Interface
 	virtual void PostSpawnActor( UObject* Asset, AActor* NewActor ) override;
-	virtual void PostCreateBlueprint( UObject* Asset, AActor* CDO ) override;
 	virtual bool CanCreateActorFrom( const FAssetData& AssetData, FText& OutErrorMsg ) override;
 	//~ End UActorFactory Interface
 
diff --git a/Engine/Source/Editor/UnrealEd/Classes/ActorFactories/ActorFactoryEmitter.h b/Engine/Source/Editor/UnrealEd/Classes/ActorFactories/ActorFactoryEmitter.h
--- a/Engine/Source/Editor/UnrealEd/Classes/ActorFactories/ActorFactoryEmitter.h
+++ b/Engine/Source/Editor/UnrealEd/Classes/ActorFactories/ActorFactoryEmitter.h
@@ -18,7 +18,6 @@ class UActorFactoryEmitter : public UActorFactory
 
 	//~ Begin UActorFactory Interface
 	virtual void PostSpawnActor( UObject* Asset, AActor* NewActor ) override;
-	virtual void PostCreateBlueprint( UObject* Asset, AActor* CDO ) override;
 	virtual bool CanCreateActorFrom( const FAssetData& AssetData, FText& OutErrorMsg ) override;
 	virtual UObject* GetAssetFromActorInstance(AActor* ActorInstance) override;
 	//~ End UActorFactory Interface
diff --git a/Engine/Source/Editor/UnrealEd/Classes/ActorFactories/ActorFactoryPhysicsAsset.h b/Engine/Source/Editor/UnrealEd/Classes/ActorFactories/ActorFactoryPhysicsAsset.h
--- a/Engine/Source/Editor/UnrealEd/Classes/ActorFactories/ActorFactoryPhysicsAsset.h
+++ b/Engine/Source/Editor/UnrealEd/Classes/ActorFactories/ActorFactoryPhysicsAsset.h
@@ -19,7 +19,6 @@ class UActorFactoryPhysicsAsset : public UActorFactory
 	//~ Begin UActorFactory Interface
 	virtual bool PreSpawnActor( UObject* Asset, FTransform& InOutLocation ) override;
 	virtual void PostSpawnActor( UObject* Asset, AActor* NewActor ) override;
-	virtual void PostCreateBlueprint( UObject* Asset, AActor* CDO ) override;
 	virtual bool CanCreateActorFrom( const FAssetData& AssetData, FText& OutErrorMsg ) override;
 	//~ End UActorFactory Interface
 };
diff --git a/Engine/Source/Editor/UnrealEd/Classes/ActorFactories/ActorFactorySkeletalMesh.h b/Engine/Source/Editor/UnrealEd/Classes/ActorFactories/ActorFactorySkeletalMesh.h
--- a/Engine/Source/Editor/UnrealEd/Classes/ActorFactories/ActorFactorySkeletalMesh.h
+++ b/Engine/Source/Editor/UnrealEd/Classes/ActorFactories/ActorFactorySkeletalMesh.h
@@ -39,7 +39,6 @@ protected:
 
 	//~ Begin UActorFactory Interface
 	virtual void PostSpawnActor( UObject* Asset, AActor* NewActor ) override;
-	virtual void PostCreateBlueprint( UObject* Asset, AActor* CDO ) override;
 	virtual bool CanCreateActorFrom( const FAssetData& AssetData, FText& OutErrorMsg ) override;
 	virtual FQuat AlignObjectToSurfaceNormal(const FVector& InSurfaceNormal, const FQuat& ActorRotation) const override;
 	//~ End UActorFactory Interface
diff --git a/Engine/Source/Editor/UnrealEd/Classes/ActorFactories/ActorFactoryStaticMesh.h b/Engine/Source/Editor/UnrealEd/Classes/ActorFactories/ActorFactoryStaticMesh.h
--- a/Engine/Source/Editor/UnrealEd/Classes/ActorFactories/ActorFactoryStaticMesh.h
+++ b/Engine/Source/Editor/UnrealEd/Classes/ActorFactories/ActorFactoryStaticMesh.h
@@ -19,7 +19,6 @@ class UActorFactoryStaticMesh : public UActorFactory
 	//~ Begin UActorFactory Interface
 	virtual bool CanCreateActorFrom( const FAssetData& AssetData, FText& OutErrorMsg ) override;
 	virtual void PostSpawnActor( UObject* Asset, AActor* NewActor) override;
-	virtual void PostCreateBlueprint( UObject* Asset, AActor* CDO ) override;
 	virtual UObject* GetAssetFromActorInstance(AActor* ActorInstance) override;
 	virtual FQuat AlignObjectToSurfaceNormal(const FVector& InSurfaceNormal, const FQuat& ActorRotation) const override;
 	//~ End UActorFactory Interface
diff --git a/Engine/Source/Editor/UnrealEd/Private/Factories/ActorFactory.cpp b/Engine/Source/Editor/UnrealEd/Private/Factories/ActorFactory.cpp
--- a/Engine/Source/Editor/UnrealEd/Private/Factories/ActorFactory.cpp
+++ b/Engine/Source/Editor/UnrealEd/Private/Factories/ActorFactory.cpp
@@ -410,14 +410,6 @@ UInstancedPlacemenClientSettings* UActorFactory::FactorySettingsObjectForPlaceme
 	return nullptr;
 }
 
-AActor* UActorFactory::CreateActor( UObject* Asset, ULevel* InLevel, const FTransform& SpawnTransform, EObjectFlags InObjectFlags, const FName Name)
-{
-	FActorSpawnParameters SpawnParams;
-	SpawnParams.ObjectFlags = InObjectFlags;
-	SpawnParams.Name = Name;
-	return CreateActor(Asset, InLevel, SpawnTransform, SpawnParams);
-}
-
 AActor* UActorFactory::CreateActor(UObject* InAsset, ULevel* InLevel, const FTransform& InTransform, const FActorSpawnParameters& InSpawnParams)
 {
 	AActor* NewActor = nullptr;
@@ -444,17 +436,6 @@ AActor* UActorFactory::CreateActor(UObject* InAsset, ULevel* InLevel, const FTra
 	return NewActor;
 }
 
-UBlueprint* UActorFactory::CreateBlueprint( UObject* Asset, UObject* Outer, const FName Name, const FName CallingContext )
-{
-	// @todo sequencer major: Needs to be overridden on any class that needs any custom setup for the new blueprint
-	//	(e.g. static mesh reference assignment.)  Basically, anywhere that PostSpawnActor() or CreateActor() is overridden,
-	//	we should consider overriding CreateBlueprint(), too.
-	UBlueprint* NewBlueprint = FKismetEditorUtilities::CreateBlueprint(NewActorClass, Outer, Name, EBlueprintType::BPTYPE_Normal, UBlueprint::StaticClass(), UBlueprintGeneratedClass::StaticClass(), CallingContext);
-	AActor* CDO = CastChecked<AActor>( NewBlueprint->GeneratedClass->ClassDefaultObject );
-	PostCreateBlueprint( Asset, CDO );
-	return NewBlueprint;
-}
-
 ULevel* UActorFactory::ValidateSpawnActorLevel(ULevel* InLevel, const FActorSpawnParameters& InSpawnParams) const
 {
 	ULevel* LocalLevel = InLevel;
@@ -481,14 +462,6 @@ bool UActorFactory::PreSpawnActor( UObject* Asset, FTransform& InOutLocation)
 	return true;
 }
 
-AActor* UActorFactory::SpawnActor( UObject* InAsset, ULevel* InLevel, const FTransform& Transform, EObjectFlags InObjectFlags, const FName Name )
-{
-	FActorSpawnParameters SpawnParams;
-	SpawnParams.ObjectFlags = InObjectFlags;
-	SpawnParams.Name = Name;
-	return SpawnActor(InAsset, InLevel, Transform, SpawnParams);
-}
-
 AActor* UActorFactory::SpawnActor(UObject* InAsset, ULevel* InLevel, const FTransform& InTransform, const FActorSpawnParameters& InSpawnParams)
 {
 	ULevel* LocalLevel = ValidateSpawnActorLevel(InLevel, InSpawnParams);
@@ -519,11 +492,6 @@ void UActorFactory::PostSpawnActor( UObject* Asset, AActor* NewActor)
 	UE_LOG(LogActorFactory, Log, TEXT("Actor Factory spawned %s as actor: %s"), *Asset->GetFullName(), *NewActor->GetFullName());
 }
 
-void UActorFactory::PostCreateBlueprint( UObject* Asset, AActor* CDO )
-{
-	// Override this in derived actor factories to initialize the blueprint's CDO based on the asset assigned to the factory!
-}
-
 FString UActorFactory::GetDefaultActorLabel(UObject* Asset) const
 {
 	UClass* PotentialActorClass = nullptr;
@@ -599,19 +567,6 @@ UObject* UActorFactoryStaticMesh::GetAssetFromActorInstance(AActor* Instance)
 	return SMA->GetStaticMeshComponent()->GetStaticMesh();
 }
 
-void UActorFactoryStaticMesh::PostCreateBlueprint( UObject* Asset, AActor* CDO )
-{
-	if (Asset != nullptr && CDO != nullptr)
-	{
-		UStaticMesh* StaticMesh = CastChecked<UStaticMesh>(Asset);
-		AStaticMeshActor* StaticMeshActor = CastChecked<AStaticMeshActor>(CDO);
-		UStaticMeshComponent* StaticMeshComponent = StaticMeshActor->GetStaticMeshComponent();
-
-		StaticMeshComponent->SetStaticMesh(StaticMesh);
-		StaticMeshComponent->StaticMeshDerivedDataKey = StaticMesh->GetRenderData()->DerivedDataKey;
-	}
-}
-
 FQuat UActorFactoryStaticMesh::AlignObjectToSurfaceNormal(const FVector& InSurfaceNormal, const FQuat& ActorRotation) const
 {
 	// Meshes align the Z (up) axis with the surface normal
@@ -776,30 +731,6 @@ void UActorFactoryDeferredDecal::PostSpawnActor(UObject* Asset, AActor* NewActor
 	}
 }
 
-void UActorFactoryDeferredDecal::PostCreateBlueprint( UObject* Asset, AActor* CDO )
-{
-	if (Asset != nullptr && CDO != nullptr)
-	{
-		UMaterialInterface* Material = GetMaterial(Asset);
-
-		if (Material != nullptr)
-		{
-			UDecalComponent* DecalComponent = nullptr;
-			for (UActorComponent* Component : CDO->GetComponents())
-			{
-				if (UDecalComponent* DecalComp = Cast<UDecalComponent>(Component))
-				{
-					DecalComponent = DecalComp;
-					break;
-				}
-			}
-
-			check(DecalComponent);
-			DecalComponent->SetDecalMaterial(Material);
-		}
-	}
-}
-
 UMaterialInterface* UActorFactoryDeferredDecal::GetMaterial( UObject* Asset ) const
 {
 	UMaterialInterface* TargetMaterial = Cast<UMaterialInterface>( Asset );
@@ -885,16 +816,6 @@ UObject* UActorFactoryEmitter::GetAssetFromActorInstance(AActor* Instance)
 	}
 }
 
-void UActorFactoryEmitter::PostCreateBlueprint( UObject* Asset, AActor* CDO )
-{
-	if (Asset != nullptr && CDO != nullptr)
-	{
-		UParticleSystem* ParticleSystem = CastChecked<UParticleSystem>(Asset);
-		AEmitter* Emitter = CastChecked<AEmitter>(CDO);
-		Emitter->SetTemplate(ParticleSystem);
-	}
-}
-
 
 /*-----------------------------------------------------------------------------
 UActorFactoryPlayerStart
@@ -994,33 +915,6 @@ void UActorFactoryPhysicsAsset::PostSpawnActor(UObject* Asset, AActor* NewActor)
 	NewSkelActor->GetSkeletalMeshComponent()->RegisterComponent();
 }
 
-void UActorFactoryPhysicsAsset::PostCreateBlueprint( UObject* Asset, AActor* CDO )
-{
-	if (CDO != nullptr)
-	{
-		ASkeletalMeshActor* SkeletalPhysicsActor = CastChecked<ASkeletalMeshActor>(CDO);
-
-		if (Asset != nullptr)
-		{
-			UPhysicsAsset* PhysicsAsset = CastChecked<UPhysicsAsset>(Asset);
-
-			USkeletalMesh* UseSkelMesh = PhysicsAsset->PreviewSkeletalMesh.Get();
-
-			SkeletalPhysicsActor->GetSkeletalMeshComponent()->SetSkeletalMeshAsset(UseSkelMesh);
-			SkeletalPhysicsActor->GetSkeletalMeshComponent()->PhysicsAssetOverride = PhysicsAsset;
-		}
-
-		// set physics setup
-		SkeletalPhysicsActor->GetSkeletalMeshComponent()->KinematicBonesUpdateType = EKinematicBonesUpdateToPhysics::SkipSimulatingBones;
-		SkeletalPhysicsActor->GetSkeletalMeshComponent()->BodyInstance.bSimulatePhysics = true;
-		SkeletalPhysicsActor->GetSkeletalMeshComponent()->bBlendPhysics = true;
-
-		SkeletalPhysicsActor->bAlwaysRelevant = true;
-		SkeletalPhysicsActor->SetReplicatingMovement(true);
-		SkeletalPhysicsActor->SetReplicates(true);
-	}
-}
-
 
 /*-----------------------------------------------------------------------------
 UActorFactoryAnimationAsset
@@ -1143,24 +1037,6 @@ void UActorFactoryAnimationAsset::PostSpawnActor( UObject* Asset, AActor* NewAct
 	}
 }
 
-void UActorFactoryAnimationAsset::PostCreateBlueprint( UObject* Asset,  AActor* CDO )
-{
-	Super::PostCreateBlueprint( Asset, CDO );
-	
-	if (Asset != nullptr && CDO != nullptr)
-	{
-		UAnimationAsset* AnimationAsset = Cast<UAnimationAsset>(Asset);
-
-		ASkeletalMeshActor* SkeletalMeshActor = CastChecked<ASkeletalMeshActor>(CDO);
-		USkeletalMeshComponent* SkeletalComponent = (SkeletalMeshActor->GetSkeletalMeshComponent());
-		if (AnimationAsset)
-		{
-			SkeletalComponent->SetAnimationMode(EAnimationMode::Type::AnimationSingleNode);
-			SkeletalComponent->SetAnimation(AnimationAsset);
-		}
-	}
-}
-
 UObject* UActorFactoryAnimationAsset::GetAssetFromActorInstance(AActor* ActorInstance)
 {
 	if (ASkeletalMeshActor* SkeletalMeshActor = Cast<ASkeletalMeshActor>(ActorInstance))
@@ -1373,19 +1249,6 @@ void UActorFactorySkeletalMesh::PostSpawnActor( UObject* Asset, AActor* NewActor
 
 }
 
-void UActorFactorySkeletalMesh::PostCreateBlueprint( UObject* Asset, AActor* CDO )
-{
-	if (Asset != nullptr && CDO != nullptr)
-	{
-		USkeletalMesh* SkeletalMesh = GetSkeletalMeshFromAsset(Asset);
-		UAnimBlueprint* AnimBlueprint = Cast<UAnimBlueprint>(Asset);
-
-		ASkeletalMeshActor* SkeletalMeshActor = CastChecked<ASkeletalMeshActor>(CDO);
-		SkeletalMeshActor->GetSkeletalMeshComponent()->SetSkeletalMeshAsset(SkeletalMesh);
-		SkeletalMeshActor->GetSkeletalMeshComponent()->AnimClass = AnimBlueprint ? Cast<UAnimBlueprintGeneratedClass>(AnimBlueprint->GeneratedClass) : nullptr;
-	}
-}
-
 FQuat UActorFactorySkeletalMesh::AlignObjectToSurfaceNormal(const FVector& InSurfaceNormal, const FQuat& ActorRotation) const
 {
 	// Meshes align the Z (up) axis with the surface normal
@@ -1544,20 +1407,6 @@ UObject* UActorFactoryAmbientSound::GetAssetFromActorInstance(AActor* Instance)
 	return SoundActor->GetAudioComponent()->Sound;
 }
 
-void UActorFactoryAmbientSound::PostCreateBlueprint( UObject* Asset, AActor* CDO )
-{
-	if (Asset != nullptr && CDO != nullptr)
-	{
-		USoundBase* AmbientSound = Cast<USoundBase>(Asset);
-
-		if (AmbientSound != nullptr)
-		{
-			AAmbientSound* NewSound = CastChecked<AAmbientSound>(CDO);
-			NewSound->GetAudioComponent()->SetSound(AmbientSound);
-		}
-	}
-}
-
 bool UActorFactoryAmbientSound::CanImportAmbientSounds()
 {
 	IAssetTools& AssetTools = FAssetToolsModule::GetModule().Get();
