Guard against an unexpected NULL reference during explicit implemented interface removal in the Blueprint class editor.

diff --git a/Engine/Source/Editor/UnrealEd/Private/Kismet2/BlueprintEditorUtils.cpp b/Engine/Source/Editor/UnrealEd/Private/Kismet2/BlueprintEditorUtils.cpp
--- a/Engine/Source/Editor/UnrealEd/Private/Kismet2/BlueprintEditorUtils.cpp
+++ b/Engine/Source/Editor/UnrealEd/Private/Kismet2/BlueprintEditorUtils.cpp
@@ -6361,11 +6361,11 @@ void FBlueprintEditorUtils::GetInterfaceGraphs(UBlueprint* Blueprint, FTopLevelA
 	}
 
 	// Find the implemented interface
-	for( int32 i = 0; i < Blueprint->ImplementedInterfaces.Num(); i++ )
+	for (const FBPInterfaceDescription& InterfaceDesc : Blueprint->ImplementedInterfaces)
 	{
-		if( Blueprint->ImplementedInterfaces[i].Interface->GetClassPathName() == InterfaceClassPathName)
+		if (InterfaceDesc.Interface && InterfaceDesc.Interface->GetClassPathName() == InterfaceClassPathName)
 		{
-			ChildGraphs = Blueprint->ImplementedInterfaces[i].Graphs;
+			ChildGraphs = InterfaceDesc.Graphs;
 			return;			
 		}
 	}
@@ -6435,16 +6435,17 @@ void FBlueprintEditorUtils::RemoveInterface(UBlueprint* Blueprint, FTopLevelAsse
 
 	// Find the implemented interface
 	int32 Idx = INDEX_NONE;
-	for( int32 i = 0; i < Blueprint->ImplementedInterfaces.Num(); i++ )
+	for (int32 i = 0; i < Blueprint->ImplementedInterfaces.Num(); i++)
 	{
-		if( Blueprint->ImplementedInterfaces[i].Interface->GetClassPathName() == InterfaceClassPathName)
+		const FBPInterfaceDescription& InterfaceDesc = Blueprint->ImplementedInterfaces[i];
+		if (InterfaceDesc.Interface && InterfaceDesc.Interface->GetClassPathName() == InterfaceClassPathName)
 		{
 			Idx = i;
 			break;
 		}
 	}
 
-	if( Idx != INDEX_NONE )
+	if (ensureMsgf(Idx != INDEX_NONE, TEXT("%s: No implementation was found for \'%s\'."), *Blueprint->GetName(), *InterfaceClassPathName.ToString()))
 	{
 		FBPInterfaceDescription& CurrentInterface = Blueprint->ImplementedInterfaces[Idx];
 		const UClass* InterfaceClass = Blueprint->ImplementedInterfaces[Idx].Interface;
