Add logging book-ends to the CompileShaders Commandlet to know when it starts compilation and ends compilation.

#rnx

diff --git a/Engine/Source/Editor/UnrealEd/Private/Commandlets/CompileShadersTestBedCommandlet.cpp b/Engine/Source/Editor/UnrealEd/Private/Commandlets/CompileShadersTestBedCommandlet.cpp
--- a/Engine/Source/Editor/UnrealEd/Private/Commandlets/CompileShadersTestBedCommandlet.cpp
+++ b/Engine/Source/Editor/UnrealEd/Private/Commandlets/CompileShadersTestBedCommandlet.cpp
@@ -118,6 +118,8 @@ int32 UCompileShadersTestBedCommandlet::Main(const FString& Params)
 	ITargetPlatformManagerModule* TPM = GetTargetPlatformManager();
 	const TArray<ITargetPlatform*>& Platforms = TPM->GetActiveTargetPlatforms();
 
+	UE_LOG(LogCompileShadersTestBedCommandlet, Display, TEXT("Begin Compiling Shaders"));
+
 	for (ITargetPlatform* Platform : Platforms)
 	{
 		UE_LOG(LogCompileShadersTestBedCommandlet, Display, TEXT("Compiling shaders for %s..."), *Platform->PlatformName());
@@ -266,6 +268,8 @@ int32 UCompileShadersTestBedCommandlet::Main(const FString& Params)
 		}
 	}
 
+	UE_LOG(LogCompileShadersTestBedCommandlet, Display, TEXT("End compiling shaders"));
+
 	GShaderCompilingManager->PrintStats();
 
 	return 0;
