Remove unused private variable in SViewportToolbar

This variable is a remnant from a previous overhaul of the viewport toolbar.

diff --git a/Engine/Source/Editor/UnrealEd/Private/SViewportToolBar.cpp b/Engine/Source/Editor/UnrealEd/Private/SViewportToolBar.cpp
--- a/Engine/Source/Editor/UnrealEd/Private/SViewportToolBar.cpp
+++ b/Engine/Source/Editor/UnrealEd/Private/SViewportToolBar.cpp
@@ -13,9 +13,8 @@ struct FSlateBrush;
 
 #define LOCTEXT_NAMESPACE "ViewportToolBar"
 
-void SViewportToolBar::Construct( const FArguments& InArgs )
+void SViewportToolBar::Construct(const FArguments& InArgs)
 {
-	bIsHovered = false;
 }
 
 TWeakPtr<SMenuAnchor> SViewportToolBar::GetOpenMenu() const
diff --git a/Engine/Source/Editor/UnrealEd/Public/SViewportToolBar.h b/Engine/Source/Editor/UnrealEd/Public/SViewportToolBar.h
--- a/Engine/Source/Editor/UnrealEd/Public/SViewportToolBar.h
+++ b/Engine/Source/Editor/UnrealEd/Public/SViewportToolBar.h
@@ -64,7 +64,5 @@ protected:
 private:
 	/** The pulldown menu that is open if any */
 	TWeakPtr< SMenuAnchor > OpenedMenu;
-	/** True if the mouse is inside the toolbar */
-	bool bIsHovered;
 };
 
