PR #10826: Add delegate in UAssetEditorSubsystem to broadcast when an asset editor is closing

diff --git a/Engine/Source/Editor/UnrealEd/Private/Subsystems/AssetEditorSubsystem.cpp b/Engine/Source/Editor/UnrealEd/Private/Subsystems/AssetEditorSubsystem.cpp
--- a/Engine/Source/Editor/UnrealEd/Private/Subsystems/AssetEditorSubsystem.cpp
+++ b/Engine/Source/Editor/UnrealEd/Private/Subsystems/AssetEditorSubsystem.cpp
@@ -410,6 +410,8 @@ void UAssetEditorSubsystem::NotifyAssetsOpened(const TArray< UObject* >& Assets,
 
 void UAssetEditorSubsystem::NotifyAssetClosed(UObject* Asset, IAssetEditorInstance* InInstance)
 {
+	AssetClosedInEditorEvent.Broadcast(Asset, InInstance);
+
 	OpenedEditors.RemoveSingle(InInstance, Asset);
 	OpenedAssets.RemoveSingle(Asset, InInstance);
 
@@ -424,6 +426,10 @@ void UAssetEditorSubsystem::NotifyEditorClosed(IAssetEditorInstance* InInstance)
 	OpenedEditors.MultiFind(InInstance, /*out*/ Assets);
 	for (int32 AssetIndex = 0; AssetIndex < Assets.Num(); ++AssetIndex)
 	{
+		if(UObject* Asset = Assets[AssetIndex].ObjectPtr.Get())
+		{
+			AssetClosedInEditorEvent.Broadcast(Asset, InInstance);
+		}
 		OpenedAssets.Remove(Assets[AssetIndex], InInstance);
 	}
 
diff --git a/Engine/Source/Editor/UnrealEd/Public/Subsystems/AssetEditorSubsystem.h b/Engine/Source/Editor/UnrealEd/Public/Subsystems/AssetEditorSubsystem.h
--- a/Engine/Source/Editor/UnrealEd/Public/Subsystems/AssetEditorSubsystem.h
+++ b/Engine/Source/Editor/UnrealEd/Public/Subsystems/AssetEditorSubsystem.h
@@ -168,7 +168,10 @@ public:
 	/** Remove given asset from all open editors */
 	UNREALED_API void RemoveAssetFromAllEditors(UObject* Asset);
 
-	/** Event called when CloseAllEditorsForAsset/RemoveAssetFromAllEditors is called */
+	/** Event called specifically when an external system requests an asset editor to be closed (e.g by calling CloseAllAssetEditors())
+	 *  If you want an event that is called anytime an asset is removed from an editor (which may or may not close the editor) - use
+	 *  OnAssetClosedInEditor()
+	 */
 	DECLARE_EVENT_TwoParams(UAssetEditorSubsystem, FAssetEditorRequestCloseEvent, UObject*, EAssetEditorCloseReason);
 	virtual FAssetEditorRequestCloseEvent& OnAssetEditorRequestClose() { return AssetEditorRequestCloseEvent; }
 
@@ -193,6 +196,10 @@ public:
 	/** Notify the asset editor manager that an asset editor is done editing an asset */
 	UNREALED_API void NotifyAssetClosed(UObject* Asset, IAssetEditorInstance* Instance);
 
+	/** Called when an editor is done editing an asset */
+	DECLARE_EVENT_TwoParams(UAssetEditorSubsystem, FOnAssetClosedInEditorEvent, UObject*, IAssetEditorInstance*);
+	virtual FOnAssetClosedInEditorEvent& OnAssetClosedInEditor() { return AssetClosedInEditorEvent; }
+
 	/** Notify the asset editor manager that an asset was closed */
 	UNREALED_API void NotifyEditorClosed(IAssetEditorInstance* Instance);
 
@@ -435,6 +442,9 @@ private:
 	/** Called when editor is opening and before widgets are constructed */
 	FOnAssetsOpenedInEditorEvent EditorOpeningPreWidgetsEvent;
 
+	/** Called when an editor is done editing an asset */
+	FOnAssetClosedInEditorEvent AssetClosedInEditorEvent;
+
 	/** Multicast delegate executed when an asset editor is requested to be opened */
 	FAssetEditorRequestOpenEvent AssetEditorRequestOpenEvent;
 
