TextureAssetAction to convert texture source to 8 bit

#rnx

diff --git a/Engine/Plugins/Editor/EngineAssetDefinitions/Source/Private/AssetDefinition_Texture.cpp b/Engine/Plugins/Editor/EngineAssetDefinitions/Source/Private/AssetDefinition_Texture.cpp
--- a/Engine/Plugins/Editor/EngineAssetDefinitions/Source/Private/AssetDefinition_Texture.cpp
+++ b/Engine/Plugins/Editor/EngineAssetDefinitions/Source/Private/AssetDefinition_Texture.cpp
@@ -60,7 +60,15 @@ namespace MenuExtension_Texture
 	{
 		if (const UContentBrowserAssetContextMenuContext* CBContext = UContentBrowserAssetContextMenuContext::FindContextWithAssets(InContext))
 		{
-			UE::TextureAssetActions::ResizeTextureSource_WithDialog(CBContext->LoadSelectedObjects<UTexture>());
+			UE::TextureAssetActions::TextureSource_Resize_WithDialog(CBContext->LoadSelectedObjects<UTexture>());
+		}
+	}
+	
+	static void Execute8bitTextureSource(const FToolMenuContext& InContext)
+	{
+		if (const UContentBrowserAssetContextMenuContext* CBContext = UContentBrowserAssetContextMenuContext::FindContextWithAssets(InContext))
+		{
+			UE::TextureAssetActions::TextureSource_ConvertTo8bit_WithDialog(CBContext->LoadSelectedObjects<UTexture>());
 		}
 	}
 
@@ -133,12 +141,20 @@ namespace MenuExtension_Texture
 						}
 						
 						{
-							const TAttribute<FText> Label = LOCTEXT("Texture_ResizeSource", "Resize Texture Source");
+							const TAttribute<FText> Label = LOCTEXT("Texture_ResizeSource", "Texture Source Reduce Size");
 							const TAttribute<FText> ToolTip = LOCTEXT("Texture_ResizeSourceTooltip", "Reduce texture asset size by shrinking the texture source dimensions.");
 							const FSlateIcon Icon = FSlateIcon(FAppStyle::GetAppStyleSetName(), "ClassIcon.Texture2D");
 							const FToolMenuExecuteAction UIAction = FToolMenuExecuteAction::CreateStatic(&ExecuteResizeTextureSource);
 							InSection.AddMenuEntry("Texture_ResizeSource", Label, ToolTip, Icon, UIAction);
 						}
+												
+						{
+							const TAttribute<FText> Label = LOCTEXT("Texture_ConvertTo8bit", "Texture Source Convert To 8 bit or minimum bit depth");
+							const TAttribute<FText> ToolTip = LOCTEXT("Texture_ConvertTo8bitTooltip", "Reduce texture asset size by converting 16/32 bit source data to 8 bit or minimum compatible bit depth.");
+							const FSlateIcon Icon = FSlateIcon(FAppStyle::GetAppStyleSetName(), "ClassIcon.Texture2D");
+							const FToolMenuExecuteAction UIAction = FToolMenuExecuteAction::CreateStatic(&Execute8bitTextureSource);
+							InSection.AddMenuEntry("Texture_ConvertTo8bit", Label, ToolTip, Icon, UIAction);
+						}
 					}
 				}));
 			}
diff --git a/Engine/Source/Editor/UnrealEd/Private/TextureAssetActions.cpp b/Engine/Source/Editor/UnrealEd/Private/TextureAssetActions.cpp
--- a/Engine/Source/Editor/UnrealEd/Private/TextureAssetActions.cpp
+++ b/Engine/Source/Editor/UnrealEd/Private/TextureAssetActions.cpp
@@ -33,11 +33,24 @@
 
 #define LOCTEXT_NAMESPACE "TextureAssetActions"
 
-struct FTextureAssetActionStatus
+namespace
 {
-	bool Enabled = true;
-	bool UnderSized = false;
-	bool WrongType = false;
+
+enum class ETextureAction
+{
+	Invalid = 0,
+	Resize,
+	ConvertTo8bit
+};
+
+enum class EAssetActionStatus
+{
+	Enabled = 0,
+	UnderSized,
+	WrongType,
+	HasLayers,
+	Already8bit,
+	DontChangeJPEG
 };
 
 /**
@@ -54,7 +67,7 @@ public:
 		Confirm = 1,		// Yes/Ok/Etc, normal usage would continue with action
 	};
 
-	STextureActionDlg(const TArray<UTexture*>& Textures);
+	STextureActionDlg(const TArray<UTexture*>& Textures, ETextureAction Act);
 
 	/**  Shows the dialog box and waits for the user to respond. */
 	EResult ShowModal();
@@ -77,7 +90,7 @@ public:
 	SLATE_END_ARGS()
 
 	void Construct(const FArguments& InArgs);
-	void Init(const TArray<UTexture *> &Textures);
+	void Init(const TArray<UTexture *> &Textures, ETextureAction Act);
 
 	STextureActionDlg::EResult GetUserResponse() const;
 
@@ -89,7 +102,7 @@ private:
 	/**
 	* Creates a single line showing an asset and it's status related to VT conversion
 	*/
-	TSharedRef<SWidget> CreateAssetLine(int index, const FAssetData &Asset, const FTextureAssetActionStatus &Status);
+	TSharedRef<SWidget> CreateAssetLine(int index, const FAssetData &Asset, const EAssetActionStatus Status);
 
 	FReply OnButtonClick(STextureActionDlg::EResult ButtonID);
 
@@ -112,6 +125,20 @@ private:
 	EVisibility GetErrorMessageVisibility() const;
 
 	EVisibility GetThresholdVisibility() const;
+	
+	EVisibility GetNMK16Visibility() const
+	{
+		return ( Action == ETextureAction::ConvertTo8bit ) ? EVisibility::Visible : EVisibility::Hidden;
+	}
+	ECheckBoxState HandleNMK16CheckBoxIsChecked() const
+	{
+		return bNormalMapsKeep16bits ? ECheckBoxState::Checked : ECheckBoxState::Unchecked;
+	}
+	void HandleNMK16CheckBoxCheckedStateChanged(ECheckBoxState InNewState)
+	{
+		bNormalMapsKeep16bits = ( InNewState == ECheckBoxState::Checked );
+		UpdateList();
+	}
 
 	FText GetIntroMessage() const;
 
@@ -121,7 +148,7 @@ private:
 
 	FText IntroMessage;
 	FText ErrorMessage;
-	STextureActionDlg::EResult	 UserResponse;
+	STextureActionDlg::EResult	 UserResponse = STextureActionDlg::EResult::Cancel;
 
 	TSharedPtr<SVerticalBox>	 AssetListContainer;
 	//TSharedPtr<STextBlock>		 MessageTextBlock;
@@ -135,34 +162,53 @@ private:
 	struct TextureListEntry
 	{
 		UTexture * Texture;
-		bool Enabled;
+		bool IsEnabled;
 	};
 
 	TArray<TextureListEntry> TextureList;
-	TArray<FAssetData> AssetList;
-	TArray<FTextureAssetActionStatus> AssetStatus;
 
-	int ThresholdValue;
-	bool bThresholdVisible;
+	int ThresholdValue = 0;
+	bool bThresholdVisible = true;
+
+	bool bNormalMapsKeep16bits = true;
 
 	TArray<TSharedPtr<int32>> TextureSizes;
+
+	ETextureAction Action = ETextureAction::Invalid;
 };
 
 
-void STextureAssetList::Construct(const FArguments& InArgs)
+static FText TAA_DialogTitle(ETextureAction Act)
 {
-	IntroMessage = LOCTEXT("TAA_Intro", "Reduce size of Texture Source to compact uassets.  Resizing is done using mip filter.  LODBias is adjusted but platform built size may change.");
+	switch(Act)
+	{
+	case ETextureAction::Resize: return LOCTEXT("TAA_Title_Resize", "Texture Asset : Resize Source");
+	case ETextureAction::ConvertTo8bit: return LOCTEXT("TAA_Title_Convert", "Texture Asset : Convert To 8 bit or minimum viable bit depth");
+	default: check(0); return FText();
+	}
+}
+
+static FText TAA_Intro(ETextureAction Act)
+{
+	switch(Act)
+	{
+	case ETextureAction::Resize: return LOCTEXT("TAA_Intro_Resize", "Reduce size of Texture Source to compact uassets.  Resizing is done using mip filter.  LODBias is adjusted but platform built size may change.");
+	case ETextureAction::ConvertTo8bit: return LOCTEXT("TAA_Intro_Convert", "Convert Texture Source to 8 bit, Normals to 8 or 16, HDR to 16F.  Output built texture may change.  Make sure CompressionSetting and SRGB are set correctly first!");
+	default: check(0); return FText();
+	}
+}
 
+void STextureAssetList::Construct(const FArguments& InArgs)
+{	
 	UserResponse = STextureActionDlg::Cancel;
 	ParentWindow = InArgs._ParentWindow.Get();
 	static FName ErrorIcon = "MessageLog.Error";
 
-
 	for (int i = 0; i < 16; i++)
 	{
 		TextureSizes.Add(MakeShareable(new int32(1 << i)));
 	}
-	ThresholdValue = *TextureSizes[10];
+	ThresholdValue = *TextureSizes[10]; // 1024
 
 	this->ChildSlot[
 		SNew(SVerticalBox)
@@ -227,6 +273,7 @@ void STextureAssetList::Construct(const FArguments& InArgs)
 			.Padding(8.0f, 4.0f, 8.0f, 4.0f)
 			[
 				SNew(SHorizontalBox)
+				.Visibility(this, &STextureAssetList::GetThresholdVisibility)
 				+SHorizontalBox::Slot()
 				.AutoWidth()
 				.VAlign(VAlign_Center)
@@ -250,6 +297,31 @@ void STextureAssetList::Construct(const FArguments& InArgs)
 					]
 				]
 			]
+		+ SVerticalBox::Slot()
+			.AutoHeight()
+			.HAlign(HAlign_Left)
+			.Padding(8.0f, 4.0f, 8.0f, 4.0f)
+			[
+				SNew(SHorizontalBox)
+				.Visibility(this, &STextureAssetList::GetNMK16Visibility)
+				+SHorizontalBox::Slot()
+				.AutoWidth()
+				.VAlign(VAlign_Center)
+				[
+					SNew(STextBlock)
+					.Text(LOCTEXT("TAA_NormalMapsKeep16bits", "Normal Maps keep 16 bits: "))
+					.AutoWrapText(true)
+				]
+				+ SHorizontalBox::Slot()
+				.Padding(4.0f, 0.0f, 2.0f, 0.0f)
+				.VAlign(VAlign_Center)
+				.AutoWidth()
+				[
+					SNew(SCheckBox)
+					.IsChecked(this, &STextureAssetList::HandleNMK16CheckBoxIsChecked)
+					.OnCheckStateChanged(this, &STextureAssetList::HandleNMK16CheckBoxCheckedStateChanged)
+				]
+			]
 		// Separator
 		+ SVerticalBox::Slot()
 			.AutoHeight()
@@ -289,15 +361,22 @@ void STextureAssetList::Construct(const FArguments& InArgs)
 
 	// will be done by Init :
 	//UpdateList();
+
 }
 
-void STextureAssetList::Init(const TArray<UTexture *> &Textures)
+void STextureAssetList::Init(const TArray<UTexture *> &Textures, ETextureAction InAction)
 {
+	Action = InAction;
+	
+	bThresholdVisible = (Action == ETextureAction::Resize );
+
+	IntroMessage = TAA_Intro(Action);
+
 	TextureList.SetNum(Textures.Num());
 	for(int i=0;i<Textures.Num();i++)
 	{
 		TextureList[i].Texture = Textures[i];
-		TextureList[i].Enabled = true;
+		TextureList[i].IsEnabled = true;
 	}
 
 	UpdateList();
@@ -308,26 +387,46 @@ STextureActionDlg::EResult STextureAssetList::GetUserResponse() const
 	return UserResponse;
 }
 
-TSharedRef<SWidget> STextureAssetList::CreateAssetLine(int index, const FAssetData &Asset, const FTextureAssetActionStatus &Status)
+static FText AssetActionStatus_Text(EAssetActionStatus Status)
 {
-	//const bool bEngineAsset = Asset.PackagePath.ToString().StartsWith(TEXT("/Engine/")); // use FPackageName::SplitPackageNameRoot
+	switch(Status)
+	{
+	case EAssetActionStatus::Enabled: return FText();
+	case EAssetActionStatus::WrongType: return LOCTEXT("TAAStatus_WrongType", "The texture is not a supported type.");
+	case EAssetActionStatus::UnderSized: return LOCTEXT("TAAStatus_UnderSized", "The texture was under the threshold size.");
+	case EAssetActionStatus::HasLayers: return LOCTEXT("TAAStatus_HasLayers", "Textures with more than 1 layer not supported.");
+	case EAssetActionStatus::Already8bit: return LOCTEXT("TAAStatus_Already8", "Texture format is already minimum bit depth.");
+	case EAssetActionStatus::DontChangeJPEG: return LOCTEXT("TAAStatus_JPEG", "Texture source is JPEG compressed, will not change."); 
+	default:
+		check(0);
+		return FText();
+	}	
+}
 
+TSharedRef<SWidget> STextureAssetList::CreateAssetLine(int index, const FAssetData &Asset, const EAssetActionStatus Status)
+{
 	FName SeverityIcon = NAME_None;
-	FText DetailedInfoText;
+	FText DetailedInfoText = AssetActionStatus_Text(Status);
 
-	if (Status.WrongType)
-	{
-		SeverityIcon = "MessageLog.Error";
-		DetailedInfoText = LOCTEXT("TAA_WrongType", "The texture is not a supported type.");
-	}
-	else if (Status.UnderSized)
+	bool IsEnabled = Status == EAssetActionStatus::Enabled;
+
+	if ( ! IsEnabled )
 	{
-		SeverityIcon = "MessageLog.Note";
-		DetailedInfoText = LOCTEXT("TAA_UnderSized", "The texture was under the threshold size.");
+		if ( Status == EAssetActionStatus::UnderSized ||
+			Status == EAssetActionStatus::Already8bit )
+		{
+			SeverityIcon = "MessageLog.Note";
+		}
+		else
+		{
+			SeverityIcon = "MessageLog.Error";
+		}
 	}
+
 	/*
+	//const bool bEngineAsset = Asset.PackagePath.ToString().StartsWith(TEXT("/Engine/")); // use FPackageName::SplitPackageNameRoot
 	// ConvertToVT does this, currently we do not
-	else if (bEngineAsset) // @@??
+	else if (bEngineAsset) // ??
 	{
 		SeverityIcon = "MessageLog.Note";
 		DetailedInfoText = LOCTEXT("TAA_EngineAsset", "The texture is an engine asset, a copy will be created in the current project.");
@@ -407,7 +506,7 @@ TSharedRef<SWidget> STextureAssetList::CreateAssetLine(int index, const FAssetDa
 					[
 						SNew(STextBlock)
 						.Text(FText::FromString(Asset.GetObjectPathString()))
-						.ColorAndOpacity( Status.Enabled ? FSlateColor::UseForeground() : FSlateColor::UseSubduedForeground())
+						.ColorAndOpacity( IsEnabled ? FSlateColor::UseForeground() : FSlateColor::UseSubduedForeground())
 					] // change color for enabled/not
 				]
 			]
@@ -482,6 +581,154 @@ static void DoResizeTextureSource(UTexture * Texture,int TargetSize)
 	Texture->PostEditChange();
 }
 
+static ETextureSourceFormat GetReducedTextureSourceFormat(TextureCompressionSettings TC,ETextureSourceFormat InTSF,bool NormalMapsKeep16bits)
+{
+	if ( InTSF == TSF_BGRE8 )
+	{
+		// don't change BGRE
+		return InTSF;
+	}
+
+	bool Out8bit = false;
+	bool OutBC45 = false;
+	bool OutHDR = false;
+	bool OutSingleChannel = false;
+
+	switch(TC)
+	{
+	case TC_Grayscale				: //"Grayscale (G8/16, RGB8 sRGB)"),
+	case TC_Displacementmap			: //"Displacementmap (G8/16)"),
+		// Gray and Displacement pass through G16 ; note they do not do that for RGBA16 (see GetDefaultTextureFormatName)
+		if ( InTSF == TSF_G16 ) return InTSF;
+		// otherwise we will convert to G8
+	case TC_DistanceFieldFont		: //"DistanceFieldFont (G8)"),
+		Out8bit = true;
+		OutSingleChannel = true;
+		break;
+
+	case TC_Default					: //"Default (DXT1/5, BC1/3 on DX11)"),
+	case TC_Masks					: //"Masks (no sRGB)"),
+	case TC_VectorDisplacementmap	: //"VectorDisplacementmap (RGBA8)"),
+	case TC_EditorIcon				: //"UserInterface2D (RGBA)"),
+	case TC_BC7						: //"BC7 (DX11, optional A)"),
+	case TC_LQ				        : // "Low Quality (BGR565/BGR555A1)", ToolTip = "BGR565/BGR555A1, fallback to DXT1/DXT5 on Mac platform"),
+		Out8bit = true;
+		break;
+
+	case TC_Normalmap				: //"Normalmap (DXT5, BC5 on DX11)"),
+		OutBC45 = true;
+		break;
+
+	case TC_Alpha					: //"Alpha (no sRGB, BC4 on DX11)"),
+		OutBC45 = true;
+		OutSingleChannel = true;
+		break;
+		
+	case TC_HDR						: //"HDR (RGBA16F, no sRGB)"),
+	case TC_HDR_Compressed			: //"HDR Compressed (RGB, BC6H, DX11)"),
+		OutHDR = true;
+		break;
+
+	case TC_HalfFloat				: //"Half Float (R16F)"),
+	case TC_SingleFloat				: //"Single Float (R32F)"),
+		OutHDR = true;
+		OutSingleChannel = true;
+		break;
+
+	case TC_EncodedReflectionCapture: 
+	case TC_HDR_F32					: //"HDR High Precision (RGBA32F)"),
+		// don't change :
+		return InTSF;
+
+	default:
+		check(0);
+		return InTSF;
+	}
+
+	if ( OutBC45 && ! NormalMapsKeep16bits )
+	{
+		// if NormalMaps don't keep 16 bit sources, then just treat them like 8 bit :
+		Out8bit = true;
+	}
+
+	ETextureSourceFormat OutTSF = InTSF;
+
+	if ( Out8bit )
+	{
+		OutTSF = ( OutSingleChannel ) ? TSF_G8 : TSF_BGRA8;		
+	}
+	else if ( OutBC45 )
+	{
+		check( NormalMapsKeep16bits );
+
+		// just choose a 16 bit output format, even if source was 8 bit
+		//	we will only convert if bytes per pixel goes down
+		OutTSF = ( OutSingleChannel ) ? TSF_G16 : TSF_RGBA16;		
+
+		if ( OutSingleChannel && InTSF == TSF_BGRA8 )
+		{
+			// don't do BGRA8 -> G16 , use G8 instead
+			OutTSF = TSF_G8;
+		}
+	}
+	else
+	{
+		check(OutHDR);
+		check( TC != TC_HDR_F32 ); // already handled
+
+		if ( TC == TC_SingleFloat &&
+			(InTSF == TSF_RGBA32F || InTSF == TSF_R32F) )
+		{
+			// 32 bit output, 32 bit input, keep it 32 bit
+			OutTSF = TSF_R32F;
+		}
+		else
+		{
+			OutTSF = ( OutSingleChannel ) ? TSF_R16F : TSF_RGBA16F;	
+		}
+	}
+
+	int64 InBPP  = FTextureSource::GetBytesPerPixel(InTSF);
+	int64 OutBPP = FTextureSource::GetBytesPerPixel(OutTSF);
+
+	if ( InBPP <= OutBPP )
+	{
+		// if bytes per pixel didn't go down, don't change
+		return InTSF;
+	}
+	else
+	{
+		// reducing
+		return OutTSF;
+	}
+}
+
+static void DoConvertTo8bitTextureSource(UTexture * Texture,bool NormalMapsKeep16bits)
+{
+	if ( Texture->Source.GetNumLayers() > 1 )
+	{
+		check(0);
+		return;
+	}
+
+	ETextureSourceFormat InTSF = Texture->Source.GetFormat();
+	ETextureSourceFormat OutTSF = GetReducedTextureSourceFormat(Texture->CompressionSettings,InTSF,NormalMapsKeep16bits);
+	if ( InTSF == OutTSF )
+	{
+		return;
+	}
+	
+	UE_LOG(LogTexture, Display, TEXT("Texture (%s) changing format from %s to %s, TC = %s"), 
+		*Texture->GetName() ,
+		*StaticEnum<ETextureSourceFormat>()->GetDisplayNameTextByValue(InTSF).ToString(),
+		*StaticEnum<ETextureSourceFormat>()->GetDisplayNameTextByValue(OutTSF).ToString(),
+		*StaticEnum<TextureCompressionSettings>()->GetDisplayNameTextByValue(Texture->CompressionSettings).ToString()
+		);
+
+	// calls Pre/PostEditChange :
+	UE::TextureUtilitiesCommon::Experimental::ChangeTextureSourceFormat(Texture,OutTSF);	
+}
+
 void STextureAssetList::DoAction()
 {
 	int NumEnabled = 0;
@@ -489,9 +736,9 @@ void STextureAssetList::DoAction()
 	{
 		UTexture * Texture = Entry.Texture;
 
-		UE_LOG(LogTexture, Display, TEXT("Texture (%s) Enabled=%d"), *Texture->GetName() , (int)Entry.Enabled);
+		UE_LOG(LogTexture, Display, TEXT("Texture (%s) Enabled=%d"), *Texture->GetName() , (int)Entry.IsEnabled);
 
-		if ( Entry.Enabled )
+		if ( Entry.IsEnabled )
 		{
 			NumEnabled ++;	
 		}
@@ -502,69 +749,98 @@ void STextureAssetList::DoAction()
 		return;
 	}
 
-	FScopedSlowTask Progress(NumEnabled, LOCTEXT("ResizingTextures", "Resizing Textures ..."));
+	FScopedSlowTask Progress(NumEnabled, LOCTEXT("TAA_Progress", "Applying action to TextureSources ..."));
 	Progress.MakeDialog(/*ShowCancelButton*/true);
 
 	for (TextureListEntry & Entry : TextureList)
 	{
-		if ( ! Entry.Enabled ) continue;
+		if ( ! Entry.IsEnabled ) continue;
 
 		UTexture * Texture = Entry.Texture;
 
-		UE_LOG(LogTexture, Display, TEXT("Texture (%s) Resizing to <= %d"), *Texture->GetName() , ThresholdValue);
-				
 		Progress.EnterProgressFrame(1.f);
 		if (Progress.ShouldCancel())
 		{
 			break;
 		}
 
-		DoResizeTextureSource(Texture,ThresholdValue);
+		switch(Action)
+		{
+		case ETextureAction::Resize:
+			UE_LOG(LogTexture, Display, TEXT("Texture (%s) Resizing to <= %d"), *Texture->GetName() , ThresholdValue);
+			DoResizeTextureSource(Texture,ThresholdValue);
+			break;
+		case ETextureAction::ConvertTo8bit:
+			UE_LOG(LogTexture, Display, TEXT("Texture (%s) ConvertTo8bit"), *Texture->GetName());
+			DoConvertTo8bitTextureSource(Texture,bNormalMapsKeep16bits);
+			break;
+		default:
+			check(0);
+			break;
+		}
 	}
 }
 
 void STextureAssetList::UpdateList()
 {
-	AssetList.Empty();
-	AssetStatus.Empty();
+	AssetListContainer->ClearChildren();
 
 	// filter select textures to see if they should be acted on
-	for (TextureListEntry & Entry : TextureList)
+	for (int Index=0;Index<TextureList.Num();Index++)
 	{
+		TextureListEntry & Entry = TextureList[Index];
 		UTexture * Texture = Entry.Texture;
 
-		AssetList.Add(Texture);
-		FTextureAssetActionStatus* Status = new(AssetStatus) FTextureAssetActionStatus();
-		
-		Entry.Enabled = true;
+		FAssetData AssetData(Texture);
+		EAssetActionStatus Status = EAssetActionStatus::Enabled; 
 
-		ETextureClass Class = Texture->GetTextureClass();
-		if ( Class != ETextureClass::TwoD && Class != ETextureClass::Cube ) //  Array ?
+		if ( Texture->Source.GetNumLayers() != 1 ) // only 1 layer textures
 		{
-			Status->WrongType = true;
-			Entry.Enabled = false;
+			Status = EAssetActionStatus::HasLayers;
 		}
-		
-		FIntPoint SourceSize = Texture->Source.GetLogicalSize();
-		int MaxSize = FMath::Max(SourceSize.X,SourceSize.Y);
-		if ( MaxSize <= ThresholdValue )
+		else if( Action == ETextureAction::Resize )
 		{
-			Entry.Enabled = false;
-			Status->UnderSized = true;
+			ETextureClass Class = Texture->GetTextureClass();
+			if ( Class != ETextureClass::TwoD && Class != ETextureClass::Cube ) //  Array ?
+			{
+				Status = EAssetActionStatus::WrongType;
+			}
+			else
+			{
+				// note: this does not match the logic in TextureSourceDataUtils
+				//	(@@FIXME) (eg. udim)
+				FIntPoint SourceSize = Texture->Source.GetLogicalSize();
+				int MaxSize = FMath::Max(SourceSize.X,SourceSize.Y);
+				if ( MaxSize <= ThresholdValue )
+				{
+					Status = EAssetActionStatus::UnderSized;
+				}
+			}
+		}
+		else if ( Action == ETextureAction::ConvertTo8bit )
+		{
+			ETextureSourceFormat InTSF = Texture->Source.GetFormat();
+			ETextureSourceFormat OutTSF = GetReducedTextureSourceFormat(Texture->CompressionSettings,InTSF,bNormalMapsKeep16bits);
+			if ( InTSF == OutTSF )
+			{
+				Status = EAssetActionStatus::Already8bit;
+			}
+			else if ( Texture->Source.GetSourceCompression() == ETextureSourceCompressionFormat::TSCF_JPEG )
+			{
+				Status = EAssetActionStatus::DontChangeJPEG;
+			}
+		}
+		else
+		{
+			check(0);
 		}
 
-		Status->Enabled = Entry.Enabled;
-	}
-
-	check(AssetList.Num() == AssetStatus.Num());
-	AssetListContainer->ClearChildren();
+		Entry.IsEnabled = ( Status == EAssetActionStatus::Enabled );
 
-	for (int Id = 0; Id < AssetList.Num(); Id++)
-	{
 		AssetListContainer->AddSlot()
 			.AutoHeight()
 			[
-				CreateAssetLine(Id, AssetList[Id], AssetStatus[Id])
+				CreateAssetLine(Index, AssetData, Status)
 			];
 	}
 
@@ -652,7 +928,7 @@ EVisibility STextureAssetList::GetErrorMessageVisibility() const
 
 EVisibility STextureAssetList::GetThresholdVisibility() const
 {
-	return (bThresholdVisible) ? EVisibility::Collapsed : EVisibility::Visible;
+	return (bThresholdVisible) ? EVisibility::Visible : EVisibility::Collapsed;
 }
 
 FText STextureAssetList::GetIntroMessage() const
@@ -665,12 +941,12 @@ FText STextureAssetList::GetErrorMessage() const
 	return ErrorMessage;
 }
 
-STextureActionDlg::STextureActionDlg(const TArray<UTexture *> &Textures)
+STextureActionDlg::STextureActionDlg(const TArray<UTexture *> &Textures, ETextureAction Act)
 {
 	if (FSlateApplication::IsInitialized())
 	{
 		DialogWindow = SNew(SWindow)
-			.Title(LOCTEXT("TAA_Title", "Texture Asset : Resize Source"))
+			.Title(TAA_DialogTitle(Act))
 			.SupportsMinimize(false).SupportsMaximize(false)
 			.ClientSize(FVector2D(500, 500));
 
@@ -684,7 +960,7 @@ STextureActionDlg::STextureActionDlg(const TArray<UTexture *> &Textures)
 			];
 
 		// Init will do UpdateList :
-		DialogWidget->Init(Textures);
+		DialogWidget->Init(Textures,Act);
 		DialogWindow->SetContent(DialogWrapper.ToSharedRef());
 	}
 }
@@ -698,10 +974,17 @@ STextureActionDlg::EResult STextureActionDlg::ShowModal()
 	return UserResponse;
 }
 
+}; // namespace
+
+void UE::TextureAssetActions::TextureSource_Resize_WithDialog(const TArray<UTexture*> & InTextures)
+{
+	STextureActionDlg Dlg(InTextures,ETextureAction::Resize);
+	Dlg.ShowModal();
+}
 
-void UE::TextureAssetActions::ResizeTextureSource_WithDialog(const TArray<UTexture*> & InTextures)
+void UE::TextureAssetActions::TextureSource_ConvertTo8bit_WithDialog(const TArray<UTexture*> & InTextures)
 {
-	STextureActionDlg Dlg(InTextures);
+	STextureActionDlg Dlg(InTextures,ETextureAction::ConvertTo8bit);
 	Dlg.ShowModal();
 }
 
diff --git a/Engine/Source/Editor/UnrealEd/Public/TextureAssetActions.h b/Engine/Source/Editor/UnrealEd/Public/TextureAssetActions.h
--- a/Engine/Source/Editor/UnrealEd/Public/TextureAssetActions.h
+++ b/Engine/Source/Editor/UnrealEd/Public/TextureAssetActions.h
@@ -11,7 +11,9 @@ namespace UE
 namespace TextureAssetActions
 {
 
-UNREALED_API void ResizeTextureSource_WithDialog(const TArray<UTexture*> & InTextures);
+UNREALED_API void TextureSource_Resize_WithDialog(const TArray<UTexture*> & InTextures);
+
+UNREALED_API void TextureSource_ConvertTo8bit_WithDialog(const TArray<UTexture*> & InTextures);
 
 }
 };
diff --git a/Engine/Source/Runtime/TextureUtilitiesCommon/Private/TextureSourceDataUtils.cpp b/Engine/Source/Runtime/TextureUtilitiesCommon/Private/TextureSourceDataUtils.cpp
--- a/Engine/Source/Runtime/TextureUtilitiesCommon/Private/TextureSourceDataUtils.cpp
+++ b/Engine/Source/Runtime/TextureUtilitiesCommon/Private/TextureSourceDataUtils.cpp
@@ -16,13 +16,15 @@ namespace Private
 {
 
 	// resize so that the largest dimension is <= MaxSize
-	bool ResizeTexture2D(UTexture* Texture, int32 MaxSize, const ITargetPlatform* TargetPlatform)
+	static bool ResizeTexture2D(UTexture* Texture, int32 MaxSize, const ITargetPlatform* TargetPlatform)
 	{
 		// We want to reduce the asset size so ignore the imported mip(s)
 		const int32 MipIndex = 0;
 		FImage SourceMip0;
 		if (!Texture->Source.GetMipImage(SourceMip0, MipIndex))
 		{
+			UE_LOG(LogTexture,Error,TEXT("ResizeTexture2D: Texture GetMipImage failed [%s]"),
+				*Texture->GetFullName());
 			return false;
 		}
 
@@ -32,7 +34,8 @@ namespace Private
 		bool MadeChanges;
 		if ( ! Texture->DownsizeImageUsingTextureSettings(TargetPlatform, SourceMip0, MaxSize, LayerIndex, MadeChanges) )
 		{
-			// a critical error
+			UE_LOG(LogTexture,Error,TEXT("ResizeTexture2D: Texture DownsizeImageUsingTextureSettings failed [%s]"),
+				*Texture->GetFullName());
 			return false;
 		}
 		if ( ! MadeChanges )
@@ -63,8 +66,29 @@ namespace Private
 
 		return true;
 	}
+	
+	// concatenate all the image payloads into one bulkdata, eg. for mips or blocks
+	static UE::Serialization::FEditorBulkData::FSharedBufferWithID MakeSharedBufferForImageDatas(const TArray<FImage> & InImages)
+	{
+		int64 SizeNeededInBytes = 0;
+		for (const FImage& Im : InImages)
+		{
+			check( Im.RawData.Num() == Im.GetImageSizeBytes() );
+			SizeNeededInBytes += Im.RawData.Num(); 
+		}
+		FUniqueBuffer WriteImageBuffer = FUniqueBuffer::Alloc(SizeNeededInBytes);
 
-	bool ResizeTexture2DBlocked(UTexture* Texture, int32 MaxSize, const ITargetPlatform* TargetPlatform)
+		uint8* CurrentAddress = static_cast<uint8*>(WriteImageBuffer.GetData());
+		for (const FImage& Im : InImages)
+		{
+			FMemory::Memcpy(CurrentAddress, Im.RawData.GetData(), Im.RawData.Num());
+			CurrentAddress += Im.RawData.Num();
+		}
+
+		return WriteImageBuffer.MoveToShared();
+	}
+
+	static bool ResizeTexture2DBlocked(UTexture* Texture, int32 MaxSize, const ITargetPlatform* TargetPlatform)
 	{
 		// note: does not support layers
 
@@ -83,7 +107,7 @@ namespace Private
 		}
 		*/
 
-		TArray<FTextureSourceBlock > ResizedSourceBlocks;
+		TArray<FTextureSourceBlock> ResizedSourceBlocks;
 		ResizedSourceBlocks.Reserve(Texture->Source.GetNumBlocks());
 	
 		TArray<FImage> ResizedBlocks;
@@ -99,6 +123,8 @@ namespace Private
 			const int32 LayerIndex = 0;
 			if (!Texture->Source.GetMipImage(SourceMip0, BlockIndex, LayerIndex, MipIndex))
 			{
+				UE_LOG(LogTexture,Error,TEXT("ResizeTexture2DBlocked: Texture GetMipImage failed [%s]"),
+					*Texture->GetFullName());
 				return false;
 			}
 
@@ -109,6 +135,8 @@ namespace Private
 			if ( ! Texture->DownsizeImageUsingTextureSettings(TargetPlatform, SourceMip0, MaxSize, LayerIndex, MadeChanges) )
 			{
 				// critical error
+				UE_LOG(LogTexture,Error,TEXT("ResizeTexture2DBlocked: Texture DownsizeImageUsingTextureSettings failed [%s]"),
+					*Texture->GetFullName());
 				return false;
 			}
 			MadeAnyChanges = MadeAnyChanges || MadeChanges;
@@ -132,26 +160,12 @@ namespace Private
 		// Protect the code from an async build of the texture
 		Texture->PreEditChange(nullptr);
 
-		int64 SizeNeededInBytes = 0;
-		for (const FImage& Block : ResizedBlocks)
-		{
-			SizeNeededInBytes += Block.RawData.Num(); // Block.GetImageSizeBytes()
-		}
-		FUniqueBuffer WriteImageBuffer = FUniqueBuffer::Alloc(SizeNeededInBytes);
-
-		uint8* CurrentAddress = static_cast<uint8*>(WriteImageBuffer.GetData());
-		for (FImage& Block : ResizedBlocks)
-		{
-			FMemory::Memcpy(CurrentAddress, static_cast<uint8*>(Block.RawData.GetData()), Block.RawData.Num());
-			CurrentAddress += Block.RawData.Num();
-		}
-
-		UE::Serialization::FEditorBulkData::FSharedBufferWithID ResizedImageBufferWithID = WriteImageBuffer.MoveToShared();
+		UE::Serialization::FEditorBulkData::FSharedBufferWithID ResizedImageBufferWithID = MakeSharedBufferForImageDatas(ResizedBlocks);
 
 		const ETextureSourceFormat SourceFormat = Texture->Source.GetFormat();
 		int32 NumLayers = 1;
 		Texture->Source.InitBlocked(
-			&SourceFormat,
+			&SourceFormat, // array of formats per layer
 			ResizedSourceBlocks.GetData(),
 			NumLayers,
 			ResizedSourceBlocks.Num(),
@@ -166,7 +180,7 @@ namespace Private
 }
 
 
-bool DownsizeTextureSourceData(UTexture* Texture, int32 TargetSourceSize, const ITargetPlatform* TargetPlatform)
+TEXTUREUTILITIESCOMMON_API bool DownsizeTextureSourceData(UTexture* Texture, int32 TargetSourceSize, const ITargetPlatform* TargetPlatform)
 {
 	check( Texture->Source.IsValid() );
 
@@ -229,7 +243,7 @@ bool DownsizeTextureSourceData(UTexture* Texture, int32 TargetSourceSize, const
 	return false;
 }
 
-bool DownsizeTexureSourceDataNearRenderingSize(UTexture* Texture, const ITargetPlatform* TargetPlatform)
+TEXTUREUTILITIESCOMMON_API bool DownsizeTextureSourceDataNearRenderingSize(UTexture* Texture, const ITargetPlatform* TargetPlatform)
 {
 	if ( ! Texture->Source.IsValid() )
 	{
@@ -263,7 +277,7 @@ bool DownsizeTexureSourceDataNearRenderingSize(UTexture* Texture, const ITargetP
 		if ( BeforeSizeX != AfterSizeX ||
 			 BeforeSizeY != AfterSizeY )
 		{
-			UE_LOG(LogTexture,Warning,TEXT("DownsizeTexureSourceDataNearRenderingSize failed to preserve built size; was: %dx%d now: %dx%d on [%s]"),
+			UE_LOG(LogTexture,Warning,TEXT("DownsizeTextureSourceDataNearRenderingSize failed to preserve built size; was: %dx%d now: %dx%d on [%s]"),
 				BeforeSizeX,BeforeSizeY,
 				AfterSizeX,AfterSizeY,
 				*Texture->GetFullName());
@@ -278,6 +292,115 @@ bool DownsizeTexureSourceDataNearRenderingSize(UTexture* Texture, const ITargetP
 	return false;
 }
 
+
+TEXTUREUTILITIESCOMMON_API bool ChangeTextureSourceFormat(UTexture* Texture, ETextureSourceFormat NewFormat)
+{
+	if ( ! Texture->Source.IsValid() )
+	{
+		return false;
+	}
+
+	// we only support 1 layer currently
+	if (Texture->Source.GetNumLayers() != 1)
+	{
+		return false;
+	}
+	
+	ETextureSourceFormat OldFormat = Texture->Source.GetFormat(0);
+	if ( OldFormat == NewFormat )
+	{
+		return false;
+	}
+
+	ERawImageFormat::Type NewRIF = FImageCoreUtils::ConvertToRawImageFormat(NewFormat);
+	EGammaSpace NewGamma = ( Texture->SRGB && ERawImageFormat::GetFormatNeedsGammaSpace(NewRIF) ) ? EGammaSpace::sRGB : EGammaSpace::Linear;
+
+	if ( Texture->Source.GetNumBlocks() == 1 && Texture->Source.GetNumMips() == 1 )
+	{
+		const int32 MipIndex = 0;
+		FImage SourceMip;
+		if (!Texture->Source.GetMipImage(SourceMip, MipIndex))
+		{
+			UE_LOG(LogTexture,Error,TEXT("ChangeTextureSourceFormat: Texture GetMipImage failed [%s]"),
+				*Texture->GetFullName());
+			return false;
+		}
+
+		FImage NewMip;
+		SourceMip.CopyTo(NewMip,NewRIF,NewGamma);
+		
+		UE::Serialization::FEditorBulkData::FSharedBufferWithID ResizedImageBufferWithID = MakeSharedBufferFromArray(MoveTemp(NewMip.RawData));
+		
+		Texture->PreEditChange(nullptr);
+
+		const int32 NumMips = 1;
+		Texture->Source.Init(NewMip.SizeX,NewMip.SizeY,NewMip.NumSlices
+			, NumMips
+			, NewFormat
+			, MoveTemp(ResizedImageBufferWithID));
+	}
+	else
+	{
+		// all blocks of a UDIM have the same format; Layers do not
+		int32 NumLayers = 1;
+		int32 LayerIndex = 0;
+
+		int32 NumBlocks = Texture->Source.GetNumBlocks();
+		check( NumBlocks >= 1 );
+
+		TArray<FTextureSourceBlock> NewBlocks;
+		NewBlocks.Reserve(NumBlocks);
+	
+		TArray<FImage> NewImages;
+		NewImages.Reserve(NumBlocks*16); // *16 for mips
+
+		for (int32 BlockIndex = 0; BlockIndex < NumBlocks; ++BlockIndex)
+		{
+			FTextureSourceBlock Block;
+			Texture->Source.GetBlock(BlockIndex,Block);
+
+			// NewBlocks has sizes, they don't change
+			NewBlocks.Add(Block);
+
+			for(int32 MipIndex=0; MipIndex < Block.NumMips;MipIndex++)
+			{
+				FImage SourceMip;
+				if ( ! Texture->Source.GetMipImage(SourceMip, BlockIndex, LayerIndex, MipIndex) )
+				{
+					UE_LOG(LogTexture,Error,TEXT("ChangeTextureSourceFormat: Texture GetMipImage failed [%s]"),
+						*Texture->GetFullName());
+
+					return false;
+				}
+				
+				FImage & NewMip = NewImages.AddDefaulted_GetRef();
+
+				SourceMip.CopyTo(NewMip,NewRIF,NewGamma);
+			}
+		}
+
+		UE::Serialization::FEditorBulkData::FSharedBufferWithID ResizedImageBufferWithID = Private::MakeSharedBufferForImageDatas(NewImages);
+		
+		Texture->PreEditChange(nullptr);
+
+		Texture->Source.InitBlocked(
+			&NewFormat, // array of formats per layer
+			NewBlocks.GetData(),
+			NumLayers,
+			NewBlocks.Num(),
+			MoveTemp(ResizedImageBufferWithID)
+		);
+	}
+
+	// if gamma was Pow22 it is now sRGB
+	Texture->bUseLegacyGamma = false;
+	Texture->PostEditChange();
+
+	check( Texture->Source.GetGammaSpace(0) == NewGamma );
+
+	return true;
+}
+
 } // End namespace UE::TextureUtilitiesCommon::Experimental
 
 #endif //WITH_EDITOR
diff --git a/Engine/Source/Runtime/TextureUtilitiesCommon/Public/TextureSourceDataUtils.h b/Engine/Source/Runtime/TextureUtilitiesCommon/Public/TextureSourceDataUtils.h
--- a/Engine/Source/Runtime/TextureUtilitiesCommon/Public/TextureSourceDataUtils.h
+++ b/Engine/Source/Runtime/TextureUtilitiesCommon/Public/TextureSourceDataUtils.h
@@ -4,6 +4,7 @@
 
 #if WITH_EDITOR
 #include "CoreTypes.h"
+#include "Engine/TextureDefines.h"
 
 class UTexture;
 class ITargetPlatform;
@@ -16,7 +17,12 @@ namespace UE::TextureUtilitiesCommon::Experimental
 	TEXTUREUTILITIESCOMMON_API bool DownsizeTextureSourceData(UTexture* Texture, int32 TargetSizeInGame, const ITargetPlatform* TargetPlatform);
 
 	//Note: This function does trigger the post edit change after modifying the source texture
-	TEXTUREUTILITIESCOMMON_API bool DownsizeTexureSourceDataNearRenderingSize(UTexture* Texture, const ITargetPlatform* TargetPlatform);
+	TEXTUREUTILITIESCOMMON_API bool DownsizeTextureSourceDataNearRenderingSize(UTexture* Texture, const ITargetPlatform* TargetPlatform);
+	
+	// ChangeTextureSourceFormat calls Pre/Post edit change
+	//	beware that changing format may change the interpretation of the SRGB bool in Texture
+	// ChangeTextureSourceFormat supports mips and blocks (udim) but not layers
+	TEXTUREUTILITIESCOMMON_API bool ChangeTextureSourceFormat(UTexture* Texture, ETextureSourceFormat NewFormat);
 }
 
 #endif // WITH_EDITOR
\ No newline at end of file
