Fix copy/pasted print statement mistake

diff --git a/Engine/Source/Editor/Blutility/Private/EditorUtilityWidgetBlueprintFactory.cpp b/Engine/Source/Editor/Blutility/Private/EditorUtilityWidgetBlueprintFactory.cpp
--- a/Engine/Source/Editor/Blutility/Private/EditorUtilityWidgetBlueprintFactory.cpp
+++ b/Engine/Source/Editor/Blutility/Private/EditorUtilityWidgetBlueprintFactory.cpp
@@ -127,7 +127,7 @@ UObject* UEditorUtilityWidgetBlueprintFactory::FactoryCreateNew(UClass* Class, U
 	{
 		FFormatNamedArguments Args;
 		Args.Add(TEXT("ClassName"), (ParentClass != nullptr) ? FText::FromString(ParentClass->GetName()) : NSLOCTEXT("UnrealEd", "Null", "(null)"));
-		FMessageDialog::Open(EAppMsgType::Ok, FText::Format(NSLOCTEXT("UnrealEd", "CannotCreateBlueprintFromClass", "Cannot create a blueprint based on the class '{0}'."), Args));
+		FMessageDialog::Open(EAppMsgType::Ok, FText::Format(NSLOCTEXT("UnrealEd", "CannotCreateBlueprintFromClass", "Cannot create a blueprint based on the class '{ClassName}'."), Args));
 		return nullptr;
 	}
 	else
diff --git a/Engine/Source/Editor/UnrealEd/Private/Factories/EditorFactories.cpp b/Engine/Source/Editor/UnrealEd/Private/Factories/EditorFactories.cpp
--- a/Engine/Source/Editor/UnrealEd/Private/Factories/EditorFactories.cpp
+++ b/Engine/Source/Editor/UnrealEd/Private/Factories/EditorFactories.cpp
@@ -7074,7 +7074,7 @@ UObject* UBlueprintFactory::FactoryCreateNew(UClass* Class, UObject* InParent, F
 	{
 		FFormatNamedArguments Args;
 		Args.Add( TEXT("ClassName"), (ParentClass != nullptr) ? FText::FromString( ParentClass->GetName() ) : LOCTEXT("Null", "(null)") );
-		FMessageDialog::Open( EAppMsgType::Ok, FText::Format( LOCTEXT("CannotCreateBlueprintFromClass", "Cannot create a blueprint based on the class '{0}'."), Args ) );
+		FMessageDialog::Open( EAppMsgType::Ok, FText::Format( LOCTEXT("CannotCreateBlueprintFromClass", "Cannot create a blueprint based on the class '{ClassName}'."), Args ) );
 		return nullptr;
 	}
 	else
@@ -7142,7 +7142,7 @@ UObject* UBlueprintMacroFactory::FactoryCreateNew(UClass* Class, UObject* InPare
 	{
 		FFormatNamedArguments Args;
 		Args.Add( TEXT("ClassName"), (ParentClass != nullptr) ? FText::FromString( ParentClass->GetName() ) : LOCTEXT("Null", "(null)") );
-		FMessageDialog::Open( EAppMsgType::Ok, FText::Format( LOCTEXT("CannotCreateBlueprintFromClass", "Cannot create a blueprint based on the class '{0}'."), Args ) );
+		FMessageDialog::Open( EAppMsgType::Ok, FText::Format( LOCTEXT("CannotCreateBlueprintFromClass", "Cannot create a blueprint based on the class '{ClassName}'."), Args ) );
 		return nullptr;
 	}
 	else
@@ -7205,7 +7205,7 @@ UObject* UBlueprintFunctionLibraryFactory::FactoryCreateNew(UClass* Class, UObje
 	{
 		FFormatNamedArguments Args;
 		Args.Add(TEXT("ClassName"), (ParentClass != nullptr) ? FText::FromString(ParentClass->GetName()) : LOCTEXT("Null", "(null)"));
-		FMessageDialog::Open(EAppMsgType::Ok, FText::Format(LOCTEXT("CannotCreateBlueprintFromClass", "Cannot create a blueprint based on the class '{0}'."), Args));
+		FMessageDialog::Open(EAppMsgType::Ok, FText::Format(LOCTEXT("CannotCreateBlueprintFromClass", "Cannot create a blueprint based on the class '{ClassName}'."), Args));
 		return nullptr;
 	}
 	else
@@ -7271,7 +7271,7 @@ UObject* UBlueprintInterfaceFactory::FactoryCreateNew(UClass* Class, UObject* In
 	{
 		FFormatNamedArguments Args;
 		Args.Add( TEXT("ClassName"), (ParentClass != nullptr) ? FText::FromString( ParentClass->GetName() ) : LOCTEXT("Null", "(null)") );
-		FMessageDialog::Open( EAppMsgType::Ok, FText::Format( LOCTEXT("CannotCreateBlueprintFromClass", "Cannot create a blueprint based on the class '{0}'."), Args ) );
+		FMessageDialog::Open( EAppMsgType::Ok, FText::Format( LOCTEXT("CannotCreateBlueprintFromClass", "Cannot create a blueprint based on the class '{ClassName}'."), Args ) );
 		return nullptr;
 	}
 	else
