Diff tool - When matching array elements before and after a change, prefer elements that have the same name and location, and relegate name-only matches as a fallback.
This fixes false positives when several elements have the same name

diff --git a/Engine/Source/Editor/UnrealEd/Private/DiffUtils.cpp b/Engine/Source/Editor/UnrealEd/Private/DiffUtils.cpp
--- a/Engine/Source/Editor/UnrealEd/Private/DiffUtils.cpp
+++ b/Engine/Source/Editor/UnrealEd/Private/DiffUtils.cpp
@@ -698,14 +698,23 @@ void DiffUtils::CompareUnrelatedSCS(const UBlueprint* Old, const TArray< FSCSRes
 {
 	const auto FindEntry = [](TArray< FSCSResolvedIdentifier > const& InArray, const FSCSIdentifier* Value) -> const FSCSResolvedIdentifier*
 	{
+		const FSCSResolvedIdentifier* BestMatch = nullptr;
+
 		for (const auto& Node : InArray)
 		{
 			if (Node.Identifier.Name == Value->Name)
 			{
-				return &Node;
+				if (Node.Identifier.TreeLocation == Value->TreeLocation)
+				{
+					return &Node;
+				}
+				else if (BestMatch == nullptr)
+				{
+					BestMatch = &Node;
+				}
 			}
 		}
-		return nullptr;
+		return BestMatch;
 	};
 
 	for (const auto& OldNode : OldHierarchy)
