Fix the skeletal mesh custom LOD import to properly reset the new vertices reduction flags.

#rnx

diff --git a/Engine/Source/Developer/SkeletalMeshUtilitiesCommon/Private/LODUtilities.cpp b/Engine/Source/Developer/SkeletalMeshUtilitiesCommon/Private/LODUtilities.cpp
--- a/Engine/Source/Developer/SkeletalMeshUtilitiesCommon/Private/LODUtilities.cpp
+++ b/Engine/Source/Developer/SkeletalMeshUtilitiesCommon/Private/LODUtilities.cpp
@@ -830,7 +830,11 @@ bool FLODUtilities::SetCustomLOD(USkeletalMesh* DestinationSkeletalMesh, USkelet
 	{
 		//Make sure any custom LOD have correct settings (no reduce)
 		LODInfo.ReductionSettings.NumOfTrianglesPercentage = 1.0f;
+		LODInfo.ReductionSettings.MaxNumOfTriangles = MAX_uint32;
+		LODInfo.ReductionSettings.MaxNumOfTrianglesPercentage = MAX_uint32;
 		LODInfo.ReductionSettings.NumOfVertPercentage = 1.0f;
+		LODInfo.ReductionSettings.MaxNumOfVerts = MAX_uint32;
+		LODInfo.ReductionSettings.MaxNumOfVertsPercentage = MAX_uint32;
 		LODInfo.ReductionSettings.MaxDeviationPercentage = 0.0f;
 	}
 
diff --git a/Engine/Source/Editor/UnrealEd/Private/Fbx/FbxSkeletalMeshImport.cpp b/Engine/Source/Editor/UnrealEd/Private/Fbx/FbxSkeletalMeshImport.cpp
--- a/Engine/Source/Editor/UnrealEd/Private/Fbx/FbxSkeletalMeshImport.cpp
+++ b/Engine/Source/Editor/UnrealEd/Private/Fbx/FbxSkeletalMeshImport.cpp
@@ -4232,7 +4232,11 @@ void UnFbx::FFbxImporter::InsertNewLODToBaseSkeletalMesh(USkeletalMesh* InSkelet
 	{
 		//Make sure any custom LOD have correct settings (no reduce)
 		LODInfo.ReductionSettings.NumOfTrianglesPercentage = 1.0f;
+		LODInfo.ReductionSettings.MaxNumOfTriangles = MAX_uint32;
+		LODInfo.ReductionSettings.MaxNumOfTrianglesPercentage = MAX_uint32;
 		LODInfo.ReductionSettings.NumOfVertPercentage = 1.0f;
+		LODInfo.ReductionSettings.MaxNumOfVerts = MAX_uint32;
+		LODInfo.ReductionSettings.MaxNumOfVertsPercentage = MAX_uint32;
 		LODInfo.ReductionSettings.MaxDeviationPercentage = 0.0f;
 	}
 
diff --git a/Engine/Source/Editor/UnrealEd/Private/FbxMeshUtils.cpp b/Engine/Source/Editor/UnrealEd/Private/FbxMeshUtils.cpp
--- a/Engine/Source/Editor/UnrealEd/Private/FbxMeshUtils.cpp
+++ b/Engine/Source/Editor/UnrealEd/Private/FbxMeshUtils.cpp
@@ -782,8 +782,12 @@ namespace FbxMeshUtils
 					//Update the import data for this lod
 					UnFbx::FFbxImporter::UpdateSkeletalMeshImportData(SelectedSkelMesh, nullptr, SelectedLOD, &ImportMaterialOriginalNameData, &ImportMeshLodData);
 
+					const FString SourceImportFilename = UAssetImportData::SanitizeImportFilename(Filename, nullptr);
 					if (SkelMeshDataPtr)
 					{
+						//Setting the source filename allow to not restore the reduction settings in case we import a custom LOD over a generated LOD.
+						//This value will be wipe during the restore but we put it back just after.
+						SelectedSkelMesh->GetLODInfo(SelectedLOD)->SourceImportFilename = SourceImportFilename;
 						SkeletalMeshImportUtils::RestoreExistingSkelMeshData(SkelMeshDataPtr, SelectedSkelMesh, SelectedLOD, false, ImportOptions->bImportAsSkeletalSkinning, ImportOptions->bResetToFbxOnMaterialConflict);
 					}
 
@@ -795,7 +799,7 @@ namespace FbxMeshUtils
 					bSuccess = true;
 
 					// Set LOD source filename
-					SelectedSkelMesh->GetLODInfo(SelectedLOD)->SourceImportFilename = UAssetImportData::SanitizeImportFilename(Filename, nullptr);
+					SelectedSkelMesh->GetLODInfo(SelectedLOD)->SourceImportFilename = SourceImportFilename;
 					SelectedSkelMesh->GetLODInfo(SelectedLOD)->bImportWithBaseMesh = false;
 
 					ReapplyClothing();
diff --git a/Engine/Source/Editor/UnrealEd/Private/ImportUtils/SkeletalMeshImportUtils.cpp b/Engine/Source/Editor/UnrealEd/Private/ImportUtils/SkeletalMeshImportUtils.cpp
--- a/Engine/Source/Editor/UnrealEd/Private/ImportUtils/SkeletalMeshImportUtils.cpp
+++ b/Engine/Source/Editor/UnrealEd/Private/ImportUtils/SkeletalMeshImportUtils.cpp
@@ -463,8 +463,8 @@ void SkeletalMesUtilsImpl::RestoreLODInfo(const TSharedPtr<const FExistingSkelMe
 	ImportedLODInfo.ScreenSize = ExistingLODInfo.ScreenSize;
 	ImportedLODInfo.LODHysteresis = ExistingLODInfo.LODHysteresis;
 	ImportedLODInfo.BuildSettings = ExistingLODInfo.BuildSettings;
-	//Old assets may have non-applied reduction settings, so only restore the reduction settings if the LOD was effectively reduced.
-	if (ExistingLODInfo.bHasBeenSimplified)
+	//Old assets may have non-applied reduction settings, so only restore the reduction settings if the LOD was effectively reduced and we did not import a custom LOD over this generated LOD.
+	if (ExistingLODInfo.bHasBeenSimplified && (!ExistingLODInfo.SourceImportFilename.IsEmpty() || ImportedLODInfo.SourceImportFilename.IsEmpty()))
 	{
 		ImportedLODInfo.ReductionSettings = ExistingLODInfo.ReductionSettings;
 	}
