Avoid delay and timeout in DerivedDataCacheCommandlet caused by an object getting into the CachingObjects container after having had BeginCacheForCookedPlatformData on it, but then being destroyed and recreated in place not through GC, but through NewObject and subobject recycling.  This could lead to subsequent calls to IsCachedCookedPlatformDataLoaded returning false perpetually.  Avoid  this situation by holding weak pointers so that if an object is destroyed and recreated in place, the reference to it becomes invalidated.

diff --git a/Engine/Source/Editor/UnrealEd/Classes/Commandlets/DerivedDataCacheCommandlet.h b/Engine/Source/Editor/UnrealEd/Classes/Commandlets/DerivedDataCacheCommandlet.h
--- a/Engine/Source/Editor/UnrealEd/Classes/Commandlets/DerivedDataCacheCommandlet.h
+++ b/Engine/Source/Editor/UnrealEd/Classes/Commandlets/DerivedDataCacheCommandlet.h
@@ -35,6 +35,9 @@ class UDerivedDataCacheCommandlet : public UCommandlet
 
 	struct FCachingData
 	{
+		/** Weak object reference to ensure that the object being referenced is still valid */
+		TWeakObjectPtr<UObject> ObjectValidityReference;
+
 		/** Element n == IsCachedCookedPlatformDataLoaded has returned true for Commandlet->Platforms[n] */
 		TBitArray<> PlatformIsComplete;
 		/**
diff --git a/Engine/Source/Editor/UnrealEd/Private/Commandlets/DerivedDataCacheCommandlet.cpp b/Engine/Source/Editor/UnrealEd/Private/Commandlets/DerivedDataCacheCommandlet.cpp
--- a/Engine/Source/Editor/UnrealEd/Private/Commandlets/DerivedDataCacheCommandlet.cpp
+++ b/Engine/Source/Editor/UnrealEd/Private/Commandlets/DerivedDataCacheCommandlet.cpp
@@ -248,6 +248,12 @@ void UDerivedDataCacheCommandlet::CacheLoadedPackages(UPackage* CurrentPackage,
 				{
 					FCachingData& CachingData = CachingObjects.FindOrAdd(Object);
 
+					if (!CachingData.ObjectValidityReference.IsValid())
+					{
+						CachingData = FCachingData();
+						CachingData.ObjectValidityReference = Object;
+					}
+
 					// For texture ddc fills, we want to stagger the platforms so that the base texture is only encoded
 					// once with shared linear texture encoding. If we kick everything off at the same time then all worker
 					// tasks ask for the linear encoding at the same time. If we queue the other platforms after the first one
@@ -303,6 +309,12 @@ bool UDerivedDataCacheCommandlet::ProcessCachingObjects()
 			FCachingData& CachingData = It->Value;
 			if (CurrentTime - CachingData.LastTimeTested > 1.0)
 			{
+				if (!It->Value.ObjectValidityReference.IsValid())
+				{
+					It.RemoveCurrent();
+					continue;
+				}
+
 				UObject* Object = It->Key;
 				bool bIsFinished = true;
 				const IInterface_AsyncCompilation* Interface_AsyncCompilation = Cast<IInterface_AsyncCompilation>(Object);
@@ -391,9 +403,9 @@ void UDerivedDataCacheCommandlet::FinishCachingObjects()
 		{
 			if (CurrentTime - LastActivityTime >= DDCCommandletMaxWaitSeconds)
 			{
-				UObject* Object = CachingObjects.CreateIterator()->Key;
+				UObject* Object = CachingObjects.CreateIterator()->Value.ObjectValidityReference.Get();
 				UE_LOG(LogDerivedDataCacheCommandlet, Warning, TEXT("Timed out for %.2lfs waiting for %d objects to finish caching. First object: %s."),
-					DDCCommandletMaxWaitSeconds, CachingObjects.Num(), *Object->GetFullName());
+					DDCCommandletMaxWaitSeconds, CachingObjects.Num(), Object ? *Object->GetFullName() : TEXT("Unknown deleted object"));
 				break;
 			}
 			else
