TextureAssetActions UI improvements
fix TextureCube icon being used but doesn't exist
TextureSourceDataUtils no need to use the complicated bulkdata init api, just init from image
verify that Compress is called as expected

#rnx

diff --git a/Engine/Plugins/Editor/EngineAssetDefinitions/Source/Private/AssetDefinition_TextureCube.cpp b/Engine/Plugins/Editor/EngineAssetDefinitions/Source/Private/AssetDefinition_TextureCube.cpp
--- a/Engine/Plugins/Editor/EngineAssetDefinitions/Source/Private/AssetDefinition_TextureCube.cpp
+++ b/Engine/Plugins/Editor/EngineAssetDefinitions/Source/Private/AssetDefinition_TextureCube.cpp
@@ -67,7 +67,10 @@ namespace MenuExtension_TextureCubeArray
 					{
 						const TAttribute<FText> Label = LOCTEXT("Texture_TextureCubeArray", "Create Texture Cube Array");
 						const TAttribute<FText> ToolTip = LOCTEXT("Texture_CreateTextureCubeArrayTooltip", "Creates a new texture cube array.");
-						const FSlateIcon Icon = FSlateIcon(FAppStyle::GetAppStyleSetName(), "ClassIcon.TextureCube");
+
+						//static const FName NameTextureCube("ClassIcon.TextureCube"); // <- doesn't exist! (FIXME)
+						static const FName NameTextureCube("ClassIcon.TextureRenderTargetCube"); // <- wrong (we are not a render target) but does exist
+						const FSlateIcon Icon = FSlateIcon(FAppStyle::GetAppStyleSetName(), NameTextureCube);
 						const FToolMenuExecuteAction UIAction = FToolMenuExecuteAction::CreateStatic(&ExecuteCreateTextureArray);
 
 						InSection.AddMenuEntry("Texture_TextureCubeArray", Label, ToolTip, Icon, UIAction);
diff --git a/Engine/Source/Editor/UnrealEd/Private/TextureAssetActions.cpp b/Engine/Source/Editor/UnrealEd/Private/TextureAssetActions.cpp
--- a/Engine/Source/Editor/UnrealEd/Private/TextureAssetActions.cpp
+++ b/Engine/Source/Editor/UnrealEd/Private/TextureAssetActions.cpp
@@ -31,6 +31,7 @@
 #include "Interfaces/ITargetPlatformManagerModule.h"
 #include "TextureSourceDataUtils.h"
 #include "TextureImportSettings.h"
+#include "ImageCoreUtils.h"
 
 #define LOCTEXT_NAMESPACE "TextureAssetActions"
 
@@ -47,6 +48,7 @@ enum class ETextureAction
 enum class EAssetActionStatus
 {
 	Enabled = 0,
+	NoSource,
 	UnderSized,
 	WrongType,
 	HasLayers,
@@ -103,7 +105,9 @@ private:
 	/**
 	* Creates a single line showing an asset and it's status related to VT conversion
 	*/
-	TSharedRef<SWidget> CreateAssetLine(int index, const FAssetData &Asset, const EAssetActionStatus Status);
+	TSharedRef<SWidget> CreateAssetLine(int index, const UTexture * Texture, const EAssetActionStatus Status,
+										const FString & EnabledActionString,
+										const FString & DescriptionString);
 
 	FReply OnButtonClick(STextureActionDlg::EResult ButtonID);
 
@@ -140,6 +144,16 @@ private:
 		bNormalMapsKeep16bits = ( InNewState == ECheckBoxState::Checked );
 		UpdateList();
 	}
+	
+	ECheckBoxState HandleShowExcludedCheckBoxIsChecked() const
+	{
+		return bShowExcluded ? ECheckBoxState::Checked : ECheckBoxState::Unchecked;
+	}
+	void HandleShowExcludedCheckBoxCheckedStateChanged(ECheckBoxState InNewState)
+	{
+		bShowExcluded = ( InNewState == ECheckBoxState::Checked );
+		UpdateList();
+	}
 
 	FText GetIntroMessage() const;
 
@@ -169,8 +183,8 @@ private:
 	TArray<TextureListEntry> TextureList;
 
 	int ThresholdValue = 0;
-
 	bool bNormalMapsKeep16bits = true;
+	bool bShowExcluded = true;
 
 	TArray<TSharedPtr<int32>> TextureSizes;
 
@@ -192,8 +206,8 @@ static FText TAA_Intro(ETextureAction Act)
 {
 	switch(Act)
 	{
-	case ETextureAction::Resize: return LOCTEXT("TAA_Intro_Resize", "Reduce size of Texture Source to compact uassets.  Resizing is done using mip filter.  LODBias is adjusted but platform built size may change.");
-	case ETextureAction::ConvertTo8bit: return LOCTEXT("TAA_Intro_Convert", "Convert Texture Source to 8 bit, Normals to 8 or 16, HDR to 16F.  Output built texture may change.  Make sure CompressionSetting and SRGB are set correctly first!");
+	case ETextureAction::Resize: return LOCTEXT("TAA_Intro_Resize", "Reduce size of Texture Source to compact uassets.  Resizing is done using mip filter, in power of two steps.  LODBias is adjusted but platform built size may change.");
+	case ETextureAction::ConvertTo8bit: return LOCTEXT("TAA_Intro_Convert", "Convert Texture Source to 8 bit, Normals to 8 or 16, HDR to 16F.  Only converts if bits per pixel goes down.  Output built texture may change.  Make sure CompressionSetting and SRGB are set correctly first!");
 	default: check(0); return FText();
 	}
 }
@@ -322,6 +336,30 @@ void STextureAssetList::Construct(const FArguments& InArgs)
 					.OnCheckStateChanged(this, &STextureAssetList::HandleNMK16CheckBoxCheckedStateChanged)
 				]
 			]
+		+ SVerticalBox::Slot()
+			.AutoHeight()
+			.HAlign(HAlign_Left)
+			.Padding(8.0f, 4.0f, 8.0f, 4.0f)
+			[
+				SNew(SHorizontalBox)
+				+SHorizontalBox::Slot()
+				.AutoWidth()
+				.VAlign(VAlign_Center)
+				[
+					SNew(STextBlock)
+					.Text(LOCTEXT("TAA_ShowExcluded", "Show Excluded Textures: "))
+					.AutoWrapText(true)
+				]
+				+ SHorizontalBox::Slot()
+				.Padding(4.0f, 0.0f, 2.0f, 0.0f)
+				.VAlign(VAlign_Center)
+				.AutoWidth()
+				[
+					SNew(SCheckBox)
+					.IsChecked(this, &STextureAssetList::HandleShowExcludedCheckBoxIsChecked)
+					.OnCheckStateChanged(this, &STextureAssetList::HandleShowExcludedCheckBoxCheckedStateChanged)
+				]
+			]
 		// Separator
 		+ SVerticalBox::Slot()
 			.AutoHeight()
@@ -390,6 +428,7 @@ static FText AssetActionStatus_Text(EAssetActionStatus Status)
 	switch(Status)
 	{
 	case EAssetActionStatus::Enabled: return FText();
+	case EAssetActionStatus::NoSource: return LOCTEXT("TAAStatus_NoSource", "The texture has no source data; cooked Editor?");
 	case EAssetActionStatus::WrongType: return LOCTEXT("TAAStatus_WrongType", "The texture is not a supported type.");
 	case EAssetActionStatus::UnderSized: return LOCTEXT("TAAStatus_UnderSized", "The texture was under the threshold size.");
 	case EAssetActionStatus::HasLayers: return LOCTEXT("TAAStatus_HasLayers", "Textures with more than 1 layer not supported.");
@@ -401,15 +440,61 @@ static FText AssetActionStatus_Text(EAssetActionStatus Status)
 	}	
 }
 
-TSharedRef<SWidget> STextureAssetList::CreateAssetLine(int index, const FAssetData &Asset, const EAssetActionStatus Status)
+static FName GetTextureClassIconName( ETextureClass TC )
+{
+	// see SlateEditStyle + StarshipStyle
+
+	static const FName NameTexture2d("ClassIcon.Texture2D");
+	//static const FName NameTextureCube("ClassIcon.TextureCube"); // <- is used but doesn't exist! (FIXME)
+	static const FName NameTextureCube("ClassIcon.TextureRenderTargetCube"); // <- wrong
+	static const FName NameVolume("ContentBrowser.AssetActions.VolumeTexture"); // <- should be ClassIcon.VolumeTexture !
+	static const FName NameDefault("ClassIcon.Default"); // TextureRenderTarget2D
+
+	// todo: many of the texture classes don't have icons
+
+	switch(TC)
+	{
+	case ETextureClass::TwoD: return NameTexture2d;
+	case ETextureClass::Cube: return NameTextureCube;
+	case ETextureClass::Array: return NameTexture2d;
+	case ETextureClass::CubeArray: return NameTextureCube;
+	case ETextureClass::Volume: return NameVolume;
+	
+	case ETextureClass::Invalid:
+	case ETextureClass::TwoDDynamic:
+	case ETextureClass::RenderTarget:
+	case ETextureClass::Other2DNoSource:
+	case ETextureClass::OtherUnknown:
+		return NameDefault;
+
+	default:
+		check(0);
+		return NameDefault;	
+	}
+}
+
+TSharedRef<SWidget> STextureAssetList::CreateAssetLine(int index, const UTexture * Texture, const EAssetActionStatus Status,
+										const FString & EnabledActionString,
+										const FString & DescriptionString)
 {
+	const FAssetData Asset(Texture);
+	FName TextureClassIconName = GetTextureClassIconName( Texture->GetTextureClass() );
+
 	FName SeverityIcon = NAME_None;
-	FText DetailedInfoText = AssetActionStatus_Text(Status);
+	FText DetailedInfoText;
 
-	bool IsEnabled = Status == EAssetActionStatus::Enabled;
+	bool IsEnabled = ( Status == EAssetActionStatus::Enabled );
 
-	if ( ! IsEnabled )
+	if ( IsEnabled )
 	{
+		// no error status, show what action will do in drop-down text
+		DetailedInfoText = FText::FromString(EnabledActionString);
+	}
+	else
+	{
+		// get description of status problem
+		DetailedInfoText = AssetActionStatus_Text(Status);
+
 		if ( Status == EAssetActionStatus::UnderSized ||
 			Status == EAssetActionStatus::Already8bit )
 		{
@@ -436,7 +521,7 @@ TSharedRef<SWidget> STextureAssetList::CreateAssetLine(int index, const FAssetDa
 		+ SVerticalBox::Slot()
 		.AutoHeight()
 		.VAlign(VAlign_Center)
-		.HAlign(HAlign_Left)
+		.HAlign(HAlign_Fill) // use the whole line
 		[
 			// Class icon
 			SNew(SHorizontalBox)
@@ -449,8 +534,7 @@ TSharedRef<SWidget> STextureAssetList::CreateAssetLine(int index, const FAssetDa
 				.Padding(2.0f)
 				[
 					SNew(SImage)
-					.Image(FSlateIcon(FAppStyle::GetAppStyleSetName(),
-						"ClassIcon.Texture2D" ).GetIcon())
+					.Image(FSlateIcon(FAppStyle::GetAppStyleSetName(), TextureClassIconName ).GetIcon())
 				]
 			]
 			// Error/warning icon
@@ -468,13 +552,17 @@ TSharedRef<SWidget> STextureAssetList::CreateAssetLine(int index, const FAssetDa
 			]
 			// Fold out button with asset name
 			+ SHorizontalBox::Slot()
-			.AutoWidth()
-			.HAlign(HAlign_Center)
+			.FillWidth(500.f) // take all the remaining space for this slot
+			//.AutoWidth()
+			.HAlign(HAlign_Left) // this applies to the stuff inside this slot, not the slot itself
 			.VAlign(VAlign_Center)
 			[
 					(DetailedInfoText.IsEmpty())
-					? static_cast<TSharedRef<SWidget>>(SNew(STextBlock)
-						.Text(FText::FromString(Asset.GetObjectPathString())))
+					? static_cast<TSharedRef<SWidget>>(
+						SNew(STextBlock)
+						.Text(FText::FromString(Asset.GetObjectPathString()))
+						.Clipping(EWidgetClipping::ClipToBounds)
+						)
 				: SNew(SButton)
 				.ButtonStyle(FCoreStyle::Get(), "NoBorder")
 				.VAlign(VAlign_Center)
@@ -484,7 +572,8 @@ TSharedRef<SWidget> STextureAssetList::CreateAssetLine(int index, const FAssetDa
 				.ContentPadding(0.f)
 				.ForegroundColor(FSlateColor::UseForeground())
 				.IsFocusable(false)
-				//.Text(FText::FromName(Asset.ObjectPath))
+				.Clipping(EWidgetClipping::ClipToBounds)
+				//.Text(FText::FromName(Asset.GetObjectPathString()))
 				[
 					// Fold out icon
 					SNew(SHorizontalBox)
@@ -508,6 +597,22 @@ TSharedRef<SWidget> STextureAssetList::CreateAssetLine(int index, const FAssetDa
 					] // change color for enabled/not
 				]
 			]
+			+ SHorizontalBox::Slot()
+			.AutoWidth()
+			.VAlign(VAlign_Center)
+			.HAlign(HAlign_Right)
+			[
+				SNew(SHorizontalBox)
+				+ SHorizontalBox::Slot()
+				.VAlign(VAlign_Center)
+				.HAlign(HAlign_Right)
+				[
+					SNew(STextBlock)
+					.Justification(ETextJustify::Right)
+					.Text(FText::FromString(DescriptionString))
+					.ColorAndOpacity( IsEnabled ? FSlateColor::UseForeground() : FSlateColor::UseSubduedForeground())
+				]
+			] 
 		]
 		+ SVerticalBox::Slot()
 		.AutoHeight()
@@ -530,6 +635,10 @@ static void DoResizeTextureSource(UTexture * Texture,int TargetSize)
 	int32 BeforeSizeX;
 	int32 BeforeSizeY;
 	Texture->GetBuiltTextureSize(RunningPlatform, BeforeSizeX, BeforeSizeY);
+	
+	const FIntPoint BeforeSourceSize = Texture->Source.GetLogicalSize();
+	UE_LOG(LogTexture, Display, TEXT("Texture (%s) Resizing to <= %d Source Size=%dx%d Built Size=%dx%d"), *Texture->GetName() , TargetSize,
+		BeforeSourceSize.X,BeforeSourceSize.Y,BeforeSizeX,BeforeSizeY);
 
 	if ( ! UE::TextureUtilitiesCommon::Experimental::DownsizeTextureSourceData(Texture, TargetSize, RunningPlatform))
 	{
@@ -538,8 +647,8 @@ static void DoResizeTextureSource(UTexture * Texture,int TargetSize)
 		// did not resize, but may have done PreEditChange
 		return;
 	}
-
-	UE_LOG(LogTexture, Display, TEXT("Texture (%s) did resize."), *Texture->GetName());
+	
+	const FIntPoint AfterSourceSize = Texture->Source.GetLogicalSize();
 
 	Texture->LODBias = 0;
 	
@@ -564,6 +673,9 @@ static void DoResizeTextureSource(UTexture * Texture,int TargetSize)
 		Texture->GetBuiltTextureSize(RunningPlatform, AfterSizeX, AfterSizeY);
 	}
 	
+	UE_LOG(LogTexture, Display, TEXT("Texture (%s) did resize Source Size=%dx%d Built Size=%dx%d"), *Texture->GetName(),
+		AfterSourceSize.X,AfterSourceSize.Y,AfterSizeX,AfterSizeY);
+
 	if ( BeforeSizeX != AfterSizeX || BeforeSizeY != AfterSizeY )
 	{
 		// not a warning, just FYI
@@ -582,7 +694,7 @@ static void DoResizeTextureSource(UTexture * Texture,int TargetSize)
 	Texture->PostEditChange();
 }
 
-static ETextureSourceFormat GetReducedTextureSourceFormat(TextureCompressionSettings TC,ETextureSourceFormat InTSF,bool NormalMapsKeep16bits)
+static ETextureSourceFormat GetReducedTextureSourceFormat(const TextureCompressionSettings TC,const ETextureSourceFormat InTSF,const bool NormalMapsKeep16bits)
 {
 	if ( InTSF == TSF_BGRE8 )
 	{
@@ -722,7 +834,7 @@ static void DoConvertTo8bitTextureSource(UTexture * Texture,bool NormalMapsKeep1
 	
 	UE_LOG(LogTexture, Display, TEXT("Texture (%s) changing format from %s to %s, TC = %s"), 
 		*Texture->GetName() ,
-		*StaticEnum<ETextureSourceFormat>()->GetDisplayNameTextByValue(InTSF).ToString(),
+		*StaticEnum<ETextureSourceFormat>()->GetDisplayNameTextByValue(InTSF).ToString(), // these have TSF_ on the strings
 		*StaticEnum<ETextureSourceFormat>()->GetDisplayNameTextByValue(OutTSF).ToString(),
 		*StaticEnum<TextureCompressionSettings>()->GetDisplayNameTextByValue(Texture->CompressionSettings).ToString()
 		);
@@ -769,11 +881,9 @@ void STextureAssetList::DoAction()
 		switch(Action)
 		{
 		case ETextureAction::Resize:
-			UE_LOG(LogTexture, Display, TEXT("Texture (%s) Resizing to <= %d"), *Texture->GetName() , ThresholdValue);
 			DoResizeTextureSource(Texture,ThresholdValue);
 			break;
 		case ETextureAction::ConvertTo8bit:
-			UE_LOG(LogTexture, Display, TEXT("Texture (%s) ConvertTo8bit"), *Texture->GetName());
 			DoConvertTo8bitTextureSource(Texture,bNormalMapsKeep16bits);
 			break;
 		default:
@@ -785,7 +895,10 @@ void STextureAssetList::DoAction()
 
 void STextureAssetList::UpdateList()
 {
+	ExpandedIndexes.SetNum(0);
+
 	AssetListContainer->ClearChildren();
+	int32 AssetListContainerIndex = 0;
 
 	// filter select textures to see if they should be acted on
 	for (int Index=0;Index<TextureList.Num();Index++)
@@ -793,57 +906,108 @@ void STextureAssetList::UpdateList()
 		TextureListEntry & Entry = TextureList[Index];
 		UTexture * Texture = Entry.Texture;
 
-		FAssetData AssetData(Texture);
 		EAssetActionStatus Status = EAssetActionStatus::Enabled; 
-
-		if ( Texture->Source.GetNumLayers() != 1 ) // only 1 layer textures
+		
+		FString EnabledActionString;
+		FString DescriptionString;
+		
+		ETextureClass Class = Texture->GetTextureClass();
+		
+		// first filter for exclusions that apply to all actions
+		//	don't query anything in Source until you check IsValid
+		if ( Class != ETextureClass::TwoD &&
+			Class != ETextureClass::Cube &&
+			Class != ETextureClass::Array &&
+			Class != ETextureClass::CubeArray &&
+			Class != ETextureClass::Volume )
+		{
+			// rendertarget, dynamic, or unknown
+			Status = EAssetActionStatus::WrongType;
+		}
+		else if ( ! Texture->Source.IsValid() )
+		{
+			Status = EAssetActionStatus::NoSource;
+		}
+		else if ( Texture->Source.GetNumLayers() != 1 ) // only 1 layer textures
 		{
 			Status = EAssetActionStatus::HasLayers;
 		}
-		else if( Action == ETextureAction::Resize )
+		else
 		{
-			ETextureClass Class = Texture->GetTextureClass();
-			if ( Class != ETextureClass::TwoD && Class != ETextureClass::Cube ) //  Array ?
+			// don't prep Description text until after Source IsValid check
+
+			// GetLogicalSize = sum of udim blocks
+			const FIntPoint SourceSize = Texture->Source.GetLogicalSize();
+			const ETextureSourceFormat InTSF = Texture->Source.GetFormat();
+
+			DescriptionString = FString::Printf(TEXT("%dx%d %s"),SourceSize.X,SourceSize.Y,
+				ERawImageFormat::GetName( FImageCoreUtils::ConvertToRawImageFormat(InTSF) ));
+			if ( Texture->VirtualTextureStreaming )
 			{
-				Status = EAssetActionStatus::WrongType;
+				DescriptionString += FString(TEXT(" VT"));
 			}
-			else
+			if ( Texture->Source.GetSourceCompression() == ETextureSourceCompressionFormat::TSCF_JPEG )
+			{
+				DescriptionString += FString(TEXT(" JPEG"));
+			}
+		
+			// per-Action filters :
+			if( Action == ETextureAction::Resize )
 			{
-				// GetLogicalSize = sum of udim blocks
-				FIntPoint SourceSize = Texture->Source.GetLogicalSize();
 				int MaxSize = FMath::Max(SourceSize.X,SourceSize.Y);
-				if ( MaxSize <= ThresholdValue )
+
+				if ( Class != ETextureClass::TwoD && Class != ETextureClass::Cube )
+				{
+					// Resize only supports 2d and Cube for now, no arrays or volumes
+					Status = EAssetActionStatus::WrongType;
+				}
+				else if ( MaxSize <= ThresholdValue )
 				{
 					Status = EAssetActionStatus::UnderSized;
 				}
+				else
+				{
+					// could fill EnabledActionString
+					// hard to get size we will output exactly right
+				}
 			}
-		}
-		else if ( Action == ETextureAction::ConvertTo8bit )
-		{
-			ETextureSourceFormat InTSF = Texture->Source.GetFormat();
-			ETextureSourceFormat OutTSF = GetReducedTextureSourceFormat(Texture->CompressionSettings,InTSF,bNormalMapsKeep16bits);
-			if ( InTSF == OutTSF )
+			else if ( Action == ETextureAction::ConvertTo8bit )
 			{
-				Status = EAssetActionStatus::Already8bit;
+				ETextureSourceFormat OutTSF = GetReducedTextureSourceFormat(Texture->CompressionSettings,InTSF,bNormalMapsKeep16bits);
+				if ( InTSF == OutTSF )
+				{
+					Status = EAssetActionStatus::Already8bit;
+				}
+				else if ( Texture->Source.GetSourceCompression() == ETextureSourceCompressionFormat::TSCF_JPEG )
+				{
+					// JPEG is already 8 bit; we don't want to try to change RGB JPEG to G8, just leave it alone
+					Status = EAssetActionStatus::DontChangeJPEG;
+				}
+				else
+				{
+					// Enabled!
+					EnabledActionString = FString::Printf(TEXT("TC=%s will convert to %s"),
+						*StaticEnum<TextureCompressionSettings>()->GetDisplayNameTextByValue(Texture->CompressionSettings).ToString(),
+						ERawImageFormat::GetName( FImageCoreUtils::ConvertToRawImageFormat( OutTSF ) ) );
+				}
 			}
-			else if ( Texture->Source.GetSourceCompression() == ETextureSourceCompressionFormat::TSCF_JPEG )
+			else
 			{
-				// JPEG is already 8 bit; we don't want to try to change RGB JPEG to G8, just leave it alone
-				Status = EAssetActionStatus::DontChangeJPEG;
+				check(0);
 			}
 		}
-		else
-		{
-			check(0);
-		}
 
 		Entry.IsEnabled = ( Status == EAssetActionStatus::Enabled );
 
-		AssetListContainer->AddSlot()
-			.AutoHeight()
-			[
-				CreateAssetLine(Index, AssetData, Status)
-			];
+		if ( Entry.IsEnabled || bShowExcluded )
+		{
+			AssetListContainer->AddSlot()
+				.AutoHeight()
+				[
+					CreateAssetLine(AssetListContainerIndex, Texture, Status, EnabledActionString, DescriptionString)
+				];
+			AssetListContainerIndex++;
+		}
 	}
 
 	ErrorMessage = FText();
@@ -950,7 +1114,7 @@ STextureActionDlg::STextureActionDlg(const TArray<UTexture *> &Textures, ETextur
 		DialogWindow = SNew(SWindow)
 			.Title(TAA_DialogTitle(Act))
 			.SupportsMinimize(false).SupportsMaximize(false)
-			.ClientSize(FVector2D(500, 500));
+			.ClientSize(FVector2D(1200, 600));
 
 		TSharedPtr<SBorder> DialogWrapper =
 			SNew(SBorder)
diff --git a/Engine/Source/Runtime/TextureUtilitiesCommon/Private/TextureSourceDataUtils.cpp b/Engine/Source/Runtime/TextureUtilitiesCommon/Private/TextureSourceDataUtils.cpp
--- a/Engine/Source/Runtime/TextureUtilitiesCommon/Private/TextureSourceDataUtils.cpp
+++ b/Engine/Source/Runtime/TextureUtilitiesCommon/Private/TextureSourceDataUtils.cpp
@@ -24,18 +24,16 @@ namespace Private
 
 		// We want to reduce the asset size so ignore the imported mip(s) (??)
 		const int32 MipIndex = 0;
-		FImage SourceMip0;
-		if (!Texture->Source.GetMipImage(SourceMip0, MipIndex))
+		FImage Image;
+		if (!Texture->Source.GetMipImage(Image, MipIndex))
 		{
 			UE_LOG(LogTexture,Error,TEXT("ResizeTexture2D: Texture GetMipImage failed [%s]"),
 				*Texture->GetFullName());
 			return false;
 		}
 
-		int32 NumSlices = Texture->Source.GetNumSlices(); // == 1 or 6 for cubes
-
 		bool MadeChanges;
-		if ( ! Texture->DownsizeImageUsingTextureSettings(TargetPlatform, SourceMip0, MaxSize, LayerIndex, MadeChanges) )
+		if ( ! Texture->DownsizeImageUsingTextureSettings(TargetPlatform, Image, MaxSize, LayerIndex, MadeChanges) )
 		{
 			UE_LOG(LogTexture,Error,TEXT("ResizeTexture2D: Texture DownsizeImageUsingTextureSettings failed [%s]"),
 				*Texture->GetFullName());
@@ -48,19 +46,11 @@ namespace Private
 		
 		Texture->PreEditChange(nullptr);
 
-		FImage ResizedImage = MoveTemp(SourceMip0); // this is just a variable rename
-
-		check( ResizedImage.NumSlices == NumSlices ); // slices are done one by one
-
-		UE::Serialization::FEditorBulkData::FSharedBufferWithID ResizedImageBufferWithID = MakeSharedBufferFromArray(MoveTemp(ResizedImage.RawData));
+		Texture->Source.Init(Image);
 
-		const int32 NumMips = 1;
-		Texture->Source.Init(ResizedImage.SizeX
-			, ResizedImage.SizeY
-			, ResizedImage.NumSlices
-			, NumMips
-			, FImageCoreUtils::ConvertToTextureSourceFormat(ResizedImage.Format)
-			, MoveTemp(ResizedImageBufferWithID));
+		// Compress() applies PNG filter to the BulkData
+		//	Compress is done automatically in Texture PreSave
+		//Texture->Source.Compress(); 
 
 		// if gamma was Pow22 it is now sRGB
 		Texture->bUseLegacyGamma = false;
@@ -320,6 +310,7 @@ TEXTUREUTILITIESCOMMON_API bool ChangeTextureSourceFormat(UTexture* Texture, ETe
 	if ( Texture->Source.GetNumBlocks() == 1 && Texture->Source.GetNumMips() == 1 )
 	{
 		const int32 MipIndex = 0;
+
 		FImage SourceMip;
 		if (!Texture->Source.GetMipImage(SourceMip, MipIndex))
 		{
@@ -331,15 +322,9 @@ TEXTUREUTILITIESCOMMON_API bool ChangeTextureSourceFormat(UTexture* Texture, ETe
 		FImage NewMip;
 		SourceMip.CopyTo(NewMip,NewRIF,NewGamma);
 		
-		UE::Serialization::FEditorBulkData::FSharedBufferWithID ResizedImageBufferWithID = MakeSharedBufferFromArray(MoveTemp(NewMip.RawData));
-		
 		Texture->PreEditChange(nullptr);
 
-		const int32 NumMips = 1;
-		Texture->Source.Init(NewMip.SizeX,NewMip.SizeY,NewMip.NumSlices
-			, NumMips
-			, NewFormat
-			, MoveTemp(ResizedImageBufferWithID));
+		Texture->Source.Init(NewMip);
 	}
 	else // blocks and/or mips
 	{
