TextureAssetActions compress with JPEG action

#rnx

diff --git a/Engine/Plugins/Editor/EngineAssetDefinitions/Source/Private/AssetDefinition_Texture.cpp b/Engine/Plugins/Editor/EngineAssetDefinitions/Source/Private/AssetDefinition_Texture.cpp
--- a/Engine/Plugins/Editor/EngineAssetDefinitions/Source/Private/AssetDefinition_Texture.cpp
+++ b/Engine/Plugins/Editor/EngineAssetDefinitions/Source/Private/AssetDefinition_Texture.cpp
@@ -42,7 +42,7 @@ FAssetOpenSupport UAssetDefinition_Texture::GetAssetOpenSupport(const FAssetOpen
 
 namespace MenuExtension_Texture
 {
-	void ExecuteCreateMaterial(const FToolMenuContext& InContext)
+	static void ExecuteCreateMaterial(const FToolMenuContext& InContext)
 	{
 		const UContentBrowserAssetContextMenuContext* CBContext = UContentBrowserAssetContextMenuContext::FindContextWithAssets(InContext);
 
@@ -71,6 +71,14 @@ namespace MenuExtension_Texture
 			UE::TextureAssetActions::TextureSource_ConvertTo8bit_WithDialog(CBContext->LoadSelectedObjects<UTexture>());
 		}
 	}
+	
+	static void ExecuteJPEGTextureSource(const FToolMenuContext& InContext)
+	{
+		if (const UContentBrowserAssetContextMenuContext* CBContext = UContentBrowserAssetContextMenuContext::FindContextWithAssets(InContext))
+		{
+			UE::TextureAssetActions::TextureSource_JPEG_WithDialog(CBContext->LoadSelectedObjects<UTexture>());
+		}
+	}
 
 	static void ExecuteConvertToVirtualTexture(const FToolMenuContext& InContext)
 	{
@@ -155,6 +163,14 @@ namespace MenuExtension_Texture
 							const FToolMenuExecuteAction UIAction = FToolMenuExecuteAction::CreateStatic(&Execute8bitTextureSource);
 							InSection.AddMenuEntry("Texture_ConvertTo8bit", Label, ToolTip, Icon, UIAction);
 						}
+						
+						{
+							const TAttribute<FText> Label = LOCTEXT("Texture2D_JPEGSource", "Texture Source Compress With JPEG");
+							const TAttribute<FText> ToolTip = LOCTEXT("Texture2D_JPEGSourceTooltip", "Reduce texture asset size by compressing source with JPEG.");
+							const FSlateIcon Icon = FSlateIcon(FAppStyle::GetAppStyleSetName(), "ClassIcon.Texture2D");
+							const FToolMenuExecuteAction UIAction = FToolMenuExecuteAction::CreateStatic(&ExecuteJPEGTextureSource);
+							InSection.AddMenuEntry("Texture2D_JPEGSource", Label, ToolTip, Icon, UIAction);
+						}
 					}
 				}));
 			}
diff --git a/Engine/Source/Editor/UnrealEd/Private/TextureAssetActions.cpp b/Engine/Source/Editor/UnrealEd/Private/TextureAssetActions.cpp
--- a/Engine/Source/Editor/UnrealEd/Private/TextureAssetActions.cpp
+++ b/Engine/Source/Editor/UnrealEd/Private/TextureAssetActions.cpp
@@ -42,7 +42,8 @@ enum class ETextureAction
 {
 	Invalid = 0,
 	Resize,
-	ConvertTo8bit
+	ConvertTo8bit,
+	JPEG
 };
 
 enum class EAssetActionStatus
@@ -52,8 +53,10 @@ enum class EAssetActionStatus
 	UnderSized,
 	WrongType,
 	HasLayers,
+	HasMipsLeaveExisting,
 	Already8bit,
-	DontChangeJPEG
+	DontChangeJPEG,
+	MustBe8BitForJPEG
 };
 
 /**
@@ -198,6 +201,7 @@ static FText TAA_DialogTitle(ETextureAction Act)
 	{
 	case ETextureAction::Resize: return LOCTEXT("TAA_Title_Resize", "Texture Asset : Resize Source");
 	case ETextureAction::ConvertTo8bit: return LOCTEXT("TAA_Title_Convert", "Texture Asset : Convert To 8 bit or minimum viable bit depth");
+	case ETextureAction::JPEG: return LOCTEXT("TAA_Title_JPEG", "Texture Asset : Compress with JPEG");
 	default: check(0); return FText();
 	}
 }
@@ -208,12 +212,16 @@ static FText TAA_Intro(ETextureAction Act)
 	{
 	case ETextureAction::Resize: return LOCTEXT("TAA_Intro_Resize", "Reduce size of Texture Source to compact uassets.  Resizing is done using mip filter, in power of two steps.  LODBias is adjusted but platform built size may change.");
 	case ETextureAction::ConvertTo8bit: return LOCTEXT("TAA_Intro_Convert", "Convert Texture Source to 8 bit, Normals to 8 or 16, HDR to 16F.  Only converts if bits per pixel goes down.  Output built texture may change.  Make sure CompressionSetting and SRGB are set correctly first!");
+	case ETextureAction::JPEG: return LOCTEXT("TAA_Intro_JPEG", "Compress Texture Source with JPEG.  Only works on 8 bit, 2D simple textures.  Greatly reduces uasset size with some loss of quality.  Does not affect in-game size.");
 	default: check(0); return FText();
 	}
 }
 
 void STextureAssetList::Construct(const FArguments& InArgs)
-{	
+{
+	// BEWARE : Action is not set yet (it comes in Init)
+	//	cannot do per-Action setup here
+
 	UserResponse = STextureActionDlg::Cancel;
 	ParentWindow = InArgs._ParentWindow.Get();
 	static FName ErrorIcon = "MessageLog.Error";
@@ -310,6 +318,20 @@ void STextureAssetList::Construct(const FArguments& InArgs)
 						.Text(this, &STextureAssetList::GetThresholdText)
 					]
 				]
+				/*
+				// cannot do Action-dependent setup here :(
+				+SHorizontalBox::Slot()
+				.AutoWidth()
+				.VAlign(VAlign_Center)
+				[
+					SNew(STextBlock)
+					.Text( Action == ETextureAction::JPEG ? 
+						LOCTEXT("TAA_Threshold_JPEG", "included if pixel count >= threshold^2") :
+						LOCTEXT("TAA_Threshold_Resize", "resizes so larger dimension is <= threshold")
+					)
+					.AutoWrapText(true)
+				]
+				*/
 			]
 		+ SVerticalBox::Slot()
 			.AutoHeight()
@@ -430,10 +452,12 @@ static FText AssetActionStatus_Text(EAssetActionStatus Status)
 	case EAssetActionStatus::Enabled: return FText();
 	case EAssetActionStatus::NoSource: return LOCTEXT("TAAStatus_NoSource", "The texture has no source data; cooked Editor?");
 	case EAssetActionStatus::WrongType: return LOCTEXT("TAAStatus_WrongType", "The texture is not a supported type.");
-	case EAssetActionStatus::UnderSized: return LOCTEXT("TAAStatus_UnderSized", "The texture was under the threshold size.");
+	case EAssetActionStatus::UnderSized: return LOCTEXT("TAAStatus_UnderSized", "The texture was under or equal the threshold size.");
 	case EAssetActionStatus::HasLayers: return LOCTEXT("TAAStatus_HasLayers", "Textures with more than 1 layer not supported.");
+	case EAssetActionStatus::HasMipsLeaveExisting: return LOCTEXT("TAAStatus_HasMipsLeaveExisting", "Textures has imported mips and LeaveExisting, will not change. (change MipGen if wanted)");
 	case EAssetActionStatus::Already8bit: return LOCTEXT("TAAStatus_Already8", "Texture format is already minimum bit depth.");
 	case EAssetActionStatus::DontChangeJPEG: return LOCTEXT("TAAStatus_JPEG", "Texture source is JPEG compressed, will not change."); 
+	case EAssetActionStatus::MustBe8BitForJPEG: return LOCTEXT("TAAStatus_MustBe8BitForJPEG", "Texture source must be 8 bit for JPEG; change bit depth first."); 
 	default:
 		check(0);
 		return FText();
@@ -843,6 +867,17 @@ static void DoConvertTo8bitTextureSource(UTexture * Texture,bool NormalMapsKeep1
 	UE::TextureUtilitiesCommon::Experimental::ChangeTextureSourceFormat(Texture,OutTSF);	
 }
 
+static void DoCompressTextureSourceWithJPEG(UTexture * Texture)
+{
+	// calls Pre/PostEditChange :
+	bool bDid = UE::TextureUtilitiesCommon::Experimental::CompressTextureSourceWithJPEG(Texture);
+	
+	UE_LOG(LogTexture, Display, TEXT("Texture (%s) %s"), 
+		*Texture->GetName() ,
+		bDid ? TEXT("was changed to JPEG compression.") : TEXT("was not changed.")
+		);
+}
+
 void STextureAssetList::DoAction()
 {
 	int NumEnabled = 0;
@@ -886,6 +921,9 @@ void STextureAssetList::DoAction()
 		case ETextureAction::ConvertTo8bit:
 			DoConvertTo8bitTextureSource(Texture,bNormalMapsKeep16bits);
 			break;
+		case ETextureAction::JPEG:
+			DoCompressTextureSourceWithJPEG(Texture);
+			break;
 		default:
 			check(0);
 			break;
@@ -932,6 +970,10 @@ void STextureAssetList::UpdateList()
 		{
 			Status = EAssetActionStatus::HasLayers;
 		}
+		else if ( Texture->Source.GetNumMips() != 1 && Texture->MipGenSettings == TMGS_LeaveExistingMips )
+		{
+			Status = EAssetActionStatus::HasMipsLeaveExisting;
+		}
 		else
 		{
 			// don't prep Description text until after Source IsValid check
@@ -991,6 +1033,37 @@ void STextureAssetList::UpdateList()
 						ERawImageFormat::GetName( FImageCoreUtils::ConvertToRawImageFormat( OutTSF ) ) );
 				}
 			}
+			else if ( Action == ETextureAction::JPEG )
+			{
+				if ( Class != ETextureClass::TwoD || Texture->Source.GetNumBlocks() > 1 || Texture->Source.GetNumSlices() > 1 )
+				{
+					// JPEG only supports 2d
+					// Blocked/UDIM doesn't support JPEG (@@ ??)
+					Status = EAssetActionStatus::WrongType;
+				}
+				else if ( Texture->Source.GetSourceCompression() == ETextureSourceCompressionFormat::TSCF_JPEG )
+				{
+					// Already JPEG
+					Status = EAssetActionStatus::DontChangeJPEG;
+				}
+				else if ( InTSF != TSF_BGRA8 && InTSF != TSF_G8 )
+				{
+					Status = EAssetActionStatus::MustBe8BitForJPEG;
+				}
+				else if ( SourceSize.X < 16 || SourceSize.Y < 16 )
+				{
+					// hard-coded requirement of at least 16 in both dimensions
+					Status = EAssetActionStatus::UnderSized;
+				}
+				else if ( (int64)SourceSize.X*SourceSize.Y < (int64)ThresholdValue*ThresholdValue ) // pixel count
+				{
+					Status = EAssetActionStatus::UnderSized;
+				}
+				else
+				{
+					// enabled!
+				}				
+			}
 			else
 			{
 				check(0);
@@ -1094,7 +1167,7 @@ EVisibility STextureAssetList::GetErrorMessageVisibility() const
 
 EVisibility STextureAssetList::GetThresholdVisibility() const
 {
-	return ( Action == ETextureAction::Resize ) ? EVisibility::Visible : EVisibility::Collapsed;
+	return ( Action == ETextureAction::Resize || Action == ETextureAction::JPEG ) ? EVisibility::Visible : EVisibility::Collapsed;
 }
 
 FText STextureAssetList::GetIntroMessage() const
@@ -1154,4 +1227,10 @@ void UE::TextureAssetActions::TextureSource_ConvertTo8bit_WithDialog(const TArra
 	Dlg.ShowModal();
 }
 
+void UE::TextureAssetActions::TextureSource_JPEG_WithDialog(const TArray<UTexture*> & InTextures)
+{
+	STextureActionDlg Dlg(InTextures,ETextureAction::JPEG);
+	Dlg.ShowModal();
+}
+
 #undef LOCTEXT_NAMESPACE
diff --git a/Engine/Source/Editor/UnrealEd/Public/TextureAssetActions.h b/Engine/Source/Editor/UnrealEd/Public/TextureAssetActions.h
--- a/Engine/Source/Editor/UnrealEd/Public/TextureAssetActions.h
+++ b/Engine/Source/Editor/UnrealEd/Public/TextureAssetActions.h
@@ -15,5 +15,8 @@ UNREALED_API void TextureSource_Resize_WithDialog(const TArray<UTexture*> & InTe
 
 UNREALED_API void TextureSource_ConvertTo8bit_WithDialog(const TArray<UTexture*> & InTextures);
 
+UNREALED_API void TextureSource_JPEG_WithDialog(const TArray<UTexture*> & InTextures);
+
+
 }
 };
diff --git a/Engine/Source/Runtime/TextureUtilitiesCommon/Private/TextureSourceDataUtils.cpp b/Engine/Source/Runtime/TextureUtilitiesCommon/Private/TextureSourceDataUtils.cpp
--- a/Engine/Source/Runtime/TextureUtilitiesCommon/Private/TextureSourceDataUtils.cpp
+++ b/Engine/Source/Runtime/TextureUtilitiesCommon/Private/TextureSourceDataUtils.cpp
@@ -9,6 +9,7 @@
 #include "HAL/UnrealMemory.h"
 #include "EngineLogs.h"
 #include "TextureImportSettings.h"
+#include "ImageUtils.h"
 
 namespace UE::TextureUtilitiesCommon::Experimental
 {
@@ -392,6 +393,80 @@ TEXTUREUTILITIESCOMMON_API bool ChangeTextureSourceFormat(UTexture* Texture, ETe
 	return true;
 }
 
+// calls Pre/PostEditChange :
+TEXTUREUTILITIESCOMMON_API bool CompressTextureSourceWithJPEG(UTexture* Texture,int32 Quality)
+{
+	if ( ! Texture->Source.IsValid() )
+	{
+		return false;
+	}
+
+	// we only support 1 layer currently
+	if (Texture->Source.GetNumLayers() != 1)
+	{
+		return false;
+	}
+	
+	if (Texture->Source.GetNumBlocks() != 1 )
+	{
+		// JPEG does not support UDIM/blocks ; fix me?
+		return false;
+	}
+	
+	if (Texture->Source.GetSourceCompression() == ETextureSourceCompressionFormat::TSCF_JPEG )
+	{
+		// already JPEG
+		return false;
+	}
+	
+	if ( Texture->Source.GetNumSlices() != 1 )
+	{
+		// 1 mip, 1 slice only
+		return false;
+	}
+
+	ETextureSourceFormat Format = Texture->Source.GetFormat(0);
+	if ( Format != TSF_G8 && Format != TSF_BGRA8 )
+	{
+		// must be 8 bit
+		return false;
+	}
+
+	// we do kill existing mips to match the behavior of the other conversions in here
+
+	// okay, looks good, do it!
+
+	FImage Image;
+	if ( ! Texture->Source.GetMipImage(Image,0) )
+	{
+		UE_LOG(LogTexture,Error,TEXT("CompressTextureSourceWithJPEG: Texture GetMipImage failed [%s]"),
+			*Texture->GetFullName());
+		return false;
+	}
+
+	// JPEG it :
+
+	TArray64<uint8> JPEGData;
+	if ( ! FImageUtils::CompressImage(JPEGData,TEXT(".jpg"),Image,Quality) )
+	{
+		UE_LOG(LogTexture,Error,TEXT("CompressTextureSourceWithJPEG: Texture CompressImage failed [%s]"),
+			*Texture->GetFullName());
+		return false;
+	}
+
+	Texture->PreEditChange(nullptr);
+
+	// Format stays BGRA8 or G8
+	const int32 NumMips = 1;
+	Texture->Source.InitWithCompressedSourceData(Image.SizeX,Image.SizeY,NumMips,Format,
+		JPEGData,
+		ETextureSourceCompressionFormat::TSCF_JPEG);
+
+	Texture->PostEditChange();
+	
+	return true;
+}
+
 } // End namespace UE::TextureUtilitiesCommon::Experimental
 
 #endif //WITH_EDITOR
diff --git a/Engine/Source/Runtime/TextureUtilitiesCommon/Public/TextureSourceDataUtils.h b/Engine/Source/Runtime/TextureUtilitiesCommon/Public/TextureSourceDataUtils.h
--- a/Engine/Source/Runtime/TextureUtilitiesCommon/Public/TextureSourceDataUtils.h
+++ b/Engine/Source/Runtime/TextureUtilitiesCommon/Public/TextureSourceDataUtils.h
@@ -23,6 +23,12 @@ namespace UE::TextureUtilitiesCommon::Experimental
 	//	beware that changing format may change the interpretation of the SRGB bool in Texture
 	// ChangeTextureSourceFormat supports mips and blocks (udim) but not layers
 	TEXTUREUTILITIESCOMMON_API bool ChangeTextureSourceFormat(UTexture* Texture, ETextureSourceFormat NewFormat);
+	
+	// Replace TextureSource with JPEG compressed data
+	// returns true if change was made
+	// calls Pre/PostEditChange :
+	//	Quality == 0 for default (85)
+	TEXTUREUTILITIESCOMMON_API bool CompressTextureSourceWithJPEG(UTexture* Texture,int32 Quality = 0);
 }
 
 #endif // WITH_EDITOR
\ No newline at end of file
