Fix a crash when importing a LOD at the wrong index

#rnx

diff --git a/Engine/Source/Editor/UnrealEd/Private/FbxMeshUtils.cpp b/Engine/Source/Editor/UnrealEd/Private/FbxMeshUtils.cpp
--- a/Engine/Source/Editor/UnrealEd/Private/FbxMeshUtils.cpp
+++ b/Engine/Source/Editor/UnrealEd/Private/FbxMeshUtils.cpp
@@ -51,6 +51,11 @@ namespace FbxMeshUtils
 {
 	namespace Private
 	{
+		void ShowFailedToImportLodDialog(int32 LodIndex)
+		{
+			FMessageDialog::Open(EAppMsgType::Ok, FText::Format(LOCTEXT("LODImport_Failure", "Failed to import LOD{0}"), FText::AsNumber(LodIndex)));
+		}
+
 		void SetupFbxImportOptions(const UStaticMesh* BaseStaticMesh, UnFbx::FBXImportOptions* ImportOptions)
 		{
 			check(BaseStaticMesh);
@@ -895,26 +900,43 @@ namespace FbxMeshUtils
 
 		FString FilenameToImport("");
 		UInterchangeAssetImportData* SelectedInterchangeAssetImportData = nullptr;
+
+		//Make sure the LODLevel is valid, it should not be more then one over the existing lod count
+		bool bInvalidLodIndex = false;
 		if (SkeletalMesh)
 		{
-			if (!bReimportWithNewFile && SkeletalMesh->IsValidLODIndex(LODLevel))
+			if (LODLevel > SkeletalMesh->GetLODNum())
 			{
-				FilenameToImport = SkeletalMesh->GetLODInfo(LODLevel)->SourceImportFilename.IsEmpty() ?
-					SkeletalMesh->GetLODInfo(LODLevel)->SourceImportFilename :
-					UAssetImportData::ResolveImportFilename(SkeletalMesh->GetLODInfo(LODLevel)->SourceImportFilename, nullptr);
+				bInvalidLodIndex = true;
+			}
+			else
+			{
+				if (!bReimportWithNewFile && SkeletalMesh->IsValidLODIndex(LODLevel))
+				{
+					FilenameToImport = SkeletalMesh->GetLODInfo(LODLevel)->SourceImportFilename.IsEmpty() ?
+						SkeletalMesh->GetLODInfo(LODLevel)->SourceImportFilename :
+						UAssetImportData::ResolveImportFilename(SkeletalMesh->GetLODInfo(LODLevel)->SourceImportFilename, nullptr);
+				}
+				SelectedInterchangeAssetImportData = Cast<UInterchangeAssetImportData>(SkeletalMesh->GetAssetImportData());
 			}
-			SelectedInterchangeAssetImportData = Cast<UInterchangeAssetImportData>(SkeletalMesh->GetAssetImportData());
 		}
 		else if (StaticMesh)
 		{
-			if (!bReimportWithNewFile && StaticMesh->IsSourceModelValid(LODLevel))
+			if (LODLevel > StaticMesh->GetNumSourceModels())
+			{
+				bInvalidLodIndex = true;
+			}
+			else
 			{
-				const FStaticMeshSourceModel& SourceModel = StaticMesh->GetSourceModel(LODLevel);
-				FilenameToImport = SourceModel.SourceImportFilename.IsEmpty() ?
-					SourceModel.SourceImportFilename :
-					UAssetImportData::ResolveImportFilename(SourceModel.SourceImportFilename, nullptr);
+				if (!bReimportWithNewFile && StaticMesh->IsSourceModelValid(LODLevel))
+				{
+					const FStaticMeshSourceModel& SourceModel = StaticMesh->GetSourceModel(LODLevel);
+					FilenameToImport = SourceModel.SourceImportFilename.IsEmpty() ?
+						SourceModel.SourceImportFilename :
+						UAssetImportData::ResolveImportFilename(SourceModel.SourceImportFilename, nullptr);
+				}
+				SelectedInterchangeAssetImportData = Cast<UInterchangeAssetImportData>(StaticMesh->GetAssetImportData());
 			}
-			SelectedInterchangeAssetImportData = Cast<UInterchangeAssetImportData>(StaticMesh->GetAssetImportData());
 		}
 		else
 		{
@@ -922,6 +944,14 @@ namespace FbxMeshUtils
 			return false;
 		}
 
+
+		if (bInvalidLodIndex)
+		{
+			UE_LOG(LogExportMeshUtils, Warning, TEXT("ImportMeshLODDialog: Invalid mesh LOD index %d, no prior LOD index exists."), LODLevel);
+			FbxMeshUtils::Private::ShowFailedToImportLodDialog(LODLevel);
+			return false;
+		}
+
 		// Check the file exists first
 		const bool bSourceFileExists = FPaths::FileExists(FilenameToImport);
 		// We'll give the user a chance to choose a new file if a previously set file fails to import
@@ -946,6 +976,10 @@ namespace FbxMeshUtils
 								GEditor->GetEditorSubsystem<UImportSubsystem>()->BroadcastAssetPostLODImport(StaticMesh, LODLevel);
 							}
 						}
+						else
+						{
+							FbxMeshUtils::Private::ShowFailedToImportLodDialog(LODLevel);
+						}
 					});
 			};
 
@@ -1016,7 +1050,7 @@ namespace FbxMeshUtils
 		if(!bImportSuccess && !FilenameToImport.IsEmpty())
 		{
 			// Failed to import a LOD, even after retries (if applicable)
-			FMessageDialog::Open(EAppMsgType::Ok, FText::Format(LOCTEXT("LODImport_Failure", "Failed to import LOD{0}"), FText::AsNumber(LODLevel)));
+			FbxMeshUtils::Private::ShowFailedToImportLodDialog(LODLevel);
 		}
 
 		if (bImportSuccess && bNotifyCB)
@@ -1127,4 +1161,4 @@ namespace FbxMeshUtils
 	}
 }  //end namespace MeshUtils
 
-#undef LOCTEXT_NAMESPACE
+#undef LOCTEXT_NAMESPACE
\ No newline at end of file
diff --git a/Engine/Source/Runtime/Interchange/Engine/Private/InterchangeMeshUtilities.cpp b/Engine/Source/Runtime/Interchange/Engine/Private/InterchangeMeshUtilities.cpp
--- a/Engine/Source/Runtime/Interchange/Engine/Private/InterchangeMeshUtilities.cpp
+++ b/Engine/Source/Runtime/Interchange/Engine/Private/InterchangeMeshUtilities.cpp
@@ -94,6 +94,7 @@ TFuture<bool> UInterchangeMeshUtilities::InternalImportCustomLodAsync(TSharedPtr
 	USkeletalMesh* SkeletalMesh = Cast<USkeletalMesh>(MeshObject);
 	UStaticMesh* StaticMesh = Cast<UStaticMesh>(MeshObject);
 	EInterchangePipelineContext ImportType = EInterchangePipelineContext::AssetCustomLODImport;
+	bool bInvalidLodIndex = false;
 	if (SkeletalMesh)
 	{
 		InterchangeAssetImportData = Cast<UInterchangeAssetImportData>(SkeletalMesh->GetAssetImportData());
@@ -101,6 +102,10 @@ TFuture<bool> UInterchangeMeshUtilities::InternalImportCustomLodAsync(TSharedPtr
 		{
 			ImportType = EInterchangePipelineContext::AssetCustomLODReimport;
 		}
+		if (LodIndex > SkeletalMesh->GetLODNum())
+		{
+			bInvalidLodIndex = true;
+		}
 	}
 	else if (StaticMesh)
 	{
@@ -109,6 +114,10 @@ TFuture<bool> UInterchangeMeshUtilities::InternalImportCustomLodAsync(TSharedPtr
 		{
 			ImportType = EInterchangePipelineContext::AssetCustomLODReimport;
 		}
+		if (LodIndex > StaticMesh->GetNumSourceModels())
+		{
+			bInvalidLodIndex = true;
+		}
 	}
 	else
 	{
@@ -117,6 +126,13 @@ TFuture<bool> UInterchangeMeshUtilities::InternalImportCustomLodAsync(TSharedPtr
 		return Promise->GetFuture();
 	}
 
+	if (bInvalidLodIndex)
+	{
+		UE_LOG(LogInterchangeEngine, Warning, TEXT("FInterchangeMeshUtilities::InternalImportCustomLodAsync: Invalid mesh LOD index %d, no prior LOD index exists."), LodIndex);
+		Promise->SetValue(false);
+		return Promise->GetFuture();
+	}
+
 	FImportAssetParameters ImportAssetParameters;
 	ImportAssetParameters.bIsAutomated = true;
 	if (InterchangeAssetImportData)
