Fixing a bug where the checkout error dialog box won't show even though there are packages which failed to checkout. This is due to CheckOutResult holding the return value of the FCheckOut operation yet it may have been set to failure when assembling the final check out list above. The use of CheckOutResult was a little overloaded. The dialog box will now show if any package names have been written to PkgsWhichFailedCheckout

diff --git a/Engine/Source/Editor/UnrealEd/Private/FileHelpers.cpp b/Engine/Source/Editor/UnrealEd/Private/FileHelpers.cpp
--- a/Engine/Source/Editor/UnrealEd/Private/FileHelpers.cpp
+++ b/Engine/Source/Editor/UnrealEd/Private/FileHelpers.cpp
@@ -2311,7 +2311,6 @@ ECommandResult::Type FEditorFileUtils::CheckoutPackages(const TArray<FString>& P
 			if(bShowCheckoutError)
 			{
 				PkgsWhichFailedCheckout += FString::Printf( TEXT("\n%s"), *PackageToCheckOutName );
-				CheckOutResult = ECommandResult::Failed;
 			}
 		}
 	}
@@ -2392,15 +2391,16 @@ ECommandResult::Type FEditorFileUtils::CheckoutPackages(const TArray<FString>& P
 				else
 				{
 					PkgsWhichFailedCheckout += FString::Printf( TEXT("\n%s"), *CurPackageName );
-					CheckOutResult = ECommandResult::Failed;
 				}
 			}
 		}
 	}
 
 	// If any packages failed the check out process, report them to the user so they know
-	if ( CheckOutResult == ECommandResult::Failed )
+	if (!PkgsWhichFailedCheckout.IsEmpty())
 	{
+		CheckOutResult = ECommandResult::Type::Failed;
+
 		FFormatNamedArguments Arguments;
 		Arguments.Add(TEXT("Packages"), FText::FromString( PkgsWhichFailedCheckout ));
 		FText MessageFormat = NSLOCTEXT("FileHelper", "FailedCheckoutDlgMessageFormatting", "The following assets could not be successfully checked out from revision control:{Packages}");
