Fixes retargeted animation not scrubbable in animation timeline.

diff --git a/Engine/Source/Editor/UnrealEd/Classes/Animation/DebugSkelMeshComponent.h b/Engine/Source/Editor/UnrealEd/Classes/Animation/DebugSkelMeshComponent.h
--- a/Engine/Source/Editor/UnrealEd/Classes/Animation/DebugSkelMeshComponent.h
+++ b/Engine/Source/Editor/UnrealEd/Classes/Animation/DebugSkelMeshComponent.h
@@ -370,6 +370,10 @@ class UDebugSkelMeshComponent : public USkeletalMeshComponent
 	UNREALED_API virtual void SetSkeletalMesh(USkeletalMesh* InSkelMesh, bool bReinitPose = true) override;
 	//~ End SkeletalMeshComponent Interface
 
+	//~ Begin UObject interface
+	UNREALED_API virtual void PostInitProperties() override;
+	//~ End UObject interface
+
 	// return true if currently preview animation asset is on
 	UNREALED_API virtual bool IsPreviewOn() const;
 
diff --git a/Engine/Source/Editor/UnrealEd/Private/Animation/DebugSkelMeshComponent.cpp b/Engine/Source/Editor/UnrealEd/Private/Animation/DebugSkelMeshComponent.cpp
--- a/Engine/Source/Editor/UnrealEd/Private/Animation/DebugSkelMeshComponent.cpp
+++ b/Engine/Source/Editor/UnrealEd/Private/Animation/DebugSkelMeshComponent.cpp
@@ -602,6 +602,21 @@ void UDebugSkelMeshComponent::SetSkeletalMesh(USkeletalMesh* InSkelMesh, bool bR
 	}
 }
 
+void UDebugSkelMeshComponent::PostInitProperties()
+{
+	EAnimationMode::Type OriginalMode = AnimationMode;
+
+	// potentially reverts the mode to "AnimationSingleNode" which is not compatible with animation editors
+	Super::PostInitProperties();
+
+	if (OriginalMode == EAnimationMode::AnimationBlueprint)
+	{
+		// in cases where AnimationBlueprint mode is not supported, revert to custom mode to prevent
+		// the base USkeletalMeshComponent from overriding the anim instance used by this component
+		AnimationMode = EAnimationMode::AnimationCustomMode;
+	}
+}
+
 void UDebugSkelMeshComponent::EnablePreview(bool bEnable, UAnimationAsset* PreviewAsset)
 {
 	if (PreviewInstance)
