TextureAssetAction new action : resize  to nearest power of two

diff --git a/Engine/Plugins/Editor/EngineAssetDefinitions/Source/Private/AssetDefinition_Texture.cpp b/Engine/Plugins/Editor/EngineAssetDefinitions/Source/Private/AssetDefinition_Texture.cpp
--- a/Engine/Plugins/Editor/EngineAssetDefinitions/Source/Private/AssetDefinition_Texture.cpp
+++ b/Engine/Plugins/Editor/EngineAssetDefinitions/Source/Private/AssetDefinition_Texture.cpp
@@ -64,6 +64,14 @@ namespace MenuExtension_Texture
 		}
 	}
 	
+	static void ExecuteResizeToPowerOfTwoTextureSource(const FToolMenuContext& InContext)
+	{
+		if (const UContentBrowserAssetContextMenuContext* CBContext = UContentBrowserAssetContextMenuContext::FindContextWithAssets(InContext))
+		{
+			UE::TextureAssetActions::TextureSource_ResizeToPowerOfTwo_WithDialog(CBContext->LoadSelectedObjects<UTexture>());
+		}
+	}
+
 	static void Execute8bitTextureSource(const FToolMenuContext& InContext)
 	{
 		if (const UContentBrowserAssetContextMenuContext* CBContext = UContentBrowserAssetContextMenuContext::FindContextWithAssets(InContext))
@@ -148,6 +156,14 @@ namespace MenuExtension_Texture
 							InSection.AddMenuEntry("Texture_FindMaterials", Label, ToolTip, Icon, UIAction);
 						}
 						
+						{
+							const TAttribute<FText> Label = LOCTEXT("Texture_ResizeSourceToPowerOfTwo", "Texture Source Resize To Power of Two");
+							const TAttribute<FText> ToolTip = LOCTEXT("Texture_ResizeSourceToPowerOfTwoTooltip", "Change texture source dimensions to the nearest power of two.");
+							const FSlateIcon Icon = FSlateIcon(FAppStyle::GetAppStyleSetName(), "ClassIcon.Texture2D");
+							const FToolMenuExecuteAction UIAction = FToolMenuExecuteAction::CreateStatic(&ExecuteResizeToPowerOfTwoTextureSource);
+							InSection.AddMenuEntry("Texture_ResizeSourceToPowerOfTwo", Label, ToolTip, Icon, UIAction);
+						}
+
 						{
 							const TAttribute<FText> Label = LOCTEXT("Texture_ResizeSource", "Texture Source Reduce Size");
 							const TAttribute<FText> ToolTip = LOCTEXT("Texture_ResizeSourceTooltip", "Reduce texture asset size by shrinking the texture source dimensions.");
diff --git a/Engine/Source/Editor/UnrealEd/Private/TextureAssetActions.cpp b/Engine/Source/Editor/UnrealEd/Private/TextureAssetActions.cpp
--- a/Engine/Source/Editor/UnrealEd/Private/TextureAssetActions.cpp
+++ b/Engine/Source/Editor/UnrealEd/Private/TextureAssetActions.cpp
@@ -42,6 +42,7 @@ enum class ETextureAction
 {
 	Invalid = 0,
 	Resize,
+	ResizePow2,
 	ConvertTo8bit,
 	JPEG
 };
@@ -56,7 +57,8 @@ enum class EAssetActionStatus
 	HasMipsLeaveExisting,
 	Already8bit,
 	DontChangeJPEG,
-	MustBe8BitForJPEG
+	MustBe8BitForJPEG,
+	AlreadyPow2
 };
 
 /**
@@ -132,8 +134,11 @@ private:
 
 	EVisibility GetErrorMessageVisibility() const;
 
-	EVisibility GetThresholdVisibility() const;
-	
+	EVisibility GetThresholdVisibility() const
+	{
+		return ( Action == ETextureAction::Resize || Action == ETextureAction::ResizePow2 || Action == ETextureAction::JPEG ) ? EVisibility::Visible : EVisibility::Collapsed;
+	}
+
 	EVisibility GetNMK16Visibility() const
 	{
 		return ( Action == ETextureAction::ConvertTo8bit ) ? EVisibility::Visible : EVisibility::Hidden;
@@ -200,6 +205,7 @@ static FText TAA_DialogTitle(ETextureAction Act)
 	switch(Act)
 	{
 	case ETextureAction::Resize: return LOCTEXT("TAA_Title_Resize", "Texture Asset : Resize Source");
+	case ETextureAction::ResizePow2: return LOCTEXT("TAA_Title_ResizePow2", "Texture Asset : Resize Source To Power of Two");
 	case ETextureAction::ConvertTo8bit: return LOCTEXT("TAA_Title_Convert", "Texture Asset : Convert To 8 bit or minimum viable bit depth");
 	case ETextureAction::JPEG: return LOCTEXT("TAA_Title_JPEG", "Texture Asset : Compress with JPEG");
 	default: check(0); return FText();
@@ -211,6 +217,7 @@ static FText TAA_Intro(ETextureAction Act)
 	switch(Act)
 	{
 	case ETextureAction::Resize: return LOCTEXT("TAA_Intro_Resize", "Reduce size of Texture Source to compact uassets.  Resizing is done using mip filter, in power of two steps.  LODBias is adjusted but platform built size may change.");
+	case ETextureAction::ResizePow2: return LOCTEXT("TAA_Intro_ResizePow2", "Change Texture Source dimensions to the nearest power or two.  May shrink or enlarge depending on which is closer.  Only textures larger than threshold are changed.  Does not resize to be <= threshold size, that's just a selection filter.");
 	case ETextureAction::ConvertTo8bit: return LOCTEXT("TAA_Intro_Convert", "Convert Texture Source to 8 bit, Normals to 8 or 16, HDR to 16F.  Only converts if bits per pixel goes down.  Output built texture may change.  Make sure CompressionSetting and SRGB are set correctly first!");
 	case ETextureAction::JPEG: return LOCTEXT("TAA_Intro_JPEG", "Compress Texture Source with JPEG.  Only works on 8 bit, 2D simple textures.  Greatly reduces uasset size with some loss of quality.  Does not affect in-game size.");
 	default: check(0); return FText();
@@ -453,6 +460,7 @@ static FText AssetActionStatus_Text(EAssetActionStatus Status)
 	case EAssetActionStatus::NoSource: return LOCTEXT("TAAStatus_NoSource", "The texture has no source data; cooked Editor?");
 	case EAssetActionStatus::WrongType: return LOCTEXT("TAAStatus_WrongType", "The texture is not a supported type.");
 	case EAssetActionStatus::UnderSized: return LOCTEXT("TAAStatus_UnderSized", "The texture was under or equal the threshold size.");
+	case EAssetActionStatus::AlreadyPow2: return LOCTEXT("TAAStatus_AlreadyPow2", "The texture is already power of two.");
 	case EAssetActionStatus::HasLayers: return LOCTEXT("TAAStatus_HasLayers", "Textures with more than 1 layer not supported.");
 	case EAssetActionStatus::HasMipsLeaveExisting: return LOCTEXT("TAAStatus_HasMipsLeaveExisting", "Textures has imported mips and LeaveExisting, will not change. (change MipGen if wanted)");
 	case EAssetActionStatus::Already8bit: return LOCTEXT("TAAStatus_Already8", "Texture format is already minimum bit depth.");
@@ -520,7 +528,8 @@ TSharedRef<SWidget> STextureAssetList::CreateAssetLine(int index, const UTexture
 		DetailedInfoText = AssetActionStatus_Text(Status);
 
 		if ( Status == EAssetActionStatus::UnderSized ||
-			Status == EAssetActionStatus::Already8bit )
+			Status == EAssetActionStatus::Already8bit ||
+			Status == EAssetActionStatus::AlreadyPow2 )
 		{
 			SeverityIcon = "MessageLog.Note";
 		}
@@ -651,6 +660,36 @@ TSharedRef<SWidget> STextureAssetList::CreateAssetLine(int index, const UTexture
 	return Result;
 }
 
+static void DoResizeTextureSourceToPowerOfTwo(UTexture * Texture)
+{
+	check( ! Texture->Source.AreAllBlocksPowerOfTwo() ); // already filtered for
+
+	const FIntPoint BeforeSourceSize = Texture->Source.GetLogicalSize();
+
+	if ( ! UE::TextureUtilitiesCommon::Experimental::ResizeTextureSourceDataToNearestPowerOfTwo(Texture) )
+	{
+		UE_LOG(LogTexture, Display, TEXT("Texture (%s) did not resize."), *Texture->GetName());
+
+		// did not resize, but may have done PreEditChange
+		return;
+	}
+	
+	const FIntPoint AfterSourceSize = Texture->Source.GetLogicalSize();
+
+	UE_LOG(LogTexture, Display, TEXT("Texture (%s) did resize Before Size=%dx%d After Size=%dx%d"), *Texture->GetName(),
+		BeforeSourceSize.X,BeforeSourceSize.Y,AfterSourceSize.X,AfterSourceSize.Y);
+
+	// ?? if Texture was set to stretch to Pow2, we could remove that now, but leaving it is harmless
+
+	// @@ ?? if Texture was set to MipGen = NoMipMaps, change to FromTextureGroup ?
+
+	// this counts as a reimport :
+	UE::TextureUtilitiesCommon::ApplyDefaultsForNewlyImportedTextures(Texture,true);
+
+	// DownsizeTextureSourceData did the PreEditChange
+	Texture->PostEditChange();
+}
+
 static void DoResizeTextureSource(UTexture * Texture,int TargetSize)
 {
 	// we do the resizing considering only mip/LOD/build settings for the running Editor platform (eg. Windows)
@@ -682,15 +721,9 @@ static void DoResizeTextureSource(UTexture * Texture,int TargetSize)
 
 	// if AfterSize > BeforeSize , kick up LODBias
 	//	to try to preserve GetBuiltTextureSize
-	while( AfterSizeX > BeforeSizeX || AfterSizeY > BeforeSizeY )
-	{
-		Texture->LODBias ++;
-		// just shifting down AfterSize is not exactly right
-		//	but ensures that our loop terminates
-		AfterSizeX = (AfterSizeX+1)>>1;
-		AfterSizeY = (AfterSizeY+1)>>1;
-	}
-	
+	Texture->LODBias = FMath::RoundToInt32( FMath::Log2( (double) FMath::Max(AfterSizeX,AfterSizeY) / FMath::Max(BeforeSizeX,BeforeSizeY) ) );
+	Texture->LODBias = FMath::Max(0,Texture->LODBias);
+		
 	// recompute AfterSize if we changed LODBias
 	if ( Texture->LODBias != 0 )
 	{
@@ -705,7 +738,7 @@ static void DoResizeTextureSource(UTexture * Texture,int TargetSize)
 		// not a warning, just FYI
 		// changing built size is totally possible and expected to happen sometimes
 		//	basically any time you resize smaller than the previous in-game size
-		UE_LOG(LogTexture,Verbose,TEXT("DoResizeTextureSource failed to preserve built size; was: %dx%d now: %dx%d on [%s]"),
+		UE_LOG(LogTexture,Display,TEXT("DoResizeTextureSource failed to preserve built size; was: %dx%d now: %dx%d on [%s]"),
 			BeforeSizeX,BeforeSizeY,
 			AfterSizeX,AfterSizeY,
 			*Texture->GetFullName());
@@ -918,6 +951,9 @@ void STextureAssetList::DoAction()
 		case ETextureAction::Resize:
 			DoResizeTextureSource(Texture,ThresholdValue);
 			break;
+		case ETextureAction::ResizePow2:
+			DoResizeTextureSourceToPowerOfTwo(Texture);
+			break;
 		case ETextureAction::ConvertTo8bit:
 			DoConvertTo8bitTextureSource(Texture,bNormalMapsKeep16bits);
 			break;
@@ -994,7 +1030,7 @@ void STextureAssetList::UpdateList()
 			}
 		
 			// per-Action filters :
-			if( Action == ETextureAction::Resize )
+			if( Action == ETextureAction::Resize || Action == ETextureAction::ResizePow2 )
 			{
 				int MaxSize = FMath::Max(SourceSize.X,SourceSize.Y);
 
@@ -1003,6 +1039,10 @@ void STextureAssetList::UpdateList()
 					// Resize only supports 2d and Cube for now, no arrays or volumes
 					Status = EAssetActionStatus::WrongType;
 				}
+				else if ( Action == ETextureAction::ResizePow2 && Texture->Source.AreAllBlocksPowerOfTwo() )
+				{
+					Status = EAssetActionStatus::AlreadyPow2;
+				}
 				else if ( MaxSize <= ThresholdValue )
 				{
 					Status = EAssetActionStatus::UnderSized;
@@ -1038,7 +1078,7 @@ void STextureAssetList::UpdateList()
 				if ( Class != ETextureClass::TwoD || Texture->Source.GetNumBlocks() > 1 || Texture->Source.GetNumSlices() > 1 )
 				{
 					// JPEG only supports 2d
-					// Blocked/UDIM doesn't support JPEG (@@ ??)
+					// Blocked/UDIM doesn't support JPEG
 					Status = EAssetActionStatus::WrongType;
 				}
 				else if ( Texture->Source.GetSourceCompression() == ETextureSourceCompressionFormat::TSCF_JPEG )
@@ -1165,11 +1205,6 @@ EVisibility STextureAssetList::GetErrorMessageVisibility() const
 	return (ErrorMessage.IsEmpty()) ? EVisibility::Collapsed : EVisibility::Visible;
 }
 
-EVisibility STextureAssetList::GetThresholdVisibility() const
-{
-	return ( Action == ETextureAction::Resize || Action == ETextureAction::JPEG ) ? EVisibility::Visible : EVisibility::Collapsed;
-}
-
 FText STextureAssetList::GetIntroMessage() const
 {
 	return IntroMessage;
@@ -1233,4 +1268,10 @@ void UE::TextureAssetActions::TextureSource_JPEG_WithDialog(const TArray<UTextur
 	Dlg.ShowModal();
 }
 
+void UE::TextureAssetActions::TextureSource_ResizeToPowerOfTwo_WithDialog(const TArray<UTexture*> & InTextures)
+{
+	STextureActionDlg Dlg(InTextures,ETextureAction::ResizePow2);
+	Dlg.ShowModal();
+}
+
 #undef LOCTEXT_NAMESPACE
diff --git a/Engine/Source/Editor/UnrealEd/Public/TextureAssetActions.h b/Engine/Source/Editor/UnrealEd/Public/TextureAssetActions.h
--- a/Engine/Source/Editor/UnrealEd/Public/TextureAssetActions.h
+++ b/Engine/Source/Editor/UnrealEd/Public/TextureAssetActions.h
@@ -13,10 +13,11 @@ namespace TextureAssetActions
 
 UNREALED_API void TextureSource_Resize_WithDialog(const TArray<UTexture*> & InTextures);
 
+UNREALED_API void TextureSource_ResizeToPowerOfTwo_WithDialog(const TArray<UTexture*> & InTextures);
+
 UNREALED_API void TextureSource_ConvertTo8bit_WithDialog(const TArray<UTexture*> & InTextures);
 
 UNREALED_API void TextureSource_JPEG_WithDialog(const TArray<UTexture*> & InTextures);
 
-
 }
 };
diff --git a/Engine/Source/Runtime/TextureUtilitiesCommon/Private/TextureSourceDataUtils.cpp b/Engine/Source/Runtime/TextureUtilitiesCommon/Private/TextureSourceDataUtils.cpp
--- a/Engine/Source/Runtime/TextureUtilitiesCommon/Private/TextureSourceDataUtils.cpp
+++ b/Engine/Source/Runtime/TextureUtilitiesCommon/Private/TextureSourceDataUtils.cpp
@@ -21,6 +21,7 @@ namespace Private
 	static bool ResizeTexture2D(UTexture* Texture, int32 MaxSize, const ITargetPlatform* TargetPlatform)
 	{
 		check( Texture->Source.GetNumLayers() == 1 );
+		check( Texture->Source.GetNumBlocks() == 1 );
 		const int32 LayerIndex = 0;
 
 		// We want to reduce the asset size so ignore the imported mip(s) (??)
@@ -49,9 +50,63 @@ namespace Private
 
 		Texture->Source.Init(Image);
 
-		// Compress() applies PNG filter to the BulkData
-		//	Compress is done automatically in Texture PreSave
-		//Texture->Source.Compress(); 
+		// if gamma was Pow22 it is now sRGB
+		Texture->bUseLegacyGamma = false;
+
+		// PostEditChange is called outside by our caller
+
+		return true;
+	}
+
+	static int32 RoundToNearestInt32PowerOfTwo(double X)
+	{
+		double Log2X = FMath::Log2(X);
+		int32 IntLog2X = FMath::RoundToInt32(Log2X);
+		
+		return 1 << FMath::Clamp(IntLog2X,0,30);
+	}
+
+	// resize so that the largest dimension is == MaxSize
+	static bool ResizeTextureToNearestPow2(UTexture* Texture)
+	{
+		check( Texture->Source.GetNumLayers() == 1 );
+		check( Texture->Source.GetNumBlocks() == 1 );
+		const int32 LayerIndex = 0;
+
+		// discard imported mips :
+		const int32 MipIndex = 0;
+
+		FImage Image;
+		if (!Texture->Source.GetMipImage(Image, MipIndex))
+		{
+			UE_LOG(LogTexture,Error,TEXT("ResizeTexture2D: Texture GetMipImage failed [%s]"),
+				*Texture->GetFullName());
+			return false;
+		}
+
+		check( Image.SizeX == Texture->Source.GetSizeX() );
+		check( Image.SizeY == Texture->Source.GetSizeY() );
+
+		int32 TargetSizeX,TargetSizeY;
+
+		// make the larger dimension go to nearest pow2 first
+		// then fix the smaller dimension for aspect ratio
+		if ( Image.SizeX >= Image.SizeY )
+		{
+			TargetSizeX = RoundToNearestInt32PowerOfTwo(Image.SizeX);
+			TargetSizeY = RoundToNearestInt32PowerOfTwo( (double) TargetSizeX * Image.SizeY / Image.SizeX );
+		}
+		else
+		{
+			TargetSizeY = RoundToNearestInt32PowerOfTwo(Image.SizeY);
+			TargetSizeX = RoundToNearestInt32PowerOfTwo( (double) TargetSizeY * Image.SizeX / Image.SizeY );
+		}
+
+		FImageCore::ResizeImageInPlace(Image,TargetSizeX,TargetSizeY);
+		
+		Texture->PreEditChange(nullptr);
+
+		Texture->Source.Init(Image);
 
 		// if gamma was Pow22 it is now sRGB
 		Texture->bUseLegacyGamma = false;
@@ -88,6 +143,8 @@ namespace Private
 		check( Texture->Source.GetNumLayers() == 1 );
 		const int32 NumLayers = 1;
 		const int32 LayerIndex = 0;
+		
+		check( Texture->Source.GetNumBlocks() > 1 );
 
 		// MaxSize is applied to the total UDIM size
 
@@ -166,12 +223,107 @@ namespace Private
 
 		return true;
 	}
-}
 
+	static bool ResizeTextureToNearestPow2Blocked(UTexture* Texture)
+	{
+		// does not support layers
+		check( Texture->Source.GetNumLayers() == 1 );
+		const int32 NumLayers = 1;
+		const int32 LayerIndex = 0;
+		
+		const int32 NumBlocks = Texture->Source.GetNumBlocks();
+		check( NumBlocks > 1 );
+	
+		// Resize to Pow2 acts on each UDIM block, not the net size
 
-TEXTUREUTILITIESCOMMON_API bool DownsizeTextureSourceData(UTexture* Texture, int32 TargetSourceSize, const ITargetPlatform* TargetPlatform)
+		TArray<FTextureSourceBlock> SourceBlocks;
+		SourceBlocks.SetNum(NumBlocks);
+		
+		int32 BlockMaxSizeX = 0;
+		int32 BlockMaxSizeY = 0;
+
+		// get the largest of all blocks (same as VT builder)
+		for (int32 BlockIndex = 0; BlockIndex < NumBlocks; ++BlockIndex)
+		{
+			Texture->Source.GetBlock(BlockIndex, SourceBlocks[BlockIndex]);
+
+			BlockMaxSizeX = FMath::Max(BlockMaxSizeX, SourceBlocks[BlockIndex].SizeX );
+			BlockMaxSizeY = FMath::Max(BlockMaxSizeY, SourceBlocks[BlockIndex].SizeY );
+		}
+		
+		int32 TargetSizeX,TargetSizeY;
+
+		// make the larger dimension go to nearest pow2 first
+		// then fix the smaller dimension for aspect ratio
+		if ( BlockMaxSizeX >= BlockMaxSizeY )
+		{
+			TargetSizeX = RoundToNearestInt32PowerOfTwo(BlockMaxSizeX);
+			TargetSizeY = RoundToNearestInt32PowerOfTwo( (double) TargetSizeX * BlockMaxSizeY / BlockMaxSizeX );
+		}
+		else
+		{
+			TargetSizeY = RoundToNearestInt32PowerOfTwo(BlockMaxSizeY);
+			TargetSizeX = RoundToNearestInt32PowerOfTwo( (double) TargetSizeY * BlockMaxSizeX / BlockMaxSizeY );
+		}
+
+		TArray<FImage> ResizedBlocks;
+		ResizedBlocks.SetNum(NumBlocks);
+
+		for (int32 BlockIndex = 0; BlockIndex < NumBlocks; ++BlockIndex)
+		{
+			FImage & SourceMip0 = ResizedBlocks[BlockIndex];
+
+			// Drops any imported mip data (yuck?)
+			const int32 MipIndex = 0;
+			if (!Texture->Source.GetMipImage(SourceMip0, BlockIndex, LayerIndex, MipIndex))
+			{
+				UE_LOG(LogTexture,Error,TEXT("ResizeTexture2DBlocked: Texture GetMipImage failed [%s]"),
+					*Texture->GetFullName());
+				return false;
+			}
+
+			// all blocks have to be pow2 and the same aspect ratio
+			// but they don't have to all be the same size
+			// if they were different sizes, we could use a mip of TargetSize here; eg. (TargetSizeX>>1) 
+			//	if that was closer to the original size
+			// don't bother for now, just make them all the same size
+
+			// if this block is == TargetSize, this is a nop
+			FImageCore::ResizeImageInPlace(SourceMip0,TargetSizeX,TargetSizeY);
+					
+			SourceBlocks[BlockIndex].SizeX = TargetSizeX;
+			SourceBlocks[BlockIndex].SizeY = TargetSizeY;
+			SourceBlocks[BlockIndex].NumMips = 1;
+		}
+				
+		// Protect the code from an async build of the texture
+		Texture->PreEditChange(nullptr);
+
+		UE::Serialization::FEditorBulkData::FSharedBufferWithID ResizedImageBufferWithID = MakeSharedBufferForImageDatas(ResizedBlocks);
+
+		const ETextureSourceFormat SourceFormat = Texture->Source.GetFormat();
+		Texture->Source.InitBlocked(
+			&SourceFormat, // array of formats per layer
+			SourceBlocks.GetData(),
+			NumLayers,
+			NumBlocks,
+			MoveTemp(ResizedImageBufferWithID)
+		);
+
+		// if gamma was Pow22 it is now sRGB
+		Texture->bUseLegacyGamma = false;
+
+		return true;
+	}
+
+} // namespace Private
+
+static bool IsTextureResizeSupported(UTexture * Texture)
 {
-	check( Texture->Source.IsValid() );
+	if ( ! ensure( Texture->Source.IsValid() ) )
+	{
+		return false;
+	}
 
 	// Check if we don't know how to resize that texture
 
@@ -213,12 +365,47 @@ TEXTUREUTILITIESCOMMON_API bool DownsizeTextureSourceData(UTexture* Texture, int
 		}
 	}
 
+	return true;
+}
+	
+TEXTUREUTILITIESCOMMON_API bool ResizeTextureSourceDataToNearestPowerOfTwo(UTexture* Texture)
+{
+	if ( ! IsTextureResizeSupported(Texture) )
+	{
+		return false;
+	}
+
+	if ( Texture->Source.AreAllBlocksPowerOfTwo() )
+	{
+		return false;
+	}
+	
+	if (Texture->Source.GetNumBlocks() == 1)
+	{
+		return Private::ResizeTextureToNearestPow2(Texture);
+	}
+	else
+	{
+		// UDIM VT
+		return Private::ResizeTextureToNearestPow2Blocked(Texture);
+	}
+}
+
+TEXTUREUTILITIESCOMMON_API bool DownsizeTextureSourceData(UTexture* Texture, int32 TargetSourceSize, const ITargetPlatform* TargetPlatform)
+{
+	if ( ! IsTextureResizeSupported(Texture) )
+	{
+		return false;
+	}
+
 	FIntPoint SourceSize = Texture->Source.GetLogicalSize();
 	if (SourceSize.X <= TargetSourceSize && SourceSize.Y <= TargetSourceSize)
 	{
 		return false;
 	}
 
+	// ?? if Texture has StretchToPow2 set, then do ResizeTextureToNearestPow2 automatically here ?
+
 	if (Texture->Source.GetNumBlocks() == 1)
 	{
 		return Private::ResizeTexture2D(Texture, TargetSourceSize, TargetPlatform);
diff --git a/Engine/Source/Runtime/TextureUtilitiesCommon/Public/TextureSourceDataUtils.h b/Engine/Source/Runtime/TextureUtilitiesCommon/Public/TextureSourceDataUtils.h
--- a/Engine/Source/Runtime/TextureUtilitiesCommon/Public/TextureSourceDataUtils.h
+++ b/Engine/Source/Runtime/TextureUtilitiesCommon/Public/TextureSourceDataUtils.h
@@ -14,7 +14,12 @@ namespace UE::TextureUtilitiesCommon::Experimental
 {
 	// Try to resize the texture source data so that it is not larger than one build for the specified platform (only Texture2D and TextureCube are supported at the moment)
 	// Note it does not trigger the post edit change after modifying the source texture
-	TEXTUREUTILITIESCOMMON_API bool DownsizeTextureSourceData(UTexture* Texture, int32 TargetSizeInGame, const ITargetPlatform* TargetPlatform);
+	//	TargetSizeInGame should be a power of 2
+	TEXTUREUTILITIESCOMMON_API bool DownsizeTextureSourceData(UTexture* Texture, int32 TargetSourceSize, const ITargetPlatform* TargetPlatform);
+	
+	// Try to resize the texture source data to a power of two
+	// Note it does not trigger the post edit change after modifying the source texture
+	TEXTUREUTILITIESCOMMON_API bool ResizeTextureSourceDataToNearestPowerOfTwo(UTexture* Texture);
 
 	//Note: This function does trigger the post edit change after modifying the source texture
 	TEXTUREUTILITIESCOMMON_API bool DownsizeTextureSourceDataNearRenderingSize(UTexture* Texture, const ITargetPlatform* TargetPlatform);
