[FScriptSetHelper and FScriptMapHelper iterators]
fixes for potential issues

diff --git a/Engine/Plugins/Experimental/Water/Source/Runtime/Private/WaterBodyIslandActor.cpp b/Engine/Plugins/Experimental/Water/Source/Runtime/Private/WaterBodyIslandActor.cpp
--- a/Engine/Plugins/Experimental/Water/Source/Runtime/Private/WaterBodyIslandActor.cpp
+++ b/Engine/Plugins/Experimental/Water/Source/Runtime/Private/WaterBodyIslandActor.cpp
@@ -141,9 +141,9 @@ void AWaterBodyIsland::PostLoad()
 						if (FMapProperty* OldLayerWeightmapSettingsMapProperty = CastField<FMapProperty>(BPProperty))
 						{
 							FScriptMapHelper MapHelper(OldLayerWeightmapSettingsMapProperty, OldLayerWeightmapSettingsMapProperty->ContainerPtrToValuePtr<void>(this));
-							for (int32 I = 0; I < MapHelper.Num(); ++I)
+							for (FScriptMapHelper::FIterator It(MapHelper); It; ++It)
 							{
-								uint8* PairPtr = MapHelper.GetPairPtr(I);
+								uint8* PairPtr = MapHelper.GetPairPtr(It);
 								const FName* Key = MapHelper.GetKeyProperty()->ContainerPtrToValuePtr<FName>(PairPtr);
 								const FWaterBodyWeightmapSettings* Value = MapHelper.GetValueProperty()->ContainerPtrToValuePtr<FWaterBodyWeightmapSettings>(PairPtr);
 								WaterWeightmapSettings.FindOrAdd(*Key) = *Value;
diff --git a/Engine/Plugins/Runtime/RigVM/Source/RigVM/Public/RigVMCore/RigVMExternalVariable.h b/Engine/Plugins/Runtime/RigVM/Source/RigVM/Public/RigVMCore/RigVMExternalVariable.h
--- a/Engine/Plugins/Runtime/RigVM/Source/RigVM/Public/RigVMCore/RigVMExternalVariable.h
+++ b/Engine/Plugins/Runtime/RigVM/Source/RigVM/Public/RigVMCore/RigVMExternalVariable.h
@@ -218,10 +218,10 @@ struct RIGVM_API FRigVMExternalVariable : public FRigVMExternalVariableDef
 		{
 			FScriptMapHelper MapHelper(MapProperty, InMemory);
 			int32 Hash = ::GetTypeHash(MapHelper.Num());
-			for (int32 Index = 0; Index < MapHelper.Num(); Index++)
+			for (FScriptMapHelper::FIterator It(MapHelper); It; ++It)
 			{
-				Hash = HashCombine(Hash, GetPropertyTypeHash(MapProperty->KeyProp, MapHelper.GetKeyPtr(Index)));
-				Hash = HashCombine(Hash, GetPropertyTypeHash(MapProperty->ValueProp, MapHelper.GetValuePtr(Index)));
+				Hash = HashCombine(Hash, GetPropertyTypeHash(MapProperty->KeyProp, MapHelper.GetKeyPtr(It)));
+				Hash = HashCombine(Hash, GetPropertyTypeHash(MapProperty->ValueProp, MapHelper.GetValuePtr(It)));
 			}
 			return Hash;
 		}
@@ -229,9 +229,9 @@ struct RIGVM_API FRigVMExternalVariable : public FRigVMExternalVariableDef
 		{
 			FScriptSetHelper SetHelper(SetProperty, InMemory);
 			int32 Hash = ::GetTypeHash(SetHelper.Num());
-			for (int32 Index = 0; Index < SetHelper.Num(); Index++)
+			for (FScriptSetHelper::FIterator It(SetHelper); It; ++It)
 			{
-				Hash = HashCombine(Hash, GetPropertyTypeHash(SetProperty->ElementProp, SetHelper.GetElementPtr(Index)));
+				Hash = HashCombine(Hash, GetPropertyTypeHash(SetProperty->ElementProp, SetHelper.GetElementPtr(It)));
 			}
 			return Hash;
 		}
diff --git a/Engine/Plugins/VirtualProduction/LevelSnapshots/Source/LevelSnapshots/Private/Filtering/Diffing/EquivalenceUtil.cpp b/Engine/Plugins/VirtualProduction/LevelSnapshots/Source/LevelSnapshots/Private/Filtering/Diffing/EquivalenceUtil.cpp
--- a/Engine/Plugins/VirtualProduction/LevelSnapshots/Source/LevelSnapshots/Private/Filtering/Diffing/EquivalenceUtil.cpp
+++ b/Engine/Plugins/VirtualProduction/LevelSnapshots/Source/LevelSnapshots/Private/Filtering/Diffing/EquivalenceUtil.cpp
@@ -325,11 +325,14 @@ bool UE::LevelSnapshots::Private::AreMapPropertiesEquivalent(ULevelSnapshot* Sna
 	{
 		return false;
 	}
+	
+	FScriptMapHelper::FIterator SnapshotMapIterator(SnapshotMap);
+	FScriptMapHelper::FIterator WorldMapIterator(WorldMap);
 
-	for (int32 j = 0; j < SnapshotMap.Num(); ++j)
+	for (; SnapshotMapIterator && WorldMapIterator; ++SnapshotMapIterator, ++WorldMapIterator)
 	{
-		void* const SnapshotPairPtr = SnapshotMap.GetPairPtr(j);
-		void* const WorldPairPtr = WorldMap.GetPairPtr(j);
+		void* const SnapshotPairPtr = SnapshotMap.GetPairPtr(SnapshotMapIterator);
+		void* const WorldPairPtr = WorldMap.GetPairPtr(WorldMapIterator);
 
 		const bool bAreKeysEquivalent = AreSnapshotAndOriginalPropertiesEquivalent(Snapshot, SnapshotMap.KeyProp, SnapshotPairPtr, WorldPairPtr, SnapshotActor, WorldActor);
 		const bool bAreValuesEquivalent = AreSnapshotAndOriginalPropertiesEquivalent(Snapshot, SnapshotMap.ValueProp, SnapshotPairPtr, WorldPairPtr, SnapshotActor, WorldActor);
diff --git a/Engine/Source/Editor/BlueprintGraph/Private/K2Node_MakeVariable.cpp b/Engine/Source/Editor/BlueprintGraph/Private/K2Node_MakeVariable.cpp
--- a/Engine/Source/Editor/BlueprintGraph/Private/K2Node_MakeVariable.cpp
+++ b/Engine/Source/Editor/BlueprintGraph/Private/K2Node_MakeVariable.cpp
@@ -141,22 +141,17 @@ void UK2Node_MakeVariable::SetupVariable(const FBPVariableDescription& InVariabl
 		const FSetProperty* SetProperty = CastFieldChecked<const FSetProperty>(Property);
 		FScriptSetHelper SetHelper(SetProperty, SetProperty->ContainerPtrToValuePtr<void>(StructData->GetStructMemory()));
 
-		// Go through each element in the array to set the default value
-		int32 Num = SetHelper.Num();
-		for( int32 I = 0 ; Num ; I++ )
+		// Go through each element in the set to set the default value
+		for (FScriptSetHelper::FIterator It(SetHelper); It; ++It)
 		{
-			if(SetHelper.IsValidIndex(I))
-			{
-				--Num;
-				const uint8* PropData = SetHelper.GetElementPtr(I);
-
-				// Retrieve the element's default value
-				FString DefaultValue;
-				FBlueprintEditorUtils::PropertyValueToString_Direct(SetProperty->ElementProp, PropData, DefaultValue);
-				UEdGraphPin* VariableInputPin = CreatePin(EGPD_Input, ContainedType, *FString::FromInt(I));
-				// Add one to the index for the pin to set the default on to skip the output pin
-				Schema->SetPinAutogeneratedDefaultValue(VariableInputPin, DefaultValue);
-			}
+			const uint8* PropData = SetHelper.GetElementPtr(It);
+
+			// Retrieve the element's default value
+			FString DefaultValue;
+			FBlueprintEditorUtils::PropertyValueToString_Direct(SetProperty->ElementProp, PropData, DefaultValue);
+			UEdGraphPin* VariableInputPin = CreatePin(EGPD_Input, ContainedType, *FString::FromInt(It.GetLogicalIndex()));
+			// Add one to the index for the pin to set the default on to skip the output pin
+			Schema->SetPinAutogeneratedDefaultValue(VariableInputPin, DefaultValue);
 		}
 	}
 	else if( VariableType.VarType.IsMap())
@@ -166,30 +161,27 @@ void UK2Node_MakeVariable::SetupVariable(const FBPVariableDescription& InVariabl
 		
 		const FEdGraphPinType ValueType = FEdGraphPinType::GetPinTypeForTerminalType(InVariableType.VarType.PinValueType);
 
-		// Go through each element in the array to set the default value
-		int32 Num = MapHelper.Num();
-		for( int32 I = 0 ; Num ; I++ )
+		// Go through each element in the map to set the default value
+		for (FScriptMapHelper::FIterator It(MapHelper); It; ++It)
 		{
-			if(MapHelper.IsValidIndex(I))
-			{
-				--Num;
-				const uint8* KeyData = MapHelper.GetKeyPtr(I);
-
-				// Retrieve the element's default value
-				FString KeyDefaultValue;
-				FBlueprintEditorUtils::PropertyValueToString_Direct(MapProperty->KeyProp, KeyData, KeyDefaultValue);
-				UEdGraphPin* VariableInputPin = CreatePin(EGPD_Input, ContainedType, *FString::FromInt(I).Append(TEXT("_Key")));
-				// Add one to the index for the pin to set the default on to skip the output pin
-				Schema->SetPinAutogeneratedDefaultValue(VariableInputPin, KeyDefaultValue);
-
-				const uint8* ValueData = MapHelper.GetValuePtr(I);
-
-				FString ValueDefaultValue;
-				FBlueprintEditorUtils::PropertyValueToString_Direct(MapProperty->ValueProp, ValueData, ValueDefaultValue);
-				VariableInputPin = CreatePin(EGPD_Input, ValueType, *FString::FromInt(I).Append(TEXT("_Value")));
-				// Add one to the index for the pin to set the default on to skip the output pin
-				Schema->SetPinAutogeneratedDefaultValue(VariableInputPin, ValueDefaultValue);
-			}
+			const uint8* KeyData = MapHelper.GetKeyPtr(It);
+
+			// Retrieve the element's default value
+			FString KeyDefaultValue;
+			FBlueprintEditorUtils::PropertyValueToString_Direct(MapProperty->KeyProp, KeyData, KeyDefaultValue);
+
+			UEdGraphPin* VariableInputPin = CreatePin(EGPD_Input, ContainedType, *FString::FromInt(It.GetLogicalIndex()).Append(TEXT("_Key")));
+			// Add one to the index for the pin to set the default on to skip the output pin
+			Schema->SetPinAutogeneratedDefaultValue(VariableInputPin, KeyDefaultValue);
+
+			const uint8* ValueData = MapHelper.GetValuePtr(It);
+
+			FString ValueDefaultValue;
+			FBlueprintEditorUtils::PropertyValueToString_Direct(MapProperty->ValueProp, ValueData, ValueDefaultValue);
+
+			VariableInputPin = CreatePin(EGPD_Input, ValueType, *FString::FromInt(It.GetLogicalIndex()).Append(TEXT("_Value")));
+			// Add one to the index for the pin to set the default on to skip the output pin
+			Schema->SetPinAutogeneratedDefaultValue(VariableInputPin, ValueDefaultValue);
 		}
 	}
 
diff --git a/Engine/Source/Editor/EditorConfig/Private/EditorConfig.cpp b/Engine/Source/Editor/EditorConfig/Private/EditorConfig.cpp
--- a/Engine/Source/Editor/EditorConfig/Private/EditorConfig.cpp
+++ b/Engine/Source/Editor/EditorConfig/Private/EditorConfig.cpp
@@ -495,14 +495,11 @@ TSharedPtr<FJsonValue> FEditorConfig::WriteSet(const FSetProperty* SetProperty,
 	TArray<TSharedPtr<FJsonValue>> JsonValuesArray;
 	JsonValuesArray.Reserve(SetHelper.Num());
 
-	for (int32 Idx = 0; Idx < SetHelper.Num(); ++Idx)
+	for (FScriptSetHelper::FIterator It(SetHelper); It; ++It)
 	{
-		if (SetHelper.IsValidIndex(Idx))
-		{
-			TSharedPtr<FJsonValue> ElementValue = WriteValue(InnerProperty, SetHelper.GetElementPtr(Idx), nullptr);
-			check(ElementValue.IsValid());
-			JsonValuesArray.Add(ElementValue);
-		}
+		TSharedPtr<FJsonValue> ElementValue = WriteValue(InnerProperty, SetHelper.GetElementPtr(It), nullptr);
+		check(ElementValue.IsValid());
+		JsonValuesArray.Add(ElementValue);
 	}
 
 	return MakeShared<FJsonValueArray>(JsonValuesArray);
@@ -529,18 +526,15 @@ TSharedPtr<FJsonValue> FEditorConfig::WriteMap(const FMapProperty* MapProperty,
 		TArray<TSharedPtr<FJsonValue>> JsonValuesArray;
 		JsonValuesArray.Reserve(MapHelper.Num());
 
-		for (int32 Idx = 0; Idx < MapHelper.Num(); ++Idx)
+		for (FScriptMapHelper::FIterator It(MapHelper); It; ++It)
 		{
-			if (MapHelper.IsValidIndex(Idx))
-			{
-				TSharedPtr<FJsonValue> JsonKey = WriteValue(KeyProperty, MapHelper.GetKeyPtr(Idx), nullptr);
-				check(JsonKey.IsValid());
-				JsonKeysArray.Add(JsonKey);
+			TSharedPtr<FJsonValue> JsonKey = WriteValue(KeyProperty, MapHelper.GetKeyPtr(It), nullptr);
+			check(JsonKey.IsValid());
+			JsonKeysArray.Add(JsonKey);
 
-				TSharedPtr<FJsonValue> JsonValue = WriteValue(ValueProperty, MapHelper.GetValuePtr(Idx), nullptr);
-				check(JsonValue.IsValid());
-				JsonValuesArray.Add(JsonValue);
-			}
+			TSharedPtr<FJsonValue> JsonValue = WriteValue(ValueProperty, MapHelper.GetValuePtr(It), nullptr);
+			check(JsonValue.IsValid());
+			JsonValuesArray.Add(JsonValue);
 		}
 
 		// maps can either be stored as $key, $value pairs or, if the keys can be stringified, as a JSON object
diff --git a/Engine/Source/Editor/UnrealEd/Private/Commandlets/FixConflictingLocalizationKeys.cpp b/Engine/Source/Editor/UnrealEd/Private/Commandlets/FixConflictingLocalizationKeys.cpp
--- a/Engine/Source/Editor/UnrealEd/Private/Commandlets/FixConflictingLocalizationKeys.cpp
+++ b/Engine/Source/Editor/UnrealEd/Private/Commandlets/FixConflictingLocalizationKeys.cpp
@@ -194,26 +194,11 @@ bool ReKeyTextProperty(UStruct* InOuterType, void* InAddrToUpdate, const TArray<
 
 						FScriptMapHelper ScriptMapHelper(MapProp, AddrToUpdate);
 
-						// ContainerIndex is the element index, but we need the sparse index
-						int32 SparseIndex = 0;
+						// ContainerIndex is the element index, but we need the sparse index which is computed by the iterator.
+						const int32 InternalIndex = ScriptMapHelper.FindInternalIndex(ContainerIndex);
+						if (InternalIndex != INDEX_NONE)
 						{
-							const int32 ElementCount = ScriptMapHelper.Num();
-							for (int32 ElementIndex = 0; ElementIndex < ElementCount; ++SparseIndex)
-							{
-								if (ScriptMapHelper.IsValidIndex(SparseIndex))
-								{
-									if (ElementIndex == ContainerIndex)
-									{
-										break;
-									}
-									++ElementIndex;
-								}
-							}
-						}
-
-						if (ScriptMapHelper.IsValidIndex(SparseIndex))
-						{
-							AddrToUpdate = ScriptMapHelper.GetPairPtr(SparseIndex) + MapProp->MapLayout.ValueOffset;
+							AddrToUpdate = ScriptMapHelper.GetPairPtr(InternalIndex) + MapProp->MapLayout.ValueOffset;
 
 							// Is this a complex property? If so, we need to recurse into it
 							if (FStructProperty* StructProp = CastField<FStructProperty>(MapProp->ValueProp))
@@ -231,26 +216,11 @@ bool ReKeyTextProperty(UStruct* InOuterType, void* InAddrToUpdate, const TArray<
 
 						FScriptSetHelper ScriptSetHelper(SetProp, AddrToUpdate);
 
-						// ContainerIndex is the element index, but we need the sparse index
-						int32 SparseIndex = 0;
+						// ContainerIndex is the element index, but we need the sparse index which is computed by the iterator.
+						const int32 InternalIndex = ScriptSetHelper.FindInternalIndex(ContainerIndex);
+						if (InternalIndex != INDEX_NONE)
 						{
-							const int32 ElementCount = ScriptSetHelper.Num();
-							for (int32 ElementIndex = 0; ElementIndex < ElementCount; ++SparseIndex)
-							{
-								if (ScriptSetHelper.IsValidIndex(SparseIndex))
-								{
-									if (ElementIndex == ContainerIndex)
-									{
-										break;
-									}
-									++ElementIndex;
-								}
-							}
-						}
-
-						if (ScriptSetHelper.IsValidIndex(SparseIndex))
-						{
-							AddrToUpdate = ScriptSetHelper.GetElementPtr(SparseIndex);
+							AddrToUpdate = ScriptSetHelper.GetElementPtr(InternalIndex);
 
 							// Is this a complex property? If so, we need to recurse into it
 							if (FStructProperty* StructProp = CastField<FStructProperty>(SetProp->ElementProp))
@@ -270,9 +240,12 @@ bool ReKeyTextProperty(UStruct* InOuterType, void* InAddrToUpdate, const TArray<
 						AddrToUpdate = MangledPropToUpdate->ContainerPtrToValuePtr<void>(AddrToUpdate);
 
 						FScriptMapHelper ScriptMapHelper(MapProp, AddrToUpdate);
-						if (ScriptMapHelper.IsValidIndex(ContainerIndex))
+
+						// ContainerIndex is the element index, but we need the sparse index to get the element.
+						const int32 InternalIndex = ScriptMapHelper.FindInternalIndex(ContainerIndex);
+						if (InternalIndex != INDEX_NONE)
 						{
-							AddrToUpdate = ScriptMapHelper.GetPairPtr(ContainerIndex);
+							AddrToUpdate = ScriptMapHelper.GetPairPtr(InternalIndex);
 
 							// Is this a complex property? If so, we need to recurse into it
 							if (FStructProperty* StructProp = CastField<FStructProperty>(MapProp->KeyProp))
diff --git a/Engine/Source/Editor/UnrealEd/Private/DiffUtils.cpp b/Engine/Source/Editor/UnrealEd/Private/DiffUtils.cpp
--- a/Engine/Source/Editor/UnrealEd/Private/DiffUtils.cpp
+++ b/Engine/Source/Editor/UnrealEd/Private/DiffUtils.cpp
@@ -505,63 +505,41 @@ FPropertyPath FPropertySoftPath::ResolvePath(const UObject* Object) const
 		}
 		else if( const FSetProperty* SetProperty = CastField<FSetProperty>(ResolvedProperty) )
 		{
-			if(PropertyIndex != INDEX_NONE)
+			FScriptSetHelper SetHelper(SetProperty, SetProperty->ContainerPtrToValuePtr<const void*>( ContainerAddress ));
+			if (SetHelper.IsValidIndex(PropertyIndex))
 			{
-				FScriptSetHelper SetHelper(SetProperty, SetProperty->ContainerPtrToValuePtr<const void*>( ContainerAddress ));
+				const int32 InternalIndex = SetHelper.FindInternalIndex(PropertyIndex);
+				UpdateContainerAddress( SetProperty->ElementProp, SetHelper.GetElementPtr(InternalIndex), ContainerAddress, ContainerStruct );
 
-				// Figure out the real index in this instance of the set (sets have gaps in them):
-				int32 RealIndex = -1;
-				for( int32 J = 0; PropertyIndex >= 0; ++J)
-				{
-					++RealIndex;
-					if(SetHelper.IsValidIndex(J))
-					{
-						--PropertyIndex;
-					}
-				}
-
-				UpdateContainerAddress( SetProperty->ElementProp, SetHelper.GetElementPtr(RealIndex), ContainerAddress, ContainerStruct );
-
-				FPropertyInfo SetInfo(SetProperty->ElementProp, RealIndex);
+				FPropertyInfo SetInfo(SetProperty->ElementProp, PropertyIndex);
 				Ret.AddProperty(SetInfo);
 			}
 		}
 		else if( const FMapProperty* MapProperty = CastField<FMapProperty>(ResolvedProperty) )
 		{
-			if(PropertyIndex != INDEX_NONE)
+			FScriptMapHelper MapHelper(MapProperty, MapProperty->ContainerPtrToValuePtr<const void*>( ContainerAddress ));
+			if (MapHelper.IsValidIndex(PropertyIndex))
 			{
-				FScriptMapHelper MapHelper(MapProperty, MapProperty->ContainerPtrToValuePtr<const void*>( ContainerAddress ));
-				
-				// Figure out the real index in this instance of the map (map have gaps in them):
-				int32 RealIndex = -1;
-				for( int32 J = 0; PropertyIndex >= 0; ++J)
-				{
-					++RealIndex;
-					if(MapHelper.IsValidIndex(J))
-					{
-						--PropertyIndex;
-					}
-				}
-
+				const int32 InternalIndex = MapHelper.FindInternalIndex(PropertyIndex);
 				// we have an index, but are we looking into a key or value? Peek ahead to find out:
-				if(ensure(I + 1 < PropertyChain.Num()))
+				if(ensure((I + 1 < PropertyChain.Num())))
 				{
 					if(PropertyChain[I+1].PropertyName == MapProperty->KeyProp->GetFName())
 					{
 						++I;
 
-						UpdateContainerAddress( MapProperty->KeyProp, MapHelper.GetKeyPtr(RealIndex), ContainerAddress, ContainerStruct );
+						UpdateContainerAddress( MapProperty->KeyProp, MapHelper.GetKeyPtr(InternalIndex), ContainerAddress, ContainerStruct );
 
-						FPropertyInfo MakKeyInfo(MapProperty->KeyProp, RealIndex);
+						FPropertyInfo MakKeyInfo(MapProperty->KeyProp, PropertyIndex);
 						Ret.AddProperty(MakKeyInfo);
 					}
 					else if(ensure( PropertyChain[I+1].PropertyName == MapProperty->ValueProp->GetFName() ))
 					{	
 						++I;
 
-						UpdateContainerAddress( MapProperty->ValueProp, MapHelper.GetValuePtr(RealIndex), ContainerAddress, ContainerStruct );
-						
-						FPropertyInfo MapValueInfo(MapProperty->ValueProp, RealIndex);
+						UpdateContainerAddress( MapProperty->ValueProp, MapHelper.GetValuePtr(InternalIndex), ContainerAddress, ContainerStruct );
+
+						FPropertyInfo MapValueInfo(MapProperty->ValueProp, PropertyIndex);
 						Ret.AddProperty(MapValueInfo);
 					}
 				}
@@ -798,25 +776,6 @@ void DiffUtils::CompareUnrelatedSCS(const UBlueprint* Old, const TArray< FSCSRes
 	}
 }
 
-static void AdvanceSetIterator( FScriptSetHelper& SetHelper, int32& Index)
-{
-	do
-	{
-		++Index;
-	}
-	while(Index < SetHelper.GetMaxIndex() && !SetHelper.IsValidIndex(Index));
-}
-
-static void AdvanceMapIterator( FScriptMapHelper& MapHelper, int32& Index)
-{
-	do
-	{
-		++Index;
-	}
-	while(Index < MapHelper.GetMaxIndex() && !MapHelper.IsValidIndex(Index));
-}
-
-
 static void IdenticalHelper(const FProperty* AProperty, const FProperty* BProperty, const void* AValue, const void* BValue,
 	const UObject* OwningOuterA, const UObject* OwningOuterB, const FPropertySoftPath& RootPath,
 	TArray<FPropertySoftPath>& DifferingSubProperties, bool bStaticArrayHandled = false)
@@ -924,25 +883,14 @@ static void IdenticalHelper(const FProperty* AProperty, const FProperty* BProper
 			}
 
 			// note any differences in contained elements:
-			const int32 SetSizeA = SetHelperA.Num();
-			const int32 SetSizeB = SetHelperB.Num();
-			
-			int32 SetIndexA = -1;
-			int32 SetIndexB = -1;
-
-			AdvanceSetIterator(SetHelperA, SetIndexA);
-			AdvanceSetIterator(SetHelperB, SetIndexB);
-
-			for (int32 VirtualIndex = 0; VirtualIndex < SetSizeA && VirtualIndex < SetSizeB; ++VirtualIndex)
+			FScriptSetHelper::FIterator IteratorA(SetHelperA);
+			FScriptSetHelper::FIterator IteratorB(SetHelperB);
+			for (; IteratorA && IteratorB; ++IteratorA, ++IteratorB)
 			{
-				const void* SubValueA = SetHelperA.GetElementPtr(SetIndexA);
-				const void* SubValueB = SetHelperB.GetElementPtr(SetIndexB);
+				const void* SubValueA = SetHelperA.GetElementPtr(IteratorA);
+				const void* SubValueB = SetHelperB.GetElementPtr(IteratorB);
 				IdenticalHelper(APropAsSet->ElementProp, BPropAsSet->ElementProp, SubValueA, SubValueB,
-					OwningOuterA, OwningOuterB, FPropertySoftPath(RootPath, VirtualIndex), DifferingSubProperties);
-
-				// advance iterators in step:
-				AdvanceSetIterator(SetHelperA, SetIndexA);
-				AdvanceSetIterator(SetHelperB, SetIndexB);
+					OwningOuterA, OwningOuterB, FPropertySoftPath(RootPath, IteratorA.GetLogicalIndex()), DifferingSubProperties);
 			}
 		}
 		else
@@ -957,7 +905,6 @@ static void IdenticalHelper(const FProperty* AProperty, const FProperty* BProper
 		{
 			FScriptMapHelper MapHelperA(APropAsMap, AValue);
 			FScriptMapHelper MapHelperB(BPropAsMap, BValue);
-
 			if (MapHelperA.Num() != MapHelperB.Num())
 			{
 				// API not robust enough to indicate changes made to # of set elements, would
@@ -965,24 +912,15 @@ static void IdenticalHelper(const FProperty* AProperty, const FProperty* BProper
 				DifferingSubProperties.Push(RootPath);
 			}
 
-			int32 MapSizeA = MapHelperA.Num();
-			int32 MapSizeB = MapHelperB.Num();
-			
-			int32 MapIndexA = -1;
-			int32 MapIndexB = -1;
-
-			AdvanceMapIterator(MapHelperA, MapIndexA);
-			AdvanceMapIterator(MapHelperB, MapIndexB);
-			
-			for (int32 VirtualIndex = 0; VirtualIndex < MapSizeA && VirtualIndex < MapSizeB; ++VirtualIndex)
+			FScriptMapHelper::FIterator IteratorA(MapHelperA);
+			FScriptMapHelper::FIterator IteratorB(MapHelperB);
+			for (; IteratorA && IteratorB; ++IteratorA, ++IteratorB)
 			{
-				IdenticalHelper(APropAsMap->KeyProp, BPropAsMap->KeyProp, MapHelperA.GetKeyPtr(MapIndexA), MapHelperB.GetKeyPtr(MapIndexB),
-					OwningOuterA, OwningOuterB, FPropertySoftPath(RootPath, VirtualIndex), DifferingSubProperties);
-				IdenticalHelper(APropAsMap->ValueProp, BPropAsMap->ValueProp, MapHelperA.GetValuePtr(MapIndexA), MapHelperB.GetValuePtr(MapIndexB),
-					OwningOuterA, OwningOuterB, FPropertySoftPath(RootPath, VirtualIndex), DifferingSubProperties);
+				IdenticalHelper(APropAsMap->KeyProp, BPropAsMap->KeyProp, MapHelperA.GetKeyPtr(IteratorA), MapHelperB.GetKeyPtr(IteratorB),
+					OwningOuterA, OwningOuterB, FPropertySoftPath(RootPath, IteratorA.GetLogicalIndex()), DifferingSubProperties);
+				IdenticalHelper(APropAsMap->ValueProp, BPropAsMap->ValueProp, MapHelperA.GetValuePtr(IteratorA), MapHelperB.GetValuePtr(IteratorB),
+					OwningOuterA, OwningOuterB, FPropertySoftPath(RootPath, IteratorA.GetLogicalIndex()), DifferingSubProperties);
 
-				AdvanceMapIterator(MapHelperA, MapIndexA);
-				AdvanceMapIterator(MapHelperB, MapIndexB);
 			}
 		}
 		else
diff --git a/Engine/Source/Runtime/CoreUObject/Private/Serialization/UnversionedPropertySerializationTest.cpp b/Engine/Source/Runtime/CoreUObject/Private/Serialization/UnversionedPropertySerializationTest.cpp
--- a/Engine/Source/Runtime/CoreUObject/Private/Serialization/UnversionedPropertySerializationTest.cpp
+++ b/Engine/Source/Runtime/CoreUObject/Private/Serialization/UnversionedPropertySerializationTest.cpp
@@ -252,17 +252,14 @@ struct FUnversionedPropertyTest : public FUnversionedPropertyTestInput
 			return false;
 		}
 
-		for (int32 Num = HelperA.Num(), IndexA = 0; IndexA < Num; ++IndexA)
+		for (FScriptSetHelper::FIterator It(HelperA); It; ++It)
 		{
-			if (HelperA.IsValidIndex(IndexA))
+			const uint8* ElemA = HelperA.GetElementPtr(It);
+			const uint8* ElemB = FindElementPtr(HelperB, ElemA, OutDiff);
+
+			if (!ElemB)
 			{
-				const uint8* ElemA = HelperA.GetElementPtr(IndexA);
-				const uint8* ElemB = FindElementPtr(HelperB, ElemA, OutDiff);
-		
-				if (!ElemB)
-				{
-					return false;
-				} 
+				return false;
 			}
 		}
 
@@ -291,22 +288,19 @@ struct FUnversionedPropertyTest : public FUnversionedPropertyTestInput
 			return false;
 		}
 		
-		for (int32 Num = HelperA.Num(), IndexA = 0; IndexA < Num; ++IndexA)
+		for (FScriptMapHelper::FIterator It(HelperA); It; ++It)
 		{
-			if (HelperA.IsValidIndex(IndexA))
+			const uint8* PairA = HelperA.GetPairPtr(It);
+			const uint8* PairB = FindPairPtr(HelperB, PairA, OutDiff);
+
+			if (!PairB)
 			{
-				const uint8* PairA = HelperA.GetPairPtr(IndexA);
-				const uint8* PairB = FindPairPtr(HelperB, PairA, OutDiff);
-		
-				if (!PairB)
-				{
-					return false;
-				} 
+				return false;
+			}
 
-				if (!Equals(ValueProp, PairA + ValueOffset, PairB + ValueOffset, OutDiff))
-				{
-					return false;
-				}
+			if (!Equals(ValueProp, PairA + ValueOffset, PairB + ValueOffset, OutDiff))
+			{
+				return false;
 			}
 		}
 
diff --git a/Engine/Source/Runtime/CoreUObject/Private/UObject/UObjectGlobals.cpp b/Engine/Source/Runtime/CoreUObject/Private/UObject/UObjectGlobals.cpp
--- a/Engine/Source/Runtime/CoreUObject/Private/UObject/UObjectGlobals.cpp
+++ b/Engine/Source/Runtime/CoreUObject/Private/UObject/UObjectGlobals.cpp
@@ -4880,34 +4880,27 @@ template<EPropertyCollectFlags CollectFlags>
 static void CollectMapReferences(FReferenceCollector& Collector, FMapProperty& Property, void* Instance, const UObject* Referencer)
 {
 	FScriptMapHelper MapHelper(&Property, Instance);
-	const int32 Num = MapHelper.Num();
-	const bool bCollectKeys = MayContainStrongReference(*MapHelper.GetKeyProperty());
-	const bool bCollectValues = MayContainStrongReference(*MapHelper.GetValueProperty());
 
-	if (Num == 0)
+	if (MapHelper.Num() == 0)
 	{
-		return;	
+		return;
 	}
 
+	const bool bCollectKeys = MayContainStrongReference(*MapHelper.GetKeyProperty());
 	if (bCollectKeys)
 	{
-		for (int32 Idx = 0; Idx < Num; ++Idx)
+		for (FScriptMapHelper::FIterator It(MapHelper); It; ++It)
 		{
-			if (MapHelper.IsValidIndex(Idx))
-			{
-				CollectPropertyReferences<CollectFlags>(Collector, *MapHelper.GetKeyProperty(), MapHelper.GetPairPtr(Idx), Referencer);
-			}
+			CollectPropertyReferences<CollectFlags>(Collector, *MapHelper.GetKeyProperty(), MapHelper.GetPairPtr(It), Referencer);
 		}
 	}
 
+	const bool bCollectValues = MayContainStrongReference(*MapHelper.GetValueProperty());
 	if (bCollectValues)
 	{
-		for (int32 Idx = 0; Idx < Num; ++Idx)
+		for (FScriptMapHelper::FIterator It(MapHelper); It; ++It)
 		{
-			if (MapHelper.IsValidIndex(Idx))
-			{
-				CollectPropertyReferences<CollectFlags>(Collector, *MapHelper.GetValueProperty(), MapHelper.GetPairPtr(Idx), Referencer);
-			}
+			CollectPropertyReferences<CollectFlags>(Collector, *MapHelper.GetValueProperty(), MapHelper.GetPairPtr(It), Referencer);
 		}
 	}
 }
@@ -4918,12 +4911,9 @@ void CollectSetReferences(FReferenceCollector& Collector, FSetProperty& Property
 	FScriptSetHelper SetHelper(&Property, Instance);
 	if (MayContainStrongReference(*SetHelper.GetElementProperty()))
 	{
-		for (int32 Idx = 0, Num = SetHelper.Num(); Idx < Num; ++Idx)
-		{	
-			if (SetHelper.IsValidIndex(Idx))
-			{
-				CollectPropertyReferences<CollectFlags>(Collector, *SetHelper.GetElementProperty(), SetHelper.GetElementPtr(Idx), Referencer);
-			}
+		for (FScriptSetHelper::FIterator It(SetHelper); It; ++It)
+		{
+			CollectPropertyReferences<CollectFlags>(Collector, *SetHelper.GetElementProperty(), SetHelper.GetElementPtr(It), Referencer);
 		}
 	}
 }
diff --git a/Engine/Source/Runtime/CoreUObject/Private/UObject/UnrealType.cpp b/Engine/Source/Runtime/CoreUObject/Private/UObject/UnrealType.cpp
--- a/Engine/Source/Runtime/CoreUObject/Private/UObject/UnrealType.cpp
+++ b/Engine/Source/Runtime/CoreUObject/Private/UObject/UnrealType.cpp
@@ -143,21 +143,15 @@ FORCEINLINE_DEBUGGABLE bool FPropertyValueIterator::NextValue(EPropertyValueIter
 				if ((KeyFlags | ValueFlags) != EPropertyValueFlags::None)
 				{
 					FScriptMapHelper Helper(MapProperty, PropertyValue);
-					const int32 Num = Helper.Num();
-					for (int32 DynamicIndex = 0; DynamicIndex < Num; ++DynamicIndex)
+					for (FScriptMapHelper::FIterator It(Helper); It; ++It)
 					{
-						if (Helper.IsValidIndex(DynamicIndex))
+						if (KeyFlags != EPropertyValueFlags::None)
+						{
+							NewEntry.ValueArray.Emplace(BasePairType(KeyProperty, Helper.GetKeyPtr(It)), KeyFlags);
+						}
+						if (ValueFlags != EPropertyValueFlags::None)
 						{
-							if (KeyFlags != EPropertyValueFlags::None)
-							{
-								NewEntry.ValueArray.Emplace(
-									BasePairType(KeyProperty, Helper.GetKeyPtr(DynamicIndex)), KeyFlags);
-							}
-							if (ValueFlags != EPropertyValueFlags::None)
-							{
-								NewEntry.ValueArray.Emplace(
-									BasePairType(ValueProperty, Helper.GetValuePtr(DynamicIndex)), ValueFlags);
-							}
+							NewEntry.ValueArray.Emplace(BasePairType(ValueProperty, Helper.GetValuePtr(It)), ValueFlags);
 						}
 					}
 				}
@@ -170,14 +164,9 @@ FORCEINLINE_DEBUGGABLE bool FPropertyValueIterator::NextValue(EPropertyValueIter
 				if (InnerFlags != EPropertyValueFlags::None)
 				{
 					FScriptSetHelper Helper(SetProperty, PropertyValue);
-					const int32 Num = Helper.Num();
-					for (int32 DynamicIndex = 0; DynamicIndex < Num; ++DynamicIndex)
+					for (FScriptSetHelper::FIterator It(Helper); It; ++It)
 					{
-						if (Helper.IsValidIndex(DynamicIndex))
-						{
-							NewEntry.ValueArray.Emplace(
-								BasePairType(InnerProperty, Helper.GetElementPtr(DynamicIndex)), InnerFlags);
-						}
+						NewEntry.ValueArray.Emplace(BasePairType(InnerProperty, Helper.GetElementPtr(It)), InnerFlags);
 					}
 				}
 			}
diff --git a/Engine/Source/Runtime/Engine/Private/Kismet/BlueprintTypeConversions.cpp b/Engine/Source/Runtime/Engine/Private/Kismet/BlueprintTypeConversions.cpp
--- a/Engine/Source/Runtime/Engine/Private/Kismet/BlueprintTypeConversions.cpp
+++ b/Engine/Source/Runtime/Engine/Private/Kismet/BlueprintTypeConversions.cpp
@@ -291,18 +291,16 @@ DEFINE_FUNCTION(UBlueprintTypeConversions::execConvertSetType)
 	FScriptSetHelper SourceSet(SourceSetProperty, SourceSetAddr);
 	FScriptSetHelper DestSet(DestSetProperty, DestSetAddr);
 
-	int SourceSetSize = SourceSet.Num();
-	DestSet.EmptyElements(SourceSetSize);
+	DestSet.EmptyElements(SourceSet.Num());
 
 	BlueprintTypeConversions::ConversionFunctionT ConversionFunction = 
 		BlueprintTypeConversions::Internal::FindConversionFunction(SourceSetProperty->ElementProp, DestSetProperty->ElementProp);
 	check(ConversionFunction);
 
-	for (int i = 0; i < SourceSetSize; ++i)
+	for (FScriptSetHelper::FIterator It(SourceSet); It; ++It)
 	{
-		const void* SrcData = SourceSet.GetElementPtr(i);
-
-		int32 NewIndex = DestSet.AddDefaultValue_Invalid_NeedsRehash();
+		const void* SrcData = SourceSet.GetElementPtr(It);
+		const int32 NewIndex = DestSet.AddDefaultValue_Invalid_NeedsRehash();
 		void* DestData = DestSet.GetElementPtr(NewIndex);
 		(*ConversionFunction)(SrcData, DestData);
 	}
@@ -327,19 +325,18 @@ DEFINE_FUNCTION(UBlueprintTypeConversions::execConvertMapType)
 	FScriptMapHelper SourceMap(SourceMapProperty, SourceMapAddr);
 	FScriptMapHelper DestMap(DestMapProperty, DestMapAddr);
 
-	int SourceMapSize = SourceMap.Num();
-	DestMap.EmptyValues(SourceMapSize);
+	DestMap.EmptyValues(SourceMap.Num());
 
 	BlueprintTypeConversions::ConversionFunctionT KeyConversionFunction = 
 		BlueprintTypeConversions::Internal::FindConversionFunction(SourceMapProperty->KeyProp, DestMapProperty->KeyProp);
 	BlueprintTypeConversions::ConversionFunctionT ValueConversionFunction = 
 		BlueprintTypeConversions::Internal::FindConversionFunction(SourceMapProperty->ValueProp, DestMapProperty->ValueProp);
 
-	for (int i = 0; i < SourceMapSize; ++i)
+	for (FScriptMapHelper::FIterator It(SourceMap); It; ++It)
 	{
-		int32 NewIndex = DestMap.AddDefaultValue_Invalid_NeedsRehash();
+		const int32 NewIndex = DestMap.AddDefaultValue_Invalid_NeedsRehash();
 
-		const void* SourceKeyRawData = SourceMap.GetKeyPtr(i);
+		const void* SourceKeyRawData = SourceMap.GetKeyPtr(It);
 		void* DestinationKeyRawData = DestMap.GetKeyPtr(NewIndex);
 
 		if (KeyConversionFunction)
@@ -351,7 +348,7 @@ DEFINE_FUNCTION(UBlueprintTypeConversions::execConvertMapType)
 			SourceMapProperty->KeyProp->CopySingleValue(DestinationKeyRawData, SourceKeyRawData);
 		}
 
-		const void* SourceValueRawData = SourceMap.GetValuePtr(i);
+		const void* SourceValueRawData = SourceMap.GetValuePtr(It);
 		void* DestinationValueRawData = DestMap.GetValuePtr(NewIndex);
 
 		if (ValueConversionFunction)
diff --git a/Engine/Source/Runtime/Engine/Private/PropertyAccess.cpp b/Engine/Source/Runtime/Engine/Private/PropertyAccess.cpp
--- a/Engine/Source/Runtime/Engine/Private/PropertyAccess.cpp
+++ b/Engine/Source/Runtime/Engine/Private/PropertyAccess.cpp
@@ -681,10 +681,10 @@ struct FPropertyAccessSystem
 		FScriptMapHelper DestinationMapHelper(DestinationMapProperty, DestinationAddress);
 
 		DestinationMapHelper.EmptyValues();
-		for (int32 i = 0; i < SourceMapHelper.Num(); ++i)
+		for (FScriptMapHelper::FIterator It(SourceMapHelper); It; ++It)
 		{
-			const void* KeyData = SourceMapHelper.GetKeyPtr(i);
-			const SourceType* SourceValueData = reinterpret_cast<const SourceType*>(SourceMapHelper.GetValuePtr(i));
+			const void* KeyData = SourceMapHelper.GetKeyPtr(It);
+			const SourceType* SourceValueData = reinterpret_cast<const SourceType*>(SourceMapHelper.GetValuePtr(It));
 			DestinationType CastedType = static_cast<DestinationType>(*SourceValueData);
 			DestinationMapHelper.AddPair(KeyData, &CastedType);
 		}
diff --git a/Engine/Source/Runtime/JsonUtilities/Private/JsonObjectConverter.cpp b/Engine/Source/Runtime/JsonUtilities/Private/JsonObjectConverter.cpp
--- a/Engine/Source/Runtime/JsonUtilities/Private/JsonObjectConverter.cpp
+++ b/Engine/Source/Runtime/JsonUtilities/Private/JsonObjectConverter.cpp
@@ -107,18 +107,13 @@ TSharedPtr<FJsonValue> ConvertScalarFPropertyToJsonValue(FProperty* Property, co
 	{
 		TArray< TSharedPtr<FJsonValue> > Out;
 		FScriptSetHelper Helper(SetProperty, Value);
-		for ( int32 i=0, n=Helper.Num(); n; ++i )
+		for (FScriptSetHelper::FIterator It(Helper); It; ++It)
 		{
-			if ( Helper.IsValidIndex(i) )
+			TSharedPtr<FJsonValue> Elem = FJsonObjectConverter::UPropertyToJsonValue(SetProperty->ElementProp, Helper.GetElementPtr(It), CheckFlags & (~CPF_ParmFlags), SkipFlags, ExportCb, SetProperty);
+			if (Elem.IsValid())
 			{
-				TSharedPtr<FJsonValue> Elem = FJsonObjectConverter::UPropertyToJsonValue(SetProperty->ElementProp, Helper.GetElementPtr(i), CheckFlags & ( ~CPF_ParmFlags ), SkipFlags, ExportCb, SetProperty);
-				if ( Elem.IsValid() )
-				{
-					// add to the array
-					Out.Push(Elem);
-				}
-
-				--n;
+				// add to the array
+				Out.Push(Elem);
 			}
 		}
 		return MakeShared<FJsonValueArray>(Out);
@@ -128,38 +123,33 @@ TSharedPtr<FJsonValue> ConvertScalarFPropertyToJsonValue(FProperty* Property, co
 		TSharedRef<FJsonObject> Out = MakeShared<FJsonObject>();
 
 		FScriptMapHelper Helper(MapProperty, Value);
-		for ( int32 i=0, n = Helper.Num(); n; ++i )
+		for (FScriptMapHelper::FIterator It(Helper); It; ++It)
 		{
-			if ( Helper.IsValidIndex(i) )
+			TSharedPtr<FJsonValue> KeyElement = FJsonObjectConverter::UPropertyToJsonValue(MapProperty->KeyProp, Helper.GetKeyPtr(It), CheckFlags & (~CPF_ParmFlags), SkipFlags, ExportCb, MapProperty, ConversionFlags);
+			TSharedPtr<FJsonValue> ValueElement = FJsonObjectConverter::UPropertyToJsonValue(MapProperty->ValueProp, Helper.GetValuePtr(It), CheckFlags & (~CPF_ParmFlags), SkipFlags, ExportCb, MapProperty, ConversionFlags);
+			if (KeyElement.IsValid() && ValueElement.IsValid())
 			{
-				TSharedPtr<FJsonValue> KeyElement = FJsonObjectConverter::UPropertyToJsonValue(MapProperty->KeyProp, Helper.GetKeyPtr(i), CheckFlags & ( ~CPF_ParmFlags ), SkipFlags, ExportCb, MapProperty, ConversionFlags);
-				TSharedPtr<FJsonValue> ValueElement = FJsonObjectConverter::UPropertyToJsonValue(MapProperty->ValueProp, Helper.GetValuePtr(i), CheckFlags & ( ~CPF_ParmFlags ), SkipFlags, ExportCb, MapProperty, ConversionFlags);
-				if ( KeyElement.IsValid() && ValueElement.IsValid() )
+				FString KeyString;
+				if (!KeyElement->TryGetString(KeyString))
 				{
-					FString KeyString;
-					if (!KeyElement->TryGetString(KeyString))
+					MapProperty->KeyProp->ExportTextItem_Direct(KeyString, Helper.GetKeyPtr(It), nullptr, nullptr, 0);
+					if (KeyString.IsEmpty())
 					{
-						MapProperty->KeyProp->ExportTextItem_Direct(KeyString, Helper.GetKeyPtr(i), nullptr, nullptr, 0);
-						if (KeyString.IsEmpty())
-						{
-							UE_LOG(LogJson, Error, TEXT("Unable to convert key to string for property %s."), *MapProperty->GetAuthoredName())
-							KeyString = FString::Printf(TEXT("Unparsed Key %d"), i);
-						}
+						UE_LOG(LogJson, Error, TEXT("Unable to convert key to string for property %s."), *MapProperty->GetAuthoredName())
+						KeyString = FString::Printf(TEXT("Unparsed Key %d"), It.GetLogicalIndex());
 					}
+				}
 
-					// Coerce camelCase map keys for Enum/FName properties
-					if (CastField<FEnumProperty>(MapProperty->KeyProp) ||
-						CastField<FNameProperty>(MapProperty->KeyProp))
+				// Coerce camelCase map keys for Enum/FName properties
+				if (CastField<FEnumProperty>(MapProperty->KeyProp) ||
+					CastField<FNameProperty>(MapProperty->KeyProp))
+				{
+					if (!EnumHasAnyFlags(ConversionFlags, EJsonObjectConversionFlags::SkipStandardizeCase))
 					{
-						if (!EnumHasAnyFlags(ConversionFlags, EJsonObjectConversionFlags::SkipStandardizeCase))
-						{
-							KeyString = FJsonObjectConverter::StandardizeCase(KeyString);
-						}
+						KeyString = FJsonObjectConverter::StandardizeCase(KeyString);
 					}
-					Out->SetField(KeyString, ValueElement);
 				}
-
-				--n;
+				Out->SetField(KeyString, ValueElement);
 			}
 		}
 
diff --git a/Engine/Source/Runtime/Serialization/Private/StructDeserializer.cpp b/Engine/Source/Runtime/Serialization/Private/StructDeserializer.cpp
--- a/Engine/Source/Runtime/Serialization/Private/StructDeserializer.cpp
+++ b/Engine/Source/Runtime/Serialization/Private/StructDeserializer.cpp
@@ -570,10 +570,11 @@ bool FStructDeserializer::DeserializeElement(void* OutAddress, UStruct& OwnerInf
 				// handle set element
 				FSetProperty* SetProperty = CastField<FSetProperty>(CurrentState.Property);
 				FScriptSetHelper SetHelper(SetProperty, CurrentState.Data);
-				FProperty* Property = SetProperty->ElementProp;
-				if (SetHelper.IsValidIndex(CurrentState.ArrayIndex))
+				const int32 InternalIndex = SetHelper.FindInternalIndex(CurrentState.ArrayIndex);
+				if (InternalIndex != INDEX_NONE)
 				{
-					uint8* ElementPtr = SetHelper.GetElementPtr(CurrentState.ArrayIndex);
+					uint8* ElementPtr = SetHelper.GetElementPtr(InternalIndex);
+					FProperty* Property = SetProperty->ElementProp;
 					constexpr int32 ReadIndex = 0; //Pointer is offset so reading index is 0
 					if (!Backend.ReadProperty(Property, CurrentState.Property, ElementPtr, ReadIndex))
 					{
@@ -594,12 +595,13 @@ bool FStructDeserializer::DeserializeElement(void* OutAddress, UStruct& OwnerInf
 				// handle map element
 				FMapProperty* MapProperty = CastField<FMapProperty>(CurrentState.Property);
 				FScriptMapHelper MapHelper(MapProperty, CurrentState.Data);
-				FProperty* Property = MapProperty->ValueProp;
 
 				//When written as array, maps won't include the key, only values
-				if (MapHelper.IsValidIndex(CurrentState.ArrayIndex))
+				const int32 InternalIndex = MapHelper.FindInternalIndex(CurrentState.ArrayIndex);
+				if (InternalIndex != INDEX_NONE)
 				{
-					uint8* PairPtr = MapHelper.GetPairPtr(CurrentState.ArrayIndex);
+					uint8* PairPtr = MapHelper.GetPairPtr(InternalIndex);
+					FProperty* Property = MapProperty->ValueProp;
 					constexpr int32 ReadIndex = 0; //Pointer is offset so reading index is 0
 					if (!Backend.ReadProperty(Property, CurrentState.Property, PairPtr, ReadIndex))
 					{
@@ -694,12 +696,13 @@ bool FStructDeserializer::DeserializeElement(void* OutAddress, UStruct& OwnerInf
 						FScriptSetHelper SetHelper(SetProperty, SetProperty->ContainerPtrToValuePtr<void>(CurrentState.Data));
 
 						//If a specific index is asked and it's not valid, skip the property
-						if (SetHelper.IsValidIndex(CurrentState.ArrayIndex))
+						const int32 InternalIndex = SetHelper.FindInternalIndex(CurrentState.ArrayIndex);
+						if (InternalIndex != INDEX_NONE)
 						{
 							Property = SetProperty->ElementProp;
 
 							//Offset the pointer directly and give index 0 to be read so no offsetting is done during deserialization
-							CurrentState.Data = SetHelper.GetElementPtr(CurrentState.ArrayIndex);
+							CurrentState.Data = SetHelper.GetElementPtr(InternalIndex);
 							CurrentState.ArrayIndex = 0;
 
 							if (!Backend.ReadProperty(Property, nullptr, CurrentState.Data, CurrentState.ArrayIndex))
@@ -724,12 +727,13 @@ bool FStructDeserializer::DeserializeElement(void* OutAddress, UStruct& OwnerInf
 						FScriptMapHelper MapHelper(MapProperty, MapProperty->ContainerPtrToValuePtr<void>(CurrentState.Data));
 
 						//If a specific index is asked and it's not valid, skip the property
-						if (MapHelper.IsValidIndex(CurrentState.ArrayIndex))
+						const int32 InternalIndex = MapHelper.FindInternalIndex(CurrentState.ArrayIndex);
+						if (InternalIndex != INDEX_NONE)
 						{
 							Property = MapProperty->ValueProp;
 
 							//Offset the pointer directly and give index 0 to be read so no offsetting is done during deserialization
-							CurrentState.Data = MapHelper.GetPairPtr(CurrentState.ArrayIndex);
+							CurrentState.Data = MapHelper.GetPairPtr(InternalIndex);
 							CurrentState.ArrayIndex = 0;
 
 							if (!Backend.ReadProperty(Property, nullptr, CurrentState.Data, CurrentState.ArrayIndex))
@@ -834,10 +838,11 @@ bool FStructDeserializer::DeserializeElement(void* OutAddress, UStruct& OwnerInf
 				if (FSetProperty* SetProperty = CastField<FSetProperty>(CurrentState.Property))
 				{
 					FScriptSetHelper SetHelper(SetProperty, CurrentState.Data);
-					if (SetHelper.IsValidIndex(CurrentState.ArrayIndex))
+					const int32 InternalIndex = SetHelper.FindInternalIndex(CurrentState.ArrayIndex);
+					if (InternalIndex != INDEX_NONE)
 					{
 						NewState.Property = SetProperty->ElementProp;
-						NewState.Data = SetHelper.GetElementPtr(CurrentState.ArrayIndex);
+						NewState.Data = SetHelper.GetElementPtr(InternalIndex);
 						NewState.ArrayIndex = 0;
 						++CurrentState.ArrayIndex;
 					}
@@ -852,10 +857,11 @@ bool FStructDeserializer::DeserializeElement(void* OutAddress, UStruct& OwnerInf
 				else if (FMapProperty* MapProperty = CastField<FMapProperty>(CurrentState.Property))
 				{
 					FScriptMapHelper MapHelper(MapProperty, CurrentState.Data);
-					if (MapHelper.IsValidIndex(CurrentState.ArrayIndex))
+					const int32 InternalIndex = MapHelper.FindInternalIndex(CurrentState.ArrayIndex);
+					if (InternalIndex != INDEX_NONE)
 					{
 						NewState.Property = MapProperty->ValueProp;
-						NewState.Data = MapHelper.GetValuePtr(CurrentState.ArrayIndex);
+						NewState.Data = MapHelper.GetValuePtr(InternalIndex);
 						NewState.ArrayIndex = 0;
 						++CurrentState.ArrayIndex;
 					}
@@ -931,14 +937,15 @@ bool FStructDeserializer::DeserializeElement(void* OutAddress, UStruct& OwnerInf
 					if (FStructProperty* SetStructProperty = CastField<FStructProperty>(MapProperty->ValueProp))
 					{
 						FScriptMapHelper MapHelper(MapProperty, MapProperty->ContainerPtrToValuePtr<void>(CurrentState.Data));
-
+						
 						//If a specific index is asked and it's not valid, skip the property
-						if (MapHelper.IsValidIndex(CurrentState.ArrayIndex))
+						const int32 InternalIndex = MapHelper.FindInternalIndex(CurrentState.ArrayIndex);
+						if (InternalIndex != INDEX_NONE)
 						{
 							//We're skipping a level directly so CurrentState becomes the outer (set) and NewState the inner (Element prop)
 							CurrentState.Property = NewState.Property;
 							NewState.Property = SetStructProperty;
-							NewState.Data = MapHelper.GetValuePtr(CurrentState.ArrayIndex);
+							NewState.Data = MapHelper.GetValuePtr(InternalIndex);
 							NewState.ArrayIndex = 0;
 						}
 						else
@@ -956,14 +963,15 @@ bool FStructDeserializer::DeserializeElement(void* OutAddress, UStruct& OwnerInf
 					if (FStructProperty* SetStructProperty = CastField<FStructProperty>(SetProperty->ElementProp))
 					{
 						FScriptSetHelper SetHelper(SetProperty, SetProperty->ContainerPtrToValuePtr<void>(CurrentState.Data));
+						const int32 InternalIndex = SetHelper.FindInternalIndex(CurrentState.ArrayIndex);
 
 						//If a specific index is asked and it's not valid, skip the property
-						if (SetHelper.IsValidIndex(CurrentState.ArrayIndex))
+						if (InternalIndex != INDEX_NONE)
 						{
 							//We're skipping a level directly so CurrentState becomes the outer (set) and NewState the inner (Element prop)
 							CurrentState.Property = NewState.Property;
 							NewState.Property = SetProperty->ElementProp;
-							NewState.Data = SetHelper.GetElementPtr(CurrentState.ArrayIndex);
+							NewState.Data = SetHelper.GetElementPtr(InternalIndex);
 							NewState.ArrayIndex = 0;
 						}
 						else
diff --git a/Engine/Source/Runtime/Serialization/Private/Tests/StructSerializerTest.cpp b/Engine/Source/Runtime/Serialization/Private/Tests/StructSerializerTest.cpp
--- a/Engine/Source/Runtime/Serialization/Private/Tests/StructSerializerTest.cpp
+++ b/Engine/Source/Runtime/Serialization/Private/Tests/StructSerializerTest.cpp
@@ -486,7 +486,9 @@ namespace StructSerializerTest
 				FStructSerializer::SerializeElement(&TestStruct, Property, INDEX_NONE, SerializerBackend, Policies);
 
 				FStructSerializerTestStruct TestStruct2(NoInit);
-				TestStruct2.Sets.IntSet = { -1, -2, -3 };
+				TestStruct2.Sets.IntSet = { -1, -2, -3, -4, -5 };
+				TestStruct2.Sets.IntSet.Remove(-1);
+				TestStruct2.Sets.IntSet.Remove(-2);
 				TestStruct2.Sets.NameSet = { TEXT("Pre1"), TEXT("Pre2") , TEXT("Pre3") };
 				TestStruct2.Sets.StrSet = { TEXT("Pre4"), TEXT("Pre5") , TEXT("Pre6"), TEXT("Pre7") };
 				TestStruct2.Sets.StructSet = { FStructSerializerBuiltinTestStruct(NoInit) };
@@ -494,7 +496,7 @@ namespace StructSerializerTest
 
 				ValidateSets(Test, TestStruct.Sets, TestStruct2.Sets);
 			}
-			
+
 			//TArray<uint8> element
 			{
 				TArray<uint8> Buffer;
