[Backout] - CL29948080
Original CL Desc
-----------------------------------------------------------------
[FScriptSetHelper and FScriptMapHelper iterators]
fixes for potential issues

diff --git a/Engine/Plugins/Experimental/Water/Source/Runtime/Private/WaterBodyIslandActor.cpp b/Engine/Plugins/Experimental/Water/Source/Runtime/Private/WaterBodyIslandActor.cpp
--- a/Engine/Plugins/Experimental/Water/Source/Runtime/Private/WaterBodyIslandActor.cpp
+++ b/Engine/Plugins/Experimental/Water/Source/Runtime/Private/WaterBodyIslandActor.cpp
@@ -141,9 +141,9 @@ void AWaterBodyIsland::PostLoad()
 						if (FMapProperty* OldLayerWeightmapSettingsMapProperty = CastField<FMapProperty>(BPProperty))
 						{
 							FScriptMapHelper MapHelper(OldLayerWeightmapSettingsMapProperty, OldLayerWeightmapSettingsMapProperty->ContainerPtrToValuePtr<void>(this));
-							for (FScriptMapHelper::FIterator It(MapHelper); It; ++It)
+							for (int32 I = 0; I < MapHelper.Num(); ++I)
 							{
-								uint8* PairPtr = MapHelper.GetPairPtr(It);
+								uint8* PairPtr = MapHelper.GetPairPtr(I);
 								const FName* Key = MapHelper.GetKeyProperty()->ContainerPtrToValuePtr<FName>(PairPtr);
 								const FWaterBodyWeightmapSettings* Value = MapHelper.GetValueProperty()->ContainerPtrToValuePtr<FWaterBodyWeightmapSettings>(PairPtr);
 								WaterWeightmapSettings.FindOrAdd(*Key) = *Value;
diff --git a/Engine/Plugins/Runtime/RigVM/Source/RigVM/Public/RigVMCore/RigVMExternalVariable.h b/Engine/Plugins/Runtime/RigVM/Source/RigVM/Public/RigVMCore/RigVMExternalVariable.h
--- a/Engine/Plugins/Runtime/RigVM/Source/RigVM/Public/RigVMCore/RigVMExternalVariable.h
+++ b/Engine/Plugins/Runtime/RigVM/Source/RigVM/Public/RigVMCore/RigVMExternalVariable.h
@@ -218,10 +218,10 @@ struct RIGVM_API FRigVMExternalVariable : public FRigVMExternalVariableDef
 		{
 			FScriptMapHelper MapHelper(MapProperty, InMemory);
 			int32 Hash = ::GetTypeHash(MapHelper.Num());
-			for (FScriptMapHelper::FIterator It(MapHelper); It; ++It)
+			for (int32 Index = 0; Index < MapHelper.Num(); Index++)
 			{
-				Hash = HashCombine(Hash, GetPropertyTypeHash(MapProperty->KeyProp, MapHelper.GetKeyPtr(It)));
-				Hash = HashCombine(Hash, GetPropertyTypeHash(MapProperty->ValueProp, MapHelper.GetValuePtr(It)));
+				Hash = HashCombine(Hash, GetPropertyTypeHash(MapProperty->KeyProp, MapHelper.GetKeyPtr(Index)));
+				Hash = HashCombine(Hash, GetPropertyTypeHash(MapProperty->ValueProp, MapHelper.GetValuePtr(Index)));
 			}
 			return Hash;
 		}
@@ -229,9 +229,9 @@ struct RIGVM_API FRigVMExternalVariable : public FRigVMExternalVariableDef
 		{
 			FScriptSetHelper SetHelper(SetProperty, InMemory);
 			int32 Hash = ::GetTypeHash(SetHelper.Num());
-			for (FScriptSetHelper::FIterator It(SetHelper); It; ++It)
+			for (int32 Index = 0; Index < SetHelper.Num(); Index++)
 			{
-				Hash = HashCombine(Hash, GetPropertyTypeHash(SetProperty->ElementProp, SetHelper.GetElementPtr(It)));
+				Hash = HashCombine(Hash, GetPropertyTypeHash(SetProperty->ElementProp, SetHelper.GetElementPtr(Index)));
 			}
 			return Hash;
 		}
diff --git a/Engine/Plugins/VirtualProduction/LevelSnapshots/Source/LevelSnapshots/Private/Filtering/Diffing/EquivalenceUtil.cpp b/Engine/Plugins/VirtualProduction/LevelSnapshots/Source/LevelSnapshots/Private/Filtering/Diffing/EquivalenceUtil.cpp
--- a/Engine/Plugins/VirtualProduction/LevelSnapshots/Source/LevelSnapshots/Private/Filtering/Diffing/EquivalenceUtil.cpp
+++ b/Engine/Plugins/VirtualProduction/LevelSnapshots/Source/LevelSnapshots/Private/Filtering/Diffing/EquivalenceUtil.cpp
@@ -325,14 +325,11 @@ bool UE::LevelSnapshots::Private::AreMapPropertiesEquivalent(ULevelSnapshot* Sna
 	{
 		return false;
 	}
-	
-	FScriptMapHelper::FIterator SnapshotMapIterator(SnapshotMap);
-	FScriptMapHelper::FIterator WorldMapIterator(WorldMap);
 
-	for (; SnapshotMapIterator && WorldMapIterator; ++SnapshotMapIterator, ++WorldMapIterator)
+	for (int32 j = 0; j < SnapshotMap.Num(); ++j)
 	{
-		void* const SnapshotPairPtr = SnapshotMap.GetPairPtr(SnapshotMapIterator);
-		void* const WorldPairPtr = WorldMap.GetPairPtr(WorldMapIterator);
+		void* const SnapshotPairPtr = SnapshotMap.GetPairPtr(j);
+		void* const WorldPairPtr = WorldMap.GetPairPtr(j);
 
 		const bool bAreKeysEquivalent = AreSnapshotAndOriginalPropertiesEquivalent(Snapshot, SnapshotMap.KeyProp, SnapshotPairPtr, WorldPairPtr, SnapshotActor, WorldActor);
 		const bool bAreValuesEquivalent = AreSnapshotAndOriginalPropertiesEquivalent(Snapshot, SnapshotMap.ValueProp, SnapshotPairPtr, WorldPairPtr, SnapshotActor, WorldActor);
diff --git a/Engine/Source/Editor/BlueprintGraph/Private/K2Node_MakeVariable.cpp b/Engine/Source/Editor/BlueprintGraph/Private/K2Node_MakeVariable.cpp
--- a/Engine/Source/Editor/BlueprintGraph/Private/K2Node_MakeVariable.cpp
+++ b/Engine/Source/Editor/BlueprintGraph/Private/K2Node_MakeVariable.cpp
@@ -141,17 +141,22 @@ void UK2Node_MakeVariable::SetupVariable(const FBPVariableDescription& InVariabl
 		const FSetProperty* SetProperty = CastFieldChecked<const FSetProperty>(Property);
 		FScriptSetHelper SetHelper(SetProperty, SetProperty->ContainerPtrToValuePtr<void>(StructData->GetStructMemory()));
 
-		// Go through each element in the set to set the default value
-		for (FScriptSetHelper::FIterator It(SetHelper); It; ++It)
+		// Go through each element in the array to set the default value
+		int32 Num = SetHelper.Num();
+		for( int32 I = 0 ; Num ; I++ )
 		{
-			const uint8* PropData = SetHelper.GetElementPtr(It);
-
-			// Retrieve the element's default value
-			FString DefaultValue;
-			FBlueprintEditorUtils::PropertyValueToString_Direct(SetProperty->ElementProp, PropData, DefaultValue);
-			UEdGraphPin* VariableInputPin = CreatePin(EGPD_Input, ContainedType, *FString::FromInt(It.GetLogicalIndex()));
-			// Add one to the index for the pin to set the default on to skip the output pin
-			Schema->SetPinAutogeneratedDefaultValue(VariableInputPin, DefaultValue);
+			if(SetHelper.IsValidIndex(I))
+			{
+				--Num;
+				const uint8* PropData = SetHelper.GetElementPtr(I);
+
+				// Retrieve the element's default value
+				FString DefaultValue;
+				FBlueprintEditorUtils::PropertyValueToString_Direct(SetProperty->ElementProp, PropData, DefaultValue);
+				UEdGraphPin* VariableInputPin = CreatePin(EGPD_Input, ContainedType, *FString::FromInt(I));
+				// Add one to the index for the pin to set the default on to skip the output pin
+				Schema->SetPinAutogeneratedDefaultValue(VariableInputPin, DefaultValue);
+			}
 		}
 	}
 	else if( VariableType.VarType.IsMap())
@@ -161,27 +166,30 @@ void UK2Node_MakeVariable::SetupVariable(const FBPVariableDescription& InVariabl
 		
 		const FEdGraphPinType ValueType = FEdGraphPinType::GetPinTypeForTerminalType(InVariableType.VarType.PinValueType);
 
-		// Go through each element in the map to set the default value
-		for (FScriptMapHelper::FIterator It(MapHelper); It; ++It)
+		// Go through each element in the array to set the default value
+		int32 Num = MapHelper.Num();
+		for( int32 I = 0 ; Num ; I++ )
 		{
-			const uint8* KeyData = MapHelper.GetKeyPtr(It);
-
-			// Retrieve the element's default value
-			FString KeyDefaultValue;
-			FBlueprintEditorUtils::PropertyValueToString_Direct(MapProperty->KeyProp, KeyData, KeyDefaultValue);
-
-			UEdGraphPin* VariableInputPin = CreatePin(EGPD_Input, ContainedType, *FString::FromInt(It.GetLogicalIndex()).Append(TEXT("_Key")));
-			// Add one to the index for the pin to set the default on to skip the output pin
-			Schema->SetPinAutogeneratedDefaultValue(VariableInputPin, KeyDefaultValue);
-
-			const uint8* ValueData = MapHelper.GetValuePtr(It);
-
-			FString ValueDefaultValue;
-			FBlueprintEditorUtils::PropertyValueToString_Direct(MapProperty->ValueProp, ValueData, ValueDefaultValue);
-
-			VariableInputPin = CreatePin(EGPD_Input, ValueType, *FString::FromInt(It.GetLogicalIndex()).Append(TEXT("_Value")));
-			// Add one to the index for the pin to set the default on to skip the output pin
-			Schema->SetPinAutogeneratedDefaultValue(VariableInputPin, ValueDefaultValue);
+			if(MapHelper.IsValidIndex(I))
+			{
+				--Num;
+				const uint8* KeyData = MapHelper.GetKeyPtr(I);
+
+				// Retrieve the element's default value
+				FString KeyDefaultValue;
+				FBlueprintEditorUtils::PropertyValueToString_Direct(MapProperty->KeyProp, KeyData, KeyDefaultValue);
+				UEdGraphPin* VariableInputPin = CreatePin(EGPD_Input, ContainedType, *FString::FromInt(I).Append(TEXT("_Key")));
+				// Add one to the index for the pin to set the default on to skip the output pin
+				Schema->SetPinAutogeneratedDefaultValue(VariableInputPin, KeyDefaultValue);
+
+				const uint8* ValueData = MapHelper.GetValuePtr(I);
+
+				FString ValueDefaultValue;
+				FBlueprintEditorUtils::PropertyValueToString_Direct(MapProperty->ValueProp, ValueData, ValueDefaultValue);
+				VariableInputPin = CreatePin(EGPD_Input, ValueType, *FString::FromInt(I).Append(TEXT("_Value")));
+				// Add one to the index for the pin to set the default on to skip the output pin
+				Schema->SetPinAutogeneratedDefaultValue(VariableInputPin, ValueDefaultValue);
+			}
 		}
 	}
 
diff --git a/Engine/Source/Editor/EditorConfig/Private/EditorConfig.cpp b/Engine/Source/Editor/EditorConfig/Private/EditorConfig.cpp
--- a/Engine/Source/Editor/EditorConfig/Private/EditorConfig.cpp
+++ b/Engine/Source/Editor/EditorConfig/Private/EditorConfig.cpp
@@ -495,11 +495,14 @@ TSharedPtr<FJsonValue> FEditorConfig::WriteSet(const FSetProperty* SetProperty,
 	TArray<TSharedPtr<FJsonValue>> JsonValuesArray;
 	JsonValuesArray.Reserve(SetHelper.Num());
 
-	for (FScriptSetHelper::FIterator It(SetHelper); It; ++It)
+	for (int32 Idx = 0; Idx < SetHelper.Num(); ++Idx)
 	{
-		TSharedPtr<FJsonValue> ElementValue = WriteValue(InnerProperty, SetHelper.GetElementPtr(It), nullptr);
-		check(ElementValue.IsValid());
-		JsonValuesArray.Add(ElementValue);
+		if (SetHelper.IsValidIndex(Idx))
+		{
+			TSharedPtr<FJsonValue> ElementValue = WriteValue(InnerProperty, SetHelper.GetElementPtr(Idx), nullptr);
+			check(ElementValue.IsValid());
+			JsonValuesArray.Add(ElementValue);
+		}
 	}
 
 	return MakeShared<FJsonValueArray>(JsonValuesArray);
@@ -526,15 +529,18 @@ TSharedPtr<FJsonValue> FEditorConfig::WriteMap(const FMapProperty* MapProperty,
 		TArray<TSharedPtr<FJsonValue>> JsonValuesArray;
 		JsonValuesArray.Reserve(MapHelper.Num());
 
-		for (FScriptMapHelper::FIterator It(MapHelper); It; ++It)
+		for (int32 Idx = 0; Idx < MapHelper.Num(); ++Idx)
 		{
-			TSharedPtr<FJsonValue> JsonKey = WriteValue(KeyProperty, MapHelper.GetKeyPtr(It), nullptr);
-			check(JsonKey.IsValid());
-			JsonKeysArray.Add(JsonKey);
+			if (MapHelper.IsValidIndex(Idx))
+			{
+				TSharedPtr<FJsonValue> JsonKey = WriteValue(KeyProperty, MapHelper.GetKeyPtr(Idx), nullptr);
+				check(JsonKey.IsValid());
+				JsonKeysArray.Add(JsonKey);
 
-			TSharedPtr<FJsonValue> JsonValue = WriteValue(ValueProperty, MapHelper.GetValuePtr(It), nullptr);
-			check(JsonValue.IsValid());
-			JsonValuesArray.Add(JsonValue);
+				TSharedPtr<FJsonValue> JsonValue = WriteValue(ValueProperty, MapHelper.GetValuePtr(Idx), nullptr);
+				check(JsonValue.IsValid());
+				JsonValuesArray.Add(JsonValue);
+			}
 		}
 
 		// maps can either be stored as $key, $value pairs or, if the keys can be stringified, as a JSON object
diff --git a/Engine/Source/Editor/UnrealEd/Private/Commandlets/FixConflictingLocalizationKeys.cpp b/Engine/Source/Editor/UnrealEd/Private/Commandlets/FixConflictingLocalizationKeys.cpp
--- a/Engine/Source/Editor/UnrealEd/Private/Commandlets/FixConflictingLocalizationKeys.cpp
+++ b/Engine/Source/Editor/UnrealEd/Private/Commandlets/FixConflictingLocalizationKeys.cpp
@@ -194,11 +194,26 @@ bool ReKeyTextProperty(UStruct* InOuterType, void* InAddrToUpdate, const TArray<
 
 						FScriptMapHelper ScriptMapHelper(MapProp, AddrToUpdate);
 
-						// ContainerIndex is the element index, but we need the sparse index which is computed by the iterator.
-						const int32 InternalIndex = ScriptMapHelper.FindInternalIndex(ContainerIndex);
-						if (InternalIndex != INDEX_NONE)
+						// ContainerIndex is the element index, but we need the sparse index
+						int32 SparseIndex = 0;
 						{
-							AddrToUpdate = ScriptMapHelper.GetPairPtr(InternalIndex) + MapProp->MapLayout.ValueOffset;
+							const int32 ElementCount = ScriptMapHelper.Num();
+							for (int32 ElementIndex = 0; ElementIndex < ElementCount; ++SparseIndex)
+							{
+								if (ScriptMapHelper.IsValidIndex(SparseIndex))
+								{
+									if (ElementIndex == ContainerIndex)
+									{
+										break;
+									}
+									++ElementIndex;
+								}
+							}
+						}
+
+						if (ScriptMapHelper.IsValidIndex(SparseIndex))
+						{
+							AddrToUpdate = ScriptMapHelper.GetPairPtr(SparseIndex) + MapProp->MapLayout.ValueOffset;
 
 							// Is this a complex property? If so, we need to recurse into it
 							if (FStructProperty* StructProp = CastField<FStructProperty>(MapProp->ValueProp))
@@ -216,11 +231,26 @@ bool ReKeyTextProperty(UStruct* InOuterType, void* InAddrToUpdate, const TArray<
 
 						FScriptSetHelper ScriptSetHelper(SetProp, AddrToUpdate);
 
-						// ContainerIndex is the element index, but we need the sparse index which is computed by the iterator.
-						const int32 InternalIndex = ScriptSetHelper.FindInternalIndex(ContainerIndex);
-						if (InternalIndex != INDEX_NONE)
+						// ContainerIndex is the element index, but we need the sparse index
+						int32 SparseIndex = 0;
 						{
-							AddrToUpdate = ScriptSetHelper.GetElementPtr(InternalIndex);
+							const int32 ElementCount = ScriptSetHelper.Num();
+							for (int32 ElementIndex = 0; ElementIndex < ElementCount; ++SparseIndex)
+							{
+								if (ScriptSetHelper.IsValidIndex(SparseIndex))
+								{
+									if (ElementIndex == ContainerIndex)
+									{
+										break;
+									}
+									++ElementIndex;
+								}
+							}
+						}
+
+						if (ScriptSetHelper.IsValidIndex(SparseIndex))
+						{
+							AddrToUpdate = ScriptSetHelper.GetElementPtr(SparseIndex);
 
 							// Is this a complex property? If so, we need to recurse into it
 							if (FStructProperty* StructProp = CastField<FStructProperty>(SetProp->ElementProp))
@@ -240,12 +270,9 @@ bool ReKeyTextProperty(UStruct* InOuterType, void* InAddrToUpdate, const TArray<
 						AddrToUpdate = MangledPropToUpdate->ContainerPtrToValuePtr<void>(AddrToUpdate);
 
 						FScriptMapHelper ScriptMapHelper(MapProp, AddrToUpdate);
-
-						// ContainerIndex is the element index, but we need the sparse index to get the element.
-						const int32 InternalIndex = ScriptMapHelper.FindInternalIndex(ContainerIndex);
-						if (InternalIndex != INDEX_NONE)
+						if (ScriptMapHelper.IsValidIndex(ContainerIndex))
 						{
-							AddrToUpdate = ScriptMapHelper.GetPairPtr(InternalIndex);
+							AddrToUpdate = ScriptMapHelper.GetPairPtr(ContainerIndex);
 
 							// Is this a complex property? If so, we need to recurse into it
 							if (FStructProperty* StructProp = CastField<FStructProperty>(MapProp->KeyProp))
diff --git a/Engine/Source/Editor/UnrealEd/Private/DiffUtils.cpp b/Engine/Source/Editor/UnrealEd/Private/DiffUtils.cpp
--- a/Engine/Source/Editor/UnrealEd/Private/DiffUtils.cpp
+++ b/Engine/Source/Editor/UnrealEd/Private/DiffUtils.cpp
@@ -505,41 +505,63 @@ FPropertyPath FPropertySoftPath::ResolvePath(const UObject* Object) const
 		}
 		else if( const FSetProperty* SetProperty = CastField<FSetProperty>(ResolvedProperty) )
 		{
-			FScriptSetHelper SetHelper(SetProperty, SetProperty->ContainerPtrToValuePtr<const void*>( ContainerAddress ));
-			if (SetHelper.IsValidIndex(PropertyIndex))
+			if(PropertyIndex != INDEX_NONE)
 			{
-				const int32 InternalIndex = SetHelper.FindInternalIndex(PropertyIndex);
-				UpdateContainerAddress( SetProperty->ElementProp, SetHelper.GetElementPtr(InternalIndex), ContainerAddress, ContainerStruct );
+				FScriptSetHelper SetHelper(SetProperty, SetProperty->ContainerPtrToValuePtr<const void*>( ContainerAddress ));
 
-				FPropertyInfo SetInfo(SetProperty->ElementProp, PropertyIndex);
+				// Figure out the real index in this instance of the set (sets have gaps in them):
+				int32 RealIndex = -1;
+				for( int32 J = 0; PropertyIndex >= 0; ++J)
+				{
+					++RealIndex;
+					if(SetHelper.IsValidIndex(J))
+					{
+						--PropertyIndex;
+					}
+				}
+
+				UpdateContainerAddress( SetProperty->ElementProp, SetHelper.GetElementPtr(RealIndex), ContainerAddress, ContainerStruct );
+
+				FPropertyInfo SetInfo(SetProperty->ElementProp, RealIndex);
 				Ret.AddProperty(SetInfo);
 			}
 		}
 		else if( const FMapProperty* MapProperty = CastField<FMapProperty>(ResolvedProperty) )
 		{
-			FScriptMapHelper MapHelper(MapProperty, MapProperty->ContainerPtrToValuePtr<const void*>( ContainerAddress ));
-			if (MapHelper.IsValidIndex(PropertyIndex))
+			if(PropertyIndex != INDEX_NONE)
 			{
-				const int32 InternalIndex = MapHelper.FindInternalIndex(PropertyIndex);
+				FScriptMapHelper MapHelper(MapProperty, MapProperty->ContainerPtrToValuePtr<const void*>( ContainerAddress ));
+				
+				// Figure out the real index in this instance of the map (map have gaps in them):
+				int32 RealIndex = -1;
+				for( int32 J = 0; PropertyIndex >= 0; ++J)
+				{
+					++RealIndex;
+					if(MapHelper.IsValidIndex(J))
+					{
+						--PropertyIndex;
+					}
+				}
+
 				// we have an index, but are we looking into a key or value? Peek ahead to find out:
-				if(ensure((I + 1 < PropertyChain.Num())))
+				if(ensure(I + 1 < PropertyChain.Num()))
 				{
 					if(PropertyChain[I+1].PropertyName == MapProperty->KeyProp->GetFName())
 					{
 						++I;
 
-						UpdateContainerAddress( MapProperty->KeyProp, MapHelper.GetKeyPtr(InternalIndex), ContainerAddress, ContainerStruct );
+						UpdateContainerAddress( MapProperty->KeyProp, MapHelper.GetKeyPtr(RealIndex), ContainerAddress, ContainerStruct );
 
-						FPropertyInfo MakKeyInfo(MapProperty->KeyProp, PropertyIndex);
+						FPropertyInfo MakKeyInfo(MapProperty->KeyProp, RealIndex);
 						Ret.AddProperty(MakKeyInfo);
 					}
 					else if(ensure( PropertyChain[I+1].PropertyName == MapProperty->ValueProp->GetFName() ))
 					{	
 						++I;
 
-						UpdateContainerAddress( MapProperty->ValueProp, MapHelper.GetValuePtr(InternalIndex), ContainerAddress, ContainerStruct );
-
-						FPropertyInfo MapValueInfo(MapProperty->ValueProp, PropertyIndex);
+						UpdateContainerAddress( MapProperty->ValueProp, MapHelper.GetValuePtr(RealIndex), ContainerAddress, ContainerStruct );
+						
+						FPropertyInfo MapValueInfo(MapProperty->ValueProp, RealIndex);
 						Ret.AddProperty(MapValueInfo);
 					}
 				}
@@ -776,6 +798,25 @@ void DiffUtils::CompareUnrelatedSCS(const UBlueprint* Old, const TArray< FSCSRes
 	}
 }
 
+static void AdvanceSetIterator( FScriptSetHelper& SetHelper, int32& Index)
+{
+	do
+	{
+		++Index;
+	}
+	while(Index < SetHelper.GetMaxIndex() && !SetHelper.IsValidIndex(Index));
+}
+
+static void AdvanceMapIterator( FScriptMapHelper& MapHelper, int32& Index)
+{
+	do
+	{
+		++Index;
+	}
+	while(Index < MapHelper.GetMaxIndex() && !MapHelper.IsValidIndex(Index));
+}
+
+
 static void IdenticalHelper(const FProperty* AProperty, const FProperty* BProperty, const void* AValue, const void* BValue,
 	const UObject* OwningOuterA, const UObject* OwningOuterB, const FPropertySoftPath& RootPath,
 	TArray<FPropertySoftPath>& DifferingSubProperties, bool bStaticArrayHandled = false)
@@ -883,14 +924,25 @@ static void IdenticalHelper(const FProperty* AProperty, const FProperty* BProper
 			}
 
 			// note any differences in contained elements:
-			FScriptSetHelper::FIterator IteratorA(SetHelperA);
-			FScriptSetHelper::FIterator IteratorB(SetHelperB);
-			for (; IteratorA && IteratorB; ++IteratorA, ++IteratorB)
+			const int32 SetSizeA = SetHelperA.Num();
+			const int32 SetSizeB = SetHelperB.Num();
+			
+			int32 SetIndexA = -1;
+			int32 SetIndexB = -1;
+
+			AdvanceSetIterator(SetHelperA, SetIndexA);
+			AdvanceSetIterator(SetHelperB, SetIndexB);
+
+			for (int32 VirtualIndex = 0; VirtualIndex < SetSizeA && VirtualIndex < SetSizeB; ++VirtualIndex)
 			{
-				const void* SubValueA = SetHelperA.GetElementPtr(IteratorA);
-				const void* SubValueB = SetHelperB.GetElementPtr(IteratorB);
+				const void* SubValueA = SetHelperA.GetElementPtr(SetIndexA);
+				const void* SubValueB = SetHelperB.GetElementPtr(SetIndexB);
 				IdenticalHelper(APropAsSet->ElementProp, BPropAsSet->ElementProp, SubValueA, SubValueB,
-					OwningOuterA, OwningOuterB, FPropertySoftPath(RootPath, IteratorA.GetLogicalIndex()), DifferingSubProperties);
+					OwningOuterA, OwningOuterB, FPropertySoftPath(RootPath, VirtualIndex), DifferingSubProperties);
+
+				// advance iterators in step:
+				AdvanceSetIterator(SetHelperA, SetIndexA);
+				AdvanceSetIterator(SetHelperB, SetIndexB);
 			}
 		}
 		else
@@ -905,6 +957,7 @@ static void IdenticalHelper(const FProperty* AProperty, const FProperty* BProper
 		{
 			FScriptMapHelper MapHelperA(APropAsMap, AValue);
 			FScriptMapHelper MapHelperB(BPropAsMap, BValue);
+
 			if (MapHelperA.Num() != MapHelperB.Num())
 			{
 				// API not robust enough to indicate changes made to # of set elements, would
@@ -912,15 +965,24 @@ static void IdenticalHelper(const FProperty* AProperty, const FProperty* BProper
 				DifferingSubProperties.Push(RootPath);
 			}
 
-			FScriptMapHelper::FIterator IteratorA(MapHelperA);
-			FScriptMapHelper::FIterator IteratorB(MapHelperB);
-			for (; IteratorA && IteratorB; ++IteratorA, ++IteratorB)
+			int32 MapSizeA = MapHelperA.Num();
+			int32 MapSizeB = MapHelperB.Num();
+			
+			int32 MapIndexA = -1;
+			int32 MapIndexB = -1;
+
+			AdvanceMapIterator(MapHelperA, MapIndexA);
+			AdvanceMapIterator(MapHelperB, MapIndexB);
+			
+			for (int32 VirtualIndex = 0; VirtualIndex < MapSizeA && VirtualIndex < MapSizeB; ++VirtualIndex)
 			{
-				IdenticalHelper(APropAsMap->KeyProp, BPropAsMap->KeyProp, MapHelperA.GetKeyPtr(IteratorA), MapHelperB.GetKeyPtr(IteratorB),
-					OwningOuterA, OwningOuterB, FPropertySoftPath(RootPath, IteratorA.GetLogicalIndex()), DifferingSubProperties);
-				IdenticalHelper(APropAsMap->ValueProp, BPropAsMap->ValueProp, MapHelperA.GetValuePtr(IteratorA), MapHelperB.GetValuePtr(IteratorB),
-					OwningOuterA, OwningOuterB, FPropertySoftPath(RootPath, IteratorA.GetLogicalIndex()), DifferingSubProperties);
+				IdenticalHelper(APropAsMap->KeyProp, BPropAsMap->KeyProp, MapHelperA.GetKeyPtr(MapIndexA), MapHelperB.GetKeyPtr(MapIndexB),
+					OwningOuterA, OwningOuterB, FPropertySoftPath(RootPath, VirtualIndex), DifferingSubProperties);
+				IdenticalHelper(APropAsMap->ValueProp, BPropAsMap->ValueProp, MapHelperA.GetValuePtr(MapIndexA), MapHelperB.GetValuePtr(MapIndexB),
+					OwningOuterA, OwningOuterB, FPropertySoftPath(RootPath, VirtualIndex), DifferingSubProperties);
 
+				AdvanceMapIterator(MapHelperA, MapIndexA);
+				AdvanceMapIterator(MapHelperB, MapIndexB);
 			}
 		}
 		else
diff --git a/Engine/Source/Runtime/CoreUObject/Private/Serialization/UnversionedPropertySerializationTest.cpp b/Engine/Source/Runtime/CoreUObject/Private/Serialization/UnversionedPropertySerializationTest.cpp
--- a/Engine/Source/Runtime/CoreUObject/Private/Serialization/UnversionedPropertySerializationTest.cpp
+++ b/Engine/Source/Runtime/CoreUObject/Private/Serialization/UnversionedPropertySerializationTest.cpp
@@ -252,14 +252,17 @@ struct FUnversionedPropertyTest : public FUnversionedPropertyTestInput
 			return false;
 		}
 
-		for (FScriptSetHelper::FIterator It(HelperA); It; ++It)
+		for (int32 Num = HelperA.Num(), IndexA = 0; IndexA < Num; ++IndexA)
 		{
-			const uint8* ElemA = HelperA.GetElementPtr(It);
-			const uint8* ElemB = FindElementPtr(HelperB, ElemA, OutDiff);
-
-			if (!ElemB)
+			if (HelperA.IsValidIndex(IndexA))
 			{
-				return false;
+				const uint8* ElemA = HelperA.GetElementPtr(IndexA);
+				const uint8* ElemB = FindElementPtr(HelperB, ElemA, OutDiff);
+		
+				if (!ElemB)
+				{
+					return false;
+				} 
 			}
 		}
 
@@ -288,19 +291,22 @@ struct FUnversionedPropertyTest : public FUnversionedPropertyTestInput
 			return false;
 		}
 		
-		for (FScriptMapHelper::FIterator It(HelperA); It; ++It)
+		for (int32 Num = HelperA.Num(), IndexA = 0; IndexA < Num; ++IndexA)
 		{
-			const uint8* PairA = HelperA.GetPairPtr(It);
-			const uint8* PairB = FindPairPtr(HelperB, PairA, OutDiff);
-
-			if (!PairB)
+			if (HelperA.IsValidIndex(IndexA))
 			{
-				return false;
-			}
+				const uint8* PairA = HelperA.GetPairPtr(IndexA);
+				const uint8* PairB = FindPairPtr(HelperB, PairA, OutDiff);
+		
+				if (!PairB)
+				{
+					return false;
+				} 
 
-			if (!Equals(ValueProp, PairA + ValueOffset, PairB + ValueOffset, OutDiff))
-			{
-				return false;
+				if (!Equals(ValueProp, PairA + ValueOffset, PairB + ValueOffset, OutDiff))
+				{
+					return false;
+				}
 			}
 		}
 
diff --git a/Engine/Source/Runtime/CoreUObject/Private/UObject/UObjectGlobals.cpp b/Engine/Source/Runtime/CoreUObject/Private/UObject/UObjectGlobals.cpp
--- a/Engine/Source/Runtime/CoreUObject/Private/UObject/UObjectGlobals.cpp
+++ b/Engine/Source/Runtime/CoreUObject/Private/UObject/UObjectGlobals.cpp
@@ -4880,27 +4880,34 @@ template<EPropertyCollectFlags CollectFlags>
 static void CollectMapReferences(FReferenceCollector& Collector, FMapProperty& Property, void* Instance, const UObject* Referencer)
 {
 	FScriptMapHelper MapHelper(&Property, Instance);
+	const int32 Num = MapHelper.Num();
+	const bool bCollectKeys = MayContainStrongReference(*MapHelper.GetKeyProperty());
+	const bool bCollectValues = MayContainStrongReference(*MapHelper.GetValueProperty());
 
-	if (MapHelper.Num() == 0)
+	if (Num == 0)
 	{
-		return;
+		return;	
 	}
 
-	const bool bCollectKeys = MayContainStrongReference(*MapHelper.GetKeyProperty());
 	if (bCollectKeys)
 	{
-		for (FScriptMapHelper::FIterator It(MapHelper); It; ++It)
+		for (int32 Idx = 0; Idx < Num; ++Idx)
 		{
-			CollectPropertyReferences<CollectFlags>(Collector, *MapHelper.GetKeyProperty(), MapHelper.GetPairPtr(It), Referencer);
+			if (MapHelper.IsValidIndex(Idx))
+			{
+				CollectPropertyReferences<CollectFlags>(Collector, *MapHelper.GetKeyProperty(), MapHelper.GetPairPtr(Idx), Referencer);
+			}
 		}
 	}
 
-	const bool bCollectValues = MayContainStrongReference(*MapHelper.GetValueProperty());
 	if (bCollectValues)
 	{
-		for (FScriptMapHelper::FIterator It(MapHelper); It; ++It)
+		for (int32 Idx = 0; Idx < Num; ++Idx)
 		{
-			CollectPropertyReferences<CollectFlags>(Collector, *MapHelper.GetValueProperty(), MapHelper.GetPairPtr(It), Referencer);
+			if (MapHelper.IsValidIndex(Idx))
+			{
+				CollectPropertyReferences<CollectFlags>(Collector, *MapHelper.GetValueProperty(), MapHelper.GetPairPtr(Idx), Referencer);
+			}
 		}
 	}
 }
@@ -4911,9 +4918,12 @@ void CollectSetReferences(FReferenceCollector& Collector, FSetProperty& Property
 	FScriptSetHelper SetHelper(&Property, Instance);
 	if (MayContainStrongReference(*SetHelper.GetElementProperty()))
 	{
-		for (FScriptSetHelper::FIterator It(SetHelper); It; ++It)
-		{
-			CollectPropertyReferences<CollectFlags>(Collector, *SetHelper.GetElementProperty(), SetHelper.GetElementPtr(It), Referencer);
+		for (int32 Idx = 0, Num = SetHelper.Num(); Idx < Num; ++Idx)
+		{	
+			if (SetHelper.IsValidIndex(Idx))
+			{
+				CollectPropertyReferences<CollectFlags>(Collector, *SetHelper.GetElementProperty(), SetHelper.GetElementPtr(Idx), Referencer);
+			}
 		}
 	}
 }
diff --git a/Engine/Source/Runtime/CoreUObject/Private/UObject/UnrealType.cpp b/Engine/Source/Runtime/CoreUObject/Private/UObject/UnrealType.cpp
--- a/Engine/Source/Runtime/CoreUObject/Private/UObject/UnrealType.cpp
+++ b/Engine/Source/Runtime/CoreUObject/Private/UObject/UnrealType.cpp
@@ -143,15 +143,21 @@ FORCEINLINE_DEBUGGABLE bool FPropertyValueIterator::NextValue(EPropertyValueIter
 				if ((KeyFlags | ValueFlags) != EPropertyValueFlags::None)
 				{
 					FScriptMapHelper Helper(MapProperty, PropertyValue);
-					for (FScriptMapHelper::FIterator It(Helper); It; ++It)
+					const int32 Num = Helper.Num();
+					for (int32 DynamicIndex = 0; DynamicIndex < Num; ++DynamicIndex)
 					{
-						if (KeyFlags != EPropertyValueFlags::None)
-						{
-							NewEntry.ValueArray.Emplace(BasePairType(KeyProperty, Helper.GetKeyPtr(It)), KeyFlags);
-						}
-						if (ValueFlags != EPropertyValueFlags::None)
+						if (Helper.IsValidIndex(DynamicIndex))
 						{
-							NewEntry.ValueArray.Emplace(BasePairType(ValueProperty, Helper.GetValuePtr(It)), ValueFlags);
+							if (KeyFlags != EPropertyValueFlags::None)
+							{
+								NewEntry.ValueArray.Emplace(
+									BasePairType(KeyProperty, Helper.GetKeyPtr(DynamicIndex)), KeyFlags);
+							}
+							if (ValueFlags != EPropertyValueFlags::None)
+							{
+								NewEntry.ValueArray.Emplace(
+									BasePairType(ValueProperty, Helper.GetValuePtr(DynamicIndex)), ValueFlags);
+							}
 						}
 					}
 				}
@@ -164,9 +170,14 @@ FORCEINLINE_DEBUGGABLE bool FPropertyValueIterator::NextValue(EPropertyValueIter
 				if (InnerFlags != EPropertyValueFlags::None)
 				{
 					FScriptSetHelper Helper(SetProperty, PropertyValue);
-					for (FScriptSetHelper::FIterator It(Helper); It; ++It)
+					const int32 Num = Helper.Num();
+					for (int32 DynamicIndex = 0; DynamicIndex < Num; ++DynamicIndex)
 					{
-						NewEntry.ValueArray.Emplace(BasePairType(InnerProperty, Helper.GetElementPtr(It)), InnerFlags);
+						if (Helper.IsValidIndex(DynamicIndex))
+						{
+							NewEntry.ValueArray.Emplace(
+								BasePairType(InnerProperty, Helper.GetElementPtr(DynamicIndex)), InnerFlags);
+						}
 					}
 				}
 			}
diff --git a/Engine/Source/Runtime/Engine/Private/Kismet/BlueprintTypeConversions.cpp b/Engine/Source/Runtime/Engine/Private/Kismet/BlueprintTypeConversions.cpp
--- a/Engine/Source/Runtime/Engine/Private/Kismet/BlueprintTypeConversions.cpp
+++ b/Engine/Source/Runtime/Engine/Private/Kismet/BlueprintTypeConversions.cpp
@@ -291,16 +291,18 @@ DEFINE_FUNCTION(UBlueprintTypeConversions::execConvertSetType)
 	FScriptSetHelper SourceSet(SourceSetProperty, SourceSetAddr);
 	FScriptSetHelper DestSet(DestSetProperty, DestSetAddr);
 
-	DestSet.EmptyElements(SourceSet.Num());
+	int SourceSetSize = SourceSet.Num();
+	DestSet.EmptyElements(SourceSetSize);
 
 	BlueprintTypeConversions::ConversionFunctionT ConversionFunction = 
 		BlueprintTypeConversions::Internal::FindConversionFunction(SourceSetProperty->ElementProp, DestSetProperty->ElementProp);
 	check(ConversionFunction);
 
-	for (FScriptSetHelper::FIterator It(SourceSet); It; ++It)
+	for (int i = 0; i < SourceSetSize; ++i)
 	{
-		const void* SrcData = SourceSet.GetElementPtr(It);
-		const int32 NewIndex = DestSet.AddDefaultValue_Invalid_NeedsRehash();
+		const void* SrcData = SourceSet.GetElementPtr(i);
+
+		int32 NewIndex = DestSet.AddDefaultValue_Invalid_NeedsRehash();
 		void* DestData = DestSet.GetElementPtr(NewIndex);
 		(*ConversionFunction)(SrcData, DestData);
 	}
@@ -325,18 +327,19 @@ DEFINE_FUNCTION(UBlueprintTypeConversions::execConvertMapType)
 	FScriptMapHelper SourceMap(SourceMapProperty, SourceMapAddr);
 	FScriptMapHelper DestMap(DestMapProperty, DestMapAddr);
 
-	DestMap.EmptyValues(SourceMap.Num());
+	int SourceMapSize = SourceMap.Num();
+	DestMap.EmptyValues(SourceMapSize);
 
 	BlueprintTypeConversions::ConversionFunctionT KeyConversionFunction = 
 		BlueprintTypeConversions::Internal::FindConversionFunction(SourceMapProperty->KeyProp, DestMapProperty->KeyProp);
 	BlueprintTypeConversions::ConversionFunctionT ValueConversionFunction = 
 		BlueprintTypeConversions::Internal::FindConversionFunction(SourceMapProperty->ValueProp, DestMapProperty->ValueProp);
 
-	for (FScriptMapHelper::FIterator It(SourceMap); It; ++It)
+	for (int i = 0; i < SourceMapSize; ++i)
 	{
-		const int32 NewIndex = DestMap.AddDefaultValue_Invalid_NeedsRehash();
+		int32 NewIndex = DestMap.AddDefaultValue_Invalid_NeedsRehash();
 
-		const void* SourceKeyRawData = SourceMap.GetKeyPtr(It);
+		const void* SourceKeyRawData = SourceMap.GetKeyPtr(i);
 		void* DestinationKeyRawData = DestMap.GetKeyPtr(NewIndex);
 
 		if (KeyConversionFunction)
@@ -348,7 +351,7 @@ DEFINE_FUNCTION(UBlueprintTypeConversions::execConvertMapType)
 			SourceMapProperty->KeyProp->CopySingleValue(DestinationKeyRawData, SourceKeyRawData);
 		}
 
-		const void* SourceValueRawData = SourceMap.GetValuePtr(It);
+		const void* SourceValueRawData = SourceMap.GetValuePtr(i);
 		void* DestinationValueRawData = DestMap.GetValuePtr(NewIndex);
 
 		if (ValueConversionFunction)
diff --git a/Engine/Source/Runtime/Engine/Private/PropertyAccess.cpp b/Engine/Source/Runtime/Engine/Private/PropertyAccess.cpp
--- a/Engine/Source/Runtime/Engine/Private/PropertyAccess.cpp
+++ b/Engine/Source/Runtime/Engine/Private/PropertyAccess.cpp
@@ -681,10 +681,10 @@ struct FPropertyAccessSystem
 		FScriptMapHelper DestinationMapHelper(DestinationMapProperty, DestinationAddress);
 
 		DestinationMapHelper.EmptyValues();
-		for (FScriptMapHelper::FIterator It(SourceMapHelper); It; ++It)
+		for (int32 i = 0; i < SourceMapHelper.Num(); ++i)
 		{
-			const void* KeyData = SourceMapHelper.GetKeyPtr(It);
-			const SourceType* SourceValueData = reinterpret_cast<const SourceType*>(SourceMapHelper.GetValuePtr(It));
+			const void* KeyData = SourceMapHelper.GetKeyPtr(i);
+			const SourceType* SourceValueData = reinterpret_cast<const SourceType*>(SourceMapHelper.GetValuePtr(i));
 			DestinationType CastedType = static_cast<DestinationType>(*SourceValueData);
 			DestinationMapHelper.AddPair(KeyData, &CastedType);
 		}
diff --git a/Engine/Source/Runtime/JsonUtilities/Private/JsonObjectConverter.cpp b/Engine/Source/Runtime/JsonUtilities/Private/JsonObjectConverter.cpp
--- a/Engine/Source/Runtime/JsonUtilities/Private/JsonObjectConverter.cpp
+++ b/Engine/Source/Runtime/JsonUtilities/Private/JsonObjectConverter.cpp
@@ -107,13 +107,18 @@ TSharedPtr<FJsonValue> ConvertScalarFPropertyToJsonValue(FProperty* Property, co
 	{
 		TArray< TSharedPtr<FJsonValue> > Out;
 		FScriptSetHelper Helper(SetProperty, Value);
-		for (FScriptSetHelper::FIterator It(Helper); It; ++It)
+		for ( int32 i=0, n=Helper.Num(); n; ++i )
 		{
-			TSharedPtr<FJsonValue> Elem = FJsonObjectConverter::UPropertyToJsonValue(SetProperty->ElementProp, Helper.GetElementPtr(It), CheckFlags & (~CPF_ParmFlags), SkipFlags, ExportCb, SetProperty);
-			if (Elem.IsValid())
+			if ( Helper.IsValidIndex(i) )
 			{
-				// add to the array
-				Out.Push(Elem);
+				TSharedPtr<FJsonValue> Elem = FJsonObjectConverter::UPropertyToJsonValue(SetProperty->ElementProp, Helper.GetElementPtr(i), CheckFlags & ( ~CPF_ParmFlags ), SkipFlags, ExportCb, SetProperty);
+				if ( Elem.IsValid() )
+				{
+					// add to the array
+					Out.Push(Elem);
+				}
+
+				--n;
 			}
 		}
 		return MakeShared<FJsonValueArray>(Out);
@@ -123,33 +128,38 @@ TSharedPtr<FJsonValue> ConvertScalarFPropertyToJsonValue(FProperty* Property, co
 		TSharedRef<FJsonObject> Out = MakeShared<FJsonObject>();
 
 		FScriptMapHelper Helper(MapProperty, Value);
-		for (FScriptMapHelper::FIterator It(Helper); It; ++It)
+		for ( int32 i=0, n = Helper.Num(); n; ++i )
 		{
-			TSharedPtr<FJsonValue> KeyElement = FJsonObjectConverter::UPropertyToJsonValue(MapProperty->KeyProp, Helper.GetKeyPtr(It), CheckFlags & (~CPF_ParmFlags), SkipFlags, ExportCb, MapProperty, ConversionFlags);
-			TSharedPtr<FJsonValue> ValueElement = FJsonObjectConverter::UPropertyToJsonValue(MapProperty->ValueProp, Helper.GetValuePtr(It), CheckFlags & (~CPF_ParmFlags), SkipFlags, ExportCb, MapProperty, ConversionFlags);
-			if (KeyElement.IsValid() && ValueElement.IsValid())
+			if ( Helper.IsValidIndex(i) )
 			{
-				FString KeyString;
-				if (!KeyElement->TryGetString(KeyString))
+				TSharedPtr<FJsonValue> KeyElement = FJsonObjectConverter::UPropertyToJsonValue(MapProperty->KeyProp, Helper.GetKeyPtr(i), CheckFlags & ( ~CPF_ParmFlags ), SkipFlags, ExportCb, MapProperty, ConversionFlags);
+				TSharedPtr<FJsonValue> ValueElement = FJsonObjectConverter::UPropertyToJsonValue(MapProperty->ValueProp, Helper.GetValuePtr(i), CheckFlags & ( ~CPF_ParmFlags ), SkipFlags, ExportCb, MapProperty, ConversionFlags);
+				if ( KeyElement.IsValid() && ValueElement.IsValid() )
 				{
-					MapProperty->KeyProp->ExportTextItem_Direct(KeyString, Helper.GetKeyPtr(It), nullptr, nullptr, 0);
-					if (KeyString.IsEmpty())
+					FString KeyString;
+					if (!KeyElement->TryGetString(KeyString))
 					{
-						UE_LOG(LogJson, Error, TEXT("Unable to convert key to string for property %s."), *MapProperty->GetAuthoredName())
-						KeyString = FString::Printf(TEXT("Unparsed Key %d"), It.GetLogicalIndex());
+						MapProperty->KeyProp->ExportTextItem_Direct(KeyString, Helper.GetKeyPtr(i), nullptr, nullptr, 0);
+						if (KeyString.IsEmpty())
+						{
+							UE_LOG(LogJson, Error, TEXT("Unable to convert key to string for property %s."), *MapProperty->GetAuthoredName())
+							KeyString = FString::Printf(TEXT("Unparsed Key %d"), i);
+						}
 					}
-				}
 
-				// Coerce camelCase map keys for Enum/FName properties
-				if (CastField<FEnumProperty>(MapProperty->KeyProp) ||
-					CastField<FNameProperty>(MapProperty->KeyProp))
-				{
-					if (!EnumHasAnyFlags(ConversionFlags, EJsonObjectConversionFlags::SkipStandardizeCase))
+					// Coerce camelCase map keys for Enum/FName properties
+					if (CastField<FEnumProperty>(MapProperty->KeyProp) ||
+						CastField<FNameProperty>(MapProperty->KeyProp))
 					{
-						KeyString = FJsonObjectConverter::StandardizeCase(KeyString);
+						if (!EnumHasAnyFlags(ConversionFlags, EJsonObjectConversionFlags::SkipStandardizeCase))
+						{
+							KeyString = FJsonObjectConverter::StandardizeCase(KeyString);
+						}
 					}
+					Out->SetField(KeyString, ValueElement);
 				}
-				Out->SetField(KeyString, ValueElement);
+
+				--n;
 			}
 		}
 
diff --git a/Engine/Source/Runtime/Serialization/Private/StructDeserializer.cpp b/Engine/Source/Runtime/Serialization/Private/StructDeserializer.cpp
--- a/Engine/Source/Runtime/Serialization/Private/StructDeserializer.cpp
+++ b/Engine/Source/Runtime/Serialization/Private/StructDeserializer.cpp
@@ -570,11 +570,10 @@ bool FStructDeserializer::DeserializeElement(void* OutAddress, UStruct& OwnerInf
 				// handle set element
 				FSetProperty* SetProperty = CastField<FSetProperty>(CurrentState.Property);
 				FScriptSetHelper SetHelper(SetProperty, CurrentState.Data);
-				const int32 InternalIndex = SetHelper.FindInternalIndex(CurrentState.ArrayIndex);
-				if (InternalIndex != INDEX_NONE)
+				FProperty* Property = SetProperty->ElementProp;
+				if (SetHelper.IsValidIndex(CurrentState.ArrayIndex))
 				{
-					uint8* ElementPtr = SetHelper.GetElementPtr(InternalIndex);
-					FProperty* Property = SetProperty->ElementProp;
+					uint8* ElementPtr = SetHelper.GetElementPtr(CurrentState.ArrayIndex);
 					constexpr int32 ReadIndex = 0; //Pointer is offset so reading index is 0
 					if (!Backend.ReadProperty(Property, CurrentState.Property, ElementPtr, ReadIndex))
 					{
@@ -595,13 +594,12 @@ bool FStructDeserializer::DeserializeElement(void* OutAddress, UStruct& OwnerInf
 				// handle map element
 				FMapProperty* MapProperty = CastField<FMapProperty>(CurrentState.Property);
 				FScriptMapHelper MapHelper(MapProperty, CurrentState.Data);
+				FProperty* Property = MapProperty->ValueProp;
 
 				//When written as array, maps won't include the key, only values
-				const int32 InternalIndex = MapHelper.FindInternalIndex(CurrentState.ArrayIndex);
-				if (InternalIndex != INDEX_NONE)
+				if (MapHelper.IsValidIndex(CurrentState.ArrayIndex))
 				{
-					uint8* PairPtr = MapHelper.GetPairPtr(InternalIndex);
-					FProperty* Property = MapProperty->ValueProp;
+					uint8* PairPtr = MapHelper.GetPairPtr(CurrentState.ArrayIndex);
 					constexpr int32 ReadIndex = 0; //Pointer is offset so reading index is 0
 					if (!Backend.ReadProperty(Property, CurrentState.Property, PairPtr, ReadIndex))
 					{
@@ -696,13 +694,12 @@ bool FStructDeserializer::DeserializeElement(void* OutAddress, UStruct& OwnerInf
 						FScriptSetHelper SetHelper(SetProperty, SetProperty->ContainerPtrToValuePtr<void>(CurrentState.Data));
 
 						//If a specific index is asked and it's not valid, skip the property
-						const int32 InternalIndex = SetHelper.FindInternalIndex(CurrentState.ArrayIndex);
-						if (InternalIndex != INDEX_NONE)
+						if (SetHelper.IsValidIndex(CurrentState.ArrayIndex))
 						{
 							Property = SetProperty->ElementProp;
 
 							//Offset the pointer directly and give index 0 to be read so no offsetting is done during deserialization
-							CurrentState.Data = SetHelper.GetElementPtr(InternalIndex);
+							CurrentState.Data = SetHelper.GetElementPtr(CurrentState.ArrayIndex);
 							CurrentState.ArrayIndex = 0;
 
 							if (!Backend.ReadProperty(Property, nullptr, CurrentState.Data, CurrentState.ArrayIndex))
@@ -727,13 +724,12 @@ bool FStructDeserializer::DeserializeElement(void* OutAddress, UStruct& OwnerInf
 						FScriptMapHelper MapHelper(MapProperty, MapProperty->ContainerPtrToValuePtr<void>(CurrentState.Data));
 
 						//If a specific index is asked and it's not valid, skip the property
-						const int32 InternalIndex = MapHelper.FindInternalIndex(CurrentState.ArrayIndex);
-						if (InternalIndex != INDEX_NONE)
+						if (MapHelper.IsValidIndex(CurrentState.ArrayIndex))
 						{
 							Property = MapProperty->ValueProp;
 
 							//Offset the pointer directly and give index 0 to be read so no offsetting is done during deserialization
-							CurrentState.Data = MapHelper.GetPairPtr(InternalIndex);
+							CurrentState.Data = MapHelper.GetPairPtr(CurrentState.ArrayIndex);
 							CurrentState.ArrayIndex = 0;
 
 							if (!Backend.ReadProperty(Property, nullptr, CurrentState.Data, CurrentState.ArrayIndex))
@@ -838,11 +834,10 @@ bool FStructDeserializer::DeserializeElement(void* OutAddress, UStruct& OwnerInf
 				if (FSetProperty* SetProperty = CastField<FSetProperty>(CurrentState.Property))
 				{
 					FScriptSetHelper SetHelper(SetProperty, CurrentState.Data);
-					const int32 InternalIndex = SetHelper.FindInternalIndex(CurrentState.ArrayIndex);
-					if (InternalIndex != INDEX_NONE)
+					if (SetHelper.IsValidIndex(CurrentState.ArrayIndex))
 					{
 						NewState.Property = SetProperty->ElementProp;
-						NewState.Data = SetHelper.GetElementPtr(InternalIndex);
+						NewState.Data = SetHelper.GetElementPtr(CurrentState.ArrayIndex);
 						NewState.ArrayIndex = 0;
 						++CurrentState.ArrayIndex;
 					}
@@ -857,11 +852,10 @@ bool FStructDeserializer::DeserializeElement(void* OutAddress, UStruct& OwnerInf
 				else if (FMapProperty* MapProperty = CastField<FMapProperty>(CurrentState.Property))
 				{
 					FScriptMapHelper MapHelper(MapProperty, CurrentState.Data);
-					const int32 InternalIndex = MapHelper.FindInternalIndex(CurrentState.ArrayIndex);
-					if (InternalIndex != INDEX_NONE)
+					if (MapHelper.IsValidIndex(CurrentState.ArrayIndex))
 					{
 						NewState.Property = MapProperty->ValueProp;
-						NewState.Data = MapHelper.GetValuePtr(InternalIndex);
+						NewState.Data = MapHelper.GetValuePtr(CurrentState.ArrayIndex);
 						NewState.ArrayIndex = 0;
 						++CurrentState.ArrayIndex;
 					}
@@ -937,15 +931,14 @@ bool FStructDeserializer::DeserializeElement(void* OutAddress, UStruct& OwnerInf
 					if (FStructProperty* SetStructProperty = CastField<FStructProperty>(MapProperty->ValueProp))
 					{
 						FScriptMapHelper MapHelper(MapProperty, MapProperty->ContainerPtrToValuePtr<void>(CurrentState.Data));
-						
+
 						//If a specific index is asked and it's not valid, skip the property
-						const int32 InternalIndex = MapHelper.FindInternalIndex(CurrentState.ArrayIndex);
-						if (InternalIndex != INDEX_NONE)
+						if (MapHelper.IsValidIndex(CurrentState.ArrayIndex))
 						{
 							//We're skipping a level directly so CurrentState becomes the outer (set) and NewState the inner (Element prop)
 							CurrentState.Property = NewState.Property;
 							NewState.Property = SetStructProperty;
-							NewState.Data = MapHelper.GetValuePtr(InternalIndex);
+							NewState.Data = MapHelper.GetValuePtr(CurrentState.ArrayIndex);
 							NewState.ArrayIndex = 0;
 						}
 						else
@@ -963,15 +956,14 @@ bool FStructDeserializer::DeserializeElement(void* OutAddress, UStruct& OwnerInf
 					if (FStructProperty* SetStructProperty = CastField<FStructProperty>(SetProperty->ElementProp))
 					{
 						FScriptSetHelper SetHelper(SetProperty, SetProperty->ContainerPtrToValuePtr<void>(CurrentState.Data));
-						const int32 InternalIndex = SetHelper.FindInternalIndex(CurrentState.ArrayIndex);
 
 						//If a specific index is asked and it's not valid, skip the property
-						if (InternalIndex != INDEX_NONE)
+						if (SetHelper.IsValidIndex(CurrentState.ArrayIndex))
 						{
 							//We're skipping a level directly so CurrentState becomes the outer (set) and NewState the inner (Element prop)
 							CurrentState.Property = NewState.Property;
 							NewState.Property = SetProperty->ElementProp;
-							NewState.Data = SetHelper.GetElementPtr(InternalIndex);
+							NewState.Data = SetHelper.GetElementPtr(CurrentState.ArrayIndex);
 							NewState.ArrayIndex = 0;
 						}
 						else
diff --git a/Engine/Source/Runtime/Serialization/Private/Tests/StructSerializerTest.cpp b/Engine/Source/Runtime/Serialization/Private/Tests/StructSerializerTest.cpp
--- a/Engine/Source/Runtime/Serialization/Private/Tests/StructSerializerTest.cpp
+++ b/Engine/Source/Runtime/Serialization/Private/Tests/StructSerializerTest.cpp
@@ -486,9 +486,7 @@ namespace StructSerializerTest
 				FStructSerializer::SerializeElement(&TestStruct, Property, INDEX_NONE, SerializerBackend, Policies);
 
 				FStructSerializerTestStruct TestStruct2(NoInit);
-				TestStruct2.Sets.IntSet = { -1, -2, -3, -4, -5 };
-				TestStruct2.Sets.IntSet.Remove(-1);
-				TestStruct2.Sets.IntSet.Remove(-2);
+				TestStruct2.Sets.IntSet = { -1, -2, -3 };
 				TestStruct2.Sets.NameSet = { TEXT("Pre1"), TEXT("Pre2") , TEXT("Pre3") };
 				TestStruct2.Sets.StrSet = { TEXT("Pre4"), TEXT("Pre5") , TEXT("Pre6"), TEXT("Pre7") };
 				TestStruct2.Sets.StructSet = { FStructSerializerBuiltinTestStruct(NoInit) };
@@ -496,7 +494,7 @@ namespace StructSerializerTest
 
 				ValidateSets(Test, TestStruct.Sets, TestStruct2.Sets);
 			}
-
+			
 			//TArray<uint8> element
 			{
 				TArray<uint8> Buffer;
