Preserve the outer package dirty state after reinstancing a Blueprint object containing one or more subobjects that are no longer load-compatible with the original's serialized data.

Example: This can occur after reparenting into a new hierarchy with matching subobjects having incompatible/downcast types. In that case, if data was previously saved using the old type, it cannot be loaded against the new type without re-saving the outermost package to update its serialized instance data.

diff --git a/Engine/Source/Editor/UnrealEd/Private/Kismet2/KismetReinstanceUtilities.cpp b/Engine/Source/Editor/UnrealEd/Private/Kismet2/KismetReinstanceUtilities.cpp
--- a/Engine/Source/Editor/UnrealEd/Private/Kismet2/KismetReinstanceUtilities.cpp
+++ b/Engine/Source/Editor/UnrealEd/Private/Kismet2/KismetReinstanceUtilities.cpp
@@ -2835,6 +2835,37 @@ void FBlueprintCompileReinstancer::ReplaceInstancesOfClass_Inner(const TMap<UCla
 				// Patch the new object into the old object linker's export map; subsequent loads may import
 				// this entry and we need to make sure that it returns the new object instead of the old one.
 				FLinkerLoad::PRIVATE_PatchNewObjectIntoExport(Obj, *NewObject);
+
+				// In some cases (e.g. reparenting across a hierarchy), the new object may contain subobjects
+				// (e.g. components) which are no longer binary-compatible with the old object's instance, which
+				// may have been delta-serialized to the outermost package. In that case, we need to ensure that
+				// the package (e.g. level/actor) remains dirty, so that the user sees that it requires a re-save.
+				const UPackage* NewObjectPackage = (*NewObject)->GetPackage();
+				if (CleanPackageList.Contains(NewObjectPackage) && NewObjectPackage->IsDirty())
+				{
+					bool bShouldPreservePackageDirtyState = false;
+					ForEachObjectWithOuterBreakable(Obj, [&OldToNewInstanceMap, &bShouldPreservePackageDirtyState](UObject* OldSubobject)
+					{
+						if (UObject* NewSubobject = OldToNewInstanceMap.FindRef(OldSubobject))
+						{
+							// If the new subobject type is not of the old subobject type, then the old subobject's
+							// data (if serialized) is no longer binary-compatible, and can no longer be imported on
+							// load (i.e. it will fail), resulting in data loss. In that case, we need the dirty state.
+							if (!NewSubobject->GetClass()->IsChildOf(OldSubobject->GetClass()))
+							{
+								bShouldPreservePackageDirtyState = true;
+							}
+						}
+
+						// No need to continue the iteration once we've found a discrepancy.
+						return !bShouldPreservePackageDirtyState;
+					});
+
+					if (bShouldPreservePackageDirtyState)
+					{
+						CleanPackageList.Remove(NewObjectPackage);
+					}
+				}
 			}
 
 			if (UAnimInstance* AnimTree = Cast<UAnimInstance>(*NewObject))
