Choose File To Save: Fix issue where Saving a map would save it's dirty actors even if they werent selected by the user

diff --git a/Engine/Source/Editor/MainFrame/Private/Frame/MainFrameActions.cpp b/Engine/Source/Editor/MainFrame/Private/Frame/MainFrameActions.cpp
--- a/Engine/Source/Editor/MainFrame/Private/Frame/MainFrameActions.cpp
+++ b/Engine/Source/Editor/MainFrame/Private/Frame/MainFrameActions.cpp
@@ -293,7 +293,11 @@ void FMainFrameActionCallbacks::ChoosePackagesToSave()
 	const bool bClosingEditor = false;
 	const bool bNotifyNoPackagesSaved = true;
 	const bool bCanBeDeclined = false;
-	FEditorFileUtils::SaveDirtyPackages( bPromptUserToSave, bSaveMapPackages, bSaveContentPackages, bFastSave, bNotifyNoPackagesSaved, bCanBeDeclined );
+
+	// Skip because user will be prompted to select the packages to save, so we don't want the parent package save actually saving an unselected external package
+	const bool bSkipExternalObjectSave = true;
+		
+	FEditorFileUtils::SaveDirtyPackages( bPromptUserToSave, bSaveMapPackages, bSaveContentPackages, bFastSave, bNotifyNoPackagesSaved, bCanBeDeclined, nullptr, FEditorFileUtils::FShouldIgnorePackage::Default, bSkipExternalObjectSave);
 }
 
 void FMainFrameActionCallbacks::ViewChangelists()
diff --git a/Engine/Source/Editor/UnrealEd/Private/FileHelpers.cpp b/Engine/Source/Editor/UnrealEd/Private/FileHelpers.cpp
--- a/Engine/Source/Editor/UnrealEd/Private/FileHelpers.cpp
+++ b/Engine/Source/Editor/UnrealEd/Private/FileHelpers.cpp
@@ -908,7 +908,7 @@ static bool SaveWorld(UWorld* World,
 
 		SlowTask.EnterProgressFrame(50);
 
-		if (!bAutosaving && !FEditorFileUtils::ShouldSkipExternalObjectSave())
+		if (!bAutosaving && (!FEditorFileUtils::ShouldSkipExternalObjectSave() || bNewlyCreated))
 		{
 			if (bSuccess)
 			{
@@ -3907,7 +3907,7 @@ bool FEditorFileUtils::SaveLevel(ULevel* Level, const FString& DefaultFilename,
 	return bLevelWasSaved;
 }
 
-bool FEditorFileUtils::SaveDirtyPackages(const bool bPromptUserToSave, const bool bSaveMapPackages, const bool bSaveContentPackages, const bool bFastSave, const bool bNotifyNoPackagesSaved, const bool bCanBeDeclined, bool* bOutPackagesNeededSaving, const FShouldIgnorePackageFunctionRef& ShouldIgnorePackageFunction)
+bool FEditorFileUtils::SaveDirtyPackages(const bool bPromptUserToSave, const bool bSaveMapPackages, const bool bSaveContentPackages, const bool bFastSave, const bool bNotifyNoPackagesSaved, const bool bCanBeDeclined, bool* bOutPackagesNeededSaving, const FShouldIgnorePackageFunctionRef& ShouldIgnorePackageFunction, bool bInSkipExternalObjectSave)
 {
 	TRACE_CPUPROFILER_EVENT_SCOPE(FEditorFileUtils::SaveDirtyPackages);
 
@@ -3938,6 +3938,7 @@ bool FEditorFileUtils::SaveDirtyPackages(const bool bPromptUserToSave, const boo
 			*bOutPackagesNeededSaving = true;
 		}
 
+		TGuardValue<bool> SkipExternalObjectSaveGuard(bSkipExternalObjectSave, bInSkipExternalObjectSave);
 		const bool bCheckDirty = true;
 		bReturnCode = InternalSavePackages(PackagesToSave, bPromptUserToSave, bFastSave, bCanBeDeclined, bCheckDirty);
 	}
@@ -4065,8 +4066,8 @@ bool FEditorFileUtils::SaveCurrentLevel()
 		{
 			TGuardValue<bool> IsExplicitSaveGuard(EditorFileUtils::bIsExplicitSave, true);
 
-			// If Level gets saved we don't want it to save its external packages because we've already filtered out the ones that need saving and they are part of the PackagesToSave array (unless level is PKG_NewlyCreated then we should save all actors)
-			TGuardValue<bool> GuardValue(bSkipExternalObjectSave, !LevelPackage->HasAnyPackageFlags(PKG_NewlyCreated));
+			// If Level gets saved we don't want it to save its external packages because we've already filtered out the ones that need saving and they are part of the PackagesToSave array (Worlds in package with PKG_NewlyCreated will ignore this flag)
+			TGuardValue<bool> GuardValue(bSkipExternalObjectSave, true);
 
 			const bool bPromptUserToSave = false;
 			const bool bFastSave = false;
diff --git a/Engine/Source/Editor/UnrealEd/Public/FileHelpers.h b/Engine/Source/Editor/UnrealEd/Public/FileHelpers.h
--- a/Engine/Source/Editor/UnrealEd/Public/FileHelpers.h
+++ b/Engine/Source/Editor/UnrealEd/Public/FileHelpers.h
@@ -378,9 +378,11 @@ public:
 	 * @param	bNotifyNoPackagesSaved		true if a notification should be displayed when no packages need to be saved.
 	 * @param	bCanBeDeclined				true if the user prompt should contain a "Don't Save" button in addition to "Cancel", which won't result in a failure return code.
 	 * @param	bOutPackagesNeededSaving	when not NULL, will be set to true if there was any work to be done, and false otherwise.
+	 * @param	ShouldIgnorePackageFunction callback that allows filtering if package should be saved or not.
+	 * @param	bSkipExternalObjectSave		true if saved package should skip saving their own external packages.
 	 * @return								true on success, false on fail.
 	 */
-	UNREALED_API static bool SaveDirtyPackages(const bool bPromptUserToSave, const bool bSaveMapPackages, const bool bSaveContentPackages, const bool bFastSave = false, const bool bNotifyNoPackagesSaved = false, const bool bCanBeDeclined = true, bool* bOutPackagesNeededSaving = NULL, const FShouldIgnorePackageFunctionRef& ShouldIgnorePackageFunction = FShouldIgnorePackage::Default);
+	UNREALED_API static bool SaveDirtyPackages(const bool bPromptUserToSave, const bool bSaveMapPackages, const bool bSaveContentPackages, const bool bFastSave = false, const bool bNotifyNoPackagesSaved = false, const bool bCanBeDeclined = true, bool* bOutPackagesNeededSaving = NULL, const FShouldIgnorePackageFunctionRef& ShouldIgnorePackageFunction = FShouldIgnorePackage::Default, bool bSkipExternalObjectSave = false);
 
 	/**
 	* Looks at all currently loaded packages and saves them if their "bDirty" flag is set and they include specified clasees, optionally prompting the user to select which packages to save)
