Remove broken blueprint performance tracking code now that Insights works well:
BlueprintCompileAndLoadTimerData has been wrong ever since the compilation manager was added
FScopedCompilerEvent was broken because there are two active compilation logs in many cases
Delete unused stats and replace a few with trace scopes

diff --git a/Engine/Plugins/Editor/EngineAssetDefinitions/Source/Private/Animation/AssetDefinition_AnimBlueprint.cpp b/Engine/Plugins/Editor/EngineAssetDefinitions/Source/Private/Animation/AssetDefinition_AnimBlueprint.cpp
--- a/Engine/Plugins/Editor/EngineAssetDefinitions/Source/Private/Animation/AssetDefinition_AnimBlueprint.cpp
+++ b/Engine/Plugins/Editor/EngineAssetDefinitions/Source/Private/Animation/AssetDefinition_AnimBlueprint.cpp
@@ -219,8 +219,6 @@ namespace MenuExtension_AnimBlueprint
 		FCompilerResultsLog	ResultsLog;
 		ResultsLog.SetSourcePath(AnimBlueprint->GetPathName());
 		ResultsLog.BeginEvent(TEXT("Compile"));
-		ResultsLog.bLogDetailedResults = GetDefault<UBlueprintEditorSettings>()->bShowDetailedCompileResults;
-		ResultsLog.EventDisplayThresholdMs = GetDefault<UBlueprintEditorSettings>()->CompileEventDisplayThresholdMs;
 
 		FKismetEditorUtilities::CompileBlueprint(AnimBlueprint, EBlueprintCompileOptions::None, &ResultsLog);
 		ResultsLog.EndEvent();
diff --git a/Engine/Plugins/Runtime/RigVM/Source/RigVMDeveloper/Private/RigVMBlueprint.cpp b/Engine/Plugins/Runtime/RigVM/Source/RigVMDeveloper/Private/RigVMBlueprint.cpp
--- a/Engine/Plugins/Runtime/RigVM/Source/RigVMDeveloper/Private/RigVMBlueprint.cpp
+++ b/Engine/Plugins/Runtime/RigVM/Source/RigVMDeveloper/Private/RigVMBlueprint.cpp
@@ -202,8 +202,6 @@ URigVMBlueprint::URigVMBlueprint(const FObjectInitializer& ObjectInitializer)
 		CompileLog.bSilentMode = true;
 	}
 	CompileLog.SetSourcePath(GetPathName());
-	CompileLog.bLogDetailedResults = false;
-	CompileLog.EventDisplayThresholdMs = false;
 #endif
 
 	if(GetClass() == URigVMBlueprint::StaticClass())
diff --git a/Engine/Source/Editor/BlueprintGraph/Private/BlueprintEditorSettings.cpp b/Engine/Source/Editor/BlueprintGraph/Private/BlueprintEditorSettings.cpp
--- a/Engine/Source/Editor/BlueprintGraph/Private/BlueprintEditorSettings.cpp
+++ b/Engine/Source/Editor/BlueprintGraph/Private/BlueprintEditorSettings.cpp
@@ -67,8 +67,6 @@ UBlueprintEditorSettings::UBlueprintEditorSettings(const FObjectInitializer& Obj
 	// Developer Settings
 	, bShowActionMenuItemSignatures(false)
 	// Perf Settings
-	, bShowDetailedCompileResults(false)
-	, CompileEventDisplayThresholdMs(5)
 	, NodeTemplateCacheCapMB(20.f)
 	// No category
 	, bShowInheritedVariables(false)
diff --git a/Engine/Source/Editor/BlueprintGraph/Public/BlueprintEditorSettings.h b/Engine/Source/Editor/BlueprintGraph/Public/BlueprintEditorSettings.h
--- a/Engine/Source/Editor/BlueprintGraph/Public/BlueprintEditorSettings.h
+++ b/Engine/Source/Editor/BlueprintGraph/Public/BlueprintEditorSettings.h
@@ -252,12 +252,10 @@ public:
 
 	// Perf Settings
 public:
-	/** If enabled, additional details will be displayed in the Compiler Results tab after compiling a blueprint. */
-	UPROPERTY(EditAnywhere, config, Category = Performance)
+	UE_DEPRECATED(5.4, "BP-specific perf tracking has been removed, use Insights")
 	bool bShowDetailedCompileResults;
 
-	/** Minimum event time threshold used as a filter when additional details are enabled for display in the Compiler Results tab. A value of zero means that all events will be included in the final summary. */
-	UPROPERTY(EditAnywhere, AdvancedDisplay, config, Category = Performance, DisplayName = "Compile Event Results Threshold (ms)", meta = (ClampMin = "0", UIMin = "0"))
+	UE_DEPRECATED(5.4, "BP-specific perf tracking has been removed, use Insights")
 	int32 CompileEventDisplayThresholdMs;
 
 	/** The node template cache is used to speed up blueprint menuing. This determines the peak data size for that cache. */
diff --git a/Engine/Source/Editor/Kismet/Private/BlueprintCompilationManager.cpp b/Engine/Source/Editor/Kismet/Private/BlueprintCompilationManager.cpp
--- a/Engine/Source/Editor/Kismet/Private/BlueprintCompilationManager.cpp
+++ b/Engine/Source/Editor/Kismet/Private/BlueprintCompilationManager.cpp
@@ -310,16 +310,6 @@ void FBlueprintCompilationManagerImpl::CompileSynchronouslyImpl(const FBPCompile
 	FlushCompilationQueueImpl(bSuppressBroadcastCompiled, &CompiledBlueprints, &SkeletonCompiledBlueprints, nullptr, bFindAndReplaceCDOReferences ? &OldToNewTemplates : nullptr);
 	FlushReinstancingQueueImpl(bFindAndReplaceCDOReferences, bFindAndReplaceCDOReferences ? &OldToNewTemplates : nullptr);
 	
-	if( Request.UserData.ClientResultsLog && Request.UserData.ClientResultsLog->bLogDetailedResults)
-	{
-		Request.UserData.ClientResultsLog->Note(
-			*FText::Format(
-				LOCTEXT("TotalCompiledBPs", "Compiled {0} total blueprints"),
-				CompiledBlueprints.Num()
-			).ToString()
-		);
-	}
-
 	if (FBlueprintEditorUtils::IsLevelScriptBlueprint(Request.UserData.BPToCompile) && !bRegenerateSkeletonOnly)
 	{
 		// When the Blueprint is recompiled, then update the bound events for level scripting
diff --git a/Engine/Source/Editor/Kismet/Private/BlueprintEditor.cpp b/Engine/Source/Editor/Kismet/Private/BlueprintEditor.cpp
--- a/Engine/Source/Editor/Kismet/Private/BlueprintEditor.cpp
+++ b/Engine/Source/Editor/Kismet/Private/BlueprintEditor.cpp
@@ -4113,8 +4113,6 @@ void FBlueprintEditor::Compile()
 		FCompilerResultsLog LogResults;
 		LogResults.SetSourcePath(BlueprintObj->GetPathName());
 		LogResults.BeginEvent(TEXT("Compile"));
-		LogResults.bLogDetailedResults = GetDefault<UBlueprintEditorSettings>()->bShowDetailedCompileResults;
-		LogResults.EventDisplayThresholdMs = GetDefault<UBlueprintEditorSettings>()->CompileEventDisplayThresholdMs;
 		EBlueprintCompileOptions CompileOptions = EBlueprintCompileOptions::None;
 		if( bSaveIntermediateBuildProducts )
 		{
diff --git a/Engine/Source/Editor/KismetCompiler/Private/KismetCompiler.cpp b/Engine/Source/Editor/KismetCompiler/Private/KismetCompiler.cpp
--- a/Engine/Source/Editor/KismetCompiler/Private/KismetCompiler.cpp
+++ b/Engine/Source/Editor/KismetCompiler/Private/KismetCompiler.cpp
@@ -88,7 +88,6 @@ DECLARE_CYCLE_STAT(TEXT("Clean and Sanitize Class"), EKismetCompilerStats_CleanA
 DECLARE_CYCLE_STAT(TEXT("Create Class Properties"), EKismetCompilerStats_CreateClassVariables, STATGROUP_KismetCompiler );
 DECLARE_CYCLE_STAT(TEXT("Bind and Link Class"), EKismetCompilerStats_BindAndLinkClass, STATGROUP_KismetCompiler );
 DECLARE_CYCLE_STAT(TEXT("Calculate checksum of CDO"), EKismetCompilerStats_ChecksumCDO, STATGROUP_KismetCompiler );
-DECLARE_CYCLE_STAT(TEXT("Analyze execution path"), EKismetCompilerStats_AnalyzeExecutionPath, STATGROUP_KismetCompiler);
 DECLARE_CYCLE_STAT(TEXT("Calculate checksum of signature"), EKismetCompilerStats_ChecksumSignature, STATGROUP_KismetCompiler);
 DECLARE_CYCLE_STAT(TEXT("Pruning"), EKismetCompilerStats_PruneIsolatedNodes, STATGROUP_KismetCompiler);
 DECLARE_CYCLE_STAT(TEXT("Merge Ubergraph Pages In"), EKismetCompilerStats_MergeUbergraphPagesIn, STATGROUP_KismetCompiler);
diff --git a/Engine/Source/Editor/KismetCompiler/Private/KismetCompilerModule.cpp b/Engine/Source/Editor/KismetCompiler/Private/KismetCompilerModule.cpp
--- a/Engine/Source/Editor/KismetCompiler/Private/KismetCompilerModule.cpp
+++ b/Engine/Source/Editor/KismetCompiler/Private/KismetCompilerModule.cpp
@@ -29,9 +29,6 @@
 #include "Engine/Engine.h"
 
 DEFINE_LOG_CATEGORY(LogK2Compiler);
-DECLARE_CYCLE_STAT(TEXT("Compile Time"), EKismetCompilerStats_CompileTime, STATGROUP_KismetCompiler);
-DECLARE_CYCLE_STAT(TEXT("Compile Skeleton Class"), EKismetCompilerStats_CompileSkeletonClass, STATGROUP_KismetCompiler);
-DECLARE_CYCLE_STAT(TEXT("Compile Generated Class"), EKismetCompilerStats_CompileGeneratedClass, STATGROUP_KismetCompiler);
 
 #define LOCTEXT_NAMESPACE "KismetCompiler"
 
@@ -90,13 +87,11 @@ public:
 
 void FKismet2CompilerModule::CompileStructure(UUserDefinedStruct* Struct, FCompilerResultsLog& Results)
 {
+	TRACE_CPUPROFILER_EVENT_SCOPE(FKismet2CompilerModule::CompileStructure);
 	Results.SetSourcePath(Struct->GetPathName());
-	BP_SCOPED_COMPILER_EVENT_STAT(EKismetCompilerStats_CompileTime);
 	FUserDefinedStructureCompilerUtils::CompileStruct(Struct, Results, true);
 }
 
-extern UNREALED_API FSecondsCounterData BlueprintCompileAndLoadTimerData;
-
 void FKismet2CompilerModule::RefreshVariables(UBlueprint* Blueprint)
 {
 #if 0 // disabled because this is a costly ensure:
diff --git a/Engine/Source/Editor/UnrealEd/Private/Kismet2/BlueprintEditorUtils.cpp b/Engine/Source/Editor/UnrealEd/Private/Kismet2/BlueprintEditorUtils.cpp
--- a/Engine/Source/Editor/UnrealEd/Private/Kismet2/BlueprintEditorUtils.cpp
+++ b/Engine/Source/Editor/UnrealEd/Private/Kismet2/BlueprintEditorUtils.cpp
@@ -1476,13 +1476,11 @@ void FBlueprintEditorUtils::PropagateParentBlueprintDefaults(UClass* ClassToProp
 	}
 }
 
-UNREALED_API FSecondsCounterData BlueprintCompileAndLoadTimerData;
-
 uint32 FBlueprintDuplicationScopeFlags::bStaticFlags = FBlueprintDuplicationScopeFlags::NoFlags;
 
 void FBlueprintEditorUtils::PostDuplicateBlueprint(UBlueprint* Blueprint, bool bDuplicateForPIE)
 {
-	FSecondsCounterScope Timer(BlueprintCompileAndLoadTimerData); 
+	TRACE_CPUPROFILER_EVENT_SCOPE(PostDuplicateBlueprint);
 	
 	// Only recompile after duplication if this isn't PIE
 	if (!bDuplicateForPIE)
@@ -1810,8 +1808,6 @@ void FBlueprintEditorUtils::UpdateDelegatesInBlueprint(UBlueprint* Blueprint)
 // Blueprint has materially changed.  Recompile the skeleton, notify observers, and mark the package as dirty.
 void FBlueprintEditorUtils::MarkBlueprintAsStructurallyModified(UBlueprint* Blueprint)
 {
-	FSecondsCounterScope Timer(BlueprintCompileAndLoadTimerData);
-
 	// The Blueprint has been structurally modified and this means that some node titles will need to be refreshed
 	GetDefault<UEdGraphSchema_K2>()->ForceVisualizationCacheClear();
 
diff --git a/Engine/Source/Editor/UnrealEd/Private/Kismet2/CompilerResultsLog.cpp b/Engine/Source/Editor/UnrealEd/Private/Kismet2/CompilerResultsLog.cpp
--- a/Engine/Source/Editor/UnrealEd/Private/Kismet2/CompilerResultsLog.cpp
+++ b/Engine/Source/Editor/UnrealEd/Private/Kismet2/CompilerResultsLog.cpp
@@ -21,13 +21,11 @@
 #define LOCTEXT_NAMESPACE "Editor.Stats"
 
 const FName FCompilerResultsLog::Name(TEXT("CompilerResultsLog"));
-FCompilerResultsLog* FCompilerResultsLog::CurrentEventTarget = nullptr;
 FDelegateHandle FCompilerResultsLog::GetGlobalModuleCompilerDumpDelegateHandle;
 
 //////////////////////////////////////////////////////////////////////////
-// FCompilerResultsLog
+// FBacktrackMap
 
-/** Update the source backtrack map to note that NewObject was most closely generated/caused by the SourceObject */
 void FBacktrackMap::NotifyIntermediateObjectCreation(UObject* NewObject, UObject* SourceObject)
 {
 	// Chase the source to make sure it's really a top-level ('source code') node
@@ -40,7 +38,6 @@ void FBacktrackMap::NotifyIntermediateObjectCreation(UObject* NewObject, UObject
 	SourceBacktrackMap.Add(NewObject, SourceObject);
 }
 
-/** Update the pin source backtrack map to note that NewPin was most closely generated/caused by the SourcePin */
 void FBacktrackMap::NotifyIntermediatePinCreation(UEdGraphPin* NewPin, UEdGraphPin* SourcePin)
 {
 	check(NewPin->GetOwningNode() && SourcePin->GetOwningNode());
@@ -54,7 +51,6 @@ void FBacktrackMap::NotifyIntermediatePinCreation(UEdGraphPin* NewPin, UEdGraphP
 	PinSourceBacktrackMap.Add(NewPin, SourcePin);
 }
 
-/** Returns the true source object for the passed in object */
 UObject* FBacktrackMap::FindSourceObject(UObject* PossiblyDuplicatedObject)
 {
 	UObject** RemappedIfExisting = SourceBacktrackMap.Find(PossiblyDuplicatedObject);
@@ -134,18 +130,10 @@ FCompilerResultsLog::FCompilerResultsLog(bool bIsCompatibleWithEvents/* = true*/
 	, EventDisplayThresholdMs(0)
 {
 	CurrentEventScope = nullptr;
-	if(bIsCompatibleWithEvents && CurrentEventTarget == nullptr)
-	{
-		CurrentEventTarget = this;
-	}
 }
 
 FCompilerResultsLog::~FCompilerResultsLog()
 {
-	if(CurrentEventTarget == this)
-	{
-		CurrentEventTarget = nullptr;
-	}
 }
 
 void FCompilerResultsLog::Register()
@@ -276,6 +264,7 @@ void FCompilerResultsLog::InternalLogSummary()
 
 		if(bLogDetailedResults)
 		{
+			// This will not do anything for the default BP compiler log
 			Note(*LOCTEXT("PerformanceSummaryHeading", "Performance summary:").ToString());
 			InternalLogEvent(*CurrentEventScope.Get());
 		}
@@ -315,7 +304,6 @@ bool FCompilerResultsLog::CommitPotentialMessages(UEdGraphNode* Source)
 	return false;
 }
 
-/** Update the source backtrack map to note that NewObject was most closely generated/caused by the SourceObject */
 void FCompilerResultsLog::NotifyIntermediateObjectCreation(UObject* NewObject, UObject* SourceObject)
 {
 	SourceBacktrackMap.NotifyIntermediateObjectCreation(NewObject, SourceObject);
@@ -326,7 +314,6 @@ void FCompilerResultsLog::NotifyIntermediatePinCreation(UEdGraphPin* NewPin, UEd
 	SourceBacktrackMap.NotifyIntermediatePinCreation(NewPin, SourcePPin);
 }
 
-/** Returns the true source object for the passed in object */
 UObject* FCompilerResultsLog::FindSourceObject(UObject* PossiblyDuplicatedObject)
 {
 	return SourceBacktrackMap.FindSourceObject(PossiblyDuplicatedObject);
diff --git a/Engine/Source/Editor/UnrealEd/Private/Kismet2/Kismet2.cpp b/Engine/Source/Editor/UnrealEd/Private/Kismet2/Kismet2.cpp
--- a/Engine/Source/Editor/UnrealEd/Private/Kismet2/Kismet2.cpp
+++ b/Engine/Source/Editor/UnrealEd/Private/Kismet2/Kismet2.cpp
@@ -86,13 +86,6 @@
 #include "ToolMenus.h"
 #include "Subsystems/AssetEditorSubsystem.h"
 
-DECLARE_CYCLE_STAT(TEXT("Compile Blueprint"), EKismetCompilerStats_CompileBlueprint, STATGROUP_KismetCompiler);
-DECLARE_CYCLE_STAT(TEXT("Broadcast Precompile"), EKismetCompilerStats_BroadcastPrecompile, STATGROUP_KismetCompiler);
-DECLARE_CYCLE_STAT(TEXT("Update Search Metadata"), EKismetCompilerStats_UpdateSearchMetaData, STATGROUP_KismetCompiler);
-DECLARE_CYCLE_STAT(TEXT("Garbage Collection"), EKismetCompilerStats_GarbageCollection, STATGROUP_KismetCompiler);
-DECLARE_CYCLE_STAT(TEXT("Refresh Dependent Blueprints"), EKismetCompilerStats_RefreshDependentBlueprints, STATGROUP_KismetCompiler);
-DECLARE_CYCLE_STAT(TEXT("Validate Generated Class"), EKismetCompilerStats_ValidateGeneratedClass, STATGROUP_KismetCompiler);
-
 #define LOCTEXT_NAMESPACE "UnrealEd.Editor"
 
 //////////////////////////////////////////////////////////////////////////
@@ -412,8 +405,6 @@ void FKismetEditorUtilities::CreateDefaultEventGraphs(UBlueprint* Blueprint)
 	Blueprint->LastEditedDocuments.AddUnique(Ubergraph);
 }
 
-extern UNREALED_API FSecondsCounterData BlueprintCompileAndLoadTimerData;
-
 /** Create a new Blueprint and initialize it to a valid state but uses the associated blueprint types. */
 UBlueprint* FKismetEditorUtilities::CreateBlueprint(UClass* ParentClass, UObject* Outer, const FName NewBPName,	EBlueprintType BlueprintType, FName CallingContext)
 {
@@ -428,7 +419,7 @@ UBlueprint* FKismetEditorUtilities::CreateBlueprint(UClass* ParentClass, UObject
 /** Create a new Blueprint and initialize it to a valid state. */
 UBlueprint* FKismetEditorUtilities::CreateBlueprint(UClass* ParentClass, UObject* Outer, const FName NewBPName, EBlueprintType BlueprintType, TSubclassOf<UBlueprint> BlueprintClassType, TSubclassOf<UBlueprintGeneratedClass> BlueprintGeneratedClassType, FName CallingContext)
 {
-	FSecondsCounterScope Timer(BlueprintCompileAndLoadTimerData);
+	TRACE_CPUPROFILER_EVENT_SCOPE(CreateBlueprint);
 	check(FindObject<UBlueprint>(Outer, *NewBPName.ToString()) == NULL); 
 
 	// Not all types are legal for all parent classes, if the parent class is const then the blueprint cannot be an ubergraph-bearing one
@@ -789,7 +780,6 @@ void FKismetEditorUtilities::CompileBlueprint(UBlueprint* BlueprintObj, EBluepri
 bool FKismetEditorUtilities::GenerateBlueprintSkeleton(UBlueprint* BlueprintObj, bool bForceRegeneration)
 {
 	bool bRegeneratedSkeleton = false;
-	FSecondsCounterScope Timer(BlueprintCompileAndLoadTimerData); 
 	check(BlueprintObj);
 
 	if( BlueprintObj->SkeletonGeneratedClass == NULL || bForceRegeneration )
diff --git a/Engine/Source/Editor/UnrealEd/Private/Kismet2/KismetReinstanceUtilities.cpp b/Engine/Source/Editor/UnrealEd/Private/Kismet2/KismetReinstanceUtilities.cpp
--- a/Engine/Source/Editor/UnrealEd/Private/Kismet2/KismetReinstanceUtilities.cpp
+++ b/Engine/Source/Editor/UnrealEd/Private/Kismet2/KismetReinstanceUtilities.cpp
@@ -52,8 +52,6 @@ DECLARE_CYCLE_STAT(TEXT("Update Bytecode References"), EKismetReinstancerStats_U
 DECLARE_CYCLE_STAT(TEXT("Recompile Child Classes"), EKismetReinstancerStats_RecompileChildClasses, STATGROUP_KismetReinstancer );
 DECLARE_CYCLE_STAT(TEXT("Replace Classes Without Reinstancing"), EKismetReinstancerStats_ReplaceClassNoReinsancing, STATGROUP_KismetReinstancer );
 DECLARE_CYCLE_STAT(TEXT("Reinstance Objects"), EKismetCompilerStats_ReinstanceObjects, STATGROUP_KismetCompiler);
-DECLARE_CYCLE_STAT(TEXT("Refresh Dependent Blueprints In Reinstancer"), EKismetCompilerStats_RefreshDependentBlueprintsInReinstancer, STATGROUP_KismetCompiler);
-DECLARE_CYCLE_STAT(TEXT("Recreate UberGraphPersistentFrame"), EKismetCompilerStats_RecreateUberGraphPersistentFrame, STATGROUP_KismetCompiler);
 
 bool GUseLegacyAnimInstanceReinstancingBehavior = false;
 static FAutoConsoleVariableRef CVarUseLegacyAnimInstanceReinstancingBehavior(
@@ -942,12 +940,8 @@ void FBlueprintCompileReinstancer::BlueprintWasRecompiled(UBlueprint* BP, bool b
 {
 }
 
-extern UNREALED_API FSecondsCounterData BlueprintCompileAndLoadTimerData;
-
 void FBlueprintCompileReinstancer::ReinstanceObjects(bool bForceAlwaysReinstance)
 {
-	FSecondsCounterScope Timer(BlueprintCompileAndLoadTimerData);
-
 	BP_SCOPED_COMPILER_EVENT_STAT(EKismetCompilerStats_ReinstanceObjects);
 	
 	// Make sure we only reinstance classes once!
@@ -3138,32 +3132,3 @@ void FBlueprintCompileReinstancer::PreCreateSubObjectsForReinstantiation_Inner(c
 		}
 	}
 }
-
-FRecreateUberGraphFrameScope::FRecreateUberGraphFrameScope(UClass* InClass, bool bRecreate)
-	: RecompiledClass(InClass)
-{
-	if (bRecreate && ensure(RecompiledClass))
-	{
-		BP_SCOPED_COMPILER_EVENT_STAT(EKismetCompilerStats_RecreateUberGraphPersistentFrame);
-
-		const bool bIncludeDerivedClasses = true;
-		GetObjectsOfClass(RecompiledClass, Objects, bIncludeDerivedClasses, RF_NoFlags);
-
-		for (UObject* Obj : Objects)
-		{
-			RecompiledClass->DestroyPersistentUberGraphFrame(Obj);
-		}
-	}
-}
-
-FRecreateUberGraphFrameScope::~FRecreateUberGraphFrameScope()
-{
-	BP_SCOPED_COMPILER_EVENT_STAT(EKismetCompilerStats_RecreateUberGraphPersistentFrame);
-	for (UObject* Obj : Objects)
-	{
-		if (IsValid(Obj))
-		{
-			RecompiledClass->CreatePersistentUberGraphFrame(Obj, false);
-		}
-	}
-}
diff --git a/Engine/Source/Editor/UnrealEd/Public/Kismet2/CompilerResultsLog.h b/Engine/Source/Editor/UnrealEd/Public/Kismet2/CompilerResultsLog.h
--- a/Engine/Source/Editor/UnrealEd/Public/Kismet2/CompilerResultsLog.h
+++ b/Engine/Source/Editor/UnrealEd/Public/Kismet2/CompilerResultsLog.h
@@ -121,10 +121,10 @@ public:
 	// Should nodes mentioned in messages be annotated for display with that message?
 	bool bAnnotateMentionedNodes;
 
-	// Should detailed results be appended to the final summary log?
+	// Should detailed BeginEvent/EndEvent timing information be written to log
 	bool bLogDetailedResults;
 
-	// Minimum event time (ms) for inclusion into the final summary log
+	// Minimum event time (ms) for events include in detailed results
 	int EventDisplayThresholdMs;
 
 	/** Tracks nodes that produced errors/warnings */
@@ -330,10 +330,10 @@ public:
 	/** End the current compiler event */
 	UNREALED_API void EndEvent();
 
-	/** Access the current event target log */
+	UE_DEPRECATED(5.4, "BP-specific perf tracking has been removed, use Insights")
 	static FCompilerResultsLog* GetEventTarget()
 	{
-		return CurrentEventTarget;
+		return nullptr;
 	}
 
 	/** Get the message log listing for this blueprint */
@@ -421,35 +421,16 @@ private:
 	/** The log's name, for easy re-use */
 	static UNREALED_API const FName Name;
 	
-	/** The log target for compile events */
-	static UNREALED_API FCompilerResultsLog* CurrentEventTarget;
-
 	/** Handle to the registered GetGlobalModuleCompilerDump delegate. */
 	static UNREALED_API FDelegateHandle GetGlobalModuleCompilerDumpDelegateHandle;
 };
 
-/** This class will begin a new compile event on construction, and automatically end it when the instance goes out of scope */
 class FScopedCompilerEvent
 {
 public:
-	/** Constructor; automatically begins a new event */
+	UE_DEPRECATED(5.4, "BP-specific perf tracking has been removed, use Insights")
 	FScopedCompilerEvent(const TCHAR* InName)
 	{
-		FCompilerResultsLog* ResultsLog = FCompilerResultsLog::GetEventTarget();
-		if(ResultsLog != nullptr)
-		{
-			ResultsLog->BeginEvent(InName);
-		}
-	}
-
-	/** Destructor; automatically ends the event */
-	~FScopedCompilerEvent()
-	{
-		FCompilerResultsLog* ResultsLog = FCompilerResultsLog::GetEventTarget();
-		if(ResultsLog != nullptr)
-		{
-			ResultsLog->EndEvent();
-		}
 	}
 };
 
@@ -471,11 +452,9 @@ public:
 #if STATS
 #define BP_SCOPED_COMPILER_EVENT_STAT(Stat) \
 	SCOPE_CYCLE_COUNTER(Stat); \
-	TRACE_CPUPROFILER_EVENT_SCOPE(Stat); \
-	FScopedCompilerEvent PREPROCESSOR_JOIN(ScopedCompilerEvent,__LINE__)(GET_STATDESCRIPTION(Stat))
+	TRACE_CPUPROFILER_EVENT_SCOPE(Stat);
 #else
 #define BP_SCOPED_COMPILER_EVENT_STAT(Stat) \
-	TRACE_CPUPROFILER_EVENT_SCOPE(Stat); \
-	FScopedCompilerEvent PREPROCESSOR_JOIN(ScopedCompilerEvent,__LINE__)(ANSI_TO_TCHAR(#Stat))
+	TRACE_CPUPROFILER_EVENT_SCOPE(Stat);
 #endif
 #endif	//#if WITH_EDITOR
diff --git a/Engine/Source/Editor/UnrealEd/Public/Kismet2/KismetReinstanceUtilities.h b/Engine/Source/Editor/UnrealEd/Public/Kismet2/KismetReinstanceUtilities.h
--- a/Engine/Source/Editor/UnrealEd/Public/Kismet2/KismetReinstanceUtilities.h
+++ b/Engine/Source/Editor/UnrealEd/Public/Kismet2/KismetReinstanceUtilities.h
@@ -25,18 +25,6 @@ enum class EBlueprintCompileReinstancerFlags
 
 ENUM_CLASS_FLAGS(EBlueprintCompileReinstancerFlags)
 
-class FReinstanceFinalizer;
-
-struct FRecreateUberGraphFrameScope
-{
-private:
-	TArray<UObject*> Objects;
-	UClass* RecompiledClass;
-public:
-	UNREALED_API FRecreateUberGraphFrameScope(UClass* InClass, bool bRecreate);
-	UNREALED_API ~FRecreateUberGraphFrameScope();
-};
-
 struct FReplaceInstancesOfClassParameters
 {
 	UE_NONCOPYABLE(FReplaceInstancesOfClassParameters)
diff --git a/Engine/Source/Runtime/Engine/Private/BlueprintGeneratedClass.cpp b/Engine/Source/Runtime/Engine/Private/BlueprintGeneratedClass.cpp
--- a/Engine/Source/Runtime/Engine/Private/BlueprintGeneratedClass.cpp
+++ b/Engine/Source/Runtime/Engine/Private/BlueprintGeneratedClass.cpp
@@ -727,8 +727,6 @@ struct FConditionalRecompileClassHepler
 	}
 };
 
-extern UNREALED_API FSecondsCounterData BlueprintCompileAndLoadTimerData;
-
 void UBlueprintGeneratedClass::ConditionalRecompileClass(FUObjectSerializeContext* InLoadContext)
 {
 	FBlueprintCompilationManager::FlushCompilationQueue(InLoadContext);
diff --git a/Engine/Source/Runtime/Launch/Private/Launch.cpp b/Engine/Source/Runtime/Launch/Private/Launch.cpp
--- a/Engine/Source/Runtime/Launch/Private/Launch.cpp
+++ b/Engine/Source/Runtime/Launch/Private/Launch.cpp
@@ -81,10 +81,6 @@ void LaunchStaticShutdownAfterError()
 	TermGamePhys();
 }
 
-#if WITH_EDITOR
-extern UNREALED_API FSecondsCounterData BlueprintCompileAndLoadTimerData;
-#endif
-
 /**
  * Static guarded main function. Rolled into own function so we can have error handling for debug/ release builds depending
  * on whether a debugger is attached or not.
@@ -170,10 +166,6 @@ int32 GuardedMain( const TCHAR* CmdLine )
 	double EngineInitializationTime = FPlatformTime::Seconds() - GStartTime;
 	UE_LOG(LogLoad, Log, TEXT("(Engine Initialization) Total time: %.2f seconds"), EngineInitializationTime);
 
-#if WITH_EDITOR
-	UE_LOG(LogLoad, Log, TEXT("(Engine Initialization) Total Blueprint compile time: %.2f seconds"), BlueprintCompileAndLoadTimerData.GetTime());
-#endif
-
 	ACCUM_LOADTIME(TEXT("EngineInitialization"), EngineInitializationTime);
 
 	BootTimingPoint("Tick loop starting");
