Make sure UWorld::WorldType is initialized correctly when reloading a loaded map, this is to ensure UEditorEngine::InitializeNewlyCreatedInactiveWorld runs at a consistent time.

diff --git a/Engine/Source/Editor/UnrealEd/Private/PackageTools.cpp b/Engine/Source/Editor/UnrealEd/Private/PackageTools.cpp
--- a/Engine/Source/Editor/UnrealEd/Private/PackageTools.cpp
+++ b/Engine/Source/Editor/UnrealEd/Private/PackageTools.cpp
@@ -840,6 +840,7 @@ UPackageTools::UPackageTools(const FObjectInitializer& ObjectInitializer)
 
 		// Check to see if we need to reload the current world.
 		FName WorldNameToReload;
+		FName CurrentWorldPackageName;
 		TArray<ULevelStreaming*> RemovedStreamingLevels;
 		if (UWorld* CurrentWorldPtr = CurrentWorld.Get())
 		{
@@ -850,6 +851,7 @@ UPackageTools::UPackageTools(const FObjectInitializer& ObjectInitializer)
 			{
 				// Cache this so we can reload the world later
 				WorldNameToReload = *CurrentWorldPtr->GetPathName();
+				CurrentWorldPackageName = CurrentWorldPtr->GetPackage()->GetFName();
 
 				// Remove the world package from the reload list
 				Filtered.Remove(CurrentWorldPtr->GetOutermost());
@@ -975,9 +977,13 @@ UPackageTools::UPackageTools(const FObjectInitializer& ObjectInitializer)
 		{
 			if (GIsEditor)
 			{
+				UWorld::WorldTypePreLoadMap.FindOrAdd(CurrentWorldPackageName) = EWorldType::Editor;
+
 				TArray<FName> WorldNamesToReload;
 				WorldNamesToReload.Add(WorldNameToReload);
 				GEditor->GetEditorSubsystem<UAssetEditorSubsystem>()->OpenEditorsForAssets(WorldNamesToReload);
+
+				UWorld::WorldTypePreLoadMap.Remove(CurrentWorldPackageName);
 			}
 			else if (UGameEngine* GameEngine = Cast<UGameEngine>(GEngine))
 			{
