Make it possible to move the UE window by dragging the project badge

* Pad around the SProjectBadge text by using the margin attribute of its SText instead of wrapping the SText in an SBorder. This is needed because the SBorder doesn't override SWidget::GetWindowZoneOverride, and therefore the window cannot be dragged on top of that SBorder.
* Add a WindowZoneOverride attribute to SExtensionPanel. This attribute's value controls the return value of SExtensionPanel::GetWindowZoneOverride.
* Set the WindowZoneOverride to EWindowZone::TitleBar for the SExtensionPanel in SProjectBadge. This allows the window to be dragged on the project badge.
* Stop using "ProjectBadgeBackgroundColor" from GEditorPerProjectIni because that was only used by the removed SBorder.
* Stop using "SProjectBadge.BadgeShape" from SlateEditorStyle and StarshipStyle because that was only used by the removed SBorder.

diff --git a/Engine/Source/Editor/LevelEditor/Private/LevelEditor.cpp b/Engine/Source/Editor/LevelEditor/Private/LevelEditor.cpp
--- a/Engine/Source/Editor/LevelEditor/Private/LevelEditor.cpp
+++ b/Engine/Source/Editor/LevelEditor/Private/LevelEditor.cpp
@@ -105,19 +105,11 @@ public:
 		FString ProjectNameWatermarkPrefix;
 		GConfig->GetString(TEXT("LevelEditor"), TEXT("ProjectNameWatermarkPrefix"), /*out*/ ProjectNameWatermarkPrefix, GEditorPerProjectIni);
 
-		FSlateColor BadgeBackgroundColor = FAppStyle::Get().GetSlateColor("Colors.Title");
-
-		FColor ConfigColor;
-		if (GConfig->GetColor(TEXT("LevelEditor"), TEXT("ProjectBadgeBackgroundColor"), /*out*/ ConfigColor, GEditorPerProjectIni))
-		{
-			BadgeBackgroundColor = FLinearColor(ConfigColor);
-		}
-
-		FColor BadgeTextColor = FColor(128,128,128,255);
+		FColor BadgeTextColor = FColor(128, 128, 128, 255);
 		GConfig->GetColor(TEXT("LevelEditor"), TEXT("ProjectBadgeTextColor"), /*out*/ BadgeTextColor, GEditorPerProjectIni);
 
 		const FString EngineVersionString = FEngineVersion::Current().ToString(FEngineVersion::Current().HasChangelist() ? EVersionComponent::Changelist : EVersionComponent::Patch);
-		
+
 		FFormatNamedArguments Args;
 
 		Args.Add(TEXT("ProjectNameWatermarkPrefix"), FText::FromString(ProjectNameWatermarkPrefix));
@@ -156,6 +148,7 @@ public:
 			.Text(RightContentText)
 			.Visibility(EVisibility::HitTestInvisible)
 			.TextStyle(FAppStyle::Get(), "SProjectBadge.Text")
+			.Margin(FAppStyle::Get().GetMargin("SProjectBadge.BadgePadding"))
 			.ColorAndOpacity(BadgeTextColor);
 
 		SBox::Construct(SBox::FArguments()
@@ -163,16 +156,10 @@ public:
 			.VAlign(VAlign_Top)
 			.Padding(FMargin(0.0f, 0.0f, 2.0f, 0.0f))
 			[
-				SNew(SBorder)
-				.BorderImage(FAppStyle::GetBrush("SProjectBadge.BadgeShape"))
-				.Padding(FAppStyle::Get().GetMargin("SProjectBadge.BadgePadding"))
-				.BorderBackgroundColor(BadgeBackgroundColor)
-				.VAlign(VAlign_Top)
-				[
-					SNew(SExtensionPanel)
-					.ExtensionPanelID("LevelEditorProjectNamePlate")
-					.DefaultWidget(DefaultNamePlate)
-				]
+				SNew(SExtensionPanel)
+				.ExtensionPanelID("LevelEditorProjectNamePlate")
+				.DefaultWidget(DefaultNamePlate)
+				.WindowZoneOverride(EWindowZone::TitleBar)
 			]);
 	}
 
diff --git a/Engine/Source/Editor/UnrealEd/Private/Subsystems/PanelExtensionSubsystem.cpp b/Engine/Source/Editor/UnrealEd/Private/Subsystems/PanelExtensionSubsystem.cpp
--- a/Engine/Source/Editor/UnrealEd/Private/Subsystems/PanelExtensionSubsystem.cpp
+++ b/Engine/Source/Editor/UnrealEd/Private/Subsystems/PanelExtensionSubsystem.cpp
@@ -24,6 +24,7 @@ void SExtensionPanel::Construct(const FArguments& InArgs)
 	ExtensionPanelID = InArgs._ExtensionPanelID.Get(NAME_None);
 	DefaultWidget = InArgs._DefaultWidget.Get(nullptr);
 	ExtensionContext = InArgs._ExtensionContext.Get(nullptr);
+	WindowZoneOverride = InArgs._WindowZoneOverride;
 
 	if (GEditor && ExtensionPanelID != NAME_None)
 	{
diff --git a/Engine/Source/Editor/UnrealEd/Public/Subsystems/PanelExtensionSubsystem.h b/Engine/Source/Editor/UnrealEd/Public/Subsystems/PanelExtensionSubsystem.h
--- a/Engine/Source/Editor/UnrealEd/Public/Subsystems/PanelExtensionSubsystem.h
+++ b/Engine/Source/Editor/UnrealEd/Public/Subsystems/PanelExtensionSubsystem.h
@@ -14,7 +14,7 @@
 #include "PanelExtensionSubsystem.generated.h"
 
 PRAGMA_DISABLE_DEPRECATION_WARNINGS
-	
+
 struct UNREALED_API FPanelExtensionFactory
 {
 public:
@@ -25,7 +25,7 @@ public:
 	int32 SortWeight = 0;
 
 	/**
-	 * Delegate that generates the SExtensionPanel content widget. 
+	 * Delegate that generates the SExtensionPanel content widget.
 	 * The FWeakObjectPtr param is an opaque context specific to the panel that is being extended
 	 * which can be used to customize or populate the extension widget.
 	 */
@@ -50,7 +50,9 @@ public:
 		: _ExtensionPanelID()
 		, _DefaultWidget()
 		, _ExtensionContext()
-	{}
+		, _WindowZoneOverride()
+	{
+	}
 
 		/** The ID to identify this Extension point */
 		SLATE_ATTRIBUTE(FName, ExtensionPanelID)
@@ -58,13 +60,23 @@ public:
 		SLATE_ATTRIBUTE(TSharedPtr<SWidget>, DefaultWidget)
 		/** Context used to customize or populate the extension widget (specific to each panel extension) */
 		SLATE_ATTRIBUTE(FWeakObjectPtr, ExtensionContext)
+		/** The window zone to return for this widget. Set to EWindowZone::TitleBar to make the window draggable on this widget. */
+		SLATE_ATTRIBUTE(EWindowZone::Type, WindowZoneOverride)
 
 	SLATE_END_ARGS()
 
 	/** Constructs the widget */
 	void Construct(const FArguments& InArgs);
 
-	UObject* GetExtensionContext() const { return ExtensionContext.Get(); }
+	UObject* GetExtensionContext() const
+	{
+		return ExtensionContext.Get();
+	}
+
+	virtual EWindowZone::Type GetWindowZoneOverride() const override
+	{
+		return WindowZoneOverride.Get(EWindowZone::Unspecified);
+	}
 
 private:
 	void RebuildWidget();
@@ -72,6 +84,7 @@ private:
 	FName ExtensionPanelID;
 	TSharedPtr<SWidget> DefaultWidget;
 	FWeakObjectPtr ExtensionContext;
+	TAttribute<EWindowZone::Type> WindowZoneOverride;
 };
 
 /**
@@ -103,5 +116,4 @@ protected:
 private:
 	TMap<FName, TArray<FPanelExtensionFactory>> ExtensionPointMap;
 	TMap<FName, FPanelFactoryRegistryChanged> PanelFactoryRegistryChangedCallbackMap;
-
 };
