nullcheck for editor selection in case invalid postundo call occurs
#tests Copied actors and exited editor while in session and without saving

diff --git a/Engine/Source/Editor/UnrealEd/Private/EditorServer.cpp b/Engine/Source/Editor/UnrealEd/Private/EditorServer.cpp
--- a/Engine/Source/Editor/UnrealEd/Private/EditorServer.cpp
+++ b/Engine/Source/Editor/UnrealEd/Private/EditorServer.cpp
@@ -1260,27 +1260,31 @@ void UEditorEngine::PostUndo(bool)
 {
 	UTypedElementSelectionSet* LevelEditorSelection = GetSelectedActors()->GetElementSelectionSet();
 
-	// Cache any Actor that needs to be re-instanced because it still points to a REINST_ class
-	TMap< UClass*, UClass* > OldToNewClassMapToReinstance;
-	LevelEditorSelection->ForEachSelectedObject<AActor>([&OldToNewClassMapToReinstance](AActor* InActor)
+	// This is a safeguard in case the function call is triggered after shutdown has initiated
+	if (LevelEditorSelection)
 	{
-		// If the Actor's Class is not the AuthoritativeClass, then it needs to be re-instanced
-		UClass* OldClass = InActor->GetClass();
-		if (OldClass->HasAnyClassFlags(CLASS_NewerVersionExists))
-		{
-			UClass* NewClass = OldClass->GetAuthoritativeClass();
-			if (!ensure(NewClass != OldClass))
+		// Cache any Actor that needs to be re-instanced because it still points to a REINST_ class
+		TMap< UClass*, UClass* > OldToNewClassMapToReinstance;
+		LevelEditorSelection->ForEachSelectedObject<AActor>([&OldToNewClassMapToReinstance](AActor* InActor)
 			{
-				UE_LOG(LogActor, Warning, TEXT("WARNING: %s is out of date and is the same as its AuthoritativeClass during PostUndo!"), *OldClass->GetName());
-			};
+				// If the Actor's Class is not the AuthoritativeClass, then it needs to be re-instanced
+				UClass* OldClass = InActor->GetClass();
+				if (OldClass->HasAnyClassFlags(CLASS_NewerVersionExists))
+				{
+					UClass* NewClass = OldClass->GetAuthoritativeClass();
+					if (!ensure(NewClass != OldClass))
+					{
+						UE_LOG(LogActor, Warning, TEXT("WARNING: %s is out of date and is the same as its AuthoritativeClass during PostUndo!"), *OldClass->GetName());
+					};
 
-			OldToNewClassMapToReinstance.Add(OldClass, NewClass);
-		}
-		return true;
-	});
+					OldToNewClassMapToReinstance.Add(OldClass, NewClass);
+				}
+				return true;
+			});
 
-	// Re-instance any actors that need it
-	FBlueprintCompileReinstancer::BatchReplaceInstancesOfClass(OldToNewClassMapToReinstance, FReplaceInstancesOfClassParameters());
+		// Re-instance any actors that need it
+		FBlueprintCompileReinstancer::BatchReplaceInstancesOfClass(OldToNewClassMapToReinstance, FReplaceInstancesOfClassParameters());
+	}
 
 	RedrawLevelEditingViewports();
 }
