MPCook: Fix retraction never retrying on local worker if it was ever found at some point to be blocked due to only having unretractable packages assigned.
#rnx

diff --git a/Engine/Source/Editor/UnrealEd/Private/Cooker/CookPackageData.cpp b/Engine/Source/Editor/UnrealEd/Private/Cooker/CookPackageData.cpp
--- a/Engine/Source/Editor/UnrealEd/Private/Cooker/CookPackageData.cpp
+++ b/Engine/Source/Editor/UnrealEd/Private/Cooker/CookPackageData.cpp
@@ -2574,15 +2574,13 @@ void FPackageDataMonitor::OnUrgencyChanged(FPackageData& PackageData)
 
 void FPackageDataMonitor::OnStateChanged(FPackageData& PackageData, EPackageState OldState)
 {
-	if (!PackageData.GetIsUrgent())
+	EPackageState NewState = PackageData.GetState();
+	if (PackageData.GetIsUrgent())
 	{
-		return;
+		TrackUrgentRequests(OldState, -1);
+		TrackUrgentRequests(NewState, 1);
 	}
 
-	TrackUrgentRequests(OldState, -1);
-	TrackUrgentRequests(PackageData.GetState(), 1);
-
-	EPackageState NewState = PackageData.GetState();
 	bool bOldStateAssignedToLocal = OldState != EPackageState::Idle && OldState != EPackageState::AssignedToWorker;
 	bool bNewStateAssignedToLocal = NewState != EPackageState::Idle && NewState != EPackageState::AssignedToWorker;
 	if (bOldStateAssignedToLocal != bNewStateAssignedToLocal)
