[Spline Component Visualizer] Fix a few issues with selection of component visualizers

Selecting the actor in the outliner would cause component visualizers to stop drawing for the selection

When dragging and dropping something containing a spline component to the viewport, a double selection was occurring which was clearing the component visualizers to draw.

Fix an issue where the spline component visualizer being active didn't select the component, which could result in loss of the editing state when selecting the component in the component tree view. This feature is on by default, but toggle-able by the cvar Editor.ComponentVisualizer.AutoSelectComponent

diff --git a/Engine/Source/Editor/ComponentVisualizers/Private/SplineComponentVisualizer.cpp b/Engine/Source/Editor/ComponentVisualizers/Private/SplineComponentVisualizer.cpp
--- a/Engine/Source/Editor/ComponentVisualizers/Private/SplineComponentVisualizer.cpp
+++ b/Engine/Source/Editor/ComponentVisualizers/Private/SplineComponentVisualizer.cpp
@@ -1815,7 +1815,6 @@ void FSplineComponentVisualizer::TrackingStopped(FEditorViewportClient* InViewpo
 	{
 		// After dragging, notify that the spline curves property has changed one last time, this time as a EPropertyChangeType::ValueSet :
 		USplineComponent* SplineComp = GetEditedSplineComponent();
-		check(SplineComp != nullptr);
 		NotifyPropertyModified(SplineComp, SplineCurvesProperty, EPropertyChangeType::ValueSet);
 	}
 }
diff --git a/Engine/Source/Editor/SceneOutliner/Private/ActorBrowsingMode.cpp b/Engine/Source/Editor/SceneOutliner/Private/ActorBrowsingMode.cpp
--- a/Engine/Source/Editor/SceneOutliner/Private/ActorBrowsingMode.cpp
+++ b/Engine/Source/Editor/SceneOutliner/Private/ActorBrowsingMode.cpp
@@ -1083,6 +1083,7 @@ void FActorBrowsingMode::OnItemSelectionChanged(FSceneOutlinerTreeItemPtr TreeIt
 				// Avoid senting out notification via typed element and call NoteSelectionChange to preserve previous behavior
 				FTypedElementList::FScopedClearNewPendingChange ClearNewPendingChange = SelectionSet->GetScopedClearNewPendingChange();
 				SelectionSet->SetSelection(ElementsToSelect, SelectionOptions);
+				SelectionSet->NotifyPendingChanges();
 			}
 
 			// Fire selection changed event
diff --git a/Engine/Source/Editor/UnrealEd/Classes/Editor/UnrealEdEngine.h b/Engine/Source/Editor/UnrealEd/Classes/Editor/UnrealEdEngine.h
--- a/Engine/Source/Editor/UnrealEd/Classes/Editor/UnrealEdEngine.h
+++ b/Engine/Source/Editor/UnrealEd/Classes/Editor/UnrealEdEngine.h
@@ -874,10 +874,6 @@ public:
 	UNREALED_API const TArray<FTemplateMapInfo>& GetProjectDefaultMapTemplates() const;
 
 protected:
-
-	/** Called when global editor selection is about to change */
-	void OnEditorSelectionPreChange(const UTypedElementSelectionSet* SelectionSet);
-
 	/** Called when the element selection set pointer set on the global editor selection changes */
 	UNREALED_API void OnEditorElementSelectionPtrChanged(USelection* Selection, UTypedElementSelectionSet* OldSelectionSet, UTypedElementSelectionSet* NewSelectionSet);
 
diff --git a/Engine/Source/Editor/UnrealEd/Private/ComponentVisualizerManager.cpp b/Engine/Source/Editor/UnrealEd/Private/ComponentVisualizerManager.cpp
--- a/Engine/Source/Editor/UnrealEd/Private/ComponentVisualizerManager.cpp
+++ b/Engine/Source/Editor/UnrealEd/Private/ComponentVisualizerManager.cpp
@@ -7,6 +7,20 @@
 #include "Editor/UnrealEdEngine.h"
 #include "UnrealEdGlobals.h"
 #include "SEditorViewport.h"
+#include "HAL/IConsoleManager.h"
+#include "EditorModeManager.h"
+#include "Elements/Framework/TypedElementSelectionSet.h"
+
+namespace ComponentVisualizers
+{
+	static bool GAutoSelectComponentWithPointSelection = true;
+	static FAutoConsoleVariableRef CVarAutoSelectComponentWithPointSelection(
+		TEXT("Editor.ComponentVisualizer.AutoSelectComponent"),
+		GAutoSelectComponentWithPointSelection,
+		TEXT("Automatically adds the spline component to the selection set if avaialable when a point is selected on the spline")
+	);
+
+}; // namespace ComponentVisualizers
 
 
 FComponentVisualizerManager::FComponentVisualizerManager()
@@ -63,18 +77,29 @@ bool FComponentVisualizerManager::HandleProxyForComponentVis(FEditorViewportClie
 			TSharedPtr<FComponentVisualizer> Visualizer = GUnrealEd->FindComponentVisualizer(ClickedComponent->GetClass());
 			if (Visualizer.IsValid())
 			{
+				FTypedElementHandle SelectedComponentHandle = VisProxy->GetElementHandle();
+				UTypedElementSelectionSet* ElementSelectionSet = (InViewportClient && InViewportClient->GetModeTools()) ? InViewportClient->GetModeTools()->GetEditorSelectionSet() : nullptr;
+
 				bool bIsActive = Visualizer->VisProxyHandleClick(InViewportClient, VisProxy, Click);
 				if (bIsActive)
 				{
-					// call EndEditing on any currently edited visualizer, if we are going to change it
-					TSharedPtr<FComponentVisualizer> EditedVisualizer = EditedVisualizerPtr.Pin();
-					if (EditedVisualizer.IsValid() && Visualizer.Get() != EditedVisualizer.Get())
+					// Set the selection set to the spline component, so clicking the component doesn't reselect and clear the current editing state.
+					if (ComponentVisualizers::GAutoSelectComponentWithPointSelection && ElementSelectionSet)
 					{
-						EditedVisualizer->EndEditing();
+						if (!ElementSelectionSet->IsElementSelected(SelectedComponentHandle, FTypedElementIsSelectedOptions()))
+						{
+							TArray<FTypedElementHandle> TmpArray = { SelectedComponentHandle };
+							ElementSelectionSet->SetSelection(TmpArray, FTypedElementSelectionOptions());
+						}
+					}
+
+					SetActiveComponentVis(InViewportClient, Visualizer);
+
+					if (ComponentVisualizers::GAutoSelectComponentWithPointSelection && ElementSelectionSet)
+					{
+						ElementSelectionSet->NotifyPendingChanges();
 					}
 
-					EditedVisualizerPtr = Visualizer;
-					EditedVisualizerViewportClient = InViewportClient;
 					return true;
 				}
 			}
@@ -237,7 +262,7 @@ void FComponentVisualizerManager::TrackingStopped(FEditorViewportClient* InViewp
 {
 	TSharedPtr<FComponentVisualizer> EditedVisualizer = EditedVisualizerPtr.Pin();
 
-	if (EditedVisualizer.IsValid())
+	if (EditedVisualizer.IsValid() && EditedVisualizer->GetEditedComponent() != nullptr)
 	{
 		return EditedVisualizer->TrackingStopped(InViewportClient, bInDidMove);
 	}
diff --git a/Engine/Source/Editor/UnrealEd/Private/EditorSelectUtils.cpp b/Engine/Source/Editor/UnrealEd/Private/EditorSelectUtils.cpp
--- a/Engine/Source/Editor/UnrealEd/Private/EditorSelectUtils.cpp
+++ b/Engine/Source/Editor/UnrealEd/Private/EditorSelectUtils.cpp
@@ -278,12 +278,6 @@ void UUnrealEdEngine::ResetPivot()
 /*-----------------------------------------------------------------------------
 	Selection.
 -----------------------------------------------------------------------------*/
-void UUnrealEdEngine::OnEditorSelectionPreChange(const UTypedElementSelectionSet* SelectionSet)
-{
-	// Clear active editing visualizer on selection change
-	ComponentVisManager.ClearActiveComponentVis();
-	VisualizersForSelection.Empty();
-}
 
 void UUnrealEdEngine::OnEditorElementSelectionPtrChanged(USelection* Selection, UTypedElementSelectionSet* OldSelectionSet, UTypedElementSelectionSet* NewSelectionSet)
 {
@@ -292,13 +286,11 @@ void UUnrealEdEngine::OnEditorElementSelectionPtrChanged(USelection* Selection,
 		if (OldSelectionSet)
 		{
 			OldSelectionSet->OnChanged().RemoveAll(this);
-			OldSelectionSet->OnPreChange().RemoveAll(this);
 		}
 
 		if (NewSelectionSet)
 		{
 			NewSelectionSet->OnChanged().AddUObject(this, &UUnrealEdEngine::OnEditorElementSelectionChanged);
-			NewSelectionSet->OnPreChange().AddUObject(this, &UUnrealEdEngine::OnEditorSelectionPreChange);
 		}
 	}
 }
@@ -306,6 +298,8 @@ void UUnrealEdEngine::OnEditorElementSelectionPtrChanged(USelection* Selection,
 
 void UUnrealEdEngine::OnEditorElementSelectionChanged(const UTypedElementSelectionSet* SelectionSet)
 {
+	VisualizersForSelection.Empty();
+
 	auto GetVisualizersForSelection = [this](AActor* Actor, const UActorComponent* SelectedComponent)
 	{
 		// Iterate over components of that actor (and recurse through child components)
@@ -370,6 +364,11 @@ void UUnrealEdEngine::OnEditorElementSelectionChanged(const UTypedElementSelecti
 				break;
 			}
 		}
+
+		if (!ComponentVisManager.IsActive())
+		{
+			ComponentVisManager.ClearActiveComponentVis();
+		}
 	}
 
 #if PLATFORM_MAC
diff --git a/Engine/Source/Editor/UnrealEd/Public/ComponentVisualizer.h b/Engine/Source/Editor/UnrealEd/Public/ComponentVisualizer.h
--- a/Engine/Source/Editor/UnrealEd/Public/ComponentVisualizer.h
+++ b/Engine/Source/Editor/UnrealEd/Public/ComponentVisualizer.h
@@ -23,6 +23,8 @@
 #include "UObject/UnrealType.h"
 #include "UObject/WeakObjectPtr.h"
 #include "UObject/WeakObjectPtrTemplates.h"
+#include "Elements/Framework/EngineElementsLibrary.h"
+#include "Elements/Framework/TypedElementHandle.h"
 
 #include "ComponentVisualizer.generated.h"
 
@@ -49,6 +51,11 @@ struct HComponentVisProxy : public HHitProxy
 		return EMouseCursor::Crosshairs;
 	}
 
+	virtual FTypedElementHandle GetElementHandle() const override
+	{
+		return UEngineElementsLibrary::AcquireEditorComponentElementHandle(Component.Get());
+	}
+
 	TWeakObjectPtr<const UActorComponent> Component;
 };
 
