[Deprecation Cleanup] Miscellaneous UnrealEd safe deprecation removals

diff --git a/Engine/Source/Editor/UnrealEd/Classes/Editor/EditorEngine.h b/Engine/Source/Editor/UnrealEd/Classes/Editor/EditorEngine.h
--- a/Engine/Source/Editor/UnrealEd/Classes/Editor/EditorEngine.h
+++ b/Engine/Source/Editor/UnrealEd/Classes/Editor/EditorEngine.h
@@ -624,11 +624,6 @@ public:
 	/**	Broadcasts that a blueprint just finished being reinstanced. THIS SHOULD NOT BE PUBLIC */
 	void BroadcastBlueprintReinstanced() { BlueprintReinstanced.Broadcast(); }
 
-	/** Called when UObjects have been replaced to allow others a chance to fix their references. */
-	using FObjectsReplacedEvent = FCoreUObjectDelegates::FOnObjectsReplaced;
-	UE_DEPRECATED(5.0, "Use FCoreUObjectDelegates::OnObjectsReplaced instead.")
-	FObjectsReplacedEvent& OnObjectsReplaced() { return FCoreUObjectDelegates::OnObjectsReplaced; }
-
 	/** Called when a package with data-driven classes becomes loaded or unloaded */
 	DECLARE_EVENT( UEditorEngine, FClassPackageLoadedOrUnloadedEvent );
 	FClassPackageLoadedOrUnloadedEvent& OnClassPackageLoadedOrUnloaded() { return ClassPackageLoadedOrUnloadedEvent; }
@@ -636,11 +631,6 @@ public:
 	/**	Broadcasts that a class package was just loaded or unloaded. THIS SHOULD NOT BE PUBLIC */
 	void BroadcastClassPackageLoadedOrUnloaded() { ClassPackageLoadedOrUnloadedEvent.Broadcast(); }
 
-	/** Called when an object is reimported. */
-	DECLARE_EVENT_OneParam( UEditorEngine, FObjectReimported, UObject* );
-	UE_DEPRECATED(4.22, "Use the ImportSubsystem instead. GEditor->GetEditorSubsystem<UImportSubsystem>()")
-	FObjectReimported& OnObjectReimported() { return ObjectReimportedEvent; }
-
 	/** Editor-only event triggered before an actor or component is moved, rotated or scaled by an editor system */
 	DECLARE_EVENT_OneParam( UEditorEngine, FOnBeginTransformObject, UObject& );
 	FOnBeginTransformObject& OnBeginObjectMovement() { return OnBeginObjectTransformEvent; }
@@ -682,15 +672,6 @@ public:
 	/** Called by internal engine systems after a HLOD Actor is added to a cluster */
 	void BroadcastHLODActorAdded(const AActor* InActor, const AActor* ParentActor) { HLODActorAddedEvent.Broadcast(InActor, ParentActor); }
 
-	/** Editor-only event triggered when a HLOD Actor is marked dirty */
-	DECLARE_EVENT_OneParam(UEngine, FHLODActorMarkedDirtyEvent, class ALODActor*);
-	UE_DEPRECATED(4.20, "This function is no longer used.")
-	FHLODActorMarkedDirtyEvent& OnHLODActorMarkedDirty() { return HLODActorMarkedDirtyEvent; }
-
-	/** Called by internal engine systems after a HLOD Actor is marked dirty */
-	UE_DEPRECATED(4.20, "This function is no longer used.")
-	void BroadcastHLODActorMarkedDirty(class ALODActor* InActor) { HLODActorMarkedDirtyEvent.Broadcast(InActor); }
-
 	/** Editor-only event triggered when a HLOD Actor is marked dirty */
 	DECLARE_EVENT(UEngine, FHLODTransitionScreenSizeChangedEvent);
 	FHLODTransitionScreenSizeChangedEvent& OnHLODTransitionScreenSizeChanged() { return HLODTransitionScreenSizeChangedEvent; }
@@ -1407,9 +1388,6 @@ public:
 	*/
 	static UNREALED_API bool polyFindBrush(UModel* InModel, int32 iSurf, FPoly &Poly);
 
-	UE_DEPRECATED(5.1, "polyFindMaster is deprecated; please use polyFindBrush instead")
-	UNREALED_API virtual bool polyFindMaster( UModel* InModel, int32 iSurf, FPoly& Poly );
-
 	/**
 	 * Update a the brush EdPoly corresponding to a newly-changed
 	 * poly to reflect its new properties.
@@ -1418,9 +1396,6 @@ public:
 	 */
 	static UNREALED_API void polyUpdateBrush(UModel* Model, int32 iSurf, bool bUpdateTexCoords, bool bOnlyRefreshSurfaceMaterials);
 
-	UE_DEPRECATED(5.1, "polyUpdateMaster is deprecated; please use polyUpdateBrush instead")
-	UNREALED_API virtual void polyUpdateMaster( UModel* Model, int32 iSurf, bool bUpdateTexCoords, bool bOnlyRefreshSurfaceMaterials );
-
 	/**
 	 * Populates a list with all polys that are linked to the specified poly.  The
 	 * resulting list includes the original poly.
@@ -1796,24 +1771,6 @@ public:
 	 */
 	UNREALED_API void RemoveViewportsRealtimeOverride(FText SystemDisplayName);
 
-	UE_DEPRECATED(4.26, "To remove realtime overrides, please now provide a system name to make sure you remove the correct override.")
-	UNREALED_API void RemoveViewportsRealtimeOverride();
-
-	/**
-	 * Disables any realtime viewports that are currently viewing the level.  This will not disable
-	 * things like preview viewports in Cascade, etc. Typically called before running the game.
-	 */
-	UE_DEPRECATED(4.25, "To save and restore realtime state non-permanently use SetViewportsRealtimeOverride and RemoveViewportsRealtimeOverride")
-	UNREALED_API void DisableRealtimeViewports();
-
-	/**
-	 * Restores any realtime viewports that have been disabled by DisableRealtimeViewports. This won't
-	 * disable viewporst that were realtime when DisableRealtimeViewports has been called and got
-	 * latter toggled to be realtime.
-	 */
-	UE_DEPRECATED(4.25, "To save and restore realtime state non-permanently use SetViewportsRealtimeOverride and RemoveViewportsRealtimeOverride")
-	UNREALED_API void RestoreRealtimeViewports();
-
 	/**
 	 * Checks to see if any viewport is set to update in realtime.
 	 */
@@ -2372,15 +2329,6 @@ public:
 	 * turns all navigable static geometry of ULevel into polygon soup stored in passed Level (ULevel::StaticNavigableGeometry)
 	 */
 	UNREALED_API virtual void RebuildStaticNavigableGeometry(ULevel* Level);
-
-	/**
-	 * Gets all objects which can be synced to in content browser for current selection
-	 *
-	 * @param Objects	Array to be filled with objects which can be browsed to
-	 * @param bAllowOverrideMetadata If true, allows an asset to define "BrowseToAssetOverride" in its metadata to sync to an asset other than itself
-	 */
-	UE_DEPRECATED(5.1, "Use GetAssetsToSyncToContentBrowser instead")
-	UNREALED_API void GetObjectsToSyncToContentBrowser(TArray<UObject*>& Objects, bool bAllowBrowseToAssetOverride = true);
 	/**
 	 * Gets all assets which can be synced to in content browser for current selection
 	 *
@@ -2749,9 +2697,6 @@ private:
 	/** Delegate broadcast when a package has been loaded or unloaded */
 	FClassPackageLoadedOrUnloadedEvent ClassPackageLoadedOrUnloadedEvent;
 
-	/** Delegate broadcast when an object has been reimported */
-	FObjectReimported ObjectReimportedEvent;
-
 	/** Delegate broadcast when an actor or component is about to be moved, rotated, or scaled  */
 	FOnBeginTransformObject OnBeginObjectTransformEvent;
 
@@ -2776,9 +2721,6 @@ private:
 	/** Broadcasts after an HLOD actor has added to a cluster */
 	FHLODActorAddedEvent HLODActorAddedEvent;
 
-	/** Broadcasts after an HLOD actor has been marked dirty */
-	FHLODActorMarkedDirtyEvent HLODActorMarkedDirtyEvent;
-
 	/** Broadcasts after a Draw distance value (World settings) is changed */
 	FHLODTransitionScreenSizeChangedEvent HLODTransitionScreenSizeChangedEvent;
 
@@ -3191,87 +3133,10 @@ private:
 
 	// DEPRECATED VARIABLES ONLY
 public:
-	UE_DEPRECATED(4.25, "Use the Request Parameters to specify this instead. Can be read from the current session if it was set in the request.")
-	/** An optional location for the starting location for "Play From Here"																*/
-	UPROPERTY()
-	FVector PlayWorldLocation;
-
-	UE_DEPRECATED(4.25, "Use the Request Parameters to specify this instead. Can be read from the current session if it was set in the request.")
-	/** An optional rotation for the starting location for "Play From Here"																*/
-	UPROPERTY()
-	FRotator PlayWorldRotation;
-
-	UE_DEPRECATED(4.25, "Use IsPlaySessionQueued() or IsPlaySessionInProgress() instead.")
-	/** Has a request for "Play From Here" been made?													 								*/
-	UPROPERTY()
-	uint32 bIsPlayWorldQueued:1;
-	
-	UE_DEPRECATED(4.25, "Use IsSimulateInEditorQueued() or IsSimulateInEditorInProgress() instead.")
-	/** True if we are requesting to start a simulation-in-editor session */
-	UPROPERTY()
-	uint32 bIsSimulateInEditorQueued:1;
-	
-	/** Did the request include the optional location and rotation?										 								*/
-	UE_DEPRECATED(4.25, "Use FRequestPlaySessionParams::HasPlayWorldPlacement() on the queued/current session instead.")
-	UPROPERTY()
-	uint32 bHasPlayWorldPlacement:1;
-
-	/** True to enable mobile preview mode when launching the game from the editor on PC platform */
-	UE_DEPRECATED(4.25, "Use FRequestPlaySessionParams::SessionPreviewTypeOverride on the queued/current session instead.")
-	UPROPERTY()
-	uint32 bUseMobilePreviewForPlayWorld:1;
-
-	/** True to enable VR preview mode when launching the game from the editor on PC platform */
-	UE_DEPRECATED(4.25, "Use FRequestPlaySessionParams::SessionPreviewTypeOverride on the queued/current session instead.")
-	UPROPERTY()
-	uint32 bUseVRPreviewForPlayWorld:1;
-
 	/** True if we're Simulating In Editor, as opposed to Playing In Editor.  In this mode, simulation takes place right the level editing environment */
 	// UE_DEPRECATED(4.25, "Use IsSimulateInEditorInProgress instead.")
 	UPROPERTY()
 	uint32 bIsSimulatingInEditor:1;
-	
-	/** Viewport the next PlaySession was requested to happen on */
-	UE_DEPRECATED(4.25, "This is stored as part of the FRequestPlaySessionParams now.")
-	TWeakPtr<class IAssetViewport>		RequestedDestinationSlateViewport;
-
-	/** When set to anything other than -1, indicates a specific In-Editor viewport index that PIE should use */
-	UE_DEPRECATED(4.25, "This isn't read and was replaced by RequestedDestinationSlateViewport.")
-	UPROPERTY()
-	int32 PlayInEditorViewportIndex;
-	
-protected:
-
-	/** Count of how many PIE instances are waiting to log in */
-	UE_DEPRECATED(4.25, "This has moved to the current session instead (stored in FPlayInEditorSessionInfo)")
-	int32 PIEInstancesToLogInCount;
-
-	UE_DEPRECATED(4.25, "This has moved to the current session instead (stored in FPlayInEditorSessionInfo)")
-	bool bAtLeastOnePIELoginFailed;
-
-	/* These are parameters that we need to cache for late joining */
-	UE_DEPRECATED(4.25, "This has moved to the current session instead (stored in FPlayInEditorSessionInfo)")
-	FString ServerPrefix;
-	
-	UE_DEPRECATED(4.25, "This has moved to the current session instead (stored in FPlayInEditorSessionInfo)")
-	int32 PIEInstance;
-	
-	UE_DEPRECATED(4.25, "This has moved to the current session instead (stored in FPlayInEditorSessionInfo)")
-	int32 SettingsIndex;
-	
-	UE_DEPRECATED(4.25, "This has moved to the current session instead (stored in FPlayInEditorSessionInfo)")
-	bool bStartLateJoinersInSpectatorMode;
-
-private:
-
-	/** Additional launch options requested for the next PlaySession */
-	UE_DEPRECATED(4.25, "Use FRequestPlaySessionParams::AdditionalStandaloneCommandLineParameters instead.")
-	FString RequestedAdditionalStandaloneLaunchOptions;
-
-protected:
-	UE_DEPRECATED(4.25, "Use FRequestPlaySessionParams::NumOustandingPIELogins instead.")
-	/** Number of currently running instances logged into an online platform */
-	int32 NumOnlinePIEInstances;
 };
 PRAGMA_ENABLE_DEPRECATION_WARNINGS
 
diff --git a/Engine/Source/Editor/UnrealEd/Classes/Editor/Transactor.h b/Engine/Source/Editor/UnrealEd/Classes/Editor/Transactor.h
--- a/Engine/Source/Editor/UnrealEd/Classes/Editor/Transactor.h
+++ b/Engine/Source/Editor/UnrealEd/Classes/Editor/Transactor.h
@@ -665,10 +665,6 @@ class UTransactor : public UObject
 	 */
 	virtual void SetPrimaryUndoObject( UObject* Object ) PURE_VIRTUAL(UTransactor::SetPrimaryUndoObject,);
 
-	/** Checks if a specific object is referenced by the transaction buffer */
-	UE_DEPRECATED(5.1, "IsObjectInTransationBuffer is deprecated, please use IsObjectInTransactionBuffer instead.")
-	virtual bool IsObjectInTransationBuffer( const UObject* Object ) const { return IsObjectInTransactionBuffer(Object); }
-
 	/** Checks if a specific object is referenced by the transaction buffer */
 	virtual bool IsObjectInTransactionBuffer( const UObject* Object ) const { return false; }
 
diff --git a/Engine/Source/Editor/UnrealEd/Classes/Editor/UnrealEdEngine.h b/Engine/Source/Editor/UnrealEd/Classes/Editor/UnrealEdEngine.h
--- a/Engine/Source/Editor/UnrealEd/Classes/Editor/UnrealEdEngine.h
+++ b/Engine/Source/Editor/UnrealEd/Classes/Editor/UnrealEdEngine.h
@@ -860,8 +860,6 @@ public:
 
 	/** Return if we have write permission under the mount point this packages lives in. */
 	UNREALED_API bool HasMountWritePermissionForPackage(const FString& PackageName);
-	UE_DEPRECATED(5.0, "Use HasMountWritePermissionForPackage instead")
-	UNREALED_API bool HasMountWritePersmissionForPackage(const FString& PackageName);
 
 	/* Delegate to override TemplateMapInfos */
 	DECLARE_DELEGATE_RetVal(const TArray<FTemplateMapInfo>&, FGetTemplateMapInfos);
diff --git a/Engine/Source/Editor/UnrealEd/Classes/Settings/ContentBrowserSettings.h b/Engine/Source/Editor/UnrealEd/Classes/Settings/ContentBrowserSettings.h
--- a/Engine/Source/Editor/UnrealEd/Classes/Settings/ContentBrowserSettings.h
+++ b/Engine/Source/Editor/UnrealEd/Classes/Settings/ContentBrowserSettings.h
@@ -73,48 +73,24 @@ public:
 		DisplayEngineFolder = bInDisplayEngineFolder;
 	}
 
-	UE_DEPRECATED(5.2, "The bOverride flag is no longer necessary. The override can now be set on SAssetView::OverrideShowEngineFolder().")
-	void SetDisplayEngineFolder( bool bInDisplayEngineFolder, bool bOverride )
-	{ 
-		SetDisplayEngineFolder(bInDisplayEngineFolder);
-	}
-
 	/** Gets whether we are allowed to display the engine folder or not, optional flag ignoring the override */
 	bool GetDisplayEngineFolder() const
 	{
 		return DisplayEngineFolder;
 	}
 
-	UE_DEPRECATED(5.2, "The bOverride flag is no longer necessary. The override can now be set on SAssetView::OverrideShowEngineFolder().")
-	bool GetDisplayEngineFolder( bool bExcludeOverride ) const
-	{ 
-		return GetDisplayEngineFolder();
-	}
-
 	/** Sets whether we are allowed to display the developers folder or not, optional flag for setting override instead */
 	void SetDisplayDevelopersFolder( bool bInDisplayDevelopersFolder )
 	{
 		DisplayDevelopersFolder = bInDisplayDevelopersFolder;
 	}
 
-	UE_DEPRECATED(5.2, "The bOverride flag is no longer necessary. The override can now be set on SAssetView::OverrideShowDeveloperFolders().")
-	void SetDisplayDevelopersFolder( bool bInDisplayDevelopersFolder, bool bOverride )
-	{ 
-		SetDisplayDevelopersFolder(bInDisplayDevelopersFolder);
-	}
-
 	/** Gets whether we are allowed to display the developers folder or not, optional flag ignoring the override */
 	bool GetDisplayDevelopersFolder() const
 	{
 		return DisplayDevelopersFolder;
 	}
 
-	UE_DEPRECATED(5.2, "The bOverride flag is no longer necessary. The override can now be set on SAssetView::OverrideShowDeveloperFolders().")
-	bool GetDisplayDevelopersFolder( bool bExcludeOverride ) const
-	{ 
-		return GetDisplayDevelopersFolder();
-	}
-
 	/** Sets whether we are allowed to display the L10N folder (contains localized assets) or not */
 	void SetDisplayL10NFolder(bool bInDisplayL10NFolder)
 	{
@@ -133,24 +109,12 @@ public:
 		DisplayPluginFolders = bInDisplayPluginFolders;
 	}
 
-	UE_DEPRECATED(5.2, "The bOverride flag is no longer necessary. The override can now be set on SAssetView::OverrideShowPluginFolders().")
-	void SetDisplayPluginFolders( bool bInDisplayPluginFolders, bool bOverride )
-	{ 
-		SetDisplayPluginFolders(bInDisplayPluginFolders);
-	}
-
 	/** Gets whether we are allowed to display the plugin folders or not */
 	bool GetDisplayPluginFolders() const
 	{
 		return DisplayPluginFolders;
 	}
 
-	UE_DEPRECATED(5.2, "The bOverride flag is no longer necessary. The override can now be set on SAssetView::OverrideShowPluginFolders().")
-	bool GetDisplayPluginFolders( bool bExcludeOverride ) const
-	{ 
-		return GetDisplayPluginFolders();
-	}
-
 	/** Sets whether we are allowed to display favorite folders or not */
 	void SetDisplayFavorites(bool bInDisplayFavorites)
 	{
diff --git a/Engine/Source/Editor/UnrealEd/Classes/Settings/LevelEditorPlaySettings.h b/Engine/Source/Editor/UnrealEd/Classes/Settings/LevelEditorPlaySettings.h
--- a/Engine/Source/Editor/UnrealEd/Classes/Settings/LevelEditorPlaySettings.h
+++ b/Engine/Source/Editor/UnrealEd/Classes/Settings/LevelEditorPlaySettings.h
@@ -340,10 +340,6 @@ public:
 	/** Whether to automatically recompile dirty Blueprints before launching */
 	UPROPERTY(config, EditAnywhere, Category=PlayOnDevice)
 	bool bAutoCompileBlueprintsOnLaunch;
-
-	/** A programmatically defined custom PIE window to use */
-	UE_DEPRECATED(4.25, "This variable is no longer read. Specify the Custom Window in the FRequestPlaySessionParams instead during PIE request.")
-	TWeakPtr<SWindow> CustomPIEWindow;
 	
 	/**
 	* This is a rarely used option that will launch a separate server (possibly hidden in-process depending on RunUnderOneProcess) 
@@ -364,11 +360,6 @@ private:
 	UPROPERTY(config, EditAnywhere, Category="Multiplayer Options")
 	bool RunUnderOneProcess;
 
-	/** If checked, a separate dedicated server will be launched. Otherwise the first player will act as a listen server that all other players connect to. */
-	UE_DEPRECATED(4.25, "This variable is no longer read. Use PlayNetMode = EPlayNetMode::PIE_Client and bLaunchSeparateServer instead.")
-	UPROPERTY(config)
-	bool PlayNetDedicated;
-
 	/** The number of client windows to open. The first one to open will respect the Play In Editor "Modes" option (PIE, PINW), additional clients respect the RunUnderOneProcess setting. */
 	UPROPERTY(config, EditAnywhere, Category="Multiplayer Options|Client", meta=(ClampMin = "1", UIMin = "1", UIMax = "64"))
 	int32 PlayNumberOfClients;
@@ -384,15 +375,6 @@ private:
 	/** Width to use when spawning additional windows. */
 	UPROPERTY(config, EditAnywhere, Category="Multiplayer Options|Client", meta=(ClampMin=0))
 	int32 ClientWindowWidth;
-	
-	/**
-	 * When running multiple players or a dedicated server the client need to connect to the server, this option sets how they connect
-	 *
-	 * If this is checked, the clients will automatically connect to the launched server, if false they will launch into the map and wait
-	 */
-	UE_DEPRECATED(4.25, "This variable is no longer read. Use PlayNetMode = EPlayNetMode::PIE_Standalone instead to prevent auto-connection.")
-	UPROPERTY(config)
-	bool AutoConnectToServer;
 
 	/**
 	 * When running multiple player windows in a single process, this option determines how the game pad input gets routed.
@@ -426,11 +408,6 @@ private:
 	UPROPERTY(config, EditAnywhere, Category="Multiplayer Options|Server", meta=(EditCondition = "PlayNetMode != EPlayNetMode::PIE_Standalone || bLaunchSeparateServer"))
 	FString AdditionalServerGameOptions;
 
-	/** Additional command line options that will be passed to standalone game instances, for example -debug */
-	UE_DEPRECATED(4.25, "This variable is no longer read. Use AdditionalServerLaunchParmeters instead to pass specific flags to externally launched servers.")
-	UPROPERTY(config)
-	FString AdditionalLaunchOptions;
-
 	/** Controls the default value of the show flag ServerDrawDebug */
 	UPROPERTY(config, EditAnywhere, Category = "Multiplayer Options")
 	bool bShowServerDebugDrawingByDefault;
@@ -490,14 +467,6 @@ public:
 	bool IsRunUnderOneProcessActive() const { return true; }
 	bool GetRunUnderOneProcess( bool &OutRunUnderOneProcess ) const { OutRunUnderOneProcess = RunUnderOneProcess; return IsRunUnderOneProcessActive(); }
 	
-	PRAGMA_DISABLE_DEPRECATION_WARNINGS
-	UE_DEPRECATED(4.25, "Read bLaunchSeparateServer directly instead.")
-	void SetPlayNetDedicated(const bool InPlayNetDedicated) { PlayNetDedicated = InPlayNetDedicated; }
-	UE_DEPRECATED(4.25, "Read bLaunchSeparateServer directly instead.")
-	bool IsPlayNetDedicatedActive() const { return (RunUnderOneProcess ? true : PlayNetMode == PIE_Client); }
-	UE_DEPRECATED(4.25, "Read bLaunchSeparateServer directly instead.")
-	bool GetPlayNetDedicated(bool &OutPlayNetDedicated) const { OutPlayNetDedicated = PlayNetDedicated; return IsPlayNetDedicatedActive(); }
-	
 	void SetPlayNumberOfClients( const int32 InPlayNumberOfClients ) { PlayNumberOfClients = InPlayNumberOfClients; }
 	bool IsPlayNumberOfClientsActive() const { return true; }
 	bool GetPlayNumberOfClients( int32 &OutPlayNumberOfClients ) const { OutPlayNumberOfClients = PlayNumberOfClients; return IsPlayNumberOfClientsActive(); }
@@ -508,18 +477,11 @@ public:
 	bool IsServerPortActive() const { return (PlayNetMode != PIE_Standalone) || RunUnderOneProcess; }
 	bool GetServerPort(uint16 &OutServerPort) const { OutServerPort = ServerPort; return IsServerPortActive(); }
 	
-	UE_DEPRECATED(4.25, "This feature has been removed. Set PlayNetMode == EPlayNetMode::PIE_Standalone instead.")
-	bool IsAutoConnectToServerActive() const { return PlayNumberOfClients > 1 || PlayNetDedicated; }
-	UE_DEPRECATED(4.25, "This feature has been removed. Set PlayNetMode == EPlayNetMode::PIE_Standalone instead.")
-	bool GetAutoConnectToServer(bool &OutAutoConnectToServer) const { OutAutoConnectToServer = AutoConnectToServer; return IsAutoConnectToServerActive(); }
-	UE_DEPRECATED(4.25, "This feature has been removed. Set PlayNetMode == EPlayNetMode::PIE_Standalone instead.")
-	EVisibility GetAutoConnectToServerVisibility() const { return (RunUnderOneProcess ? EVisibility::Visible : EVisibility::Hidden); }
-	
 	bool IsRouteGamepadToSecondWindowActive() const { return PlayNumberOfClients > 1; }
 	bool GetRouteGamepadToSecondWindow( bool &OutRouteGamepadToSecondWindow ) const { OutRouteGamepadToSecondWindow = RouteGamepadToSecondWindow; return IsRouteGamepadToSecondWindowActive(); }
 	EVisibility GetRouteGamepadToSecondWindowVisibility() const { return (RunUnderOneProcess ? EVisibility::Visible : EVisibility::Hidden); }
 
-	EVisibility GetNetworkEmulationVisibility() const { return (PlayNumberOfClients > 1 || PlayNetDedicated) ? EVisibility::Visible : EVisibility::Hidden; }
+	EVisibility GetNetworkEmulationVisibility() const { return (PlayNumberOfClients > 1) ? EVisibility::Visible : EVisibility::Hidden; }
 
 	bool IsServerMapNameOverrideActive() const { return false /*(PlayNetMode == PIE_StandaloneWithServer)*/; }
 	bool GetServerMapNameOverride( FString& OutStandaloneServerMapName ) const { OutStandaloneServerMapName = ServerMapNameOverride; return IsServerMapNameOverrideActive(); }
@@ -528,14 +490,6 @@ public:
 	bool IsAdditionalServerGameOptionsActive() const { return (PlayNetMode != PIE_Standalone) || RunUnderOneProcess; }
 	bool GetAdditionalServerGameOptions( FString &OutAdditionalServerGameOptions ) const { OutAdditionalServerGameOptions = AdditionalServerGameOptions; return IsAdditionalServerGameOptionsActive(); }
 
-	UE_DEPRECATED(4.25, "Read AdditionalServerLaunchParmeters for a non-zero length instead.")
-	bool IsAdditionalLaunchOptionsActive() const { return true; }
-	UE_DEPRECATED(4.25, "Read AdditionalServerLaunchParmeters directly instead.")
-	bool GetAdditionalLaunchOptions(FString &OutAdditionalLaunchOptions) const { OutAdditionalLaunchOptions = AdditionalLaunchOptions; return IsAdditionalLaunchOptionsActive(); }
-	UE_DEPRECATED(4.25, "Read AdditionalServerLaunchParmeters for a non-zero length instead.")
-	EVisibility GetAdditionalLaunchOptionsVisibility() const { return (RunUnderOneProcess ? EVisibility::Hidden : EVisibility::Visible); }
-	PRAGMA_ENABLE_DEPRECATION_WARNINGS
-
 	void SetClientWindowSize( const FIntPoint InClientWindowSize ) { ClientWindowWidth = InClientWindowSize.X; ClientWindowHeight = InClientWindowSize.Y; }
 	bool IsClientWindowSizeActive() const { return PlayNumberOfClients > 1; }
 	bool GetClientWindowSize( FIntPoint &OutClientWindowSize ) const { OutClientWindowSize = FIntPoint(ClientWindowWidth, ClientWindowHeight); return IsClientWindowSizeActive(); }
diff --git a/Engine/Source/Editor/UnrealEd/Private/EditorCsg.cpp b/Engine/Source/Editor/UnrealEd/Private/EditorCsg.cpp
--- a/Engine/Source/Editor/UnrealEd/Private/EditorCsg.cpp
+++ b/Engine/Source/Editor/UnrealEd/Private/EditorCsg.cpp
@@ -326,23 +326,11 @@ void UEditorEngine::polySetAndClearPolyFlags(UModel *Model, uint32 SetBits, uint
 	FBSPUtils::polySetAndClearPolyFlags(Model, SetBits, ClearBits, SelectedOnly, UpdateBrush);
 }
 
-
-bool UEditorEngine::polyFindMaster(UModel* InModel, int32 iSurf, FPoly &Poly)
-{
-	return polyFindBrush(InModel, iSurf, Poly);
-}
-
-bool UEditorEngine::polyFindBrush(UModel* InModel, int32 iSurf, FPoly &Poly)
+bool UEditorEngine::polyFindBrush(UModel* InModel, int32 iSurf, FPoly& Poly)
 {
 	return FBSPUtils::polyFindBrush(InModel, iSurf, Poly);
 }
 
-
-void UEditorEngine::polyUpdateMaster(UModel* Model, int32 iSurf, bool bUpdateTexCoords, bool bOnlyRefreshSurfaceMaterials)
-{
-	polyUpdateBrush(Model, iSurf, bUpdateTexCoords, bOnlyRefreshSurfaceMaterials);
-}
-
 void UEditorEngine::polyUpdateBrush
 (
 	UModel*	Model,
diff --git a/Engine/Source/Editor/UnrealEd/Private/EditorEngine.cpp b/Engine/Source/Editor/UnrealEd/Private/EditorEngine.cpp
--- a/Engine/Source/Editor/UnrealEd/Private/EditorEngine.cpp
+++ b/Engine/Source/Editor/UnrealEd/Private/EditorEngine.cpp
@@ -1242,7 +1242,6 @@ void UEditorEngine::RemoveLevelViewportClients(FLevelEditorViewportClient* Viewp
 
 void UEditorEngine::BroadcastObjectReimported(UObject* InObject)
 {
-	ObjectReimportedEvent.Broadcast(InObject);
 	GetEditorSubsystem<UImportSubsystem>()->BroadcastAssetReimport(InObject);
 }
 
@@ -3072,17 +3071,6 @@ bool UEditorEngine::CanSyncToContentBrowser()
 	return Assets.Num() > 0;
 }
 
-
-void UEditorEngine::GetObjectsToSyncToContentBrowser(TArray<UObject*>& Objects, bool bAllowBrowseToAssetOverride)
-{
-	TArray<FAssetData> Assets;
-	GetAssetsToSyncToContentBrowser(Assets, bAllowBrowseToAssetOverride);
-	for (const FAssetData& Asset : Assets)
-	{
-		Objects.Add(Asset.GetAsset());
-	}
-}
-
 void UEditorEngine::GetAssetsToSyncToContentBrowser(TArray<FAssetData>& Assets, bool bAllowBrowseToAssetOverride)
 {
 	// If the user has any BSP surfaces selected, sync to the materials on them.
@@ -6338,56 +6326,6 @@ void UEditorEngine::RemoveViewportsRealtimeOverride(FText SystemDisplayName)
 	FEditorSupportDelegates::UpdateUI.Broadcast();
 }
 
-void UEditorEngine::RemoveViewportsRealtimeOverride()
-{
-	for (FEditorViewportClient* VC : AllViewportClients)
-	{
-		if (VC)
-		{
-			VC->PopRealtimeOverride();
-		}
-	}
-
-	RedrawAllViewports();
-
-	FEditorSupportDelegates::UpdateUI.Broadcast();
-}
-
-void UEditorEngine::DisableRealtimeViewports()
-{
-PRAGMA_DISABLE_DEPRECATION_WARNINGS
-	for(FEditorViewportClient* VC : AllViewportClients)
-	{
-		if( VC )
-		{
-			VC->SetRealtime( false, true );
-		}
-	}
-
-	RedrawAllViewports();
-
-	FEditorSupportDelegates::UpdateUI.Broadcast();
-PRAGMA_ENABLE_DEPRECATION_WARNINGS
-}
-
-
-void UEditorEngine::RestoreRealtimeViewports()
-{
-PRAGMA_DISABLE_DEPRECATION_WARNINGS
-	for(FEditorViewportClient* VC : AllViewportClients)
-	{
-		if( VC )
-		{
-			VC->RestoreRealtime(true);
-		}
-	}
-
-	RedrawAllViewports();
-
-	FEditorSupportDelegates::UpdateUI.Broadcast();
-PRAGMA_ENABLE_DEPRECATION_WARNINGS
-}
-
 
 bool UEditorEngine::IsAnyViewportRealtime()
 {
diff --git a/Engine/Source/Editor/UnrealEd/Private/Settings/SettingsClasses.cpp b/Engine/Source/Editor/UnrealEd/Private/Settings/SettingsClasses.cpp
--- a/Engine/Source/Editor/UnrealEd/Private/Settings/SettingsClasses.cpp
+++ b/Engine/Source/Editor/UnrealEd/Private/Settings/SettingsClasses.cpp
@@ -501,10 +501,6 @@ ULevelEditorPlaySettings::ULevelEditorPlaySettings( const FObjectInitializer& Ob
 	bLaunchSeparateServer = false;
 	PlayNumberOfClients = 1;
 	ServerPort = 17777;
-PRAGMA_DISABLE_DEPRECATION_WARNINGS
-	PlayNetDedicated = false;
-	AutoConnectToServer = true;
-PRAGMA_ENABLE_DEPRECATION_WARNINGS
 	RunUnderOneProcess = true;
 	RouteGamepadToSecondWindow = false;
 	BuildGameBeforeLaunch = EPlayOnBuildMode::PlayOnBuild_Default;
