Skeletal Mesh: Remove the FBX import flags from FSkeletalMeshImportData. There's no need for them, since they can be maintained purely within the import cycle.

#tests Editor.Import.Fbx.*

diff --git a/Engine/Plugins/Animation/RigLogic/Source/RigLogicModule/Private/SkelMeshDNAUtils.cpp b/Engine/Plugins/Animation/RigLogic/Source/RigLogicModule/Private/SkelMeshDNAUtils.cpp
--- a/Engine/Plugins/Animation/RigLogic/Source/RigLogicModule/Private/SkelMeshDNAUtils.cpp
+++ b/Engine/Plugins/Animation/RigLogic/Source/RigLogicModule/Private/SkelMeshDNAUtils.cpp
@@ -587,8 +587,6 @@ void USkelMeshDNAUtils::UpdateSourceData(USkeletalMesh* InSkelMesh)
 				// Add Morph target name.
 				ImportData.MorphTargetNames.Add(MorphTarget->GetName());
 				FSkeletalMeshImportData MorphTargetImportDeltas;
-				MorphTargetImportDeltas.bDiffPose = ImportData.bDiffPose;
-				MorphTargetImportDeltas.bUseT0AsRefPose = ImportData.bUseT0AsRefPose;
 
 				FMorphTargetLODModel& MorphLODModel = MorphTarget->GetMorphLODModels()[LODIndex];
 
diff --git a/Engine/Source/Developer/MeshUtilities/Private/MeshUtilities.cpp b/Engine/Source/Developer/MeshUtilities/Private/MeshUtilities.cpp
--- a/Engine/Source/Developer/MeshUtilities/Private/MeshUtilities.cpp
+++ b/Engine/Source/Developer/MeshUtilities/Private/MeshUtilities.cpp
@@ -6015,12 +6015,9 @@ void FMeshUtilities::CreateImportDataFromLODModel(USkeletalMesh* SkeletalMesh) c
 		}
 
 		FSkeletalMeshImportData ImportData;
-		ImportData.bDiffPose = false;
 		ImportData.bHasNormals = true;
 		ImportData.bHasTangents = true;
-		ImportData.bUseT0AsRefPose = false;
 		ImportData.bHasVertexColors = SkeletalMesh->GetHasVertexColors();
-		ImportData.bKeepSectionsSeparate = false;
 
 		TArray<FSkeletalMaterial>& SKMaterials = SkeletalMesh->GetMaterials();
 		ImportData.Materials.Reserve(SkeletalMesh->GetMaterials().Num());
diff --git a/Engine/Source/Editor/UnrealEd/Private/Fbx/FbxFactory.cpp b/Engine/Source/Editor/UnrealEd/Private/Fbx/FbxFactory.cpp
--- a/Engine/Source/Editor/UnrealEd/Private/Fbx/FbxFactory.cpp
+++ b/Engine/Source/Editor/UnrealEd/Private/Fbx/FbxFactory.cpp
@@ -706,6 +706,7 @@ UObject* UFbxFactory::FactoryCreateFile
 								}
 							}
 							FSkeletalMeshImportData OutData;
+							bool bMapMorphTargetToTimeZero = false;
 							if (LODIndex == 0 && SkelMeshNodeArray.Num() != 0)
 							{
 								FName OutputName = NAME_None;
@@ -758,6 +759,7 @@ UObject* UFbxFactory::FactoryCreateFile
 										// Reapply the transforms for the rest of the import
 										FbxImporter->ApplyTransformSettingsToFbxNode(RootNodeToImport, ImportUI->SkeletalMeshImportData);
 									}
+									bMapMorphTargetToTimeZero = ImportSkeletalMeshArgs.bMapMorphTargetToTimeZero;
 									ImportedSuccessfulLodIndex = SuccessfulLodIndex;
 									//Increment the LOD index
 									SuccessfulLodIndex++;
@@ -797,6 +799,7 @@ UObject* UFbxFactory::FactoryCreateFile
 									FSkeletalMeshLODInfo* LODInfo = BaseSkeletalMesh->GetLODInfo(SuccessfulLodIndex);
 									LODInfo->bImportWithBaseMesh = true;
 									LODInfo->SourceImportFilename = FString(TEXT(""));
+									bMapMorphTargetToTimeZero = ImportSkeletalMeshArgs.bMapMorphTargetToTimeZero;
 									ImportedSuccessfulLodIndex = SuccessfulLodIndex;
 									SuccessfulLodIndex++;
 								}
@@ -817,7 +820,7 @@ UObject* UFbxFactory::FactoryCreateFile
 								uint32 bImportTextures = ImportOptions->bImportTextures;
 								ImportOptions->bImportTextures = 0;
 
-								FbxImporter->ImportFbxMorphTarget(SkelMeshNodeArray, BaseSkeletalMesh, ImportedSuccessfulLodIndex, OutData);
+								FbxImporter->ImportFbxMorphTarget(SkelMeshNodeArray, BaseSkeletalMesh, ImportedSuccessfulLodIndex, OutData, bMapMorphTargetToTimeZero);
 								bOperationCanceled |= FbxImporter->GetImportOperationCancelled();
 							
 								ImportOptions->bImportMaterials = !!bImportMaterials;
diff --git a/Engine/Source/Editor/UnrealEd/Private/Fbx/FbxSceneImportFactory.cpp b/Engine/Source/Editor/UnrealEd/Private/Fbx/FbxSceneImportFactory.cpp
--- a/Engine/Source/Editor/UnrealEd/Private/Fbx/FbxSceneImportFactory.cpp
+++ b/Engine/Source/Editor/UnrealEd/Private/Fbx/FbxSceneImportFactory.cpp
@@ -1951,6 +1951,7 @@ UObject* UFbxSceneImportFactory::ImportOneSkeletalMesh(void* VoidRootNodeToImpor
 			}
 		}
 		FSkeletalMeshImportData OutData;
+		bool bMapMorphTargetToTimeZero = false;
 		if (LODIndex == 0 && SkelMeshNodeArray.Num() != 0)
 		{
 			FName OutputName = FbxImporter->MakeNameForMesh(SkelMeshNodeArray[0]->GetName(), SkelMeshNodeArray[0]);
@@ -1996,6 +1997,8 @@ UObject* UFbxSceneImportFactory::ImportOneSkeletalMesh(void* VoidRootNodeToImpor
 				//Set the data in the node info
 				RootNodeInfo->AttributeInfo->SetOriginalImportPath(PackageName);
 				RootNodeInfo->AttributeInfo->SetOriginalFullImportName(NewObject->GetPathName());
+
+				bMapMorphTargetToTimeZero = ImportSkeletalMeshArgs.bMapMorphTargetToTimeZero;
 			}
 		}
 		else if (NewObject && GlobalImportSettings->bImportSkeletalMeshLODs) // the base skeletal mesh is imported successfully
@@ -2032,6 +2035,10 @@ UObject* UFbxSceneImportFactory::ImportOneSkeletalMesh(void* VoidRootNodeToImpor
 				{
 					FbxImporter->AddTokenizedErrorMessage(FTokenizedMessage::Create(EMessageSeverity::Error, LOCTEXT("FailedToImport_SkeletalMeshLOD", "Failed to import Skeletal mesh LOD.")), FFbxErrors::SkeletalMesh_LOD_FailedToImport);
 				}
+				else
+				{
+					bMapMorphTargetToTimeZero = ImportSkeletalMeshArgs.bMapMorphTargetToTimeZero;
+				}
 			}
 		}
 
@@ -2049,7 +2056,7 @@ UObject* UFbxSceneImportFactory::ImportOneSkeletalMesh(void* VoidRootNodeToImpor
 				NewSkelMesh->GetImportedModel()->LODModels.IsValidIndex(LODIndex))
 			{
 				// TODO: Disable material importing when importing morph targets
-				FbxImporter->ImportFbxMorphTarget(SkelMeshNodeArray, NewSkelMesh, LODIndex, OutData);
+				FbxImporter->ImportFbxMorphTarget(SkelMeshNodeArray, NewSkelMesh, LODIndex, OutData, bMapMorphTargetToTimeZero);
 			}
 		}
 	}
diff --git a/Engine/Source/Editor/UnrealEd/Private/Fbx/FbxSkeletalMeshImport.cpp b/Engine/Source/Editor/UnrealEd/Private/Fbx/FbxSkeletalMeshImport.cpp
--- a/Engine/Source/Editor/UnrealEd/Private/Fbx/FbxSkeletalMeshImport.cpp
+++ b/Engine/Source/Editor/UnrealEd/Private/Fbx/FbxSkeletalMeshImport.cpp
@@ -1329,7 +1329,7 @@ bool UnFbx::FFbxImporter::ImportBones(TArray<FbxNode*>& NodeArray, FSkeletalMesh
 	return true;
 }
 
-bool UnFbx::FFbxImporter::FillSkeletalMeshImportData(TArray<FbxNode*>& NodeArray, UFbxSkeletalMeshImportData* TemplateImportData, TArray<FbxShape*> *FbxShapeArray, FSkeletalMeshImportData* OutData, TArray<FbxNode*>& OutImportedSkeletonLinkNodes, TArray<FName> &LastImportedMaterialNames, const bool bIsReimport, const TMap<FVector3f, FColor>& ExistingVertexColorData)
+bool UnFbx::FFbxImporter::FillSkeletalMeshImportData(TArray<FbxNode*>& NodeArray, UFbxSkeletalMeshImportData* TemplateImportData, TArray<FbxShape*> *FbxShapeArray, FSkeletalMeshImportData* OutData, TArray<FbxNode*>& OutImportedSkeletonLinkNodes, TArray<FName> &LastImportedMaterialNames, const bool bIsReimport, const TMap<FVector3f, FColor>& ExistingVertexColorData, bool& bMapMorphTargetToTimeZero)
 {
 	if (NodeArray.Num() == 0)
 	{
@@ -1352,16 +1352,20 @@ bool UnFbx::FFbxImporter::FillSkeletalMeshImportData(TArray<FbxNode*>& NodeArray
 	TArray<FbxNode*>& SortedLinkArray = OutImportedSkeletonLinkNodes;
 	FbxArray<FbxAMatrix> GlobalsPerLink;
 
-	SkelMeshImportDataPtr->bUseT0AsRefPose = ImportOptions->bUseT0AsRefPose;
+	bool bDiffPose = false;
+	bool bUseT0AsRefPose = ImportOptions->bUseT0AsRefPose;
+
 	// Note: importing morph data causes additional passes through this function, so disable the warning dialogs
 	// from popping up again on each additional pass.
-	if (!ImportBones(NodeArray, *SkelMeshImportDataPtr, TemplateImportData, SortedLinkArray, SkelMeshImportDataPtr->bDiffPose, (FbxShapeArray != nullptr), SkelMeshImportDataPtr->bUseT0AsRefPose, Node, bIsReimport))
+	if (!ImportBones(NodeArray, *SkelMeshImportDataPtr, TemplateImportData, SortedLinkArray, bDiffPose, (FbxShapeArray != nullptr), bUseT0AsRefPose, Node, bIsReimport))
 	{
 		if (!(bIsReimport && ImportOptions->bImportAsSkeletalGeometry)) //Do not import bone if we import only the geometry and we are reimporting
 		{
 			return false;
 		}
 	}
+	
+	bMapMorphTargetToTimeZero = bUseT0AsRefPose & bDiffPose;
 
  	FbxNode* SceneRootNode = Scene->GetRootNode();
  	if(SceneRootNode && TemplateImportData)
@@ -1428,7 +1432,7 @@ bool UnFbx::FFbxImporter::FillSkeletalMeshImportData(TArray<FbxNode*>& NodeArray
 			return false;
 		}
 
-		if (SkelMeshImportDataPtr->bUseT0AsRefPose && SkelMeshImportDataPtr->bDiffPose && !ImportOptions->bImportScene)
+		if (bMapMorphTargetToTimeZero && !ImportOptions->bImportScene)
 		{
 			// deform skin vertex to the frame 0 from bind pose
 			SkinControlPointsToPose(*SkelMeshImportDataPtr, FbxMesh, FbxShape, true);
@@ -1542,7 +1546,7 @@ bool UnFbx::FFbxImporter::FillSkeletalMeshImportPoints(FSkeletalMeshImportData*
 	return true;
 }
 
-bool UnFbx::FFbxImporter::GatherPointsForMorphTarget(FSkeletalMeshImportData* OutData, TArray<FbxNode*>& NodeArray, TArray< FbxShape* >* FbxShapeArray, TSet<uint32>& ModifiedPoints)
+bool UnFbx::FFbxImporter::GatherPointsForMorphTarget(FSkeletalMeshImportData* OutData, TArray<FbxNode*>& NodeArray, TArray< FbxShape* >* FbxShapeArray, TSet<uint32>& ModifiedPoints, const bool bMapMorphTargetToTimeZero)
 {
 	check(OutData);
 	TArray<FVector3f> CompressPoints;
@@ -1565,7 +1569,7 @@ bool UnFbx::FFbxImporter::GatherPointsForMorphTarget(FSkeletalMeshImportData* Ou
 
 		FillSkeletalMeshImportPoints( &NewImportData, RootNode, Node, FbxShape );
 
-		if (OutData->bUseT0AsRefPose && OutData->bDiffPose && !ImportOptions->bImportScene)
+		if (bMapMorphTargetToTimeZero && !ImportOptions->bImportScene)
 		{
 			// deform skin vertex to the frame 0 from bind pose
 			SkinControlPointsToPose(NewImportData, FbxMesh, FbxShape, true);
@@ -1738,7 +1742,7 @@ USkeletalMesh* UnFbx::FFbxImporter::ImportSkeletalMesh(FImportSkeletalMeshArgs &
 	// We must do a maximum of fail test before backing up the data since the backup is destructive on the existing skeletal mesh.
 	// See the comment later when we call the following function (SaveExistingSkelMeshData)
 	TArray<FbxNode*> ImportedSkeletonLinkNodes;
-	if (FillSkeletalMeshImportData(ImportSkeletalMeshArgs.NodeArray, ImportSkeletalMeshArgs.TemplateImportData, ImportSkeletalMeshArgs.FbxShapeArray, SkelMeshImportDataPtr, ImportedSkeletonLinkNodes, LastImportedMaterialNames, ExistingSkelMesh != nullptr, ExistingVertexColorData) == false)
+	if (FillSkeletalMeshImportData(ImportSkeletalMeshArgs.NodeArray, ImportSkeletalMeshArgs.TemplateImportData, ImportSkeletalMeshArgs.FbxShapeArray, SkelMeshImportDataPtr, ImportedSkeletonLinkNodes, LastImportedMaterialNames, ExistingSkelMesh != nullptr, ExistingVertexColorData, ImportSkeletalMeshArgs.bMapMorphTargetToTimeZero) == false)
 	{
 		AddTokenizedErrorMessage(FTokenizedMessage::Create(EMessageSeverity::Warning, LOCTEXT("FbxSkeletaLMeshimport_FillupImportData", "Get Import Data has failed.")), FFbxErrors::SkeletalMesh_FillImportDataFailed);
 		EARLY_RETURN_ON_CANCEL(true, FailureCleanup);
@@ -2659,6 +2663,7 @@ USkeletalMesh* UnFbx::FFbxImporter::ReimportSkeletalMesh(USkeletalMesh* Mesh, UF
 				}
 			}
 			FSkeletalMeshImportData OutData;
+			bool bMapMorphTargetToTimeZero = false;
 			if (LODIndex == 0)
 			{
 				ImportMeshLodData.AddZeroed();
@@ -2682,6 +2687,7 @@ USkeletalMesh* UnFbx::FFbxImporter::ReimportSkeletalMesh(USkeletalMesh* Mesh, UF
 				}
 				if (NewMesh)
 				{
+					bMapMorphTargetToTimeZero = ImportSkeletalMeshArgs.bMapMorphTargetToTimeZero;
 					ImportedSuccessfulLodIndex = SuccessfulLodIndex;
 					SuccessfulLodIndex++;
 				}
@@ -2731,6 +2737,7 @@ USkeletalMesh* UnFbx::FFbxImporter::ReimportSkeletalMesh(USkeletalMesh* Mesh, UF
 					FSkeletalMeshLODInfo* LODInfo = BaseSkeletalMesh->GetLODInfo(SuccessfulLodIndex);
 					LODInfo->bImportWithBaseMesh = true;
 					LODInfo->SourceImportFilename = FString(TEXT(""));
+					bMapMorphTargetToTimeZero = ImportSkeletalMeshArgs.bMapMorphTargetToTimeZero;
 					ImportedSuccessfulLodIndex = SuccessfulLodIndex;
 					SuccessfulLodIndex++;
 					for (FName MaterialImportNameLOD : ImportMaterialOriginalNameDataLOD)
@@ -2766,7 +2773,7 @@ USkeletalMesh* UnFbx::FFbxImporter::ReimportSkeletalMesh(USkeletalMesh* Mesh, UF
 					NewMesh->GetImportedModel() &&
 					NewMesh->GetImportedModel()->LODModels.IsValidIndex(ImportedSuccessfulLodIndex))
 				{
-					ImportFbxMorphTarget(SkelMeshNodeArray, NewMesh, ImportedSuccessfulLodIndex, OutData);
+					ImportFbxMorphTarget(SkelMeshNodeArray, NewMesh, ImportedSuccessfulLodIndex, OutData, bMapMorphTargetToTimeZero);
 				}
 			}
 		}
@@ -3231,8 +3238,6 @@ bool UnFbx::FFbxImporter::FillSkelMeshImporterFromFbx( FSkeletalMeshImportData&
 		ImportTexturesFromNode(Node);
 	}
 
-	ImportData.bKeepSectionsSeparate = ImportOptions->bKeepSectionsSeparate;
-
 	// Maps local mesh material index to global material index
 	const int32 MaterialCount = Node->GetMaterialCount();
 	TArray<int32> MaterialMapping;
@@ -4597,7 +4602,7 @@ bool UnFbx::FFbxImporter::ImportSkeletalMeshLOD(USkeletalMesh* InSkeletalMesh, U
 }
 
 
-void UnFbx::FFbxImporter::ImportMorphTargetsInternal( TArray<FbxNode*>& SkelMeshNodeArray, USkeletalMesh* BaseSkelMesh, int32 LODIndex, FSkeletalMeshImportData &BaseImportData)
+void UnFbx::FFbxImporter::ImportMorphTargetsInternal( TArray<FbxNode*>& SkelMeshNodeArray, USkeletalMesh* BaseSkelMesh, int32 LODIndex, FSkeletalMeshImportData &BaseImportData, const bool bMapMorphTargetToTimeZero)
 {
 	FbxString ShapeNodeName;
 	TMap<FString, TArray<FbxShape*>> ShapeNameToShapeArray;
@@ -4693,7 +4698,7 @@ void UnFbx::FFbxImporter::ImportMorphTargetsInternal( TArray<FbxNode*>& SkelMesh
 		}
 		BaseImportData.MorphTargetNames.Add(ShapeName);
 		TSet<uint32>& ModifiedPoints = BaseImportData.MorphTargetModifiedPoints.AddDefaulted_GetRef();
-		GatherPointsForMorphTarget(&ShapeImportData, SkelMeshNodeArray, &ShapeArray, ModifiedPoints);
+		GatherPointsForMorphTarget(&ShapeImportData, SkelMeshNodeArray, &ShapeArray, ModifiedPoints, bMapMorphTargetToTimeZero);
 		//We do not need this data anymore empty it so we reduce the size of what we save into memory
 		ShapeImportData.PointToRawMap.Empty();
 		BaseImportData.MorphTargets.Add(ShapeImportData);
@@ -4757,7 +4762,7 @@ void UnFbx::FFbxImporter::ImportMorphTargetsInternal( TArray<FbxNode*>& SkelMesh
 }
 
 // Import Morph target
-void UnFbx::FFbxImporter::ImportFbxMorphTarget(TArray<FbxNode*> &SkelMeshNodeArray, USkeletalMesh* BaseSkelMesh, int32 LODIndex, FSkeletalMeshImportData &BaseSkeletalMeshImportData)
+void UnFbx::FFbxImporter::ImportFbxMorphTarget(TArray<FbxNode*> &SkelMeshNodeArray, USkeletalMesh* BaseSkelMesh, int32 LODIndex, FSkeletalMeshImportData &BaseSkeletalMeshImportData, const bool bMapMorphTargetToTimeZero)
 {
 	if (!CanImportClass(UMorphTarget::StaticClass()))
 	{
@@ -4786,7 +4791,7 @@ void UnFbx::FFbxImporter::ImportFbxMorphTarget(TArray<FbxNode*> &SkelMeshNodeArr
 	
 	if (bHasMorph)
 	{
-		ImportMorphTargetsInternal( SkelMeshNodeArray, BaseSkelMesh, LODIndex, BaseSkeletalMeshImportData);
+		ImportMorphTargetsInternal( SkelMeshNodeArray, BaseSkelMesh, LODIndex, BaseSkeletalMeshImportData, bMapMorphTargetToTimeZero);
 		//Save the rawMesh
 		BaseSkelMesh->SaveLODImportedData(LODIndex, BaseSkeletalMeshImportData);
 	}
diff --git a/Engine/Source/Editor/UnrealEd/Private/FbxMeshUtils.cpp b/Engine/Source/Editor/UnrealEd/Private/FbxMeshUtils.cpp
--- a/Engine/Source/Editor/UnrealEd/Private/FbxMeshUtils.cpp
+++ b/Engine/Source/Editor/UnrealEd/Private/FbxMeshUtils.cpp
@@ -798,7 +798,7 @@ namespace FbxMeshUtils
 
 					if (ImportOptions->bImportMorph)
 					{
-						FFbxImporter->ImportFbxMorphTarget(SkelMeshNodeArray, SelectedSkelMesh, SelectedLOD, OutData);
+						FFbxImporter->ImportFbxMorphTarget(SkelMeshNodeArray, SelectedSkelMesh, SelectedLOD, OutData, ImportSkeletalMeshArgs.bMapMorphTargetToTimeZero);
 					}
 
 					bSuccess = true;
diff --git a/Engine/Source/Editor/UnrealEd/Public/FbxImporter.h b/Engine/Source/Editor/UnrealEd/Public/FbxImporter.h
--- a/Engine/Source/Editor/UnrealEd/Public/FbxImporter.h
+++ b/Engine/Source/Editor/UnrealEd/Public/FbxImporter.h
@@ -849,6 +849,7 @@ public:
 			, OrderedMaterialNames(nullptr)
 			, ImportMaterialOriginalNameData(nullptr)
 			, ImportMeshSectionsData(nullptr)
+			, bMapMorphTargetToTimeZero(false)
 		{}
 
 		UObject* InParent;
@@ -864,6 +865,7 @@ public:
 
 		TArray<FName> *ImportMaterialOriginalNameData;
 		FImportMeshLodSectionsData *ImportMeshSectionsData;
+		bool bMapMorphTargetToTimeZero;
 	};
 
 	UNREALED_API USkeletalMesh* ImportSkeletalMesh(FImportSkeletalMeshArgs &ImportSkeletalMeshArgs);
@@ -934,7 +936,7 @@ public:
 	 * @param BaseSkelMesh - base Skeletal Mesh
 	 * @param LODIndex - LOD index
 	 */
-	UNREALED_API void ImportFbxMorphTarget(TArray<FbxNode*> &SkelMeshNodeArray, USkeletalMesh* BaseSkelMesh, int32 LODIndex, FSkeletalMeshImportData &BaseSkeletalMeshImportData);
+	UNREALED_API void ImportFbxMorphTarget(TArray<FbxNode*> &SkelMeshNodeArray, USkeletalMesh* BaseSkelMesh, int32 LODIndex, FSkeletalMeshImportData &BaseSkeletalMeshImportData, const bool bSkinControlPointToTimeZero);
 
 	/**
 	 * Import LOD object for skeletal mesh
@@ -1206,8 +1208,9 @@ private:
 	 * @param SkelMeshNodeArray - Fbx Nodes that the base Skeletal Mesh construct from
 	 * @param BaseSkelMesh - base Skeletal Mesh
 	 * @param LODIndex - LOD index of the skeletal mesh
+	 * @param bSkinControlPointToTimeZero Use the pose at T0 to map the morph targets onto, since the base mesh was imported at that time, rather than the ref pose.
 	 */
-	void ImportMorphTargetsInternal( TArray<FbxNode*>& SkelMeshNodeArray, USkeletalMesh* BaseSkelMesh, int32 LODIndex, FSkeletalMeshImportData &BaseSkeletalMeshImportData);
+	void ImportMorphTargetsInternal( TArray<FbxNode*>& SkelMeshNodeArray, USkeletalMesh* BaseSkelMesh, int32 LODIndex, FSkeletalMeshImportData &BaseSkeletalMeshImportData, bool bMapMorphTargetToTimeZero);
 
 	/**
 	* sub-method called from ImportSkeletalMeshLOD method
@@ -1517,7 +1520,7 @@ public:
 	*/
 	bool FillSkeletalMeshImportData(TArray<FbxNode*>& NodeArray, UFbxSkeletalMeshImportData* TemplateImportData, TArray<FbxShape*> *FbxShapeArray,
 									FSkeletalMeshImportData* OutData, TArray<FbxNode*>& OutImportedSkeletonLinkNodes, TArray<FName> &LastImportedMaterialNames, 
-									const bool bIsReimport, const TMap<FVector3f, FColor>& ExistingVertexColorData);
+									const bool bIsReimport, const TMap<FVector3f, FColor>& ExistingVertexColorData, bool& bMapMorphTargetToTimeZero);
 
 protected:
 
@@ -1543,10 +1546,11 @@ protected:
 	* @param NodeArray		Fbx node array to look at
 	* @param FbxShapeArray	Fbx Morph object, if not NULL, we are importing a morph object.
 	* @param ModifiedPoints	Set of points indices for which we've modified the value in OutData
+	* @param bInUseT0AsRefPose Use the pose at T0 to map the morph targets onto, since the base mesh was imported at that time, rather than the ref pose.
 	*
 	* @returns bool			true if import successfully.
 	*/
-	bool GatherPointsForMorphTarget(FSkeletalMeshImportData* OutData, TArray<FbxNode*>& NodeArray, TArray< FbxShape* >* FbxShapeArray, TSet<uint32>& ModifiedPoints);
+	bool GatherPointsForMorphTarget(FSkeletalMeshImportData* OutData, TArray<FbxNode*>& NodeArray, TArray< FbxShape* >* FbxShapeArray, TSet<uint32>& ModifiedPoints, bool bSkinControlPointToTimeZero);
 
 	/**
 	 * Import bones from skeletons that NodeArray bind to.
diff --git a/Engine/Source/Runtime/Engine/Private/Rendering/SkeletalMeshLODImporterData.cpp b/Engine/Source/Runtime/Engine/Private/Rendering/SkeletalMeshLODImporterData.cpp
--- a/Engine/Source/Runtime/Engine/Private/Rendering/SkeletalMeshLODImporterData.cpp
+++ b/Engine/Source/Runtime/Engine/Private/Rendering/SkeletalMeshLODImporterData.cpp
@@ -28,9 +28,6 @@ void FSkeletalMeshImportData::CopyDataNeedByMorphTargetImport(FSkeletalMeshImpor
 	Other.Points = Points;
 	//PointToRawMap should not be save when saving morph target data, we only need it temporary to gather the point from the fbx shape
 	Other.PointToRawMap = PointToRawMap;
-	Other.bDiffPose = bDiffPose;
-	Other.bUseT0AsRefPose = bUseT0AsRefPose;
-	Other.bKeepSectionsSeparate = bKeepSectionsSeparate;
 }
 
 void FSkeletalMeshImportData::KeepAlternateSkinningBuildDataOnly()
@@ -159,7 +156,6 @@ bool FSkeletalMeshImportData::ReplaceSkeletalMeshGeometryImportData(const USkele
 	ImportData->bHasTangents = OriginalSkeletalMeshImportData.bHasTangents;
 	ImportData->bHasVertexColors = OriginalSkeletalMeshImportData.bHasVertexColors;
 	ImportData->NumTexCoords = OriginalSkeletalMeshImportData.NumTexCoords;
-	ImportData->bKeepSectionsSeparate = OriginalSkeletalMeshImportData.bKeepSectionsSeparate;
 
 	ImportData->Materials.Reset();
 	ImportData->Points.Reset();
@@ -215,10 +211,6 @@ bool FSkeletalMeshImportData::ReplaceSkeletalMeshRigImportData(const USkeletalMe
 	FSkeletalMeshImportData OriginalSkeletalMeshImportData;
 	SkeletalMesh->LoadLODImportedData(LodIndex, OriginalSkeletalMeshImportData);
 
-	ImportData->bDiffPose = OriginalSkeletalMeshImportData.bDiffPose;
-	ImportData->bUseT0AsRefPose = OriginalSkeletalMeshImportData.bUseT0AsRefPose;
-	ImportData->bKeepSectionsSeparate = OriginalSkeletalMeshImportData.bKeepSectionsSeparate;
-
 	ImportData->RefBonesBinary.Reset();
 	ImportData->RefBonesBinary += OriginalSkeletalMeshImportData.RefBonesBinary;
 
@@ -687,12 +679,13 @@ FArchive& operator<<(FArchive& Ar, FSkeletalMeshImportData& RawMesh)
 	* Serialization should use the raw mesh version not the archive version.
 	* Additionally, stick to serializing basic types and arrays of basic types.
 	*/
+	bool bDummyFlag1 = false, bDummyFlag2 = false;
 
-	Ar << RawMesh.bDiffPose;
+	Ar << bDummyFlag1;
 	Ar << RawMesh.bHasNormals;
 	Ar << RawMesh.bHasTangents;
 	Ar << RawMesh.bHasVertexColors;
-	Ar << RawMesh.bUseT0AsRefPose;
+	Ar << bDummyFlag2;
 	Ar << RawMesh.MaxMaterialIndex;
 	Ar << RawMesh.NumTexCoords;
 	
@@ -766,11 +759,8 @@ FArchive& operator<<(FArchive& Ar, FSkeletalMeshImportData& RawMesh)
 
 	if (Version >= RAW_SKELETAL_MESH_BULKDATA_VER_Keep_Sections_Separate)
 	{
-		Ar << RawMesh.bKeepSectionsSeparate;
-	}
-	else if (Ar.IsLoading())
-	{
-		RawMesh.bKeepSectionsSeparate = false;
+		bool bDummyFlag3 = false;
+		Ar << bDummyFlag3;
 	}
 
 	return Ar;
diff --git a/Engine/Source/Runtime/Engine/Public/Rendering/SkeletalMeshLODImporterData.h b/Engine/Source/Runtime/Engine/Public/Rendering/SkeletalMeshLODImporterData.h
--- a/Engine/Source/Runtime/Engine/Public/Rendering/SkeletalMeshLODImporterData.h
+++ b/Engine/Source/Runtime/Engine/Public/Rendering/SkeletalMeshLODImporterData.h
@@ -338,9 +338,6 @@ public:
 	bool bHasVertexColors; // If true there are vertex colors in the imported file
 	bool bHasNormals; // If true there are normals in the imported file
 	bool bHasTangents; // If true there are tangents in the imported file
-	bool bUseT0AsRefPose; // If true, then the pose at time=0 will be used instead of the ref pose
-	bool bDiffPose; // If true, one of the bones has a different pose at time=0 vs the ref pose
-	bool bKeepSectionsSeparate; // If true, sections with matching materials are kept separate and will not get combined
 
 	// Morph targets imported(i.e. FBX) data. The name is the morph target name
 	TArray<FSkeletalMeshImportData> MorphTargets;
@@ -362,9 +359,6 @@ public:
 		, bHasVertexColors(false)
 		, bHasNormals(false)
 		, bHasTangents(false)
-		, bUseT0AsRefPose(false)
-		, bDiffPose(false)
-		, bKeepSectionsSeparate(false)
 	{
 
 	}
@@ -438,6 +432,8 @@ public:
 private:
 	ENGINE_API void CleanUpUnusedMaterials();
 	ENGINE_API void SplitVerticesBySmoothingGroups();
+	
+	friend FArchive& operator<<(FArchive& Ar, FSkeletalMeshImportData& RawMesh);
 };
 
 /**
