Editor: Move a CancelDragDrop() call out from an actor deletion function into a higher entrance point so that the deletion function can be used from slate drag/drop ops. This will allow actor drag operations to delete their preview actors using that function, which gets called through the TEv1 interfaces.

The original cancel call was added in 6613611 to avoid a crash when hitting delete while dragging in the outliner. So, I also added a guard to make it so that there is no longer a crash even if the CancelDragDrop() call is removed.

diff --git a/Engine/Source/Editor/SceneOutliner/Private/ActorMode.cpp b/Engine/Source/Editor/SceneOutliner/Private/ActorMode.cpp
--- a/Engine/Source/Editor/SceneOutliner/Private/ActorMode.cpp
+++ b/Engine/Source/Editor/SceneOutliner/Private/ActorMode.cpp
@@ -833,8 +833,12 @@ bool FActorMode::ParseDragDrop(FSceneOutlinerDragDropPayload& OutPayload, const
 		}
 		if (const auto& ActorOp = OutlinerOp.GetSubOp<FActorDragDropOp>())
 		{
-			for (const auto& Actor : ActorOp->Actors)
+			for (const TWeakObjectPtr<AActor>& Actor : ActorOp->Actors)
 			{
+				if (!Actor.IsValid())
+				{
+					continue;
+				}
 				OutPayload.DraggedItems.Add(SceneOutliner->GetTreeItem(Actor.Get()));
 			}
 		}
@@ -844,6 +848,10 @@ bool FActorMode::ParseDragDrop(FSceneOutlinerDragDropPayload& OutPayload, const
 	{
 		for (const TWeakObjectPtr<AActor>& Actor : static_cast<const FActorDragDropOp&>(Operation).Actors)
 		{
+			if (!Actor.IsValid())
+			{
+				continue;
+			}
 			OutPayload.DraggedItems.Add(SceneOutliner->GetTreeItem(Actor.Get()));
 		}
 		return true;
diff --git a/Engine/Source/Editor/UnrealEd/Private/EditorActor.cpp b/Engine/Source/Editor/UnrealEd/Private/EditorActor.cpp
--- a/Engine/Source/Editor/UnrealEd/Private/EditorActor.cpp
+++ b/Engine/Source/Editor/UnrealEd/Private/EditorActor.cpp
@@ -886,8 +886,6 @@ bool UUnrealEdEngine::DeleteActors(const TArray<AActor*>& InActorsToDelete, UWor
 	// Aggregate the time we waited on user input to remove it from the total time
 	double DialogWaitingSeconds = 0;
 
-	FSlateApplication::Get().CancelDragDrop();
-
 	// Fire ULevel::LevelDirtiedEvent when falling out of scope.
 	FScopedLevelDirtied LevelDirtyCallback;
 
diff --git a/Engine/Source/Editor/UnrealEd/Private/UnrealEdSrv.cpp b/Engine/Source/Editor/UnrealEd/Private/UnrealEdSrv.cpp
--- a/Engine/Source/Editor/UnrealEd/Private/UnrealEdSrv.cpp
+++ b/Engine/Source/Editor/UnrealEd/Private/UnrealEdSrv.cpp
@@ -1504,6 +1504,9 @@ bool UUnrealEdEngine::Exec_Edit( UWorld* InWorld, const TCHAR* Str, FOutputDevic
 				{
 					if (UTypedElementCommonActions* CommonActions = LevelEditor->GetCommonActions())
 					{
+						// End drags to avoid deleting something out from under one.
+						FSlateApplication::Get().CancelDragDrop();
+
 						const FScopedTransaction Transaction(NSLOCTEXT("UnrealEd", "Cut", "Cut"));
 
 						UTypedElementSelectionSet* SelectionSet = LevelEditor->GetMutableElementSelectionSet();
@@ -2577,6 +2580,9 @@ bool UUnrealEdEngine::Exec_Actor( UWorld* InWorld, const TCHAR* Str, FOutputDevi
 				{
 					UTypedElementSelectionSet* SelectionSet = LevelEditor->GetMutableElementSelectionSet();
 				
+					// End drags to avoid deleting something out from under one.
+					FSlateApplication::Get().CancelDragDrop();
+
 					const FScopedTransaction Transaction(NSLOCTEXT("UnrealEd", "DeleteElements", "Delete Elements"));
 					if (SelectionSet->GetNumSelectedElements() == 0)
 					{
