PR #11097: Fix to allow listener delegates to potentially mutate the managed set while iterating.

diff --git a/Engine/Source/Editor/UnrealEd/Public/Kismet2/ListenerManager.h b/Engine/Source/Editor/UnrealEd/Public/Kismet2/ListenerManager.h
--- a/Engine/Source/Editor/UnrealEd/Public/Kismet2/ListenerManager.h
+++ b/Engine/Source/Editor/UnrealEd/Public/Kismet2/ListenerManager.h
@@ -51,17 +51,19 @@ public:
 
 	void PreChange(const TType* Changed, TChangeInfo Info)
 	{
-		for (auto Listener : Listeners)
+		// Note: Not using range-based iteration here in case the listener mutates the list via add/remove.
+		for (auto ListenerIt = Listeners.CreateConstIterator(); ListenerIt; ++ListenerIt)
 		{
-			Listener->PreChange(Changed, Info);
+			(*ListenerIt)->PreChange(Changed, Info);
 		}
 	}
 
 	void PostChange(const TType* Changed, TChangeInfo Info)
 	{
-		for (auto Listener : Listeners)
+		// Note: Not using range-based iteration here in case the listener mutates the list via add/remove.
+		for (auto ListenerIt = Listeners.CreateConstIterator(); ListenerIt; ++ListenerIt)
 		{
-			Listener->PostChange(Changed, Info);
+			(*ListenerIt)->PostChange(Changed, Info);
 		}
 	}
 	
