[BugFix] accidently didn't submit this file as apart of 30029548

diff --git a/Engine/Source/Editor/UnrealEd/Internal/AsyncTreeDifferences.h b/Engine/Source/Editor/UnrealEd/Internal/AsyncTreeDifferences.h
--- a/Engine/Source/Editor/UnrealEd/Internal/AsyncTreeDifferences.h
+++ b/Engine/Source/Editor/UnrealEd/Internal/AsyncTreeDifferences.h
@@ -116,6 +116,15 @@ public:
 	{
 		return true;
 	}
+
+	/**
+	 * return true if TreeNode is considered equal when all it's children are equal.
+	 * This avoids an unnecessary call to AreValuesEqual
+	*/
+	virtual bool ShouldInheritEqualFromChildren(const InNodeType& TreeNodeA, const InNodeType& TreeNodeB) const
+	{
+		return false;
+	}
 };
 
 
@@ -139,11 +148,9 @@ struct TDiffNode
 	TDiffNode() = default;
 	TDiffNode(TDiffNode &&Other) = default;
 	
-	TDiffNode(const ValueType& InValueA, const ValueType& InValueB, const TDiffNode* InParent, TUniquePtr<TTreeDiffSpecification<ValueType>>& Specification)
+	TDiffNode(const ValueType& InValueA, const ValueType& InValueB, const TDiffNode* InParent)
 		: ValueA(InValueA), ValueB(InValueB), Children(), Parent(InParent)
-	{
-		SetDiffType(Specification);
-	}
+	{}
 
 	void SetDiffType(TUniquePtr<TTreeDiffSpecification<ValueType>>& Specification);
 
@@ -207,13 +214,30 @@ void TDiffNode<InNodeType>::SetDiffType(TUniquePtr<TTreeDiffSpecification<ValueT
 {
 	if (ValueA != NullValue && ValueB != NullValue)
 	{
-		if (Specification->AreValuesEqual(ValueA, ValueB))
+		if (Specification->ShouldInheritEqualFromChildren(ValueA, ValueB) && !Children.IsEmpty())
 		{
+			// iterate children. Parent is identical iff all of it's children are identical
+			for (const TUniquePtr<TDiffNode>& Child : Children)
+			{
+				if (Child->DiffResult != ETreeDiffResult::Identical)
+				{
+					DiffResult = ETreeDiffResult::DifferentValues;
+					return;
+				}
+			}
 			DiffResult = ETreeDiffResult::Identical;
 			return;
 		}
-		DiffResult = ETreeDiffResult::DifferentValues;
-		return;
+		else
+		{
+			if (Specification->AreValuesEqual(ValueA, ValueB))
+			{
+				DiffResult = ETreeDiffResult::Identical;
+				return;
+			}
+			DiffResult = ETreeDiffResult::DifferentValues;
+			return;
+		}
 	}
 	if (ValueA != NullValue && ValueB == NullValue)
 	{
@@ -442,27 +466,27 @@ void TAsyncTreeDifferences<InNodeType>::QueueParallelNodeLists(const TArray<Valu
         	if (IndexA == 0)
         	{
         		// ValueB doesn't match. push it on it's own
-        		FoundChildren.Add(DiffNodeType(DiffNodeType::NullValue, ValuesB[--IndexB], ParentNode, Specification));
+        		FoundChildren.Add(DiffNodeType(DiffNodeType::NullValue, ValuesB[--IndexB], ParentNode));
         	}
         	else if (IndexB == 0)
         	{
         		// ValueA doesn't match. push it on it's own
-        		FoundChildren.Add(DiffNodeType(ValuesA[--IndexA], DiffNodeType::NullValue, ParentNode, Specification));
+        		FoundChildren.Add(DiffNodeType(ValuesA[--IndexA], DiffNodeType::NullValue, ParentNode));
         	}
         	else if (AreMatching(ValuesA[IndexA - 1], ValuesB[IndexB - 1]))
         	{
         		// found a match between both nodes
-        		FoundChildren.Add(DiffNodeType(ValuesA[--IndexA], ValuesB[--IndexB], ParentNode, Specification));
+        		FoundChildren.Add(DiffNodeType(ValuesA[--IndexA], ValuesB[--IndexB], ParentNode));
         	}
         	else if (LCS[IndexA - 1][IndexB] <= LCS[IndexA][IndexB - 1])
         	{
         		// ValueB doesn't match. push it on it's own
-        		FoundChildren.Add(DiffNodeType(DiffNodeType::NullValue, ValuesB[--IndexB], ParentNode, Specification));
+        		FoundChildren.Add(DiffNodeType(DiffNodeType::NullValue, ValuesB[--IndexB], ParentNode));
         	}
         	else
         	{
         		// ValueA doesn't match. push it on it's own
-        		FoundChildren.Add(DiffNodeType(ValuesA[--IndexA], DiffNodeType::NullValue, ParentNode, Specification));
+        		FoundChildren.Add(DiffNodeType(ValuesA[--IndexA], DiffNodeType::NullValue, ParentNode));
         	}
         }
 
@@ -479,7 +503,7 @@ void TAsyncTreeDifferences<InNodeType>::QueueParallelNodeLists(const TArray<Valu
 	int32 CompressCount = 0;
 	for (DiffNodeType& FoundChild : AllFoundChildren)
 	{
-		if (FoundChild.DiffResult == ETreeDiffResult::MissingFromTree1 &&
+		if (FoundChild.ValueA == DiffNodeType::NullValue &&
 			Specification->ShouldMatchByValue(FoundChild.ValueB))
 		{
 			if (CompressCount == 0)
@@ -491,11 +515,10 @@ void TAsyncTreeDifferences<InNodeType>::QueueParallelNodeLists(const TArray<Valu
 			continue;
 		}
 		if (CompressCount > 0 &&
-            FoundChild.DiffResult == ETreeDiffResult::MissingFromTree2 &&
+            FoundChild.ValueB == DiffNodeType::NullValue &&
 			Specification->ShouldMatchByValue(FoundChild.ValueA))
 		{
 			FoundChildren[CompressIndex].ValueA = FoundChild.ValueA;
-			FoundChildren[CompressIndex].SetDiffType(Specification);
 			++CompressIndex;
 			--CompressCount;
 			continue;
@@ -524,9 +547,7 @@ void TAsyncTreeDifferences<InNodeType>::QueueParallelNodeLists(const TArray<Valu
 			{
 				// found a match between both nodes. Preserve the old data by copying it into the new child array
 				ParentNode->Children.Add(MoveTemp(ExpectedChildren[--ExpectedChildrenIndex]));
-				// update the diff type to the latest result
-				ParentNode->Children.Last()->DiffResult = FoundChildren[--FoundChildrenIndex].DiffResult;
-				
+				--FoundChildrenIndex;
 			}
 			else if (LCS[FoundChildrenIndex - 1][ExpectedChildrenIndex] <= LCS[FoundChildrenIndex][ExpectedChildrenIndex - 1])
 			{
@@ -542,6 +563,12 @@ void TAsyncTreeDifferences<InNodeType>::QueueParallelNodeLists(const TArray<Valu
 	}
 	Algo::Reverse(ParentNode->Children);
 
+	// diff the values of all the children
+	for (const TUniquePtr<DiffNodeType>& Child : ParentNode->Children)
+	{
+		Child->SetDiffType(Specification);
+	}
+
 	// queue all the children to update
 	for (const TUniquePtr<DiffNodeType>& Child : ParentNode->Children)
 	{
