ReplaceActors: Selection optimization (costing > 40seconds for about 500 selected actors). Avoid notifying for every selection change since it is already inside a batch operation.

diff --git a/Engine/Source/Editor/UnrealEd/Private/EditorEngine.cpp b/Engine/Source/Editor/UnrealEd/Private/EditorEngine.cpp
--- a/Engine/Source/Editor/UnrealEd/Private/EditorEngine.cpp
+++ b/Engine/Source/Editor/UnrealEd/Private/EditorEngine.cpp
@@ -5384,8 +5384,10 @@ void UEditorEngine::ReplaceActors(UActorFactory* Factory, const FAssetData& Asse
 
 			if (SelectedActors->IsSelected(OldActor))
 			{
-				SelectActor(OldActor, false, true);
-				SelectActor(NewActor, true, true);
+				// Avoid notifications as we are in a Batch Select Operation
+				const bool bNotify = false;
+				SelectActor(OldActor, false, bNotify);
+				SelectActor(NewActor, true, bNotify);
 			}
 
 			// Find compatible static mesh components and copy instance colors between them.
@@ -5443,7 +5445,8 @@ void UEditorEngine::ReplaceActors(UActorFactory* Factory, const FAssetData& Asse
 		}
 	}
 
-	SelectedActors->EndBatchSelectOperation();
+	const bool bNotify = true;
+	SelectedActors->EndBatchSelectOperation(bNotify);
 
 	// Reattaches actors based on their previous parent child relationship.
 	ReattachActorsHelper::ReattachActors(ConvertedMap, AttachmentInfo);
