CookIndeterminism: Add a -CookLast argument to the cooker that loads a package at start of cook and saves it only after all other packages have been cooked. This can be contrasted with bCookFirst to find packages that can be modified by other packages.
#rnx

diff --git a/Engine/Source/Editor/UnrealEd/Classes/CookOnTheSide/CookOnTheFlyServer.h b/Engine/Source/Editor/UnrealEd/Classes/CookOnTheSide/CookOnTheFlyServer.h
--- a/Engine/Source/Editor/UnrealEd/Classes/CookOnTheSide/CookOnTheFlyServer.h
+++ b/Engine/Source/Editor/UnrealEd/Classes/CookOnTheSide/CookOnTheFlyServer.h
@@ -1497,6 +1497,8 @@ private:
 	bool bCookFilter = false;
 	/** True if commandline arguments specify that packages on commandline should be cooked first. */
 	bool bCookFirst = false;
+	/** True if commandline arguments specify that packages on commandline should be cooked last. */
+	bool bCookLast = false;
 	/** True if experimental optimizations for fast startup should be used. */
 	bool bCookFastStartup = false;
 	/**
diff --git a/Engine/Source/Editor/UnrealEd/Private/CookOnTheFlyServer.cpp b/Engine/Source/Editor/UnrealEd/Private/CookOnTheFlyServer.cpp
--- a/Engine/Source/Editor/UnrealEd/Private/CookOnTheFlyServer.cpp
+++ b/Engine/Source/Editor/UnrealEd/Private/CookOnTheFlyServer.cpp
@@ -3648,20 +3648,40 @@ UE::Cook::EPollStatus UCookOnTheFlyServer::PrepareSave(UE::Cook::FPackageData& P
 {
 	using namespace UE::Cook;
 
+	EPollStatus Result = EPollStatus::Incomplete;
 	if (PackageData.GetCookedPlatformDataComplete())
 	{
-		return EPollStatus::Success;
+		Result = EPollStatus::Success;
 	}
-	if (PackageData.HasPrepareSaveFailed())
+	else if (PackageData.HasPrepareSaveFailed())
 	{
-		return EPollStatus::Error;
+		Result = EPollStatus::Error;
+	}
+	else
+	{
+		UE_SCOPED_HIERARCHICAL_COOKTIMER_AND_DURATION(PrepareSave, DetailedCookStats::TickCookOnTheSidePrepareSaveTimeSec);
+		Result = PrepareSaveInternal(PackageData, Timer, bPrecaching);
+		if (Result == EPollStatus::Error)
+		{
+			PackageData.SetHasPrepareSaveFailed(true);
+		}
 	}
-	UE_SCOPED_HIERARCHICAL_COOKTIMER_AND_DURATION(PrepareSave, DetailedCookStats::TickCookOnTheSidePrepareSaveTimeSec);
-	EPollStatus Result = PrepareSaveInternal(PackageData, Timer, bPrecaching);
-	if (Result == EPollStatus::Error)
+
+	if (Result == EPollStatus::Success && PackageData.GetIsCookLast())
 	{
-		PackageData.SetHasPrepareSaveFailed(true);
+		// No longer urgent
+		PackageData.ClearCookLastUrgency();
+		// Mark it as still not ready if there are non-cook-last packages still in progress
+		if (PackageDatas->GetMonitor().GetNumInProgress() - PackageDatas->GetMonitor().GetNumCookLast() > 0)
+		{
+			Result = EPollStatus::Incomplete;
+		}
+		else
+		{
+			UE_LOG(LogCook, Display, TEXT("CookLast: All other packages cooked. Releasing %s."), *PackageData.GetPackageName().ToString());
+		}
 	}
+
 	return Result;
 }
 
@@ -4478,7 +4498,7 @@ void UCookOnTheFlyServer::PumpSaves(UE::Cook::FTickStackData& StackData, uint32
 					// Poll the results again and check whether we are now done
 					PackageDatas->PollPendingCookedPlatformDatas(true, LastCookableObjectTickTime);
 					PrepareSaveStatus = PrepareSave(PackageData, StackData.Timer, false /* bPrecaching */);
-				} while (!StackData.Timer.IsActionTimeUp() && PrepareSaveStatus == EPollStatus::Incomplete);
+				} while (!StackData.Timer.IsActionTimeUp() && PrepareSaveStatus == EPollStatus::Incomplete && PackageData.GetIsUrgent());
 			}
 			// If we couldn't postpone or wait, then we need to exit and try again later
 			if (PrepareSaveStatus != EPollStatus::Success)
@@ -4517,7 +4537,7 @@ void UCookOnTheFlyServer::PumpSaves(UE::Cook::FTickStackData& StackData, uint32
 			}
 		}
 
-		FSaveCookedPackageContext Context(*this, PackageData, PlatformsForPackage, StackData);
+ 		FSaveCookedPackageContext Context(*this, PackageData, PlatformsForPackage, StackData);
 		SaveCookedPackage(Context);
 		if (Context.bHasTimeOut)
 		{
@@ -5019,7 +5039,7 @@ static void ConstructSoftGCPackageToObjectList(TArray<UObject*>& PackageToObject
 			if (Index > PreviousPackageStartIndex)
 			{
 				UPackage::SoftGCPackageToObjectList.Add(PreviousPackage,
-																								ObjectPtrWrap(TArrayView<UObject*>(PackageToObjectListBufferPtr + PreviousPackageStartIndex, Index - PreviousPackageStartIndex)));
+					ObjectPtrWrap(TArrayView<UObject*>(PackageToObjectListBufferPtr + PreviousPackageStartIndex, Index - PreviousPackageStartIndex)));
 			}
 			PreviousPackage = Pair.Package;
 			PreviousPackageStartIndex = Index;
@@ -5029,7 +5049,7 @@ static void ConstructSoftGCPackageToObjectList(TArray<UObject*>& PackageToObject
 	if (PackageObjectPairsNum > PreviousPackageStartIndex)
 	{
 		UPackage::SoftGCPackageToObjectList.Add(PreviousPackage,
-																						ObjectPtrWrap(TArrayView<UObject*>(PackageToObjectListBufferPtr + PreviousPackageStartIndex, PackageObjectPairsNum - PreviousPackageStartIndex)));
+			ObjectPtrWrap(TArrayView<UObject*>(PackageToObjectListBufferPtr + PreviousPackageStartIndex, PackageObjectPairsNum - PreviousPackageStartIndex)));
 	}
 }
 
@@ -5072,6 +5092,7 @@ void UCookOnTheFlyServer::PreGarbageCollect()
 	{
 		check(SavingPackageData->GetPackage());
 		GCKeepObjects.Add(SavingPackageData->GetPackage());
+		GCKeepPackageDatas.Add(SavingPackageData);
 	}
 
 
@@ -5094,6 +5115,11 @@ void UCookOnTheFlyServer::PreGarbageCollect()
 				ReleaseCookedPlatformData(*PackageData, UE::Cook::EStateChangeReason::GeneratorPreGarbageCollected);
 			}
 		}
+		if (PackageData->GetIsCookLast())
+		{
+			GCKeepPackages.Add(PackageData->GetPackage());
+			GCKeepPackageDatas.Add(PackageData);
+		}
 	}
 	
 	// Find the packages that are waiting on async jobs to finish cooking data
@@ -5293,6 +5319,11 @@ void UCookOnTheFlyServer::PostGarbageCollect()
 	for (FPackageData* PackageData : Demotes)
 	{
 		PackageData->SendToState(EPackageState::Request, ESendFlags::QueueRemove, EStateChangeReason::GarbageCollected);
+		if (PackageData->GetIsCookLast())
+		{
+			// CookLast packages in SaveState have had their urgency removed. Add it back if we need to demote them.
+			PackageData->AddUrgency(true /* bValue */, false /* bAllowUpdateState */);
+		}
 		PackageDatas->GetRequestQueue().AddRequest(PackageData, /* bForceUrgent */ true);
 	}
 
@@ -6846,7 +6877,13 @@ void UCookOnTheFlyServer::SetInitializeConfigSettings(UE::Cook::FInitializeConfi
 	}
 
 	bCookFirst = FParse::Param(FCommandLine::Get(), TEXT("CookFirst"));
-	bRandomizeCookOrder = !bCookFirst && (FParse::Param(FCommandLine::Get(), TEXT("RANDOMPACKAGEORDER")) ||
+	bCookLast = FParse::Param(FCommandLine::Get(), TEXT("CookLast"));
+	if (bCookFirst && bCookLast)
+	{
+		UE_LOG(LogCook, Error, TEXT("-CookFirst and -CookLast are mutually exclusive. Ignoring -CookLast"));
+		bCookLast = false;
+	}
+	bRandomizeCookOrder = !bCookFirst && !bCookLast && (FParse::Param(FCommandLine::Get(), TEXT("RANDOMPACKAGEORDER")) ||
 		(FParse::Param(FCommandLine::Get(), TEXT("DIFFONLY")) && !FParse::Param(FCommandLine::Get(), TEXT("DIFFNORANDCOOK"))));
 
 	ParseCookFilters();
@@ -11534,15 +11571,25 @@ void UCookOnTheFlyServer::GenerateInitialRequests(FBeginCookContext& BeginContex
 		UE_SCOPED_HIERARCHICAL_COOKTIMER(GenerateLongPackageName);
 		GenerateLongPackageNames(FilesInPath, FilesInPathInstigators);
 	}
-	TSet<FName> CookFirstPackages;
-	if (bCookFirst)
+	TSet<FName> CookFirstOrLastPackages;
+	bool bCookFirstOrLast = bCookFirst || bCookLast;
+	if (bCookFirstOrLast)
 	{
 		for (const FString& CookMap : CookMaps)
 		{
 			FString LongPackageName;
 			if (FPackageName::TryConvertFilenameToLongPackageName(CookMap, LongPackageName))
 			{
-				CookFirstPackages.Add(FName(*LongPackageName));
+				FName LongPackageFName(*LongPackageName);
+				CookFirstOrLastPackages.Add(LongPackageFName);
+				if (bCookLast)
+				{
+					UE::Cook::FPackageData* PackageDataToDelay = PackageDatas->TryAddPackageDataByPackageName(LongPackageFName);
+					if (PackageDataToDelay)
+					{
+						PackageDataToDelay->SetIsCookLast(true);
+					}
+				}
 			}
 		}
 	}
@@ -11561,7 +11608,7 @@ void UCookOnTheFlyServer::GenerateInitialRequests(FBeginCookContext& BeginContex
 		if (!PackageFileFName.IsNone())
 		{
 			UE::Cook::FFilePlatformRequest Request(PackageFileFName, MoveTemp(Instigator), TargetPlatforms);
-			Request.SetUrgent(bCookFirst && CookFirstPackages.Contains(PackageName));
+			Request.SetUrgent(bCookFirstOrLast && CookFirstOrLastPackages.Contains(PackageName));
 			WorkerRequests->AddStartCookByTheBookRequest(MoveTemp(Request));
 		}
 		else if (!FLinkerLoad::IsKnownMissingPackage(PackageName))
diff --git a/Engine/Source/Editor/UnrealEd/Private/Cooker/CookPackageData.cpp b/Engine/Source/Editor/UnrealEd/Private/Cooker/CookPackageData.cpp
--- a/Engine/Source/Editor/UnrealEd/Private/Cooker/CookPackageData.cpp
+++ b/Engine/Source/Editor/UnrealEd/Private/Cooker/CookPackageData.cpp
@@ -89,7 +89,7 @@ bool FPackagePlatformData::NeedsCooking(const ITargetPlatform* PlatformItBelongs
 
 FPackageData::FPackageData(FPackageDatas& PackageDatas, const FName& InPackageName, const FName& InFileName)
 	: GeneratedOwner(nullptr), PackageName(InPackageName), FileName(InFileName), PackageDatas(PackageDatas)
-	, Instigator(EInstigator::NotYetRequested), bIsUrgent(0)
+	, Instigator(EInstigator::NotYetRequested), bIsUrgent(0), bIsCookLast(0)
 	, bIsVisited(0), bIsPreloadAttempted(0)
 	, bIsPreloaded(0), bHasSaveCache(0), bPrepareSaveFailed(0), bPrepareSaveRequiresGC(0)
 	, bCookedPlatformDataStarted(0), bCookedPlatformDataCalled(0), bCookedPlatformDataComplete(0)
@@ -268,6 +268,25 @@ void FPackageData::AddUrgency(bool bUrgent, bool bAllowUpdateState)
 	}
 }
 
+void FPackageData::SetIsCookLast(bool bValue)
+{
+	bool bWasCookLast = GetIsCookLast();
+	if (bWasCookLast != bValue)
+	{
+		bIsCookLast = static_cast<uint32>(bValue);
+		PackageDatas.GetMonitor().OnCookLastChanged(*this);
+	}
+}
+
+void FPackageData::ClearCookLastUrgency()
+{
+	if (!GetIsCookLast() || !GetIsUrgent())
+	{
+		return;
+	}
+	SetIsUrgent(false);
+}
+
 void FPackageData::SetInstigator(FRequestCluster& Cluster, FInstigator&& InInstigator)
 {
 	SetInstigatorInternal(MoveTemp(InInstigator));
@@ -1798,7 +1817,7 @@ void FGeneratorPackage::PreGarbageCollect(FCookGenerationInfo& Info, TArray<TObj
 	check(Info.PackageData); // Caller validates this is non-null
 	if (Info.GetSaveState() > FCookGenerationInfo::ESaveState::CallPopulate)
 	{
-		if (IsUseInternalReferenceToAvoidGarbageCollect())
+		if (IsUseInternalReferenceToAvoidGarbageCollect() || Info.PackageData->GetIsCookLast())
 		{
 			UPackage* Package = Info.PackageData->GetPackage();
 			if (Package)
@@ -2499,6 +2518,7 @@ void FPendingCookedPlatformDataCancelManager::Release(FPendingCookedPlatformData
 FPackageDataMonitor::FPackageDataMonitor()
 {
 	FMemory::Memset(NumUrgentInState, 0);
+	FMemory::Memset(NumCookLastInState, 0);
 }
 
 int32 FPackageDataMonitor::GetNumUrgent() const
@@ -2513,12 +2533,30 @@ int32 FPackageDataMonitor::GetNumUrgent() const
 	return NumUrgent;
 }
 
+int32 FPackageDataMonitor::GetNumCookLast() const
+{
+	int32 Num = 0;
+	for (EPackageState State = EPackageState::Min;
+		State <= EPackageState::Max;
+		State = static_cast<EPackageState>(static_cast<uint32>(State) + 1))
+	{
+		Num += NumCookLastInState[static_cast<uint32>(State) - static_cast<uint32>(EPackageState::Min)];
+	}
+	return Num;
+}
+
 int32 FPackageDataMonitor::GetNumUrgent(EPackageState InState) const
 {
 	check(EPackageState::Min <= InState && InState <= EPackageState::Max);
 	return NumUrgentInState[static_cast<uint32>(InState) - static_cast<uint32>(EPackageState::Min)];
 }
 
+int32 FPackageDataMonitor::GetNumCookLast(EPackageState InState) const
+{
+	check(EPackageState::Min <= InState && InState <= EPackageState::Max);
+	return NumUrgentInState[static_cast<uint32>(InState) - static_cast<uint32>(EPackageState::Min)];
+}
+
 int32 FPackageDataMonitor::GetNumPreloadAllocated() const
 {
 	return NumPreloadAllocated;
@@ -2572,6 +2610,12 @@ void FPackageDataMonitor::OnUrgencyChanged(FPackageData& PackageData)
 	TrackUrgentRequests(PackageData.GetState(), Delta);
 }
 
+void FPackageDataMonitor::OnCookLastChanged(FPackageData& PackageData)
+{
+	int32 Delta = PackageData.GetIsCookLast() ? 1 : -1;
+	TrackCookLastRequests(PackageData.GetState(), Delta);
+}
+
 void FPackageDataMonitor::OnStateChanged(FPackageData& PackageData, EPackageState OldState)
 {
 	EPackageState NewState = PackageData.GetState();
@@ -2580,7 +2624,11 @@ void FPackageDataMonitor::OnStateChanged(FPackageData& PackageData, EPackageStat
 		TrackUrgentRequests(OldState, -1);
 		TrackUrgentRequests(NewState, 1);
 	}
-
+	if (PackageData.GetIsCookLast())
+	{
+		TrackCookLastRequests(OldState, -1);
+		TrackCookLastRequests(NewState, 1);
+	}
 	bool bOldStateAssignedToLocal = OldState != EPackageState::Idle && OldState != EPackageState::AssignedToWorker;
 	bool bNewStateAssignedToLocal = NewState != EPackageState::Idle && NewState != EPackageState::AssignedToWorker;
 	if (bOldStateAssignedToLocal != bNewStateAssignedToLocal)
@@ -2596,6 +2644,16 @@ void FPackageDataMonitor::TrackUrgentRequests(EPackageState State, int32 Delta)
 	check(NumUrgentInState[static_cast<uint32>(State) - static_cast<uint32>(EPackageState::Min)] >= 0);
 }
 
+void FPackageDataMonitor::TrackCookLastRequests(EPackageState State, int32 Delta)
+{
+	check(EPackageState::Min <= State && State <= EPackageState::Max);
+	if (State != EPackageState::Idle)
+	{
+		NumCookLastInState[static_cast<uint32>(State) - static_cast<uint32>(EPackageState::Min)] += Delta;
+		check(NumCookLastInState[static_cast<uint32>(State) - static_cast<uint32>(EPackageState::Min)] >= 0);
+	}
+}
+
 int32 FPackageDataMonitor::GetMPCookAssignedFenceMarker() const
 {
 	return MPCookAssignedFenceMarker;
diff --git a/Engine/Source/Editor/UnrealEd/Private/Cooker/CookPackageData.h b/Engine/Source/Editor/UnrealEd/Private/Cooker/CookPackageData.h
--- a/Engine/Source/Editor/UnrealEd/Private/Cooker/CookPackageData.h
+++ b/Engine/Source/Editor/UnrealEd/Private/Cooker/CookPackageData.h
@@ -382,6 +382,7 @@ public:
 	 * Move it to front of its current container if bAllowUpdateState.
 	 */
 	void AddUrgency(bool bUrgent, bool bAllowUpdateState);
+	void ClearCookLastUrgency();
 
 	/** Accessor for RequestClusters to add reachable platforms directly without modifying dependent data. */
 	void AddReachablePlatforms(FRequestCluster& RequestCluster, TConstArrayView<const ITargetPlatform*> Platforms,
@@ -496,6 +497,10 @@ public:
 	void AddCompletionCallback(TConstArrayView<const ITargetPlatform*> TargetPlatforms,
 		FCompletionCallback&& InCompletionCallback);
 
+	/** Get/Set whether the package has been marked for cooking after all other packages, via -cooklast */
+	void SetIsCookLast(bool bValue);
+	bool GetIsCookLast() const { return bIsCookLast != 0;}
+
 	/**
 	 * Get/Set a visited flag used when searching graphs of PackageData. User of the graph is responsible for 
 	 * setting the bIsVisited flag back to empty when graph operations are done.
@@ -786,6 +791,7 @@ private:
 	FWorkerId WorkerAssignmentConstraint = FWorkerId::Invalid();
 	uint32 State : int32(EPackageState::BitCount);
 	uint32 bIsUrgent : 1;
+	uint32 bIsCookLast : 1;
 	uint32 bIsVisited : 1;
 	uint32 bIsPreloadAttempted : 1;
 	uint32 bIsPreloaded : 1;
@@ -1103,6 +1109,10 @@ public:
 	 * Used to prioritize scheduler actions.
 	 */
 	int32 GetNumUrgent(EPackageState InState) const;
+	/** Report the number of CookLast packages. */
+	int32 GetNumCookLast() const;
+	/** Report the number of CookLast packages in the given state. */
+	int32 GetNumCookLast(EPackageState InState) const;
 
 	/** Callback called from FPackageData when it transitions to or from inprogress. */
 	void OnInProgressChanged(FPackageData& PackageData, bool bInProgress);
@@ -1113,6 +1123,8 @@ public:
 	void OnLastCookedPlatformRemoved(FPackageData& PackageData);
 	/** Callback called from FPackageData when it has changed its urgency. */
 	void OnUrgencyChanged(FPackageData& PackageData);
+	/** Callback called from FPackageData when it has changed its value of IsCookLast. */
+	void OnCookLastChanged(FPackageData& PackageData);
 	/** Callback called from FPackageData when it has changed its state. */
 	void OnStateChanged(FPackageData& PackageData, EPackageState OldState);
 
@@ -1122,11 +1134,14 @@ public:
 private:
 	/** Increment or decrement the NumUrgent counter for the given state. */
 	void TrackUrgentRequests(EPackageState State, int32 Delta);
+	/** Increment or decrement the NumCookLast counter for the given state. */
+	void TrackCookLastRequests(EPackageState State, int32 Delta);
 
 	int32 NumInProgress = 0;
 	int32 NumCooked[(uint8)ECookResult::Count]{};
 	int32 NumPreloadAllocated = 0;
 	int32 NumUrgentInState[static_cast<uint32>(EPackageState::Count)];
+	int32 NumCookLastInState[static_cast<uint32>(EPackageState::Count)];
 	int32 MPCookAssignedFenceMarker = 0;
 	int32 MPCookRetiredFenceMarker = 0;
 };
