Set a nullptr check to avoid crash when Composite node reroute can't be resolved from bad Material uasset data.

diff --git a/Engine/Source/Editor/UnrealEd/Private/MaterialGraph.cpp b/Engine/Source/Editor/UnrealEd/Private/MaterialGraph.cpp
--- a/Engine/Source/Editor/UnrealEd/Private/MaterialGraph.cpp
+++ b/Engine/Source/Editor/UnrealEd/Private/MaterialGraph.cpp
@@ -481,19 +481,22 @@ void UMaterialGraph::LinkGraphNodesFromMaterial()
 					// This is an unseen composite reroute expression, find the actual expression output to connect to.
 					UMaterialExpressionComposite* OwningComposite = CastChecked<UMaterialExpressionComposite>(CompositeReroute->SubgraphExpression);
 
-					UMaterialGraphNode* OutputGraphNode;
-					int32 OutputPinIndex = OwningComposite->InputExpressions->ReroutePins.FindLastByPredicate(ExpressionMatchesPredicate(CompositeReroute));
-					if (OutputPinIndex != INDEX_NONE)
+					if (OwningComposite && OwningComposite->InputExpressions && OwningComposite->OutputExpressions)
 					{
-						OutputGraphNode = CastChecked<UMaterialGraphNode>(OwningComposite->InputExpressions->GraphNode);
-					}
-					else
-					{
-						// Output pin base in the subgraph cannot have outputs, if this reroute isn't in the inputs connect to composite's outputs
-						OutputPinIndex = OwningComposite->OutputExpressions->ReroutePins.FindLastByPredicate(ExpressionMatchesPredicate(CompositeReroute));
-						OutputGraphNode = CastChecked<UMaterialGraphNode>(OwningComposite->GraphNode);
+						UMaterialGraphNode* OutputGraphNode;
+						int32 OutputPinIndex = OwningComposite->InputExpressions->ReroutePins.FindLastByPredicate(ExpressionMatchesPredicate(CompositeReroute));
+						if (OutputPinIndex != INDEX_NONE)
+						{
+							OutputGraphNode = CastChecked<UMaterialGraphNode>(OwningComposite->InputExpressions->GraphNode);
+						}
+						else
+						{
+							// Output pin base in the subgraph cannot have outputs, if this reroute isn't in the inputs connect to composite's outputs
+							OutputPinIndex = OwningComposite->OutputExpressions->ReroutePins.FindLastByPredicate(ExpressionMatchesPredicate(CompositeReroute));
+							OutputGraphNode = CastChecked<UMaterialGraphNode>(OwningComposite->GraphNode);
+						}
+						InputPin->MakeLinkTo(OutputGraphNode->GetOutputPin(OutputPinIndex));
 					}
-					InputPin->MakeLinkTo(OutputGraphNode->GetOutputPin(OutputPinIndex));
 				}
 			}
 		}
