Outliner: Change drag/drop incompatibility text based on mobility of actors to be more accurate

diff --git a/Engine/Source/Editor/UnrealEd/Private/EditorEngine.cpp b/Engine/Source/Editor/UnrealEd/Private/EditorEngine.cpp
--- a/Engine/Source/Editor/UnrealEd/Private/EditorEngine.cpp
+++ b/Engine/Source/Editor/UnrealEd/Private/EditorEngine.cpp
@@ -4379,7 +4379,13 @@ bool UEditorEngine::CanParentActors( const AActor* ParentActor, const AActor* Ch
 			FFormatNamedArguments Arguments;
 			Arguments.Add(TEXT("StaticActor"), FText::FromString(ChildActor->GetActorLabel()));
 			Arguments.Add(TEXT("DynamicActor"), FText::FromString(ParentActor->GetActorLabel()));
-			*ReasonText = FText::Format( NSLOCTEXT("ActorAttachmentError", "StaticDynamic_ActorAttachmentError", "Cannot attach static actor {StaticActor} to dynamic actor {DynamicActor}."), Arguments);
+			
+			Arguments.Add(TEXT("DynamicActorMobility"),
+				ParentRoot->Mobility == EComponentMobility::Stationary
+				? NSLOCTEXT("ActorAttachmentError", "StationaryMobility", "Stationary")
+				: NSLOCTEXT("ActorAttachmentError", "MovableMobility", "Movable"));
+				
+			*ReasonText = FText::Format( NSLOCTEXT("ActorAttachmentError", "StaticDynamic_ActorAttachmentError", "Cannot attach actor with Static mobility ({StaticActor}) to actor with {DynamicActorMobility} mobility ({DynamicActor})."), Arguments);
 		}
 		return false;
 	}
