Editor: Support using non-actor factories for dragging things into the level, support non-actor drag preview objects, and fix some "recently placed" category handling to handle custom placed items.

diff --git a/Engine/Source/Editor/EditorFramework/Private/Subsystems/PlacementSubsystem.cpp b/Engine/Source/Editor/EditorFramework/Private/Subsystems/PlacementSubsystem.cpp
--- a/Engine/Source/Editor/EditorFramework/Private/Subsystems/PlacementSubsystem.cpp
+++ b/Engine/Source/Editor/EditorFramework/Private/Subsystems/PlacementSubsystem.cpp
@@ -121,6 +121,11 @@ void UPlacementSubsystem::RegisterPlacementFactories()
 	PlacementFactoriesRegistered.Broadcast();
 }
 
+// Deprecated in 5.4. This function is a mistake because the user is not expected to pass in a 
+// UClass that identifies a subclass of a UClass- they are expected to pass in a UClass that
+// identifies their factory (which does not derive from UClass). TSubclassOf does not do any
+// verification when it is created, it will just give false in the if statement and the function
+// will return nullptr.
 TScriptInterface<IAssetFactoryInterface> UPlacementSubsystem::GetAssetFactoryFromFactoryClass(TSubclassOf<UClass> InFactoryInterfaceClass) const
 {
 	if (InFactoryInterfaceClass)
@@ -137,6 +142,22 @@ TScriptInterface<IAssetFactoryInterface> UPlacementSubsystem::GetAssetFactoryFro
 	return nullptr;
 }
 
+TScriptInterface<IAssetFactoryInterface> UPlacementSubsystem::GetAssetFactoryFromFactoryClass(UClass* InFactoryInterfaceClass) const
+{
+	if (InFactoryInterfaceClass)
+	{
+		for (const TScriptInterface<IAssetFactoryInterface>& AssetFactory : AssetFactories)
+		{
+			if (AssetFactory.GetObject()->IsA(InFactoryInterfaceClass))
+			{
+				return AssetFactory;
+			}
+		}
+	}
+
+	return nullptr;
+}
+
 void UPlacementSubsystem::RegisterAssetFactory(TScriptInterface<IAssetFactoryInterface> AssetFactory)
 {
 	AssetFactories.Add(AssetFactory);
diff --git a/Engine/Source/Editor/EditorFramework/Public/Subsystems/PlacementSubsystem.h b/Engine/Source/Editor/EditorFramework/Public/Subsystems/PlacementSubsystem.h
--- a/Engine/Source/Editor/EditorFramework/Public/Subsystems/PlacementSubsystem.h
+++ b/Engine/Source/Editor/EditorFramework/Public/Subsystems/PlacementSubsystem.h
@@ -125,6 +125,9 @@ public:
 
 	FSimpleMulticastDelegate& OnPlacementFactoriesRegistered();
 
+	TScriptInterface<IAssetFactoryInterface> GetAssetFactoryFromFactoryClass(UClass* InFactoryInterfaceClass) const;
+
+	UE_DEPRECATED(5.4, "This overload never worked as intended. Use the UClass* overload instead.")
 	TScriptInterface<IAssetFactoryInterface> GetAssetFactoryFromFactoryClass(TSubclassOf<UClass> InFactoryInterfaceClass) const;
 	
 	void RegisterAssetFactory(TScriptInterface<IAssetFactoryInterface> AssetFactory);
diff --git a/Engine/Source/Editor/LevelEditor/Private/SLevelViewport.cpp b/Engine/Source/Editor/LevelEditor/Private/SLevelViewport.cpp
--- a/Engine/Source/Editor/LevelEditor/Private/SLevelViewport.cpp
+++ b/Engine/Source/Editor/LevelEditor/Private/SLevelViewport.cpp
@@ -765,7 +765,7 @@ void SLevelViewport::OnDragLeave( const FDragDropEvent& DragDropEvent )
 {
 	if ( LevelViewportClient->HasDropPreviewActors() )
 	{
-		LevelViewportClient->DestroyDropPreviewActors();
+		LevelViewportClient->DestroyDropPreviewElements();
 	}
 
 	TSharedPtr<FDragDropOperation> Operation = DragDropEvent.GetOperation();
@@ -907,7 +907,7 @@ bool SLevelViewport::HandlePlaceDraggedObjects(const FGeometry& MyGeometry, cons
 {
 	bool bAllAssetWereLoaded = false;
 	bool bValidDrop = false;
-	UActorFactory* ActorFactory = nullptr;
+	TScriptInterface<IAssetFactoryInterface> AssetFactory = nullptr;
 
 	TSharedPtr< FDragDropOperation > Operation = DragDropEvent.GetOperation();
 	if (!Operation.IsValid())
@@ -953,7 +953,7 @@ bool SLevelViewport::HandlePlaceDraggedObjects(const FGeometry& MyGeometry, cons
 
 		TSharedPtr<FAssetDragDropOp> DragDropOp = StaticCastSharedPtr<FAssetDragDropOp>( Operation );
 
-		ActorFactory = DragDropOp->GetActorFactory();
+		AssetFactory = DragDropOp->GetAssetFactory();
 
 		bAllAssetWereLoaded = true;
 		for (const FAssetData& AssetData : DragDropOp->GetAssets())
@@ -1033,15 +1033,21 @@ bool SLevelViewport::HandlePlaceDraggedObjects(const FGeometry& MyGeometry, cons
 		bool bDropSuccessful = false;
 
 		// Make sure the drop preview is destroyed
-		LevelViewportClient->DestroyDropPreviewActors();
+		LevelViewportClient->DestroyDropPreviewElements();
 
 		if( !bShowDropContextMenu || !bCreateDropPreview )
 		{
 			// Otherwise just attempt to drop the object(s)
-			TArray< AActor* > TemporaryActors;
+			FLevelEditorViewportClient::FDropObjectOptions DropOptions;
+			DropOptions.FactoryToUse = AssetFactory;
+			DropOptions.bOnlyDropOnTarget = false;
+			DropOptions.bCreateDropPreview = bCreateDropPreview;
 			// Only select actor on drop
-			const bool SelectActor = !bCreateDropPreview;
-			bDropSuccessful = LevelViewportClient->DropObjectsAtCoordinates(CachedOnDropLocalMousePos.X, CachedOnDropLocalMousePos.Y, DroppedObjects, TemporaryActors, false, bCreateDropPreview, SelectActor, ActorFactory);
+			DropOptions.bSelectOutput = !bCreateDropPreview;
+
+			TArray<FTypedElementHandle> Unused;
+			bDropSuccessful = LevelViewportClient->DropObjectsAtCoordinates(CachedOnDropLocalMousePos.X, CachedOnDropLocalMousePos.Y, 
+				DroppedObjects, Unused, DropOptions);
 		}
 		else if ( bAllAssetWereLoaded && DroppedObjects.Num() > 0 )
 		{
diff --git a/Engine/Source/Editor/PlacementMode/Private/PlacementModeModule.cpp b/Engine/Source/Editor/PlacementMode/Private/PlacementModeModule.cpp
--- a/Engine/Source/Editor/PlacementMode/Private/PlacementModeModule.cpp
+++ b/Engine/Source/Editor/PlacementMode/Private/PlacementModeModule.cpp
@@ -48,6 +48,7 @@
 #include "AssetRegistry/AssetRegistryModule.h"
 #include "ActorPlacementInfo.h"
 #include "IPlacementModeModule.h"
+#include "Subsystems/PlacementSubsystem.h"
 #include "ToolMenus.h"
 #include "AssetToolsModule.h"
 #include "Kismet2/KismetEditorUtilities.h"
@@ -55,6 +56,15 @@
 #include "SPlacementModeTools.h"
 #include "AssetSelection.h"
 
+namespace PlacementModeModuleLocals
+{
+	FActorPlacementInfo MakePlacementInfo(const FPlaceableItem& Item)
+	{
+		FString ObjectString = Item.AssetData.IsValid() ? Item.AssetData.GetObjectPathString() : FString();
+		FString FactoryString = Item.AssetFactory ? Item.AssetFactory.GetObject()->GetPathName() : FString();
+		return FActorPlacementInfo(ObjectString, FactoryString);
+	}
+}
 
 TOptional<FLinearColor> GetBasicShapeColorOverride()
 {
@@ -82,6 +92,8 @@ FPlacementModeModule::FPlacementModeModule()
 
 void FPlacementModeModule::StartupModule()
 {
+	using namespace PlacementModeModuleLocals;
+
 	TArray< FString > RecentlyPlacedAsStrings;
 	GConfig->GetArray(TEXT("PlacementMode"), TEXT("RecentlyPlaced"), RecentlyPlacedAsStrings, GEditorPerProjectIni);
 
@@ -185,31 +197,31 @@ void FPlacementModeModule::StartupModule()
 		// Cube
 		{
 			TSharedPtr<FPlaceableItem> Cube = MakeShareable(new FPlaceableItem(*UActorFactoryBasicShape::StaticClass(), FAssetData(LoadObject<UStaticMesh>(nullptr, *UActorFactoryBasicShape::BasicCube.ToString())), FName("ClassThumbnail.Cube"), FName("ClassIcon.Cube"), BasicShapeColorOverride, SortOrder += 10, CubeText));
-			ManuallyCreatedPlaceableItems.Add(*UActorFactoryBasicShape::BasicCube.ToString(),Cube);
+			ManuallyCreatedPlaceableItems.Add(MakePlacementInfo(*Cube), Cube);
 			Category->Items.Add(CreateID(), Cube);
 		}
 		// Sphere
 		{
 			TSharedPtr<FPlaceableItem> Sphere = MakeShareable(new FPlaceableItem(*UActorFactoryBasicShape::StaticClass(), FAssetData(LoadObject<UStaticMesh>(nullptr, *UActorFactoryBasicShape::BasicSphere.ToString())), FName("ClassThumbnail.Sphere"), FName("ClassIcon.Sphere"), BasicShapeColorOverride, SortOrder += 10, SphereText));
-			ManuallyCreatedPlaceableItems.Add(*UActorFactoryBasicShape::BasicSphere.ToString(), Sphere);
+			ManuallyCreatedPlaceableItems.Add(MakePlacementInfo(*Sphere), Sphere);
 			Category->Items.Add(CreateID(), Sphere);
 		}
 		// Cylinder
 		{
 			TSharedPtr<FPlaceableItem> Cylinder = MakeShareable(new FPlaceableItem(*UActorFactoryBasicShape::StaticClass(), FAssetData(LoadObject<UStaticMesh>(nullptr, *UActorFactoryBasicShape::BasicCylinder.ToString())), FName("ClassThumbnail.Cylinder"), FName("ClassIcon.Cylinder"), BasicShapeColorOverride, SortOrder += 10, CylinderText));
-			ManuallyCreatedPlaceableItems.Add(*UActorFactoryBasicShape::BasicCylinder.ToString(), Cylinder);
+			ManuallyCreatedPlaceableItems.Add(MakePlacementInfo(*Cylinder), Cylinder);
 			Category->Items.Add(CreateID(), Cylinder);
 		}
 		// Cone
 		{
 			TSharedPtr<FPlaceableItem> Cone = MakeShareable(new FPlaceableItem(*UActorFactoryBasicShape::StaticClass(), FAssetData(LoadObject<UStaticMesh>(nullptr, *UActorFactoryBasicShape::BasicCone.ToString())), FName("ClassThumbnail.Cone"), FName("ClassIcon.Cone"), BasicShapeColorOverride, SortOrder += 10, ConeText));
-			ManuallyCreatedPlaceableItems.Add(*UActorFactoryBasicShape::BasicCone.ToString(),Cone);
+			ManuallyCreatedPlaceableItems.Add(MakePlacementInfo(*Cone), Cone);
 			Category->Items.Add(CreateID(), Cone);
 		}
 		// Plane
 		{
 			TSharedPtr<FPlaceableItem> Plane = MakeShareable(new FPlaceableItem(*UActorFactoryBasicShape::StaticClass(), FAssetData(LoadObject<UStaticMesh>(nullptr, *UActorFactoryBasicShape::BasicPlane.ToString())), FName("ClassThumbnail.Plane"), FName("ClassIcon.Plane"), BasicShapeColorOverride, SortOrder += 10, PlaneText));
-			ManuallyCreatedPlaceableItems.Add(*UActorFactoryBasicShape::BasicPlane.ToString(), Plane);
+			ManuallyCreatedPlaceableItems.Add(MakePlacementInfo(*Plane), Plane);
 			Category->Items.Add(CreateID(), Plane);
 		}
 	}
@@ -381,12 +393,17 @@ void FPlacementModeModule::PreUnloadCallback()
 	}
 }
 
-void FPlacementModeModule::AddToRecentlyPlaced(const TArray<UObject *>& PlacedObjects, UActorFactory* FactoryUsed /* = NULL */)
+void FPlacementModeModule::AddToRecentlyPlaced(const TArray<UObject*>& PlacedObjects, UActorFactory* FactoryUsed /* = NULL */)
+{
+	AddToRecentlyPlaced(PlacedObjects, TScriptInterface<IAssetFactoryInterface>(FactoryUsed));
+}
+
+void FPlacementModeModule::AddToRecentlyPlaced(const TArray< UObject* >& PlacedObjects, TScriptInterface<IAssetFactoryInterface> FactoryUsed)
 {
 	FString FactoryPath;
 	if (FactoryUsed != NULL)
 	{
-		FactoryPath = FactoryUsed->GetPathName();
+		FactoryPath = FactoryUsed.GetObject()->GetPathName();
 	}
 
 	TArray< UObject* > FilteredPlacedObjects;
@@ -483,6 +500,11 @@ void FPlacementModeModule::OnInitialAssetsScanComplete()
 }
 
 void FPlacementModeModule::AddToRecentlyPlaced(UObject* Asset, UActorFactory* FactoryUsed /* = NULL */)
+{
+	AddToRecentlyPlaced(Asset, TScriptInterface<IAssetFactoryInterface>(FactoryUsed));
+}
+
+void FPlacementModeModule::AddToRecentlyPlaced(UObject* Asset, TScriptInterface<IAssetFactoryInterface> FactoryUsed)
 {
 	TArray< UObject* > Assets;
 	Assets.Add(Asset);
@@ -535,11 +557,16 @@ void FPlacementModeModule::GetSortedCategories(TArray<FPlacementCategoryInfo>& O
 
 TOptional<FPlacementModeID> FPlacementModeModule::RegisterPlaceableItem(FName CategoryName, const TSharedRef<FPlaceableItem>& InItem)
 {
+	using namespace PlacementModeModuleLocals;
+
 	FPlacementCategory* Category = Categories.Find(CategoryName);
 	if (Category && !Category->CustomGenerator)
 	{
 		FPlacementModeID ID = CreateID(CategoryName);
 		Category->Items.Add(ID.UniqueID, InItem);
+
+		ManuallyCreatedPlaceableItems.Add(MakePlacementInfo(*InItem), InItem);
+
 		return ID;
 	}
 	return TOptional<FPlacementModeID>();
@@ -547,10 +574,17 @@ TOptional<FPlacementModeID> FPlacementModeModule::RegisterPlaceableItem(FName Ca
 
 void FPlacementModeModule::UnregisterPlaceableItem(FPlacementModeID ID)
 {
+	using namespace PlacementModeModuleLocals;
+
 	FPlacementCategory* Category = Categories.Find(ID.Category);
 	if (Category)
 	{
-		Category->Items.Remove(ID.UniqueID);
+		TSharedPtr<FPlaceableItem> Item;
+		Category->Items.RemoveAndCopyValue(ID.UniqueID, Item);
+		if (Item)
+		{
+			ManuallyCreatedPlaceableItems.Remove(MakePlacementInfo(*Item));
+		}
 	}
 }
 
@@ -641,29 +675,60 @@ void FPlacementModeModule::RefreshRecentlyPlaced()
 	for (const FActorPlacementInfo& RecentlyPlacedItem : RecentlyPlaced)
 	{
 		// First check if it's a manually created entry
-		if (TSharedPtr<FPlaceableItem>* ItemFound = ManuallyCreatedPlaceableItems.Find(RecentlyPlacedItem.ObjectPath))
+		TWeakPtr<FPlaceableItem>* ManualItem = ManuallyCreatedPlaceableItems.Find(RecentlyPlacedItem);
+		if (ManualItem)
 		{
-			Category->Items.Add(CreateID(), *ItemFound);
+			if (ManualItem->IsValid())
+			{
+				Category->Items.Add(CreateID(), ManualItem->Pin());
+				continue;
+			}
+
+			// A dead pointer here could theoretically be the result of an unregistered category that had 
+			// items manually registered but not manually unregistered. Just remove the dead pointer and
+			// try the other approaches.
+			ManuallyCreatedPlaceableItems.Remove(RecentlyPlacedItem);
 		}
-		else
+
+		UObject* Asset = FindObject<UObject>(nullptr, *RecentlyPlacedItem.ObjectPath);
+
+		// If asset is pending delete, it will not be marked as RF_Standalone, in which case we skip it
+		if (Asset == nullptr || !Asset->HasAnyFlags(RF_Standalone))
 		{
-			UObject* Asset = FindObject<UObject>(nullptr, *RecentlyPlacedItem.ObjectPath);
+			continue;
+		}
 
-			// If asset is pending delete, it will not be marked as RF_Standalone, in which case we skip it
-			if (Asset != nullptr && Asset->HasAnyFlags(RF_Standalone))
-			{
-				FAssetData AssetData = AssetRegistryModule.Get().GetAssetByObjectPath(FSoftObjectPath(RecentlyPlacedItem.ObjectPath));
+		FAssetData AssetData = AssetRegistryModule.Get().GetAssetByObjectPath(FSoftObjectPath(RecentlyPlacedItem.ObjectPath));
 
-				if (AssetData.IsValid())
-				{
-					if (UActorFactory* Factory = FindObject<UActorFactory>(nullptr, *RecentlyPlacedItem.Factory))
-					{
-						TSharedPtr<FPlaceableItem> Ptr = MakeShareable(new FPlaceableItem(Factory, AssetData));
-						Category->Items.Add(CreateID(), Ptr);
-					}
-				}
+		if (!AssetData.IsValid())
+		{
+			continue;
+		}
+
+		TScriptInterface<IAssetFactoryInterface> Factory;
+		if (RecentlyPlacedItem.Factory.IsEmpty())
+		{
+			// The factory portion of the recently placed data is frequently null because it only gets set
+			// high up in the drop if we gave a specific factory (see FLevelEditorViewportClient::DropObjectsAtCoordinates).
+			// Otherwise, the actually used factory is gotten from the placement subsystem later based on the asset data.
+			// So, we do the same kind of lookup here if we don't have a factory.
+			if (UPlacementSubsystem* PlacementSubsystem = GEditor->GetEditorSubsystem<UPlacementSubsystem>())
+			{
+				Factory = PlacementSubsystem->FindAssetFactoryFromAssetData(AssetData);
 			}
 		}
+		else
+		{
+			Factory = FindObject<UObject>(nullptr, *RecentlyPlacedItem.Factory);
+		}
+
+		if (!Factory)
+		{
+			continue;
+		}
+
+		TSharedPtr<FPlaceableItem> Ptr = MakeShareable(new FPlaceableItem(Factory, AssetData));
+		Category->Items.Add(CreateID(), Ptr);
 	}
 }
 
diff --git a/Engine/Source/Editor/PlacementMode/Private/SPlacementModeTools.cpp b/Engine/Source/Editor/PlacementMode/Private/SPlacementModeTools.cpp
--- a/Engine/Source/Editor/PlacementMode/Private/SPlacementModeTools.cpp
+++ b/Engine/Source/Editor/PlacementMode/Private/SPlacementModeTools.cpp
@@ -302,7 +302,7 @@ FReply SPlacementAssetEntry::OnDragDetected(const FGeometry& MyGeometry, const F
 
 	if( MouseEvent.IsMouseButtonDown( EKeys::LeftMouseButton ) )
 	{
-		return FReply::Handled().BeginDragDrop(FAssetDragDropOp::New(Item->AssetData, Item->Factory));
+		return FReply::Handled().BeginDragDrop(FAssetDragDropOp::New(Item->AssetData, Item->AssetFactory));
 	}
 	else
 	{
@@ -528,7 +528,7 @@ FReply SPlacementAssetMenuEntry::OnDragDetected(const FGeometry& MyGeometry, con
 
 	if( MouseEvent.IsMouseButtonDown( EKeys::LeftMouseButton ) )
 	{
-		return FReply::Handled().BeginDragDrop(FAssetDragDropOp::New(Item->AssetData, Item->Factory));
+		return FReply::Handled().BeginDragDrop(FAssetDragDropOp::New(Item->AssetData, Item->AssetFactory));
 	}
 	else
 	{
diff --git a/Engine/Source/Editor/PlacementMode/Public/ActorPlacementInfo.h b/Engine/Source/Editor/PlacementMode/Public/ActorPlacementInfo.h
--- a/Engine/Source/Editor/PlacementMode/Public/ActorPlacementInfo.h
+++ b/Engine/Source/Editor/PlacementMode/Public/ActorPlacementInfo.h
@@ -3,6 +3,8 @@
 
 #include "CoreMinimal.h"
 
+#include "Templates/TypeHash.h"
+
 class FActorPlacementInfo 
 {
 public:
@@ -42,3 +44,8 @@ public:
 	FString ObjectPath;
 	FString Factory;
 };
+
+inline uint32 GetTypeHash(const FActorPlacementInfo& Object)
+{
+	return HashCombine(GetTypeHash(Object.ObjectPath), GetTypeHash(Object.Factory));
+}
\ No newline at end of file
diff --git a/Engine/Source/Editor/PlacementMode/Public/IPlacementModeModule.h b/Engine/Source/Editor/PlacementMode/Public/IPlacementModeModule.h
--- a/Engine/Source/Editor/PlacementMode/Public/IPlacementModeModule.h
+++ b/Engine/Source/Editor/PlacementMode/Public/IPlacementModeModule.h
@@ -88,8 +88,27 @@ struct FPlaceableItem
 	{}
 
 	/** Constructor that takes a specific factory and asset */
-	FPlaceableItem(UActorFactory* InFactory, const FAssetData& InAssetData, TOptional<int32> InSortOrder = TOptional<int32>())
+	FPlaceableItem(TScriptInterface<IAssetFactoryInterface> InAssetFactory,
+		const FAssetData& InAssetData, TOptional<int32> InSortOrder = TOptional<int32>())
+		: Factory(Cast<UActorFactory>(InAssetFactory.GetObject()))
+		, AssetFactory(InAssetFactory)
+		, AssetData(InAssetData)
+		, bAlwaysUseGenericThumbnail(false)
+		, SortOrder(InSortOrder)
+	{
+		AutoSetNativeAndDisplayName();
+	}
+
+	//~ Note: for the constructors that expect UActorFactory, we can't forward to the IAssetFactoryInterface
+	//~ constructors because those currently have a Cast<>() call that requires linking to the module that
+	//~ contains IAssetFactoryInterface, which users sometimes don't do. This problem might go away once we
+	//~ finished deprecating and removing the Factory member.
+
+	/** Constructor that takes a specific actor factory and asset */
+	FPlaceableItem(UActorFactory* InFactory, const FAssetData& InAssetData, 
+		TOptional<int32> InSortOrder = TOptional<int32>()) 
 		: Factory(InFactory)
+		, AssetFactory(InFactory)
 		, AssetData(InAssetData)
 		, bAlwaysUseGenericThumbnail(false)
 		, SortOrder(InSortOrder)
@@ -97,9 +116,11 @@ struct FPlaceableItem
 		AutoSetNativeAndDisplayName();
 	}
 
-	/** Constructor for any placeable class */
-	FPlaceableItem(UClass& InAssetClass, TOptional<int32> InSortOrder = TOptional<int32>())
-		: Factory(GEditor->FindActorFactoryByClass(&InAssetClass))
+	//~ TODO: Might want a version of this that is works for non-actor factories.
+	/** Constructor for any actor factory class */
+	 FPlaceableItem(UClass& InActorFactoryClass, TOptional<int32> InSortOrder = TOptional<int32>())
+		: Factory(GEditor->FindActorFactoryByClass(&InActorFactoryClass))
+		, AssetFactory(Factory)
 		, AssetData(Factory ? Factory->GetDefaultActorClass(FAssetData()) : FAssetData())
 		, bAlwaysUseGenericThumbnail(false)
 		, SortOrder(InSortOrder)
@@ -107,17 +128,17 @@ struct FPlaceableItem
 		AutoSetNativeAndDisplayName();
 	}
 
-	/** Constructor for any placeable class with associated asset data, brush and display name overrides */
+	/** Constructor that takes a specific factory and asset with brush and display name overrides */
 	FPlaceableItem(
-		UClass& InAssetClass,
-		const FAssetData& InAssetData,
-		FName InClassThumbnailBrushOverride = NAME_None,
-		FName InClassIconBrushOverride = NAME_None,
-		TOptional<FLinearColor> InAssetTypeColorOverride = TOptional<FLinearColor>(),
-		TOptional<int32> InSortOrder = TOptional<int32>(),
-		TOptional<FText> InDisplayName = TOptional<FText>()
-	)
-		: Factory(GEditor->FindActorFactoryByClass(&InAssetClass))
+		 TScriptInterface<IAssetFactoryInterface> InAssetFactory,
+		 const FAssetData& InAssetData,
+		 FName InClassThumbnailBrushOverride,
+		 FName InClassIconBrushOverride,
+		 TOptional<FLinearColor> InAssetTypeColorOverride = TOptional<FLinearColor>(),
+		 TOptional<int32> InSortOrder = TOptional<int32>(),
+		 TOptional<FText> InDisplayName = TOptional<FText>())
+		: Factory(Cast<UActorFactory>(InAssetFactory.GetObject()))
+		, AssetFactory(InAssetFactory)
 		, AssetData(InAssetData)
 		, ClassThumbnailBrushOverride(InClassThumbnailBrushOverride)
 		, ClassIconBrushOverride(InClassIconBrushOverride)
@@ -132,6 +153,32 @@ struct FPlaceableItem
 		}
 	}
 
+	/** Constructor for any placeable actor class with associated asset data, brush and display name overrides */
+	 FPlaceableItem(
+		UClass& InActorFactoryClass,
+		const FAssetData& InAssetData,
+		FName InClassThumbnailBrushOverride = NAME_None,
+		FName InClassIconBrushOverride = NAME_None,
+		TOptional<FLinearColor> InAssetTypeColorOverride = TOptional<FLinearColor>(),
+		TOptional<int32> InSortOrder = TOptional<int32>(),
+		TOptional<FText> InDisplayName = TOptional<FText>()
+	 ) 
+		 : Factory(GEditor->FindActorFactoryByClass(&InActorFactoryClass))
+		 , AssetFactory(Factory)
+		 , AssetData(InAssetData)
+		 , ClassThumbnailBrushOverride(InClassThumbnailBrushOverride)
+		 , ClassIconBrushOverride(InClassIconBrushOverride)
+		 , bAlwaysUseGenericThumbnail(true)
+		 , AssetTypeColorOverride(InAssetTypeColorOverride)
+		 , SortOrder(InSortOrder)
+	{
+		 AutoSetNativeAndDisplayName();
+		 if (InDisplayName.IsSet())
+		 {
+			 DisplayName = InDisplayName.GetValue();
+		 }
+	}
+
 	/** Automatically set this item's native and display names from its class or asset */
 	void AutoSetNativeAndDisplayName()
 	{
@@ -164,8 +211,11 @@ struct FPlaceableItem
 
 public:
 
+	//~TODO: UE_DEPRECATED(5.4, "Use AssetFactory instead.")
+	UActorFactory* Factory = nullptr;
+
 	/** The factory used to create an instance of this placeable item */
-	UActorFactory* Factory;
+	TScriptInterface<IAssetFactoryInterface> AssetFactory = nullptr;
 
 	/** Asset data pertaining to the class */
 	FAssetData AssetData;
@@ -238,11 +288,23 @@ public:
 	 * Add the specified assets to the recently placed items list
 	 */
 	virtual void AddToRecentlyPlaced( const TArray< UObject* >& Assets, UActorFactory* FactoryUsed = NULL ) = 0;
+	virtual void AddToRecentlyPlaced(const TArray< UObject* >& Assets, TScriptInterface<IAssetFactoryInterface> FactoryUsed)
+	{
+		// This overload was added later, so we route to old overload, though there shouldn't be any implementers
+		// aside from the module itself.
+		AddToRecentlyPlaced(Assets, Cast<UActorFactory>(FactoryUsed.GetObject()));
+	}
 	
 	/**
 	 * Add the specified asset to the recently placed items list
 	 */
 	virtual void AddToRecentlyPlaced( UObject* Asset, UActorFactory* FactoryUsed = NULL ) = 0;
+	virtual void AddToRecentlyPlaced(UObject* Asset, TScriptInterface<IAssetFactoryInterface> FactoryUsed)
+	{
+		// This overload was added later, so we route to old overload, though there shouldn't be any implementers
+		// aside from the module itself.
+		AddToRecentlyPlaced(Asset, Cast<UActorFactory>(FactoryUsed.GetObject()));
+	}
 
 	/**
 	 * Get a copy of the recently placed items
diff --git a/Engine/Source/Editor/PlacementMode/Public/PlacementModeModule.h b/Engine/Source/Editor/PlacementMode/Public/PlacementModeModule.h
--- a/Engine/Source/Editor/PlacementMode/Public/PlacementModeModule.h
+++ b/Engine/Source/Editor/PlacementMode/Public/PlacementModeModule.h
@@ -66,11 +66,13 @@ public:
 	 * Add the specified assets to the recently placed items list
 	 */
 	virtual void AddToRecentlyPlaced(const TArray< UObject* >& PlacedObjects, UActorFactory* FactoryUsed = NULL) override;
+	virtual void AddToRecentlyPlaced(const TArray< UObject* >& Assets, TScriptInterface<IAssetFactoryInterface> FactoryUsed) override;
 
 	/**
 	 * Add the specified asset to the recently placed items list
 	 */
 	virtual void AddToRecentlyPlaced(UObject* Asset, UActorFactory* FactoryUsed = NULL) override;
+	virtual void AddToRecentlyPlaced(UObject* Asset, TScriptInterface<IAssetFactoryInterface> FactoryUsed) override;
 
 	/**
 	 * Get a copy of the recently placed items list
@@ -146,7 +148,9 @@ private:
 	TArray< TSharedPtr<FExtender> > ContentPaletteFiltersExtenders;
 	TArray< TSharedPtr<FExtender> > PaletteExtenders;
 
-	TMap<FString, TSharedPtr<FPlaceableItem>> ManuallyCreatedPlaceableItems;
+	// When users explicitly add placeable items, they may add custom icons/descriptions, so we 
+	// need to store extra data to be able to recreate the placeable item in "recently placed"
+	TMap<FActorPlacementInfo, TWeakPtr<FPlaceableItem>> ManuallyCreatedPlaceableItems;
 };
 
 IMPLEMENT_MODULE(FPlacementModeModule, PlacementMode);
\ No newline at end of file
diff --git a/Engine/Source/Editor/UnrealEd/Private/AssetSelection.cpp b/Engine/Source/Editor/UnrealEd/Private/AssetSelection.cpp
--- a/Engine/Source/Editor/UnrealEd/Private/AssetSelection.cpp
+++ b/Engine/Source/Editor/UnrealEd/Private/AssetSelection.cpp
@@ -3,6 +3,7 @@
 #include "AssetSelection.h"
 #include "Engine/Level.h"
 #include "Model.h"
+#include "UObject/ScriptInterface.h"
 #include "UObject/UnrealType.h"
 #include "GameFramework/Actor.h"
 #include "ActorFactories/ActorFactory.h"
@@ -21,6 +22,7 @@
 #include "Kismet2/ComponentEditorUtils.h"
 #include "Engine/Selection.h"
 #include "Editor.h"
+#include "EditorModeManager.h"
 #include "ScopedTransaction.h"
 
 #include "LevelUtils.h"
@@ -35,6 +37,7 @@
 #include "ContentBrowserModule.h"
 #include "SnappingUtils.h"
 #include "ActorEditorUtils.h"
+#include "LevelEditorSubsystem.h"
 #include "LevelEditorViewport.h"
 #include "LandscapeProxy.h"
 #include "Landscape.h"
@@ -52,7 +55,9 @@
 #include "ISourceControlProvider.h"
 #include "Misc/MessageDialog.h"
 #include "Subsystems/PlacementSubsystem.h"
+#include "Elements/Framework/EngineElementsLibrary.h"
 #include "Elements/Framework/TypedElementRegistry.h"
+#include "Elements/Framework/TypedElementSelectionSet.h"
 #include "Elements/Interfaces/TypedElementObjectInterface.h"
 
 namespace AssetSelectionUtils
@@ -526,25 +531,63 @@ namespace ActorPlacementUtils
 	}
 }
 
-/**
-* Creates an actor using the specified factory.  
-*
-* Does nothing if ActorClass is NULL.
-*/
-static AActor* PrivateAddActor( UObject* Asset, UActorFactory* Factory, bool SelectActor = true, EObjectFlags ObjectFlags = RF_Transactional, const FName Name = NAME_None )
+namespace AssetSelectionLocals {
+
+UTypedElementSelectionSet* GetEditorSelectionSet()
 {
-	if (!Factory)
+	if (!GEditor)
 	{
 		return nullptr;
 	}
+	
+	if (ULevelEditorSubsystem* LevelEditorSubsystem = GEditor->GetEditorSubsystem<ULevelEditorSubsystem>())
+	{
+		if (FEditorModeTools* ModeManager = LevelEditorSubsystem->GetLevelEditorModeManager())
+		{
+			return ModeManager->GetEditorSelectionSet();
+		}
+	}
+
+	return nullptr;
+}
 
-	AActor* Actor = NULL;
-	AActor* NewActorTemplate = Factory->GetDefaultActor( Asset );
-	if ( !NewActorTemplate )
+void ForEachObjectInHandles(TArray<FTypedElementHandle> Handles, TFunctionRef<void(UObject&)> Func)
+{
+	for (FTypedElementHandle Handle : Handles)
 	{
-		return nullptr;
+		TTypedElement<ITypedElementObjectInterface> ObjectInterface = UTypedElementRegistry::GetInstance()->GetElement<ITypedElementObjectInterface>(Handle);
+		if (!ObjectInterface)
+		{
+			continue;
+		}
+
+		UObject* Object = ObjectInterface.GetObject();
+		if (!Object)
+		{
+			continue;
+		}
+
+		Func(*Object);
+	}
+}
+
+/**
+ * Creates an object using the specified factory.
+ */
+TArray<FTypedElementHandle> PlaceAssetUsingFactory(UObject* Asset, TScriptInterface<IAssetFactoryInterface> Factory, bool bSelectResult = true, EObjectFlags ObjectFlags = RF_Transactional, const FName Name = NAME_None)
+{
+	TArray<FTypedElementHandle> PlacedItems;
+	if (!Factory)
+	{
+		return PlacedItems;
 	}
 
+	// Whereas going throught UPlacementSubsystem does not require the factory to be an actor factory,
+	// other legacy paths require actor factories. These two pointers, when non-null, can be used for 
+	// those paths.
+	UActorFactory* ActorFactory = Cast<UActorFactory>(Factory.GetObject());
+	AActor* NewActorTemplate = ActorFactory ? ActorFactory->GetDefaultActor(Asset) : nullptr;
+
 	UWorld* OldWorld = nullptr;
 
 	// The play world needs to be selected if it exists
@@ -554,20 +597,23 @@ static AActor* PrivateAddActor( UObject* Asset, UActorFactory* Factory, bool Sel
 	}
 
 	// For Brushes/Volumes, use the default brush as the template rather than the factory default actor
-	if (NewActorTemplate->IsA(ABrush::StaticClass()) && GWorld->GetDefaultBrush() != nullptr)
+	if (NewActorTemplate && NewActorTemplate->IsA(ABrush::StaticClass()) && GWorld->GetDefaultBrush() != nullptr)
 	{
 		NewActorTemplate = GWorld->GetDefaultBrush();
 	}
 
+	// TODO: FSnappedPositioningData should probably not require the use of an actor factory
 	const FSnappedPositioningData PositioningData = FSnappedPositioningData(GCurrentLevelEditingViewportClient, GEditor->ClickLocation, GEditor->ClickPlane)
-		.UseFactory(Factory)
-		.UsePlacementExtent(NewActorTemplate->GetPlacementExtent());
+		.UseFactory(ActorFactory)
+		.UsePlacementExtent(NewActorTemplate ? NewActorTemplate->GetPlacementExtent() : FVector3d::Zero());
 
 	FTransform ActorTransform = FActorPositioning::GetSnappedSurfaceAlignedTransform(PositioningData);
 
-	if (GetDefault<ULevelEditorViewportSettings>()->SnapToSurface.bEnabled)
+	if (NewActorTemplate && GetDefault<ULevelEditorViewportSettings>()->SnapToSurface.bEnabled)
 	{
 		// HACK: If we are aligning rotation to surfaces, we have to factor in the inverse of the actor's rotation and translation so that the resulting transform after SpawnActor is correct.
+		
+		// TODO: Do this for non-actor placeable objects
 
 		if (auto* RootComponent = NewActorTemplate->GetRootComponent())
 		{
@@ -581,24 +627,26 @@ static AActor* PrivateAddActor( UObject* Asset, UActorFactory* Factory, bool Sel
 
 	// Do not fade snapping indicators over time if the viewport is not realtime
 	bool bClearImmediately = !GCurrentLevelEditingViewportClient || !GCurrentLevelEditingViewportClient->IsRealtime();
-	FSnappingUtils::ClearSnappingHelpers( bClearImmediately );
+	FSnappingUtils::ClearSnappingHelpers(bClearImmediately);
 
 	ULevel* DesiredLevel = GWorld->GetCurrentLevel();
 
-	bool bSpawnActor = true;
+	bool bShouldSpawnObject = true;
 
 	if ((ObjectFlags & RF_Transactional) != 0)
 	{
 		TArray<FTransform> SpawningActorTransforms;
 		SpawningActorTransforms.Add(ActorTransform);
-		bSpawnActor = ActorPlacementUtils::IsLevelValidForActorPlacement(DesiredLevel, SpawningActorTransforms);
+		// TODO: At the moment, the conditions for placing non-actor items in a level are the same as actor items, but
+		// we should probably rename this function so it is not actor-specific.
+		bShouldSpawnObject = ActorPlacementUtils::IsLevelValidForActorPlacement(DesiredLevel, SpawningActorTransforms);
 	}
 
-	if(bSpawnActor)
+	if (bShouldSpawnObject)
 	{
-		FScopedTransaction Transaction( NSLOCTEXT("UnrealEd", "CreateActor", "Create Actor"), (ObjectFlags & RF_Transactional) != 0 );
-		
-		// Create the actor.
+		FScopedTransaction Transaction(NSLOCTEXT("UnrealEd", "PlaceObject", "Place Object"), (ObjectFlags & RF_Transactional) != 0);
+
+		// Create the object.
 		UPlacementSubsystem* PlacementSubsystem = GEditor->GetEditorSubsystem<UPlacementSubsystem>();
 		if (PlacementSubsystem)
 		{
@@ -612,46 +660,60 @@ static AActor* PrivateAddActor( UObject* Asset, UActorFactory* Factory, bool Sel
 			FPlacementOptions PlacementOptions;
 			PlacementOptions.bIsCreatingPreviewElements = FLevelEditorViewportClient::IsDroppingPreviewActor();
 
-			TArray<FTypedElementHandle> PlacedElements = PlacementSubsystem->PlaceAsset(PlacementInfo, PlacementOptions);
-			if (PlacedElements.Num())
+			PlacedItems = PlacementSubsystem->PlaceAsset(PlacementInfo, PlacementOptions);
+
+			ForEachObjectInHandles(PlacedItems, [ObjectFlags](UObject& PlacedObject)
 			{
-				if (TTypedElement<ITypedElementObjectInterface> ObjectInterface = UTypedElementRegistry::GetInstance()->GetElement<ITypedElementObjectInterface>(PlacedElements[0]))
-				{
-					Actor = ObjectInterface.GetObjectAs<AActor>();
-					if (Actor)
-					{
-						Actor->SetFlags(ObjectFlags);
-					}
-				}
-			}
+				PlacedObject.SetFlags(ObjectFlags);
+			});
 		}
-		
-		if (!Actor)
+
+		// If we fail to place using the placement subsystem above for some reason, we keep this legacy path that 
+		// tries using an actor factory directly. We don't bother adding a fallback for non-actor factories, as
+		// those got introduced after the existence of the placement subsystem, and should rely on it.
+		if (!PlacedItems.Num() && ActorFactory)
 		{
 			FActorSpawnParameters SpawnParams;
 			SpawnParams.ObjectFlags = ObjectFlags;
 			SpawnParams.Name = Name;
-			Actor = Factory->CreateActor(Asset, DesiredLevel, ActorTransform, SpawnParams);
+
+			AActor* PlacedActor = ActorFactory->CreateActor(Asset, DesiredLevel, ActorTransform, SpawnParams);
+			FTypedElementHandle Handle = UEngineElementsLibrary::AcquireEditorActorElementHandle(PlacedActor);
+			if (ensure(Handle))
+			{
+				PlacedItems.Add(Handle);
+			}
 		}
+	}
 
-		if (Actor)
+	if (PlacedItems.Num())
+	{
+		if (bSelectResult)
 		{
-			if ( SelectActor )
+			// TODO: It would be nice not to use this old form of selection clearing, but it has the benefit
+			// of clearing up legacy bsp selection as well...
+			GEditor->SelectNone(false, true);
+
+			UTypedElementSelectionSet* SelectionSet = GetEditorSelectionSet();
+			if (ensure(SelectionSet))
 			{
-				GEditor->SelectNone( false, true );
-				GEditor->SelectActor( Actor, true, true );
+				FTypedElementSelectionOptions SelectionOptions;
+				SelectionSet->SelectElements(PlacedItems, SelectionOptions);
 			}
-
-			Actor->InvalidateLightingCache();
-			Actor->PostEditChange();
 		}
 
-		GEditor->RedrawLevelEditingViewports();
-	}
+		ForEachObjectInHandles(PlacedItems, [](UObject& PlacedObject)
+		{
+			if (AActor* Actor = Cast<AActor>(&PlacedObject))
+			{
+				Actor->InvalidateLightingCache();
+			}
 
-	if ( Actor )
-	{
-		Actor->MarkPackageDirty();
+			PlacedObject.PostEditChange();
+			PlacedObject.MarkPackageDirty();
+		});
+
+		GEditor->RedrawLevelEditingViewports();
 		ULevel::LevelDirtiedEvent.Broadcast();
 	}
 
@@ -661,8 +723,31 @@ static AActor* PrivateAddActor( UObject* Asset, UActorFactory* Factory, bool Sel
 		RestoreEditorWorld(OldWorld);
 	}
 
-	return Actor;
+	return PlacedItems;
+}
+
+/**
+ * Helper to pull out a single actor from an array of typed element handles. Used to convert
+ * output in some legacy paths.
+ */
+AActor* GetActorFromTypedElementHandles(const TArray<FTypedElementHandle>& Handles)
+{
+	for (const FTypedElementHandle& Handle : Handles)
+	{
+		TTypedElement<ITypedElementObjectInterface> ObjectInterface = UTypedElementRegistry::GetInstance()->GetElement<ITypedElementObjectInterface>(Handle);
+		if (!ObjectInterface)
+		{
+			continue;
+		}
+
+		if (AActor* Actor = ObjectInterface.GetObjectAs<AActor>())
+		{
+			return Actor;
+		}
+	}
+	return nullptr;
 }
+}//end namespace AssetSelectionLocals
 
 
 namespace AssetUtil
@@ -825,37 +910,29 @@ UActorFactory* FActorFactoryAssetProxy::GetFactoryForAssetObject( UObject* Asset
 	return Result;
 }
 
-AActor* FActorFactoryAssetProxy::AddActorForAsset( UObject* AssetObj, bool SelectActor, EObjectFlags ObjectFlags, UActorFactory* FactoryToUse /*= NULL*/, const FName Name )
+AActor* FActorFactoryAssetProxy::AddActorForAsset( UObject* AssetObj, bool bSelectActor, EObjectFlags ObjectFlags, 
+	UActorFactory* FactoryToUse /*= NULL*/, const FName Name )
 {
-	AActor* Result = NULL;
+	UE::AssetPlacementUtil::FExtraPlaceAssetOptions Options;
+	Options.bSelectOutput = bSelectActor;
+	Options.ObjectFlags = ObjectFlags;
+	Options.FactoryToUse = FactoryToUse;
+	Options.Name = Name;
 
-	const FAssetData AssetData(AssetObj, FAssetData::ECreationFlags::AllowBlueprintClass);
-	FText UnusedErrorMessage;
-	if ( AssetObj != NULL )
-	{
-		// If a specific factory has been provided, verify its validity and then use it to create the actor
-		if ( FactoryToUse )
-		{
-			if ( FactoryToUse->CanCreateActorFrom( AssetData, UnusedErrorMessage ) )
-			{
-				Result = PrivateAddActor( AssetObj, FactoryToUse, SelectActor, ObjectFlags, Name );
-			}
-		}
-		// If no specific factory has been provided, use the placement subsystem to find the appropriate factory and place
-		else
-		{
-			UPlacementSubsystem* PlacementSubsystem = GEditor->GetEditorSubsystem<UPlacementSubsystem>();
-			TScriptInterface<IAssetFactoryInterface> AssetFactory = PlacementSubsystem->FindAssetFactoryFromAssetData(AssetData);
-			Result = PrivateAddActor(AssetObj, Cast<UActorFactory>(AssetFactory.GetObject()), SelectActor, ObjectFlags, Name);
-		}
-	}
+	TArray<FTypedElementHandle> PlacedItems = UE::AssetPlacementUtil::PlaceAssetInCurrentLevel(AssetObj, Options);
 
+	AActor* Actor = AssetSelectionLocals::GetActorFromTypedElementHandles(PlacedItems);
 
-	return Result;
+	ensureMsgf(Actor || PlacedItems.Num() == 0, TEXT("FActorFactoryAssetProxy::AddActorForAsset produced an object, "
+		"but not an actor. Use UE::AssetFactoryUtils::AddObjectForAssetToCurrentLevel instead to use the result."));
+
+	return Actor;
 }
 
 AActor* FActorFactoryAssetProxy::AddActorFromSelection( UClass* ActorClass, const FVector* ActorLocation, bool SelectActor, EObjectFlags ObjectFlags, UActorFactory* ActorFactory, const FName Name )
 {
+	using namespace AssetSelectionLocals;
+
 	check( ActorClass != NULL );
 
 	if( !ActorFactory )
@@ -874,7 +951,11 @@ AActor* FActorFactoryAssetProxy::AddActorFromSelection( UClass* ActorClass, cons
 		if( TargetObject && ActorFactory->CanCreateActorFrom( FAssetData(TargetObject), ErrorMessage ) )
 		{
 			// Attempt to add the actor
-			Result = PrivateAddActor( TargetObject, ActorFactory, SelectActor, ObjectFlags );
+			TArray<FTypedElementHandle> PlacedItems = PlaceAssetUsingFactory(TargetObject, ActorFactory, SelectActor, ObjectFlags);
+			Result = GetActorFromTypedElementHandles(PlacedItems);
+
+			ensureMsgf(Result || PlacedItems.Num() == 0, TEXT("FActorFactoryAssetProxy::AddActorFromSelection produced a "
+				"result, but not an actor."));
 		}
 	}
 
@@ -980,8 +1061,30 @@ bool FActorFactoryAssetProxy::ApplyMaterialToActor( AActor* TargetActor, UMateri
 }
 
 
-// EOF
+// AssetPlacementUtil
 
+TArray<FTypedElementHandle> UE::AssetPlacementUtil::PlaceAssetInCurrentLevel(UObject* AssetObj, const FExtraPlaceAssetOptions& ExtraParms)
+{
+	using namespace AssetSelectionLocals;
 
+	if (!AssetObj)
+	{
+		return TArray<FTypedElementHandle>();
+	}
 
+	const FAssetData AssetData(AssetObj, FAssetData::ECreationFlags::AllowBlueprintClass);
 
+	if (!ExtraParms.FactoryToUse)
+	{
+		UPlacementSubsystem* PlacementSubsystem = GEditor->GetEditorSubsystem<UPlacementSubsystem>();
+		TScriptInterface<IAssetFactoryInterface> AssetFactory = PlacementSubsystem->FindAssetFactoryFromAssetData(AssetData);
+		return PlaceAssetUsingFactory(AssetObj, AssetFactory, ExtraParms.bSelectOutput, ExtraParms.ObjectFlags, ExtraParms.Name);
+	}
+
+	if (!ExtraParms.FactoryToUse->CanPlaceElementsFromAssetData(AssetData))
+	{
+		return TArray<FTypedElementHandle>();
+	}
+
+	return PlaceAssetUsingFactory(AssetObj, ExtraParms.FactoryToUse, ExtraParms.bSelectOutput, ExtraParms.ObjectFlags, ExtraParms.Name);
+}
diff --git a/Engine/Source/Editor/UnrealEd/Private/DragAndDrop/AssetDragDropOp.cpp b/Engine/Source/Editor/UnrealEd/Private/DragAndDrop/AssetDragDropOp.cpp
--- a/Engine/Source/Editor/UnrealEd/Private/DragAndDrop/AssetDragDropOp.cpp
+++ b/Engine/Source/Editor/UnrealEd/Private/DragAndDrop/AssetDragDropOp.cpp
@@ -11,16 +11,16 @@
 #include "ThumbnailRendering/ThumbnailManager.h"
 #include "Engine/Level.h"
 
-TSharedRef<FAssetDragDropOp> FAssetDragDropOp::New(const FAssetData& InAssetData, UActorFactory* ActorFactory)
+TSharedRef<FAssetDragDropOp> FAssetDragDropOp::New(const FAssetData& InAssetData, TScriptInterface<IAssetFactoryInterface> AssetFactory)
 {
 	TArray<FAssetData> AssetDataArray;
 	AssetDataArray.Emplace(InAssetData);
-	return New(MoveTemp(AssetDataArray), TArray<FString>(), ActorFactory);
+	return New(MoveTemp(AssetDataArray), TArray<FString>(), AssetFactory);
 }
 
-TSharedRef<FAssetDragDropOp> FAssetDragDropOp::New(TArray<FAssetData> InAssetData, UActorFactory* ActorFactory)
+TSharedRef<FAssetDragDropOp> FAssetDragDropOp::New(TArray<FAssetData> InAssetData, TScriptInterface<IAssetFactoryInterface> AssetFactory)
 {
-	return New(MoveTemp(InAssetData), TArray<FString>(), ActorFactory);
+	return New(MoveTemp(InAssetData), TArray<FString>(), AssetFactory);
 }
 
 TSharedRef<FAssetDragDropOp> FAssetDragDropOp::New(FString InAssetPath)
@@ -35,16 +35,41 @@ TSharedRef<FAssetDragDropOp> FAssetDragDropOp::New(TArray<FString> InAssetPaths)
 	return New(TArray<FAssetData>(), MoveTemp(InAssetPaths), nullptr);
 }
 
-TSharedRef<FAssetDragDropOp> FAssetDragDropOp::New(TArray<FAssetData> InAssetData, TArray<FString> InAssetPaths, UActorFactory* ActorFactory)
+TSharedRef<FAssetDragDropOp> FAssetDragDropOp::New(TArray<FAssetData> InAssetData, TArray<FString> InAssetPaths, TScriptInterface<IAssetFactoryInterface> AssetFactory)
 {
 	TSharedRef<FAssetDragDropOp> Operation = MakeShared<FAssetDragDropOp>();
 
-	Operation->Init(MoveTemp(InAssetData), MoveTemp(InAssetPaths), ActorFactory);
+	Operation->Init(MoveTemp(InAssetData), MoveTemp(InAssetPaths), AssetFactory);
 
 	Operation->Construct();
 	return Operation;
 }
 
+TSharedRef<FAssetDragDropOp> FAssetDragDropOp::New(const FAssetData& InAssetData, UActorFactory* ActorFactory)
+{
+	return New(InAssetData, TScriptInterface<IAssetFactoryInterface>(ActorFactory));
+}
+
+TSharedRef<FAssetDragDropOp> FAssetDragDropOp::New(TArray<FAssetData> InAssetData, UActorFactory* ActorFactory)
+{
+	return New(InAssetData, TScriptInterface<IAssetFactoryInterface>(ActorFactory));
+}
+
+TSharedRef<FAssetDragDropOp> FAssetDragDropOp::New(TArray<FAssetData> InAssetData, TArray<FString> InAssetPaths, UActorFactory* ActorFactory)
+{
+	return New(InAssetData, InAssetPaths, TScriptInterface<IAssetFactoryInterface>(ActorFactory));
+}
+
+UActorFactory* FAssetDragDropOp::GetActorFactory() const
+{
+	return Cast<UActorFactory>(AssetFactory.GetObject());
+}
+
+TScriptInterface<IAssetFactoryInterface> FAssetDragDropOp::GetAssetFactory() const
+{
+	return AssetFactory.GetObject();
+}
+
 FAssetDragDropOp::~FAssetDragDropOp()
 {
 }
@@ -90,10 +115,14 @@ TSharedPtr<SWidget> FAssetDragDropOp::GetDefaultDecorator() const
 		SubTypeBrush = FAppStyle::GetBrush("ContentBrowser.AssetTreeFolderClosed");
 		SubTypeColor = FLinearColor::Gray;
 	}
-	else if (ActorFactory.IsValid() && HasFiles())
+	else if (AssetFactory.IsValid() && HasFiles())
 	{
-		AActor* DefaultActor = ActorFactory->GetDefaultActor(AssetData[0]);
-		SubTypeBrush = FClassIconFinder::FindIconForActor(DefaultActor);
+		// TODO: Probably need to add a function in IAssetFactoryInterface for this and use that.
+		if (UActorFactory* ActorFactory = Cast<UActorFactory>(AssetFactory.GetObject()))
+		{
+			AActor* DefaultActor = ActorFactory->GetDefaultActor(AssetData[0]);
+			SubTypeBrush = FClassIconFinder::FindIconForActor(DefaultActor);
+		}
 	}
 
 	return 
@@ -198,14 +227,14 @@ FText FAssetDragDropOp::GetDecoratorText() const
 	return CurrentHoverText;
 }
 
-void FAssetDragDropOp::Init(TArray<FAssetData> InAssetData, TArray<FString> InAssetPaths, UActorFactory* InActorFactory)
+void FAssetDragDropOp::Init(TArray<FAssetData> InAssetData, TArray<FString> InAssetPaths, TScriptInterface<IAssetFactoryInterface> InAssetFactory)
 {
 	MouseCursor = EMouseCursor::GrabHandClosed;
 	ThumbnailSize = 64;
 
 	AssetData = MoveTemp(InAssetData);
 	AssetPaths = MoveTemp(InAssetPaths);
-	ActorFactory = InActorFactory;
+	AssetFactory = InAssetFactory;
 
 	// Load all assets first so that there is no loading going on while attempting to drag
 	// Can cause unsafe frame reentry 
@@ -217,6 +246,11 @@ void FAssetDragDropOp::Init(TArray<FAssetData> InAssetData, TArray<FString> InAs
 	InitThumbnail();
 }
 
+void FAssetDragDropOp::Init(TArray<FAssetData> InAssetData, TArray<FString> InAssetPaths, UActorFactory* InActorFactory)
+{
+	Init(InAssetData, InAssetPaths, TScriptInterface<IAssetFactoryInterface>(InActorFactory));
+}
+
 void FAssetDragDropOp::InitThumbnail()
 {
 	if (AssetData.Num() > 0 && ThumbnailSize > 0)
diff --git a/Engine/Source/Editor/UnrealEd/Private/EditorViewportClient.cpp b/Engine/Source/Editor/UnrealEd/Private/EditorViewportClient.cpp
--- a/Engine/Source/Editor/UnrealEd/Private/EditorViewportClient.cpp
+++ b/Engine/Source/Editor/UnrealEd/Private/EditorViewportClient.cpp
@@ -1,6 +1,11 @@
 // Copyright Epic Games, Inc. All Rights Reserved.
 
 #include "EditorViewportClient.h"
+
+#include "ActorFactories/ActorFactory.h"
+#include "Elements/Framework/EngineElementsLibrary.h"
+#include "Elements/Framework/TypedElementRegistry.h"
+#include "Elements/Interfaces/TypedElementObjectInterface.h"
 #include "PreviewScene.h"
 #include "HAL/FileManager.h"
 #include "Misc/FileHelper.h"
@@ -5063,6 +5068,29 @@ bool FEditorViewportClient::IsMovingCamera() const
 	return bUsingOrbitCamera || IsFlightCameraActive();
 }
 
+bool FEditorViewportClient::DropObjectsAtCoordinates(int32 MouseX, int32 MouseY, const TArray<UObject*>& DroppedObjects, 
+	TArray<FTypedElementHandle>& OutNewObjects, const FDropObjectOptions& Options)
+{
+	// Forward things to the deprecated overload while it still exists, so that we don't break any
+	// existing derivations of FEditorViewportClient. Once removed, this function will just return false.
+	TArray<AActor*> OutputActors;
+PRAGMA_DISABLE_DEPRECATION_WARNINGS
+	bool bSuccess = DropObjectsAtCoordinates(MouseX, MouseY, DroppedObjects, OutputActors,
+		Options.bOnlyDropOnTarget, Options.bCreateDropPreview, Options.bSelectOutput,
+		Cast<UActorFactory>(Options.FactoryToUse.GetObject()));
+PRAGMA_ENABLE_DEPRECATION_WARNINGS
+
+	TArray<FTypedElementHandle> OutputElements;
+
+	for (const AActor* Actor : OutputActors)
+	{
+		FTypedElementHandle Handle = UEngineElementsLibrary::AcquireEditorActorElementHandle(Actor);
+		OutputElements.Add(Handle);
+	}
+
+	return bSuccess;
+}
+
 /** True if the window is maximized or floating */
 bool FEditorViewportClient::IsVisible() const
 {
diff --git a/Engine/Source/Editor/UnrealEd/Private/LevelEditorViewport.cpp b/Engine/Source/Editor/UnrealEd/Private/LevelEditorViewport.cpp
--- a/Engine/Source/Editor/UnrealEd/Private/LevelEditorViewport.cpp
+++ b/Engine/Source/Editor/UnrealEd/Private/LevelEditorViewport.cpp
@@ -49,6 +49,7 @@
 #include "Elements/Framework/TypedElementCommonActions.h"
 #include "Elements/Framework/TypedElementListObjectUtil.h"
 #include "Elements/Framework/TypedElementViewportInteraction.h"
+#include "Elements/Interfaces/TypedElementObjectInterface.h"
 #include "Elements/Actor/ActorElementLevelEditorViewportInteractionCustomization.h"
 #include "Elements/Component/ComponentElementLevelEditorViewportInteractionCustomization.h"
 #include "AudioDevice.h"
@@ -114,6 +115,7 @@ DEFINE_LOG_CATEGORY(LogEditorViewport);
 #define LOCTEXT_NAMESPACE "LevelEditorViewportClient"
 const FLevelViewportActorLock FLevelViewportActorLock::None(nullptr);
 
+TWeakPtr<FTypedElementList> FLevelEditorViewportClient::StaticDropPreviewElements;
 TArray< TWeakObjectPtr< AActor > > FLevelEditorViewportClient::DropPreviewActors;
 TMap< TObjectKey< AActor >, TWeakObjectPtr< UActorComponent > > FLevelEditorViewportClient::ViewComponentForActorCache;
 
@@ -124,6 +126,39 @@ TSet<FViewportHoverTarget> FLevelEditorViewportClient::HoveredObjects;
 
 IMPLEMENT_HIT_PROXY( HLevelSocketProxy, HHitProxy );
 
+namespace LevelEditorViewportLocals
+{
+	const FText DropObjectsTransactionName = LOCTEXT("PlaceObjects", "Place Objects");
+
+	/**
+	 * Calls Func on each element in ElementArray that represents an actor.
+	 */
+	void ForEachActorInElementArray(TArray<FTypedElementHandle> ElementArray, TFunctionRef<void(AActor&)> Func)
+	{
+		UTypedElementRegistry* Registry = UTypedElementRegistry::GetInstance();
+		if (!ensure(Registry))
+		{
+			return;
+		}
+
+		for (const FTypedElementHandle& Element : ElementArray)
+		{
+			ITypedElementObjectInterface* ObjectInterface = Registry->GetElementInterface<ITypedElementObjectInterface>(Element);
+			if (!ObjectInterface)
+			{
+				continue;
+			}
+			AActor* Actor = ObjectInterface->GetObjectAs<AActor>(Element);
+			if (!Actor)
+			{
+				continue;
+			}
+
+			Func(*Actor);
+		}
+	}
+}
+
 /** Helper function to compute a new location that is snapped to the origin plane given the users cursor location and camera angle */
 static FVector4 AttemptToSnapLocationToOriginPlane( const FViewportCursorLocation& Cursor, FVector4 Location )
 {
@@ -209,59 +244,108 @@ static void NotifyAtmosphericLightHasMoved(UDirectionalLightComponent& SelectedA
 	}
 }
 
-TArray<AActor*> FLevelEditorViewportClient::TryPlacingActorFromObject( ULevel* InLevel, UObject* ObjToUse, bool bSelectActors, EObjectFlags ObjectFlags, UActorFactory* FactoryToUse, const FName Name, const FViewportCursorLocation* Cursor )
+TArray<FTypedElementHandle> FLevelEditorViewportClient::TryPlacingAssetObject(ULevel* InLevel, UObject* ObjToUse,
+	const UE::AssetPlacementUtil::FExtraPlaceAssetOptions& PlacementOptions,
+	const FViewportCursorLocation* CursorInformation)
 {
+	using namespace LevelEditorViewportLocals;
+
+	// This is used for issuing a legacy notification.
 	TArray<AActor*> PlacedActors;
 
+	// Our actual output.
+	TArray<FTypedElementHandle> PlacedItems;
+
+	// Helper to update PlacedActors and PlacedItems. Can be called with null actor.
+	auto AddActorToPlaced = [&PlacedItems, &PlacedActors](AActor* Actor)
+	{
+		if (!Actor)
+		{
+			return;
+		}
+
+		PlacedActors.Add(Actor);
+
+		FTypedElementHandle Handle = UEngineElementsLibrary::AcquireEditorActorElementHandle(Actor);
+		if (ensure(Handle))
+		{
+			PlacedItems.Add(Handle);
+		}
+	};
+
+	if (!ensure(ObjToUse))
+	{
+		return PlacedItems;
+	}
+
 	UClass* ObjectClass = Cast<UClass>(ObjToUse);
 
 	if ( ObjectClass == NULL )
 	{
 		ObjectClass = ObjToUse->GetClass();
-		check(ObjectClass);
 	}
 
-	AActor* PlacedActor = NULL;
+	if (!ensure(ObjectClass))
+	{
+		return PlacedItems;
+	}
+
+	bool bPlaced = false;
 	if ( ObjectClass->IsChildOf( AActor::StaticClass() ) )
 	{
 		//Attempting to drop a UClass object
-		UActorFactory* ActorFactory = FactoryToUse;
+		UActorFactory* ActorFactory = Cast<UActorFactory>(PlacementOptions.FactoryToUse.GetObject());
 		if ( ActorFactory == NULL )
 		{
 			ActorFactory = GEditor->FindActorFactoryForActorClass( ObjectClass );
 		}
 
+		// TODO: This path looks to be an ancient and unused one- AddActorFromSelection queries a seldom
+		// used separate object selection set, and replacing ObjToUse with whatever happens to be stored
+		// there (and only if ObjToUse happened to be an actor) does not make sense. This should be removed.
 		if ( ActorFactory != NULL )
 		{
-			PlacedActor = FActorFactoryAssetProxy::AddActorFromSelection( ObjectClass, NULL, bSelectActors, ObjectFlags, ActorFactory, Name );
+			AActor* Actor = FActorFactoryAssetProxy::AddActorFromSelection(ObjectClass, NULL, 
+				PlacementOptions.bSelectOutput, PlacementOptions.ObjectFlags, ActorFactory);
+			AddActorToPlaced(Actor);
+			bPlaced = (Actor != nullptr);
 		}
 
-		if ( PlacedActor == NULL && ActorFactory != NULL )
+		if (!bPlaced && ActorFactory != NULL )
 		{
-			PlacedActor = FActorFactoryAssetProxy::AddActorForAsset( ObjToUse, bSelectActors, ObjectFlags, ActorFactory, Name );
+			PlacedItems = UE::AssetPlacementUtil::PlaceAssetInCurrentLevel(ObjToUse, PlacementOptions);
+			ForEachActorInElementArray(PlacedItems, [&PlacedActors](AActor& Actor)
+			{
+				PlacedActors.Add(&Actor);
+			});
+			bPlaced = PlacedItems.Num() > 0;
 		}
 		
-		if ( PlacedActor == NULL && !ObjectClass->HasAnyClassFlags(CLASS_NotPlaceable | CLASS_Abstract) )
+		if (!bPlaced && !ObjectClass->HasAnyClassFlags(CLASS_NotPlaceable | CLASS_Abstract) )
 		{
 			// If no actor factory was found or failed, add the actor directly.
-			const FTransform ActorTransform = FActorPositioning::GetCurrentViewportPlacementTransform(*ObjectClass->GetDefaultObject<AActor>(), /*bSnap=*/true, Cursor);
-			PlacedActor = GEditor->AddActor( InLevel, ObjectClass, ActorTransform, /*bSilent=*/false, ObjectFlags, bSelectActors );
-		}
-
-		if ( PlacedActor != NULL )
-		{
-			FVector Collision = ObjectClass->GetDefaultObject<AActor>()->GetPlacementExtent();
-			PlacedActors.Add(PlacedActor);
+			const FTransform ActorTransform = FActorPositioning::GetCurrentViewportPlacementTransform(
+				*ObjectClass->GetDefaultObject<AActor>(), /*bSnap=*/true, CursorInformation);
+			AActor* Actor = GEditor->AddActor( InLevel, ObjectClass, ActorTransform, 
+				/*bSilent=*/false, PlacementOptions.ObjectFlags, PlacementOptions.bSelectOutput);
+			AddActorToPlaced(Actor);
+			bPlaced = (Actor != nullptr);
 		}
 	}
 	
-	if ( (NULL == PlacedActor) && ObjToUse->IsA( UExportTextContainer::StaticClass() ) )
+	if (!bPlaced && ObjToUse->IsA( UExportTextContainer::StaticClass() ) )
 	{
+		// TODO: This path probably needs fixing for non-actor placement
+
 		UExportTextContainer* ExportContainer = CastChecked<UExportTextContainer>(ObjToUse);
-		const TArray<AActor*> NewActors = GEditor->AddExportTextActors( ExportContainer->ExportText, /*bSilent*/false, ObjectFlags);
-		PlacedActors.Append(NewActors);
+		const TArray<AActor*> NewActors = GEditor->AddExportTextActors( ExportContainer->ExportText, 
+			/*bSilent*/false, PlacementOptions.ObjectFlags);
+		for (AActor* Actor : NewActors)
+		{
+			AddActorToPlaced(Actor);
+		}
 	}
-	else if ( (NULL == PlacedActor) && ObjToUse->IsA( UBrushBuilder::StaticClass() ) )
+	else if (!bPlaced && ObjToUse->IsA( UBrushBuilder::StaticClass() ) )
 	{
 		UBrushBuilder* BrushBuilder = CastChecked<UBrushBuilder>(ObjToUse);
 		UWorld* World = InLevel->OwningWorld;
@@ -274,35 +358,36 @@ TArray<AActor*> FLevelEditorViewportClient::TryPlacingActorFromObject( ULevel* I
 			FSnappingUtils::SnapPointToGrid(ActorLoc, FVector::ZeroVector);
 
 			DefaultBrush->SetActorLocation(ActorLoc);
-			PlacedActor = DefaultBrush;
-			PlacedActors.Add(DefaultBrush);
+			AddActorToPlaced(DefaultBrush);
 		}
 	}
-	else if (NULL == PlacedActor)
+	else if (!bPlaced)
 	{
-		bool bPlace = true;
+		bool bShouldPlace = true;
 		if (ObjectClass->IsChildOf(UBlueprint::StaticClass()))
 		{
 			UBlueprint* BlueprintObj = StaticCast<UBlueprint*>(ObjToUse);
-			bPlace = BlueprintObj->GeneratedClass != NULL;
-			if(bPlace)
+			bShouldPlace = BlueprintObj->GeneratedClass != NULL;
+			if (bShouldPlace)
 			{
 				check(BlueprintObj->ParentClass == BlueprintObj->GeneratedClass->GetSuperClass());
 				if (BlueprintObj->GeneratedClass->HasAnyClassFlags(CLASS_NotPlaceable | CLASS_Abstract))
 				{
-					bPlace = false;
+					bShouldPlace = false;
 				}
 			}
 		}
 
-		if (bPlace)
+		if (bShouldPlace)
 		{
-			PlacedActor = FActorFactoryAssetProxy::AddActorForAsset( ObjToUse, bSelectActors, ObjectFlags, FactoryToUse, Name );
-			if ( PlacedActor != NULL )
+			PlacedItems = UE::AssetPlacementUtil::PlaceAssetInCurrentLevel(ObjToUse, PlacementOptions);
+			
+			LevelEditorViewportLocals::ForEachActorInElementArray(PlacedItems, [&PlacedActors](AActor& Actor)
 			{
-				PlacedActors.Add(PlacedActor);
-				PlacedActor->PostEditMove(true);
-			}
+				PlacedActors.Add(&Actor);
+				// Not clear why this would be necessary, but kept for legacy behavior.
+				Actor.PostEditMove(true);
+			});
 		}
 	}
 
@@ -311,10 +396,32 @@ TArray<AActor*> FLevelEditorViewportClient::TryPlacingActorFromObject( ULevel* I
 		FEditorDelegates::OnNewActorsPlaced.Broadcast(ObjToUse, PlacedActors);
 	}
 
-	return PlacedActors;
+	return PlacedItems;
 }
 
+TArray<AActor*> FLevelEditorViewportClient::TryPlacingActorFromObject(ULevel* InLevel, UObject* ObjToUse,
+	bool bSelectActors, EObjectFlags ObjectFlags, UActorFactory* FactoryToUse,
+	const FName Name, const FViewportCursorLocation* Cursor)
+{
+	TArray<AActor*> OutputActors;
+
+	// Forward the call to TryPlacingAssetObject
+	UE::AssetPlacementUtil::FExtraPlaceAssetOptions PlacementOptions;
+	PlacementOptions.bSelectOutput = bSelectActors;
+	PlacementOptions.ObjectFlags = ObjectFlags;
+	PlacementOptions.FactoryToUse = FactoryToUse;
+	PlacementOptions.Name = Name;
 
+	TArray<FTypedElementHandle> PlacedItems = TryPlacingAssetObject(InLevel, ObjToUse, PlacementOptions, Cursor);
+
+	// Filter out actors into the actors array
+	LevelEditorViewportLocals::ForEachActorInElementArray(PlacedItems, [&OutputActors](AActor& Actor)
+	{
+		OutputActors.Add(&Actor);
+	});
+
+	return OutputActors;
+}
 
 static FString GetSharedTextureNameAndKind( FString TextureName, EMaterialKind& Kind)
 {
@@ -1071,7 +1178,9 @@ static bool AreAnyDroppedObjectsBrushBuilders(const TArray<UObject*>& DroppedObj
 	return false;
 }
 
-bool FLevelEditorViewportClient::DropObjectsOnBackground(FViewportCursorLocation& Cursor, const TArray<UObject*>& DroppedObjects, EObjectFlags ObjectFlags, TArray<AActor*>& OutNewActors, bool bCreateDropPreview, bool bSelectActors, UActorFactory* FactoryToUse)
+bool FLevelEditorViewportClient::DropObjectsOnBackground(FViewportCursorLocation& Cursor, 
+	const TArray<UObject*>& DroppedObjects, EObjectFlags ObjectFlags, TArray<FTypedElementHandle>& OutNewItems, 
+	const FDropObjectOptions& Options)
 {
 	if (DroppedObjects.Num() == 0)
 	{
@@ -1080,12 +1189,12 @@ bool FLevelEditorViewportClient::DropObjectsOnBackground(FViewportCursorLocation
 
 	bool bSuccess = false;
 
-	const bool bTransacted = !bCreateDropPreview && !AreAllDroppedObjectsBrushBuilders(DroppedObjects);
+	const bool bTransacted = !Options.bCreateDropPreview && !AreAllDroppedObjectsBrushBuilders(DroppedObjects);
 
 	// Create a transaction if not a preview drop
 	if (bTransacted)
 	{
-		GEditor->BeginTransaction(NSLOCTEXT("UnrealEd", "CreateActors", "Create Actors"));
+		GEditor->BeginTransaction(LevelEditorViewportLocals::DropObjectsTransactionName);
 	}
 
 	for ( int32 DroppedObjectsIdx = 0; DroppedObjectsIdx < DroppedObjects.Num(); ++DroppedObjectsIdx )
@@ -1094,11 +1203,17 @@ bool FLevelEditorViewportClient::DropObjectsOnBackground(FViewportCursorLocation
 		ensure( AssetObj );
 
 		// Attempt to create actors from the dropped object
-		TArray<AActor*> NewActors = TryPlacingActorFromObject(GetWorld()->GetCurrentLevel(), AssetObj, bSelectActors, ObjectFlags, FactoryToUse, NAME_None, &Cursor);
+		UE::AssetPlacementUtil::FExtraPlaceAssetOptions PlacementOptions;
+		PlacementOptions.bSelectOutput = Options.bSelectOutput;
+		PlacementOptions.ObjectFlags = ObjectFlags;
+		PlacementOptions.FactoryToUse = Options.FactoryToUse;
+		PlacementOptions.Name = NAME_None;
+		TArray<FTypedElementHandle> PlacedItems = TryPlacingAssetObject(GetWorld()->GetCurrentLevel(), AssetObj,
+			PlacementOptions, &Cursor);
 
-		if ( NewActors.Num() > 0 )
+		if (PlacedItems.Num() > 0)
 		{
-			OutNewActors.Append(NewActors);
+			OutNewItems.Append(PlacedItems);
 			bSuccess = true;
 		}
 	}
@@ -1111,7 +1226,28 @@ bool FLevelEditorViewportClient::DropObjectsOnBackground(FViewportCursorLocation
 	return bSuccess;
 }
 
-bool FLevelEditorViewportClient::DropObjectsOnActor(FViewportCursorLocation& Cursor, const TArray<UObject*>& DroppedObjects, AActor* DroppedUponActor, int32 DroppedUponSlot, EObjectFlags ObjectFlags, TArray<AActor*>& OutNewActors, bool bCreateDropPreview, bool bSelectActors, UActorFactory* FactoryToUse)
+bool FLevelEditorViewportClient::DropObjectsOnBackground(FViewportCursorLocation& Cursor, const TArray<UObject*>& DroppedObjects, EObjectFlags ObjectFlags, 
+	TArray<AActor*>& OutNewActors, bool bCreateDropPreview, bool bSelectOutput, UActorFactory* FactoryToUse)
+{
+	FDropObjectOptions Options;
+	Options.bCreateDropPreview = bCreateDropPreview;
+	Options.bSelectOutput = bSelectOutput;
+	Options.FactoryToUse = FactoryToUse;
+
+	TArray<FTypedElementHandle> NewItems;
+	bool bSuccess = DropObjectsOnBackground(Cursor, DroppedObjects, ObjectFlags, NewItems, Options);
+
+	LevelEditorViewportLocals::ForEachActorInElementArray(NewItems, [&OutNewActors](AActor& Actor)
+	{
+		OutNewActors.Add(&Actor);
+	});
+	return bSuccess;
+}
+
+bool FLevelEditorViewportClient::DropObjectsOnActor(FViewportCursorLocation& Cursor, 
+	const TArray<UObject*>& DroppedObjects, AActor* DroppedUponActor, 
+	int32 DroppedUponSlot, EObjectFlags ObjectFlags,
+	TArray<FTypedElementHandle>& OutNewItems, const FDropObjectOptions& Options)
 {
 	if ( !DroppedUponActor || DroppedObjects.Num() == 0 )
 	{
@@ -1120,27 +1256,34 @@ bool FLevelEditorViewportClient::DropObjectsOnActor(FViewportCursorLocation& Cur
 
 	bool bSuccess = false;
 
-	const bool bTransacted = !bCreateDropPreview && !AreAllDroppedObjectsBrushBuilders(DroppedObjects);
+	const bool bTransacted = !Options.bCreateDropPreview && !AreAllDroppedObjectsBrushBuilders(DroppedObjects);
 
 	// Create a transaction if not a preview drop
 	if (bTransacted)
 	{
-		GEditor->BeginTransaction(NSLOCTEXT("UnrealEd", "CreateActors", "Create Actors"));
+		GEditor->BeginTransaction(LevelEditorViewportLocals::DropObjectsTransactionName);
 	}
 
 	for ( UObject* DroppedObject : DroppedObjects )
 	{
-		const bool bIsTestApplication = bCreateDropPreview;
-		const bool bAppliedToActor = ( FactoryToUse == nullptr ) ? AttemptApplyObjToActor( DroppedObject, DroppedUponActor, DroppedUponSlot, bIsTestApplication ) : false;
+		const bool bIsTestApplication = Options.bCreateDropPreview;
+		const bool bAppliedToActor = (Options.FactoryToUse == nullptr ) ? AttemptApplyObjToActor( DroppedObject, DroppedUponActor, DroppedUponSlot, bIsTestApplication ) : false;
 
 		if (!bAppliedToActor)
 		{
 			// Attempt to create actors from the dropped object
-			TArray<AActor*> NewActors = TryPlacingActorFromObject(GetWorld()->GetCurrentLevel(), DroppedObject, bSelectActors, ObjectFlags, FactoryToUse, NAME_None, &Cursor);
+			UE::AssetPlacementUtil::FExtraPlaceAssetOptions PlacementOptions;
+			PlacementOptions.bSelectOutput = Options.bSelectOutput;
+			PlacementOptions.ObjectFlags = ObjectFlags;
+			PlacementOptions.FactoryToUse = Options.FactoryToUse;
+			PlacementOptions.Name = NAME_None;
+
+			TArray<FTypedElementHandle> NewItems = TryPlacingAssetObject(GetWorld()->GetCurrentLevel(), 
+				DroppedObject, PlacementOptions, &Cursor);
 
-			if ( NewActors.Num() > 0 )
+			if (NewItems.Num() > 0 )
 			{
-				OutNewActors.Append( NewActors );
+				OutNewItems.Append(NewItems);
 				bSuccess = true;
 			}
 		}
@@ -1158,7 +1301,28 @@ bool FLevelEditorViewportClient::DropObjectsOnActor(FViewportCursorLocation& Cur
 	return bSuccess;
 }
 
-bool FLevelEditorViewportClient::DropObjectsOnBSPSurface(FSceneView* View, FViewportCursorLocation& Cursor, const TArray<UObject*>& DroppedObjects, HModel* TargetProxy, EObjectFlags ObjectFlags, TArray<AActor*>& OutNewActors, bool bCreateDropPreview, bool bSelectActors, UActorFactory* FactoryToUse)
+bool FLevelEditorViewportClient::DropObjectsOnActor(FViewportCursorLocation& Cursor, const TArray<UObject*>& DroppedObjects, 
+	AActor* DroppedUponActor, int32 DroppedUponSlot, EObjectFlags ObjectFlags, 
+	TArray<AActor*>& OutNewActors, bool bCreateDropPreview, bool bSelectOutput, UActorFactory* FactoryToUse)
+{
+	FDropObjectOptions Options;
+	Options.bCreateDropPreview = bCreateDropPreview;
+	Options.bSelectOutput = bSelectOutput;
+	Options.FactoryToUse = FactoryToUse;
+
+	TArray<FTypedElementHandle> NewItems;
+	bool bSuccess = DropObjectsOnActor(Cursor, DroppedObjects, DroppedUponActor, DroppedUponSlot, ObjectFlags, NewItems, Options);
+
+	LevelEditorViewportLocals::ForEachActorInElementArray(NewItems, [&OutNewActors](AActor& Actor)
+	{
+		OutNewActors.Add(&Actor);
+	});
+	return bSuccess;
+}
+
+bool FLevelEditorViewportClient::DropObjectsOnBSPSurface(FSceneView* View, FViewportCursorLocation& Cursor, 
+	const TArray<UObject*>& DroppedObjects, HModel* TargetProxy, EObjectFlags ObjectFlags, 
+	TArray<FTypedElementHandle>& OutNewItems, const FDropObjectOptions& Options)
 {
 	if (DroppedObjects.Num() == 0)
 	{
@@ -1167,26 +1331,33 @@ bool FLevelEditorViewportClient::DropObjectsOnBSPSurface(FSceneView* View, FView
 
 	bool bSuccess = false;
 
-	const bool bTransacted = !bCreateDropPreview && !AreAllDroppedObjectsBrushBuilders(DroppedObjects);
+	const bool bTransacted = !Options.bCreateDropPreview && !AreAllDroppedObjectsBrushBuilders(DroppedObjects);
 
 	// Create a transaction if not a preview drop
 	if (bTransacted)
 	{
-		GEditor->BeginTransaction(NSLOCTEXT("UnrealEd", "CreateActors", "Create Actors"));
+		GEditor->BeginTransaction(LevelEditorViewportLocals::DropObjectsTransactionName);
 	}
 
 	for (UObject* DroppedObject : DroppedObjects)
 	{
-		const bool bAppliedToActor = (!bCreateDropPreview && FactoryToUse == NULL) ? AttemptApplyObjAsMaterialToSurface(DroppedObject, TargetProxy, Cursor) : false;
+		const bool bAppliedToActor = (!Options.bCreateDropPreview && Options.FactoryToUse == NULL) ? AttemptApplyObjAsMaterialToSurface(DroppedObject, TargetProxy, Cursor) : false;
 
 		if (!bAppliedToActor)
 		{
 			// Attempt to create actors from the dropped object
-			TArray<AActor*> NewActors = TryPlacingActorFromObject(GetWorld()->GetCurrentLevel(), DroppedObject, bSelectActors, ObjectFlags, FactoryToUse, NAME_None, &Cursor);
+			UE::AssetPlacementUtil::FExtraPlaceAssetOptions PlacementOptions;
+			PlacementOptions.bSelectOutput = Options.bSelectOutput;
+			PlacementOptions.ObjectFlags = ObjectFlags;
+			PlacementOptions.FactoryToUse = Options.FactoryToUse;
+			PlacementOptions.Name = NAME_None;
+
+			TArray<FTypedElementHandle> NewItems = TryPlacingAssetObject(GetWorld()->GetCurrentLevel(), 
+				DroppedObject, PlacementOptions, &Cursor);
 
-			if (NewActors.Num() > 0)
+			if (NewItems.Num() > 0)
 			{
-				OutNewActors.Append(NewActors);
+				OutNewItems.Append(NewItems);
 				bSuccess = true;
 			}
 		}
@@ -1204,6 +1375,25 @@ bool FLevelEditorViewportClient::DropObjectsOnBSPSurface(FSceneView* View, FView
 	return bSuccess;
 }
 
+bool FLevelEditorViewportClient::DropObjectsOnBSPSurface(FSceneView* View, FViewportCursorLocation& Cursor,
+	const TArray<UObject*>& DroppedObjects, HModel* TargetProxy, EObjectFlags ObjectFlags, 
+	TArray<AActor*>& OutNewActors, bool bCreateDropPreview, bool bSelectOutput, UActorFactory* FactoryToUse)
+{
+	FDropObjectOptions Options;
+	Options.bCreateDropPreview = bCreateDropPreview;
+	Options.bSelectOutput = bSelectOutput;
+	Options.FactoryToUse = FactoryToUse;
+
+	TArray<FTypedElementHandle> NewItems;
+	bool bSuccess = DropObjectsOnBSPSurface(View, Cursor, DroppedObjects, TargetProxy, ObjectFlags, NewItems, Options);
+	
+	LevelEditorViewportLocals::ForEachActorInElementArray(NewItems, [&OutNewActors](AActor& Actor)
+	{
+		OutNewActors.Add(&Actor);
+	});
+	return bSuccess;
+}
+
 /**
  * Called when an asset is dropped upon a manipulation widget.
  *
@@ -1241,10 +1431,13 @@ bool FLevelEditorViewportClient::DropObjectsOnWidget(FSceneView* View, FViewport
 	check( !HitProxy || ( HitProxy && !HitProxy->IsA( HWidgetAxis::StaticGetType() ) ) );
 
 	// Try this again, but without the widgets this time!
-	TArray< AActor* > TemporaryActors;
+	TArray<FTypedElementHandle> NewObjects;
 	const FIntPoint& CursorPos = Cursor.GetCursorPos();
-	const bool bOnlyDropOnTarget = false;
-	bResult = DropObjectsAtCoordinates(CursorPos.X, CursorPos.Y, DroppedObjects, TemporaryActors, bOnlyDropOnTarget, bCreateDropPreview);
+
+	FDropObjectOptions DropOptions;
+	DropOptions.bOnlyDropOnTarget = false;
+	DropOptions.bCreateDropPreview = bCreateDropPreview;
+	bResult = DropObjectsAtCoordinates(CursorPos.X, CursorPos.Y, DroppedObjects, NewObjects, DropOptions);
 
 	// Restore the original flags
 	EngineShowFlags.SetModeWidgets(bOldModeWidgets1);
@@ -1255,7 +1448,7 @@ bool FLevelEditorViewportClient::DropObjectsOnWidget(FSceneView* View, FViewport
 
 bool FLevelEditorViewportClient::HasDropPreviewActors() const
 {
-	return DropPreviewActors.Num() > 0;
+	return HasDropPreviewElements();
 }
 
 /* Helper functions to find a dropped position on a 2D layer */
@@ -1303,9 +1496,26 @@ static FActorPositionTraceResult TraceForPositionOn2DLayer(const FViewportCursor
 }
 
 bool FLevelEditorViewportClient::UpdateDropPreviewActors(int32 MouseX, int32 MouseY, const TArray<UObject*>& DroppedObjects, bool& out_bDroppedObjectsVisible, UActorFactory* FactoryToUse)
+{
+	return UpdateDropPreviewElements(MouseX, MouseY, DroppedObjects, out_bDroppedObjectsVisible, FactoryToUse);
+}
+
+void FLevelEditorViewportClient::DestroyDropPreviewActors()
+{
+	DestroyDropPreviewElements();
+}
+
+bool FLevelEditorViewportClient::HasDropPreviewElements() const
+{
+	return ensure(DropPreviewElements) && DropPreviewElements->Num() > 0;
+}
+
+bool FLevelEditorViewportClient::UpdateDropPreviewElements(int32 MouseX, int32 MouseY, 
+	const TArray<UObject*>& DroppedObjects, bool& out_bDroppedObjectsVisible, 
+	TScriptInterface<IAssetFactoryInterface> Factory)
 {
 	out_bDroppedObjectsVisible = false;
-	if( !HasDropPreviewActors() )
+	if (!HasDropPreviewElements())
 	{
 		return false;
 	}
@@ -1314,7 +1524,7 @@ bool FLevelEditorViewportClient::UpdateDropPreviewActors(int32 MouseX, int32 Mou
 	bNeedsRedraw = true;
 
 	// If the mouse did not move, there is no need to update anything
-	if ( MouseX == DropPreviewMouseX && MouseY == DropPreviewMouseY )
+	if (MouseX == DropPreviewMouseX && MouseY == DropPreviewMouseY)
 	{
 		return false;
 	}
@@ -1323,39 +1533,53 @@ bool FLevelEditorViewportClient::UpdateDropPreviewActors(int32 MouseX, int32 Mou
 	DropPreviewMouseX = MouseX;
 	DropPreviewMouseY = MouseY;
 
-	// Get the center point between all the drop preview actors for use in calculations below
+	// Get the center point between all the drop preview objects for use in calculations below
 	// Also, build a list of valid AActor* pointers
-	FVector ActorOrigin = FVector::ZeroVector;
-	TArray<AActor*> DraggingActors;
-	TArray<AActor*> IgnoreActors;
-	for ( auto ActorIt = DropPreviewActors.CreateConstIterator(); ActorIt; ++ActorIt )
-	{
-		AActor* DraggingActor = (*ActorIt).Get();
+	FVector CombinedOrigin = FVector::ZeroVector;
+	int32 Count = 0;
 
-		if ( DraggingActor )
+	DropPreviewElements->ForEachElement<ITypedElementWorldInterface>([&CombinedOrigin, &Count](const TTypedElement<ITypedElementWorldInterface>& InElement)
+	{
+		FTransform Transform;
+		if (InElement.GetWorldTransform(Transform))
 		{
-			DraggingActors.Add(DraggingActor);
-			IgnoreActors.Add(DraggingActor);
-			DraggingActor->GetAllChildActors(IgnoreActors);
-			ActorOrigin += DraggingActor->GetActorLocation();
+			CombinedOrigin += Transform.GetTranslation();
+			++Count;
 		}
-	}
 
-	// If there were not valid actors after all, there is nothing to update
-	if ( DraggingActors.Num() == 0 )
+		// true means don't stop
+		return true;
+	});
+
+	if (Count == 0)
 	{
 		return false;
 	}
 
 	// Finish the calculation of the actors origin now that we know we are not dividing by zero
-	ActorOrigin /= DraggingActors.Num();
+	CombinedOrigin /= Count;
+
+	// TODO: Swap to ignored items instead of actors
+	TArray<AActor*> IgnoreActors;
+	DropPreviewElements->ForEachElement<ITypedElementObjectInterface>(
+	[&IgnoreActors](const TTypedElement<ITypedElementObjectInterface>& InElement)
+	{
+		if (AActor* Actor = InElement.GetObjectAs<AActor>())
+		{
+			IgnoreActors.Add(Actor);
+			Actor->GetAllChildActors(IgnoreActors);
+		}
+
+		// true means continue
+		return true;
+	});
 
 	FSceneViewFamilyContext ViewFamily(FSceneViewFamily::ConstructionValues(
-		Viewport, 
+		Viewport,
 		GetScene(),
 		EngineShowFlags)
-		.SetRealtimeUpdate( IsRealtime() ));
-	FSceneView* View = CalcSceneView( &ViewFamily );
+		.SetRealtimeUpdate(IsRealtime()));
+	FSceneView* View = CalcSceneView(&ViewFamily);
 	FViewportCursorLocation Cursor(View, this, MouseX, MouseY);
 
 	const FActorPositionTraceResult TraceResult = IsDroppingOn2DLayer() ? TraceForPositionOn2DLayer(Cursor) : FActorPositioning::TraceWorldForPositionWithDefault(Cursor, *View, &IgnoreActors);
@@ -1368,13 +1592,13 @@ bool FLevelEditorViewportClient::UpdateDropPreviewActors(int32 MouseX, int32 Mou
 	FSnappingUtils::SnapPointToGrid(GEditor->ClickLocation, FVector::ZeroVector);
 
 	AActor* DroppedOnActor = TraceResult.HitActor.Get();
-	
+
 	if (DroppedOnActor)
 	{
 		// We indicate that the dropped objects are visible if *any* of them are not applicable to other actors
-		out_bDroppedObjectsVisible = DroppedObjects.ContainsByPredicate([&](UObject* AssetObj){
+		out_bDroppedObjectsVisible = DroppedObjects.ContainsByPredicate([&](UObject* AssetObj) {
 			return !AttemptApplyObjToActor(AssetObj, DroppedOnActor, -1, true);
-		});
+			});
 	}
 	else
 	{
@@ -1382,48 +1606,100 @@ bool FLevelEditorViewportClient::UpdateDropPreviewActors(int32 MouseX, int32 Mou
 		out_bDroppedObjectsVisible = true;
 	}
 
-	for (AActor* Actor : DraggingActors)
+	// Perform the actual updates
+	DropPreviewElements->ForEachElementHandle(
+		[this, &TraceResult, ActorFactory = Cast<UActorFactory>(Factory.GetObject())](const FTypedElementHandle& Element)
 	{
-		const UActorFactory* ActorFactory = FactoryToUse ? FactoryToUse : GEditor->FindActorFactoryForActorClass(Actor->GetClass());
+		UTypedElementRegistry* Registry = UTypedElementRegistry::GetInstance();
+		if (!ensure(Registry))
+		{
+			return false;
+		}
 
-		const FSnappedPositioningData PositioningData = FSnappedPositioningData(this, TraceResult.Location, TraceResult.SurfaceNormal)
-			.DrawSnapHelpers(true)
-			.UseFactory(ActorFactory)
-			.UseStartTransform(PreDragActorTransforms.FindRef(Actor))
-			.UsePlacementExtent(Actor->GetPlacementExtent());
+		// Used for the actual location update
+		ITypedElementWorldInterface* WorldInterface = Registry->GetElementInterface<ITypedElementWorldInterface>(Element);
 
-		FTransform ActorTransform = FActorPositioning::GetSnappedSurfaceAlignedTransform(PositioningData);
-		ActorTransform.SetScale3D(Actor->GetActorScale3D());		// preserve scaling
+		// Used for some legacy updates if we're dealing with actors
+		ITypedElementObjectInterface* ObjectInterface = Registry->GetElementInterface<ITypedElementObjectInterface>(Element);
+		AActor* Actor = ObjectInterface ? ObjectInterface->GetObjectAs<AActor>(Element) : nullptr;
 
-		Actor->SetActorTransform(ActorTransform);
-		Actor->SetIsTemporarilyHiddenInEditor(false);
-		Actor->PostEditMove(false);
-	}
+		if (WorldInterface)
+		{
+			FTransform3d PreviousTransform = PreDragElementTransforms.FindRef(Element);
+			
+			UActorFactory* ActorFactoryToUse = (!ActorFactory && Actor) ? GEditor->FindActorFactoryForActorClass(Actor->GetClass())
+				: ActorFactory;
+
+			FSnappedPositioningData PositioningData = FSnappedPositioningData(this, TraceResult.Location, TraceResult.SurfaceNormal)
+				.DrawSnapHelpers(true)
+				.UseStartTransform(PreviousTransform)
+				// TODO: Need to support non-actor factories here
+				.UseFactory(ActorFactoryToUse)
+				.UsePlacementExtent(Actor ? Actor->GetPlacementExtent() : FVector3d::Zero());
+
+			FTransform DestinationTransform = FActorPositioning::GetSnappedSurfaceAlignedTransform(PositioningData);
+			DestinationTransform.SetScale3D(PreviousTransform.GetScale3D());		// preserve scaling
+
+			WorldInterface->SetWorldTransform(Element, DestinationTransform);
+		}
+
+		if (Actor)
+		{
+			Actor->SetIsTemporarilyHiddenInEditor(false);
+			Actor->PostEditMove(false);
+		}
+
+		// true means don't stop
+		return true;
+	});
 
 	return true;
 }
 
-void FLevelEditorViewportClient::DestroyDropPreviewActors()
+void FLevelEditorViewportClient::DestroyDropPreviewElements()
 {
-	if ( HasDropPreviewActors() )
+	using namespace LevelEditorViewportLocals;
+
+	if (!HasDropPreviewElements())
+	{
+		return;
+	}
+
+	// TODO: This code to remove the object from TEDS should not be necessary, beause the element
+	// deletion code further below should include TEDS deregistration. However, the code path for
+	// deleting preview actors in UUnrealEdEngine::DeleteActors skips explicit handle deregistration,
+	// and although it does still happen in the immediately triggered garbage cleanup, that feels
+	// potentially brittle.
+	DropPreviewElements->ForEachElement<ITypedElementObjectInterface>([this](const TTypedElement<ITypedElementObjectInterface>& InElement)
 	{
-		for ( auto ActorIt = DropPreviewActors.CreateConstIterator(); ActorIt; ++ActorIt )
+		UObject* PreviewObject = InElement.GetObject();
+		if (PreviewObject && PreviewObject != GetWorld()->GetDefaultBrush())
 		{
-			AActor* PreviewActor = (*ActorIt).Get();
-			if (PreviewActor && PreviewActor != GetWorld()->GetDefaultBrush())
+			UTypedElementRegistry* TypedElementRegistry = UTypedElementRegistry::GetInstance();
+			if (ITypedElementDataStorageCompatibilityInterface* DataStorageCompatibilityInterface = TypedElementRegistry->GetMutableDataStorageCompatibility())
 			{
-				// Temporary hack to remove the preview actor from the TypedElementRegistry (if it exists)
-				UTypedElementRegistry* TypedElementRegistry = UTypedElementRegistry::GetInstance();
-				if(ITypedElementDataStorageCompatibilityInterface* DataStorageCompatibilityInterface = TypedElementRegistry->GetMutableDataStorageCompatibility())
-				{
-					DataStorageCompatibilityInterface->RemoveCompatibleObject(PreviewActor);
-				}
-
-				GetWorld()->DestroyActor(PreviewActor);
+				DataStorageCompatibilityInterface->RemoveCompatibleObject(PreviewObject);
 			}
 		}
-		DropPreviewActors.Empty();
-	}
+		return true; // true means continue
+	});
+
+	DropPreviewElements->ForEachElement<ITypedElementWorldInterface>([this](const TTypedElement<ITypedElementWorldInterface>& InElement)
+	{
+		// Make sure to remove any items keyed by this handle, otherwise we'll get complaints about
+		// remaining references to the handle after destroying the item
+		PreDragElementTransforms.Remove(InElement);
+		
+		FTypedElementDeletionOptions Options;
+		Options
+			.SetWarnAboutReferences(false)
+			.SetWarnAboutSoftReferences(false);
+		ensure(InElement.DeleteElement(InElement.GetOwnerWorld(), GetMutableSelectionSet(), Options));
+
+		return true;  // true means continue
+	});
+
+	DropPreviewElements->Empty();
 }
 
 
@@ -1497,18 +1773,28 @@ FDropQuery FLevelEditorViewportClient::CanDropObjectsAtCoordinates(int32 MouseX,
 	return Result;
 }
 
-bool FLevelEditorViewportClient::DropObjectsAtCoordinates(int32 MouseX, int32 MouseY, const TArray<UObject*>& DroppedObjects, TArray<AActor*>& OutNewActors, bool bOnlyDropOnTarget/* = false*/, bool bCreateDropPreview/* = false*/, bool SelectActors, UActorFactory* FactoryToUse )
+bool FLevelEditorViewportClient::DropObjectsAtCoordinates(int32 MouseX, int32 MouseY, const TArray<UObject*>& DroppedObjects, 
+	TArray<FTypedElementHandle>& OutNewElements, const FDropObjectOptions& Options)
 {
 	bool bResult = false;
 
 	// Allow the drag drop handler to do anything pre-drop.
 	ULevelEditorDragDropHandler* DragDropHandler = GEditor->GetLevelEditorDragDropHandler();
-	if (!bCreateDropPreview)
+	if (!Options.bCreateDropPreview)
 	{
-		if (!DragDropHandler->PreDropObjectsAtCoordinates(MouseX, MouseY, GetWorld(), Viewport, DroppedObjects, OutNewActors))
+		// TODO: Deprecate PreDropObjectsAtCoordinates because it does not currently get used anywhere.
+		TArray<AActor*> NewActors;
+		if (!DragDropHandler->PreDropObjectsAtCoordinates(MouseX, MouseY, GetWorld(), Viewport, DroppedObjects, NewActors))
 		{
 			return bResult;
 		}
+		for (AActor* Actor : NewActors)
+		{
+			if (FTypedElementHandle Handle = UEngineElementsLibrary::AcquireEditorActorElementHandle(Actor))
+			{
+				OutNewElements.Add(Handle);
+			}
+		}
 	}
 
 	// Make sure the placement dragging actor is cleaned up.
@@ -1516,7 +1802,7 @@ bool FLevelEditorViewportClient::DropObjectsAtCoordinates(int32 MouseX, int32 Mo
 
 	if(DroppedObjects.Num() > 0)
 	{
-		bIsDroppingPreviewActor = bCreateDropPreview;
+		bIsDroppingPreviewActor = Options.bCreateDropPreview;
 		Viewport->InvalidateHitProxy();
 
 		FSceneViewFamilyContext ViewFamily(FSceneViewFamily::ConstructionValues(
@@ -1538,10 +1824,10 @@ bool FLevelEditorViewportClient::DropObjectsAtCoordinates(int32 MouseX, int32 Mo
 		// Snap the new location if snapping is enabled
 		FSnappingUtils::SnapPointToGrid(GEditor->ClickLocation, FVector::ZeroVector);
 
-		EObjectFlags ObjectFlags = bCreateDropPreview ? RF_Transient : RF_Transactional;
+		EObjectFlags ObjectFlags = Options.bCreateDropPreview ? RF_Transient : RF_Transactional;
 		if (HitProxy == nullptr || HitProxy->IsA(HInstancedStaticMeshInstance::StaticGetType()))
 		{
-			bResult = DropObjectsOnBackground(Cursor, DroppedObjects, ObjectFlags, OutNewActors, bCreateDropPreview, SelectActors, FactoryToUse);
+			bResult = DropObjectsOnBackground(Cursor, DroppedObjects, ObjectFlags, OutNewElements, Options);
 		}
 		else if (HitProxy->IsA(HActor::StaticGetType()) || HitProxy->IsA(HBSPBrushVert::StaticGetType()))
 		{
@@ -1572,23 +1858,24 @@ bool FLevelEditorViewportClient::DropObjectsAtCoordinates(int32 MouseX, int32 Mo
 
 			if (TargetActor != nullptr && IsValidChecked(TargetActor->GetWorld()) && !TargetActor->GetWorld()->IsUnreachable())
 			{
-				FNavigationLockContext LockNavigationUpdates(TargetActor->GetWorld(), ENavigationLockReason::SpawnOnDragEnter, bCreateDropPreview);
+				FNavigationLockContext LockNavigationUpdates(TargetActor->GetWorld(), ENavigationLockReason::SpawnOnDragEnter, Options.bCreateDropPreview);
 
 				// If the target actor is selected, we should drop onto all selected actors
 				// otherwise, we should drop only onto the target object
 				const bool bDropOntoSelectedActors = TargetActor->IsSelected();
 				const bool bCanApplyToComponent = AttemptApplyObjToComponent(DroppedObjects[0], TargetComponent, TargetMaterialSlot, true);
-				if (bOnlyDropOnTarget || !bDropOntoSelectedActors || !bCanApplyToComponent)
+				if (Options.bOnlyDropOnTarget || !bDropOntoSelectedActors || !bCanApplyToComponent)
 				{
 					if (bCanApplyToComponent)
 					{
-						const bool bIsTestAttempt = bCreateDropPreview;
+						const bool bIsTestAttempt = Options.bCreateDropPreview;
 						bResult = AttemptApplyObjToComponent(DroppedObjects[0], TargetComponent, TargetMaterialSlot, bIsTestAttempt);
 					}
 					else
 					{
 						// Couldn't apply to a component, so try dropping the objects on the hit actor
-						bResult = DropObjectsOnActor(Cursor, DroppedObjects, TargetActor, TargetMaterialSlot, ObjectFlags, OutNewActors, bCreateDropPreview, SelectActors, FactoryToUse);
+						bResult = DropObjectsOnActor(Cursor, DroppedObjects, TargetActor, TargetMaterialSlot, 
+							ObjectFlags, OutNewElements, Options);
 					}
 				}
 				else
@@ -1604,14 +1891,14 @@ bool FLevelEditorViewportClient::DropObjectsAtCoordinates(int32 MouseX, int32 Mo
 							for (FSelectedEditableComponentIterator It(GEditor->GetSelectedEditableComponentIterator()); It; ++It)
 							{
 								USceneComponent* SceneComponent = Cast<USceneComponent>(*It);
-								AttemptApplyObjToComponent(DroppedObjects[0], SceneComponent, TargetMaterialSlot, bCreateDropPreview);
+								AttemptApplyObjToComponent(DroppedObjects[0], SceneComponent, TargetMaterialSlot, Options.bCreateDropPreview);
 								bResult = true;
 							}
 						}
 						else
 						{
 							// The target component is not selected, so apply the object exclusively to it
-							bResult = AttemptApplyObjToComponent(DroppedObjects[0], TargetComponent, TargetMaterialSlot, bCreateDropPreview);
+							bResult = AttemptApplyObjToComponent(DroppedObjects[0], TargetComponent, TargetMaterialSlot, Options.bCreateDropPreview);
 						}
 					}
 					
@@ -1623,7 +1910,8 @@ bool FLevelEditorViewportClient::DropObjectsAtCoordinates(int32 MouseX, int32 Mo
 							TargetActor = static_cast<AActor*>( *It );
 							if (TargetActor)
 							{
-								DropObjectsOnActor(Cursor, DroppedObjects, TargetActor, TargetMaterialSlot, ObjectFlags, OutNewActors, bCreateDropPreview, SelectActors, FactoryToUse);
+								DropObjectsOnActor(Cursor, DroppedObjects, TargetActor, 
+									TargetMaterialSlot, ObjectFlags, OutNewElements, Options);
 								bResult = true;
 							}
 						}
@@ -1634,7 +1922,7 @@ bool FLevelEditorViewportClient::DropObjectsAtCoordinates(int32 MouseX, int32 Mo
 		else if (HitProxy->IsA(HModel::StaticGetType()))
 		{
 			// BSP surface
-			bResult = DropObjectsOnBSPSurface(View, Cursor, DroppedObjects, static_cast<HModel*>(HitProxy), ObjectFlags, OutNewActors, bCreateDropPreview, SelectActors, FactoryToUse);
+			bResult = DropObjectsOnBSPSurface(View, Cursor, DroppedObjects, static_cast<HModel*>(HitProxy), ObjectFlags, OutNewElements, Options);
 		}
 		else if( HitProxy->IsA( HWidgetAxis::StaticGetType() ) )
 		{
@@ -1663,10 +1951,16 @@ bool FLevelEditorViewportClient::DropObjectsAtCoordinates(int32 MouseX, int32 Mo
 				// and FLevelEditorViewportClient::DropObjectsOnActor (when not doing some kind of special "apply to actor" operation)
 				for (UObject* DroppedObject : DroppedObjects)
 				{
-					TArray<AActor*> NewActors = TryPlacingActorFromObject(GetWorld()->GetCurrentLevel(), DroppedObject, true, ObjectFlags, FactoryToUse, NAME_None, &Cursor);
-					if (NewActors.Num() > 0)
+					UE::AssetPlacementUtil::FExtraPlaceAssetOptions PlacementOptions;
+					PlacementOptions.bSelectOutput = Options.bSelectOutput;
+					PlacementOptions.ObjectFlags = ObjectFlags;
+					PlacementOptions.FactoryToUse = Options.FactoryToUse;
+					PlacementOptions.Name = NAME_None;
+
+					TArray<FTypedElementHandle> NewElements = TryPlacingAssetObject(GetWorld()->GetCurrentLevel(), DroppedObject, PlacementOptions, &Cursor);
+					if (NewElements.Num() > 0)
 					{
-						OutNewActors.Append(NewActors);
+						OutNewElements.Append(NewElements);
 						bResult = true;
 					}
 				}
@@ -1676,25 +1970,75 @@ bool FLevelEditorViewportClient::DropObjectsAtCoordinates(int32 MouseX, int32 Mo
 		if ( bResult )
 		{
 			// If we are creating a drop preview actor instead of a normal actor, we need to disable collision, selection, and make sure it is never saved.
-			if ( bCreateDropPreview && OutNewActors.Num() > 0 )
+			if (Options.bCreateDropPreview && OutNewElements.Num() > 0 )
 			{
-				DropPreviewActors.Empty();
+				// Theoretically we shouldn't have selected preview items to begin with, but we've had this safety deselection
+				// code forever, and might as well keep it.
+				UTypedElementSelectionSet* SelectionSet = GetMutableSelectionSet();
+				if (ensure(SelectionSet))
+				{
+					FTypedElementSelectionOptions SelectionOptions;
+					SelectionSet->DeselectElements(OutNewElements, SelectionOptions);
+				}
+				
+				DestroyDropPreviewElements();
 
-				for ( auto ActorIt = OutNewActors.CreateConstIterator(); ActorIt; ++ActorIt )
+				// Helper to add elements to DropPreviewElements
+				auto StorePreviewElement = [this](const FTypedElementHandle& Element)
 				{
-					AActor* NewActor = *ActorIt;
-					DropPreviewActors.Add(NewActor);
-					
-					PreDragActorTransforms.Add(NewActor, NewActor->GetTransform());
+					UTypedElementRegistry* Registry = UTypedElementRegistry::GetInstance();
+					if (!ensure(Registry))
+					{
+						return;
+					}
 
-					NewActor->SetActorEnableCollision(false);
+					ITypedElementWorldInterface* WorldInterface = Registry->GetElementInterface<ITypedElementWorldInterface>(Element);
+
+					// Don't store preview element if it doesn't have a world interface, because we would need it 
+					// to delete the preview element or to move it around
+					if (!WorldInterface)
+					{
+						return;
+					}
+
+					DropPreviewElements->Add(Element);
+
+					// If we wanted 
+				};
+
+				// Collect the preview elements and save pre drag transforms.
+				for (const FTypedElementHandle& Element : OutNewElements)
+				{
+					StorePreviewElement(Element);
+				}
+
+				// Store pre drag transforms
+				DropPreviewElements->ForEachElement<ITypedElementWorldInterface>(
+					[this](const TTypedElement<ITypedElementWorldInterface>& InElement)
+				{
+					FTransform Transform;
+					if (InElement.GetWorldTransform(Transform))
+					{
+						PreDragElementTransforms.Add(InElement, Transform);
+					}
 
-					// Deselect if selected
-					if( NewActor->IsSelected() )
+					return true; // true means continue
+				});
+				
+				// Do some legacy actor handling
+				DropPreviewElements->ForEachElement<ITypedElementObjectInterface>(
+				[this](const TTypedElement<ITypedElementObjectInterface>& InElement)
+				{
+					AActor* NewActor = InElement.GetObjectAs<AActor>();
+					if (!NewActor)
 					{
-						GEditor->SelectActor( NewActor, /*InSelected=*/false, /*bNotify=*/true );
+						return true; // true means continue
 					}
 
+					PreDragActorTransforms.Add(NewActor, NewActor->GetTransform());
+
+					NewActor->SetActorEnableCollision(false);
+
 					// Prevent future selection. This also prevents the hit proxy from interfering with placement logic.
 					for (UActorComponent* Component : NewActor->GetComponents())
 					{
@@ -1703,7 +2047,10 @@ bool FLevelEditorViewportClient::DropObjectsAtCoordinates(int32 MouseX, int32 Mo
 							PrimComp->bSelectable = false;
 						}
 					}
-				}
+
+					// true means continue
+					return true;
+				});
 
 				// Set the current MouseX/Y to prime the preview update
 				DropPreviewMouseX = MouseX;
@@ -1714,7 +2061,7 @@ bool FLevelEditorViewportClient::DropObjectsAtCoordinates(int32 MouseX, int32 Mo
 				//Viewport->InvalidateHitProxy();
 			}
 			// Dropping the actors rather than a preview? Probably want to select them all then. 
-			else if(!bCreateDropPreview && SelectActors && OutNewActors.Num() > 0
+			else if(!Options.bCreateDropPreview && Options.bSelectOutput && DroppedObjects.Num() > 0
 
 				// If we're dropping bsp brushes, that geometry actually gets created later, in 
 				// FBrushBuilderDragDropOp::OnDrop(), so we don't want to select the (preview) builder brush.
@@ -1722,9 +2069,11 @@ bool FLevelEditorViewportClient::DropObjectsAtCoordinates(int32 MouseX, int32 Mo
 				// whereas the FBrushBuilderDragDropOp::OnDrop() gets called via SlateUser->NotifyPointerReleased).
 				&& !AreAnyDroppedObjectsBrushBuilders(DroppedObjects))
 			{
-				for (auto It = OutNewActors.CreateConstIterator(); It; ++It)
+				UTypedElementSelectionSet* SelectionSet = GetMutableSelectionSet();
+				if (ensure(SelectionSet))
 				{
-					GEditor->SelectActor( (*It), true, true );
+					FTypedElementSelectionOptions SelectionOptions;
+					SelectionSet->SelectElements(OutNewElements, SelectionOptions);
 				}
 			}
 
@@ -1737,18 +2086,24 @@ bool FLevelEditorViewportClient::DropObjectsAtCoordinates(int32 MouseX, int32 Mo
 
 	if (bResult)
 	{
-		if ( !bCreateDropPreview && IPlacementModeModule::IsAvailable() )
+		if ( !Options.bCreateDropPreview && IPlacementModeModule::IsAvailable() )
 		{
-			IPlacementModeModule::Get().AddToRecentlyPlaced( DroppedObjects, FactoryToUse );
+			IPlacementModeModule::Get().AddToRecentlyPlaced(DroppedObjects, Options.FactoryToUse);
 		}
 
-		if (!bCreateDropPreview)
+		if (!Options.bCreateDropPreview)
 		{
-			FEditorDelegates::OnNewActorsDropped.Broadcast(DroppedObjects, OutNewActors);
+			TArray<AActor*> Actors;
+			LevelEditorViewportLocals::ForEachActorInElementArray(OutNewElements, [&Actors](AActor& Actor)
+			{
+				Actors.Add(&Actor);
+			});
+
+			FEditorDelegates::OnNewActorsDropped.Broadcast(DroppedObjects, Actors);
 		}
 	}
 
-	if (!bCreateDropPreview)
+	if (!Options.bCreateDropPreview)
 	{
 		DragDropHandler->PostDropObjectsAtCoordinates(MouseX, MouseX, World, Viewport, DroppedObjects);
 	}
@@ -1759,6 +2114,24 @@ bool FLevelEditorViewportClient::DropObjectsAtCoordinates(int32 MouseX, int32 Mo
 	return bResult;
 }
 
+bool FLevelEditorViewportClient::DropObjectsAtCoordinates(int32 MouseX, int32 MouseY, const TArray<UObject*>& DroppedObjects, TArray<AActor*>& OutNewActors, bool bOnlyDropOnTarget, bool bCreateDropPreview, bool bSelectOutput, UActorFactory* FactoryToUse)
+{
+	FDropObjectOptions Options;
+	Options.bOnlyDropOnTarget = bOnlyDropOnTarget;
+	Options.bCreateDropPreview = bCreateDropPreview;
+	Options.bSelectOutput = bSelectOutput;
+	Options.FactoryToUse = FactoryToUse;
+
+	TArray<FTypedElementHandle> NewElements;
+	bool bSuccess = DropObjectsAtCoordinates(MouseX, MouseY, DroppedObjects, NewElements, Options);
+	
+	LevelEditorViewportLocals::ForEachActorInElementArray(NewElements, [&OutNewActors](AActor& Actor)
+	{
+		OutNewActors.Add(&Actor);
+	});
+	return bSuccess;
+}
+
 /**
  *	Called to check if a material can be applied to an object, given the hit proxy
  */
@@ -1996,6 +2369,13 @@ FLevelEditorViewportClient::FLevelEditorViewportClient(const TSharedPtr<SLevelVi
 	GetMutableDefault<ULevelEditorViewportSettings>()->OnSettingChanged().AddRaw(this, &FLevelEditorViewportClient::HandleViewportSettingChanged);
 
 	FModuleManager::LoadModuleChecked<FLevelEditorModule>("LevelEditor").OnMapChanged().AddRaw(this, &FLevelEditorViewportClient::OnMapChanged);
+
+	DropPreviewElements = StaticDropPreviewElements.Pin();
+	if (!StaticDropPreviewElements.IsValid() && UTypedElementRegistry::GetInstance())
+	{
+		DropPreviewElements = UTypedElementRegistry::GetInstance()->CreateElementList();
+		StaticDropPreviewElements = DropPreviewElements;
+	}
 }
 
 //
@@ -3195,6 +3575,7 @@ void FLevelEditorViewportClient::TrackingStarted( const FInputEventState& InInpu
 	bHasBegunGizmoManipulation = false;
 
 	PreDragActorTransforms.Empty();
+	PreDragElementTransforms.Empty();
 
 	// Track BSP changes
 	{
@@ -3393,6 +3774,7 @@ void FLevelEditorViewportClient::TrackingStopped()
 	}
 
 	PreDragActorTransforms.Empty();
+	PreDragElementTransforms.Empty();
 }
 
 void FLevelEditorViewportClient::AbortTracking()
diff --git a/Engine/Source/Editor/UnrealEd/Public/AssetSelection.h b/Engine/Source/Editor/UnrealEd/Public/AssetSelection.h
--- a/Engine/Source/Editor/UnrealEd/Public/AssetSelection.h
+++ b/Engine/Source/Editor/UnrealEd/Public/AssetSelection.h
@@ -7,6 +7,8 @@
 #include "Input/Reply.h"
 
 class AActor;
+struct FTypedElementHandle;
+class IAssetFactoryInterface;
 class UActorFactory;
 class ULevel;
 class UMaterialInterface;
@@ -272,6 +274,7 @@ public:
 	 */
 	static UNREALED_API UActorFactory* GetFactoryForAssetObject( UObject* AssetObj );
 
+	//~ TODO: UE_DEPRECATED(5.4, "Use UE::AssetPlacementUtil::PlaceAssetInCurrentLevel instead")
 	/**
 	 * Places an actor instance using the factory appropriate for the type of asset
 	 *
@@ -285,6 +288,7 @@ public:
 	 */
 	static UNREALED_API AActor* AddActorForAsset( UObject* AssetObj, bool SelectActor = true, EObjectFlags ObjectFlags = RF_Transactional, UActorFactory* FactoryToUse = NULL, const FName Name = NAME_None );
 
+	//~ TODO: UE_DEPRECATED(5.4)
 	/**
 	 * Places an actor instance using the factory appropriate for the type of asset using the current object selection as the asset
 	 *
@@ -330,3 +334,24 @@ private:
 	{
 	}
 };
+
+namespace UE::AssetPlacementUtil
+{
+	struct FExtraPlaceAssetOptions
+	{
+		bool bSelectOutput = true;
+		EObjectFlags ObjectFlags = RF_Transactional;
+		TScriptInterface<IAssetFactoryInterface> FactoryToUse = NULL;
+		FName Name = NAME_None;
+	};
+
+	/**
+	 * Places an asset instance using the factory appropriate for the type of asset
+	 *
+	 * @param	AssetObj						the asset that is contained in the d&d operation
+	 *
+	 * @return	the object that was created by the factory, or NULL if there aren't any factories for this asset (or
+	 *			the object couldn't be created for some other reason)
+	 */
+	UNREALED_API TArray<FTypedElementHandle> PlaceAssetInCurrentLevel(UObject* AssetObj, const FExtraPlaceAssetOptions& ExtraParams);
+}
diff --git a/Engine/Source/Editor/UnrealEd/Public/DragAndDrop/AssetDragDropOp.h b/Engine/Source/Editor/UnrealEd/Public/DragAndDrop/AssetDragDropOp.h
--- a/Engine/Source/Editor/UnrealEd/Public/DragAndDrop/AssetDragDropOp.h
+++ b/Engine/Source/Editor/UnrealEd/Public/DragAndDrop/AssetDragDropOp.h
@@ -7,8 +7,10 @@
 #include "Input/DragAndDrop.h"
 #include "Layout/Visibility.h"
 #include "DragAndDrop/DecoratedDragDropOp.h"
+#include "UObject/WeakInterfacePtr.h"
 
 class FAssetThumbnail;
+class IAssetFactoryInterface;
 class UActorFactory;
 
 class FAssetDragDropOp : public FDecoratedDragDropOp
@@ -16,15 +18,24 @@ class FAssetDragDropOp : public FDecoratedDragDropOp
 public:
 	DRAG_DROP_OPERATOR_TYPE(FAssetDragDropOp, FDecoratedDragDropOp)
 
-	static UNREALED_API TSharedRef<FAssetDragDropOp> New(const FAssetData& InAssetData, UActorFactory* ActorFactory = nullptr);
+	static UNREALED_API TSharedRef<FAssetDragDropOp> New(const FAssetData& InAssetData, 
+		TScriptInterface<IAssetFactoryInterface> Factory = nullptr);
 
-	static UNREALED_API TSharedRef<FAssetDragDropOp> New(TArray<FAssetData> InAssetData, UActorFactory* ActorFactory = nullptr);
+	static UNREALED_API TSharedRef<FAssetDragDropOp> New(TArray<FAssetData> InAssetData,
+		TScriptInterface<IAssetFactoryInterface> Factory = nullptr);
 
 	static UNREALED_API TSharedRef<FAssetDragDropOp> New(FString InAssetPath);
 
 	static UNREALED_API TSharedRef<FAssetDragDropOp> New(TArray<FString> InAssetPaths);
 
-	static UNREALED_API TSharedRef<FAssetDragDropOp> New(TArray<FAssetData> InAssetData, TArray<FString> InAssetPaths, UActorFactory* ActorFactory = nullptr);
+	static UNREALED_API TSharedRef<FAssetDragDropOp> New(TArray<FAssetData> InAssetData, TArray<FString> InAssetPaths, 
+		TScriptInterface<IAssetFactoryInterface> Factory = nullptr);
+
+	//~ These overloads are redundant to the ones that use TScriptInterface<IAssetFactoryInterface>, but
+	//~ we keep them temporarily in case a user just forward declared UActorFactory.
+	static UNREALED_API TSharedRef<FAssetDragDropOp> New(const FAssetData& InAssetData, UActorFactory* ActorFactory);
+	static UNREALED_API TSharedRef<FAssetDragDropOp> New(TArray<FAssetData> InAssetData, UActorFactory* ActorFactory);
+	static UNREALED_API TSharedRef<FAssetDragDropOp> New(TArray<FAssetData> InAssetData, TArray<FString> InAssetPaths, UActorFactory* ActorFactory);
 
 	/** @return true if this drag operation contains assets */
 	bool HasAssets() const
@@ -50,11 +61,11 @@ public:
 		return AssetPaths;
 	}
 
-	/** @return The actor factory to use if converting this asset to an actor */
-	UActorFactory* GetActorFactory() const
-	{
-		return ActorFactory.Get();
-	}
+	//~ TODO: UE_DEPRECATED(5.4, "Use GetAssetFactory instead.")
+	UNREALED_API UActorFactory* GetActorFactory() const;
+
+	/** @return The asset factory to use if converting this asset to an actor */
+	UNREALED_API TScriptInterface<IAssetFactoryInterface> GetAssetFactory() const;
 
 public:
 	UNREALED_API virtual ~FAssetDragDropOp();
@@ -64,6 +75,8 @@ public:
 	UNREALED_API FText GetDecoratorText() const;
 
 protected:
+	UNREALED_API void Init(TArray<FAssetData> InAssetData, TArray<FString> InAssetPaths, TScriptInterface<IAssetFactoryInterface> InAssetFactory);
+	//TODO: Would be nice to remove this, but users could have just forward declared UActorFactory...
 	UNREALED_API void Init(TArray<FAssetData> InAssetData, TArray<FString> InAssetPaths, UActorFactory* InActorFactory);
 	UNREALED_API virtual void InitThumbnail();
 
@@ -80,8 +93,8 @@ private:
 	/** Data for the asset paths this item represents */
 	TArray<FString> AssetPaths;
 
-	/** The actor factory to use if converting this asset to an actor */
-	TWeakObjectPtr<UActorFactory> ActorFactory;
+	/** The factory to use if converting this asset to a placed object */
+	TWeakInterfacePtr<IAssetFactoryInterface> AssetFactory;
 
 protected:
 	/** Handle to the thumbnail resource */
diff --git a/Engine/Source/Editor/UnrealEd/Public/EditorViewportClient.h b/Engine/Source/Editor/UnrealEd/Public/EditorViewportClient.h
--- a/Engine/Source/Editor/UnrealEd/Public/EditorViewportClient.h
+++ b/Engine/Source/Editor/UnrealEd/Public/EditorViewportClient.h
@@ -32,6 +32,8 @@ class FEditorViewportClient;
 class FEdMode;
 class FMouseDeltaTracker;
 class FPreviewScene;
+struct FTypedElementHandle;
+class IAssetFactoryInterface;
 class SEditorViewport;
 class UActorFactory;
 class UTypedElementViewportInteraction;
@@ -825,9 +827,11 @@ public:
 	 */
 	UNREALED_API virtual void CheckHoveredHitProxy( HHitProxy* HoveredHitProxy );
 
+	//~ TODO: UE_DEPRECATED(5.4,"Use HasDropPreviewElements instead.")
 	/** Returns true if a placement dragging actor exists */
 	virtual bool HasDropPreviewActors() const { return false; }
 
+	//~ TODO: UE_DEPRECATED(5.4,"Use UpdateDropPreviewElements instead.")
 	/**
 	 * If dragging an actor for placement, this function updates its position.
 	 *
@@ -840,11 +844,34 @@ public:
 	 */
 	virtual bool UpdateDropPreviewActors(int32 MouseX, int32 MouseY, const TArray<UObject*>& DroppedObjects, bool& out_bDroppedObjectsVisible, UActorFactory* FactoryToUse = NULL) { return false; }
 
+	//~ TODO: UE_DEPRECATED(5.4,"Use DestroyDropPreviewElements instead.")
 	/**
 	 * If dragging an actor for placement, this function destroys the actor.
 	 */
 	virtual void DestroyDropPreviewActors() {}
 
+	/** Returns true if a placement drag preview elements exists */
+	virtual bool HasDropPreviewElements() const { return false; }
+
+	/**
+	 * If dragging items for placement, this function updates their position.
+	 *
+	 * @param MouseX						The position of the mouse's X coordinate
+	 * @param MouseY						The position of the mouse's Y coordinate
+	 * @param DroppedObjects				The Objects that were used to create preview objects
+	 * @param out_bDroppedObjectsVisible	Output, returns if preview objects are visible or not
+	 *
+	 * Returns true if preview elements were updated
+	 */
+	virtual bool UpdateDropPreviewElements(int32 MouseX, int32 MouseY, 
+		const TArray<UObject*>& DroppedObjects, bool& out_bDroppedObjectsVisible, 
+		TScriptInterface<IAssetFactoryInterface> Factory = nullptr) { return false; }
+
+	/**
+	 * If dragging items for placement, this function destroys the items.
+	 */
+	virtual void DestroyDropPreviewElements() {}
+
 	/**
 	 * Checks the viewport to see if the given object can be dropped using the given mouse coordinates local to this viewport
 	 *
@@ -854,19 +881,43 @@ public:
 	 */
 	virtual FDropQuery CanDropObjectsAtCoordinates(int32 MouseX, int32 MouseY, const FAssetData& AssetInfo) { return FDropQuery(); }
 
+	struct FDropObjectOptions
+	{
+		//~ This default constructor just exists as a workaround for a bug in some compilers
+		//~ where trying to use a nested class as a default parameter in a member function of
+		//~ the same outer class will fail to compile if the nested class has default member 
+		//~ initializers and an implicitly declared default constructor.
+		FDropObjectOptions() {}
+
+		// Flag that when true, will only attempt a drop on the actor targeted by the mouse position. Defaults to false.
+		bool bOnlyDropOnTarget = false;
+		// If true, a drop preview will be spawned instead of a normal result.
+		bool bCreateDropPreview = false;
+		// If true, select the newly dropped elements
+		bool bSelectOutput = true;
+		// The preferred factory to use (optional)
+		TScriptInterface<IAssetFactoryInterface> FactoryToUse = nullptr;
+	};
+
 	/**
 	 * Attempts to intelligently drop the given objects in the viewport, using the given mouse coordinates local to this viewport
 	 *
-	 * @param MouseX			 The position of the mouse's X coordinate
-	 * @param MouseY			 The position of the mouse's Y coordinate
-	 * @param DroppedObjects	 The Objects to be placed into the editor via this viewport
-	 * @param OutNewActors		 The new actor objects that were created
-	 * @param bOnlyDropOnTarget  Flag that when True, will only attempt a drop on the actor targeted by the Mouse position. Defaults to false.
-	 * @param bCreateDropPreview If true, a drop preview actor will be spawned instead of a normal actor.
-	 * @param bSelectActors		 If true, select the newly dropped actors (defaults: true)
-	 * @param FactoryToUse		 The preferred actor factory to use (optional)
-	 */
-	virtual bool DropObjectsAtCoordinates(int32 MouseX, int32 MouseY, const TArray<UObject*>& DroppedObjects, TArray<AActor*>& OutNewActors, bool bOnlyDropOnTarget = false, bool bCreateDropPreview = false, bool bSelectActors = true, UActorFactory* FactoryToUse = NULL ) { return false; }
+	 * @param MouseX			The position of the mouse's X coordinate
+	 * @param MouseY			The position of the mouse's Y coordinate
+	 * @param DroppedObjects	The Objects to be placed into the editor via this viewport
+	 * @param OutNewObjects		The new actor objects that were created
+	 * @param Options			Additional options
+	 */
+	UNREALED_API virtual bool DropObjectsAtCoordinates(int32 MouseX, int32 MouseY, const TArray<UObject*>& DroppedObjects, TArray<FTypedElementHandle>& OutNewObjects,
+		const FDropObjectOptions& Options = FDropObjectOptions())
+		// Forwards to the other overload (which returns false) during deprecation. Will return false directly once the other
+		// overload is removed.
+		;
+
+	UE_DEPRECATED(5.4, "Use the overload that takes FDropObjectOptions instead.")
+	UNREALED_API virtual bool DropObjectsAtCoordinates(int32 MouseX, int32 MouseY, const TArray<UObject*>& DroppedObjects, TArray<AActor*>& OutNewActors,
+		bool bOnlyDropOnTarget = false, bool bCreateDropPreview = false, bool bSelectActors = true,
+		UActorFactory* FactoryToUse = NULL) { return false; }
 
 	/** Returns true if the viewport is allowed to be possessed for previewing cinematic sequences or keyframe animations*/
 	bool AllowsCinematicControl() const { return bAllowCinematicControl; }
diff --git a/Engine/Source/Editor/UnrealEd/Public/LevelEditorViewport.h b/Engine/Source/Editor/UnrealEd/Public/LevelEditorViewport.h
--- a/Engine/Source/Editor/UnrealEd/Public/LevelEditorViewport.h
+++ b/Engine/Source/Editor/UnrealEd/Public/LevelEditorViewport.h
@@ -4,6 +4,7 @@
 #pragma once
 
 #include "CoreMinimal.h"
+#include "AssetSelection.h" // FExtraPlaceAssetOptions
 #include "InputCoreTypes.h"
 #include "GameFramework/Actor.h"
 #include "Camera/CameraComponent.h"
@@ -18,6 +19,7 @@ struct FMinimalViewInfo;
 class FCanvas;
 class FDragTool;
 class HModel;
+class IAssetFactoryInterface;
 class ILevelEditor;
 class SLevelViewport;
 class UActorFactory;
@@ -167,6 +169,7 @@ class FLevelEditorViewportClient : public FEditorViewportClient
 
 public:
 
+	//~ TODO: UE_DEPRECATED(5.4, "Use GetDropPreviewElements instead.")
 	/** @return Returns the current global drop preview actor, or a NULL pointer if we don't currently have one */
 	static const TArray< TWeakObjectPtr<AActor> >& GetDropPreviewActors()
 	{
@@ -397,9 +400,11 @@ public:
 
 	UNREALED_API void RemoveReferenceToWorldContext(FWorldContext& WorldContext);
 
+	//~ TODO: UE_DEPRECATED(5.4, "Use HasDropPreviewElements instead.")
 	/** Returns true if a placement dragging actor exists */
 	UNREALED_API virtual bool HasDropPreviewActors() const override;
 
+	//~ TODO: UE_DEPRECATED(5.4, "Use UpdateDropPreviewElements instead.")
 	/**
 	 * If dragging an actor for placement, this function updates its position.
 	 *
@@ -412,11 +417,18 @@ public:
 	 */
 	UNREALED_API virtual bool UpdateDropPreviewActors(int32 MouseX, int32 MouseY, const TArray<UObject*>& DroppedObjects, bool& out_bDroppedObjectsVisible, class UActorFactory* FactoryToUse = NULL) override;
 
+	//~ TODO: UE_DEPRECATED(5.4, "Use DestroyDropPreviewElements instead.")
 	/**
 	 * If dragging an actor for placement, this function destroys the actor.
 	 */
 	UNREALED_API virtual void DestroyDropPreviewActors() override;
 
+	UNREALED_API virtual bool HasDropPreviewElements() const override;
+	UNREALED_API virtual bool UpdateDropPreviewElements(int32 MouseX, int32 MouseY, 
+		const TArray<UObject*>& DroppedObjects, bool& out_bDroppedObjectsVisible, 
+		TScriptInterface<IAssetFactoryInterface> Factory = nullptr) override;
+	UNREALED_API virtual void DestroyDropPreviewElements() override;
+
 	/**
 	 * Checks the viewport to see if the given object can be dropped using the given mouse coordinates local to this viewport
 	 *
@@ -429,16 +441,18 @@ public:
 	/**
 	 * Attempts to intelligently drop the given objects in the viewport, using the given mouse coordinates local to this viewport
 	 *
-	 * @param MouseX			 The position of the mouse's X coordinate
-	 * @param MouseY			 The position of the mouse's Y coordinate
-	 * @param DroppedObjects	 The Objects to be placed into the editor via this viewport
-	 * @param OutNewActors		 The new actor objects that were created
-	 * @param bOnlyDropOnTarget  Flag that when True, will only attempt a drop on the actor targeted by the Mouse position. Defaults to false.
-	 * @param bCreateDropPreview If true, a drop preview actor will be spawned instead of a normal actor.
-	 * @param bSelectActors		 If true, select the newly dropped actors (defaults: true)
-	 * @param FactoryToUse		 The preferred actor factory to use (optional)
+	 * @param MouseX			The position of the mouse's X coordinate
+	 * @param MouseY			The position of the mouse's Y coordinate
+	 * @param DroppedObjects	The asset objects to be placed into the editor via this viewport
+	 * @param OutNewItems		The new items that were created
+	 * @param Options			Additional options
 	 */
-	UNREALED_API virtual bool DropObjectsAtCoordinates(int32 MouseX, int32 MouseY, const TArray<UObject*>& DroppedObjects, TArray<AActor*>& OutNewActors, bool bOnlyDropOnTarget = false, bool bCreateDropPreview = false, bool bSelectActors = true, UActorFactory* FactoryToUse = NULL ) override;
+	UNREALED_API virtual bool DropObjectsAtCoordinates(int32 MouseX, int32 MouseY, const TArray<UObject*>& DroppedObjects, 
+		TArray<FTypedElementHandle>& OutNewItems, const FDropObjectOptions& Options = FDropObjectOptions()) override;
+
+	//~ TODO: UE_DEPRECATED(5.4, "Use the overload that uses FDropObjectOptions instead.")
+	UNREALED_API virtual bool DropObjectsAtCoordinates(int32 MouseX, int32 MouseY, const TArray<UObject*>& DroppedObjects, TArray<AActor*>& OutNewActors, bool bOnlyDropOnTarget = false, bool bCreateDropPreview = false, bool bSelectActors = true, UActorFactory* FactoryToUse = NULL) override;
+
 
 	/**
 	 * Sets GWorld to the appropriate world for this client
@@ -612,6 +626,15 @@ public:
 	 */
 	UNREALED_API void SetViewportTypeFromTool(ELevelViewportType InViewportType);
 
+	/**
+	 * Static: Attempts to place the specified asset object in the level, returning one or more 
+	 * newly-created objects by their FTypedElementHandles if successful.
+	 */
+	static UNREALED_API TArray<FTypedElementHandle> TryPlacingAssetObject(ULevel* InLevel, UObject* AssetObject,
+		const UE::AssetPlacementUtil::FExtraPlaceAssetOptions& AdditionalParams,
+		const FViewportCursorLocation* CursorInformation = nullptr);
+
+	//~TODO: UE_DEPRECATED(5.4, "Use TryPlacingAssetObject instead")
 	/**
 	 * Static: Attempts to place the specified object in the level, returning one or more newly-created actors if successful.
 	 * IMPORTANT: The placed actor's location must be first set using GEditor->ClickLocation and GEditor->ClickPlane.
@@ -626,7 +649,9 @@ public:
 	 *
 	 * @return	true if the object was successfully used to place an actor; false otherwise
 	 */
-	static UNREALED_API TArray<AActor*> TryPlacingActorFromObject( ULevel* InLevel, UObject* ObjToUse, bool bSelectActors, EObjectFlags ObjectFlags, UActorFactory* FactoryToUse, const FName Name = NAME_None, const FViewportCursorLocation* Cursor = nullptr);
+	static UNREALED_API TArray<AActor*> TryPlacingActorFromObject( ULevel* InLevel, UObject* ObjToUse, bool bSelectActors,
+		EObjectFlags ObjectFlags, UActorFactory* FactoryToUse, 
+		const FName Name = NAME_None, const FViewportCursorLocation* Cursor = nullptr);
 
 	/** 
 	 * Returns true if creating a preview actor in the viewport. 
@@ -750,17 +775,19 @@ private:
 	/**
 	 * Called when an asset is dropped onto the blank area of a viewport.
 	 *
-	 * @param	Cursor				Mouse cursor location
-	 * @param	DroppedObjects		Array of objects dropped into the viewport
-	 * @param	ObjectFlags			The object flags to place on the actors that this function spawns.
-	 * @param	OutNewActors		The list of actors created while dropping
-	 * @param	bCreateDropPreview	If true, the actor being dropped is a preview actor (defaults: false)
-	 * @param	bSelectActors		If true, select the newly dropped actors (defaults: true)
-	 * @param	FactoryToUse		The preferred actor factory to use (optional)
+	 * @param	Cursor			Mouse cursor location
+	 * @param	DroppedObjects	Array of objects dropped into the viewport
+	 * @param	ObjectFlags		The object flags to place on the actors that this function spawns.
+	 * @param	OutNewItems		The list of actors created while dropping
+	 * @param	Options			Additional options
 	 *
 	 * @return	true if the drop operation was successfully handled; false otherwise
 	 */
-	UNREALED_API bool DropObjectsOnBackground( struct FViewportCursorLocation& Cursor, const TArray<UObject*>& DroppedObjects, EObjectFlags ObjectFlags, TArray<AActor*>& OutNewActors, bool bCreateDropPreview = false, bool bSelectActors = true, class UActorFactory* FactoryToUse = NULL );
+	UNREALED_API bool DropObjectsOnBackground(struct FViewportCursorLocation& Cursor, const TArray<UObject*>& DroppedObjects, 
+		EObjectFlags ObjectFlags, TArray<FTypedElementHandle>& OutNewItems, const FDropObjectOptions& Options);
+	
+	//~ TODO: UE_DEPRECATED(5.4, "Use the overload that uses FDropObjectOptions instead.")
+	UNREALED_API bool DropObjectsOnBackground(struct FViewportCursorLocation& Cursor, const TArray<UObject*>& DroppedObjects, EObjectFlags ObjectFlags, TArray<AActor*>& OutNewActors, bool bCreateDropPreview = false, bool bSelectActors = true, class UActorFactory* FactoryToUse = NULL);
 
 	/**
 	* Called when an asset is dropped upon an existing actor.
@@ -770,14 +797,20 @@ private:
 	* @param	DroppedUponActor	The actor that we are dropping upon
 	* @param    DroppedUponSlot     The material slot/submesh that was identified as the drop location.  If unknown use -1.
 	* @param	ObjectFlags			The object flags to place on the actors that this function spawns.
-	* @param	OutNewActors		The list of actors created while dropping
-	* @param	bCreateDropPreview	If true, the actor being dropped is a preview actor (defaults: false)
-	* @param	bSelectActors		If true, select the newly dropped actors (defaults: true)
-	* @param	FactoryToUse		The preferred actor factory to use (optional)
+	* @param	OutNewItems			The list of items created while dropping
+	* @param	Options				Additional options. Note that bOnlyDropOnTarget is ignored, since this function only drops on the actor
 	*
 	* @return	true if the drop operation was successfully handled; false otherwise
 	*/
-	UNREALED_API bool DropObjectsOnActor(struct FViewportCursorLocation& Cursor, const TArray<UObject*>& DroppedObjects, AActor* DroppedUponActor, int32 DroppedUponSlot, EObjectFlags ObjectFlags, TArray<AActor*>& OutNewActors, bool bCreateDropPreview = false, bool bSelectActors = true, class UActorFactory* FactoryToUse = NULL);
+	UNREALED_API bool DropObjectsOnActor(struct FViewportCursorLocation& Cursor, const TArray<UObject*>& DroppedObjects,
+		AActor* DroppedUponActor, int32 DroppedUponSlot, EObjectFlags ObjectFlags,
+		TArray<FTypedElementHandle>& OutNewItems, const FDropObjectOptions& Options);
+
+	//~ TODO: UE_DEPRECATED(5.4, "Use the overload that uses FDropObjectOptions instead.")
+	UNREALED_API bool DropObjectsOnActor(struct FViewportCursorLocation& Cursor, const TArray<UObject*>& DroppedObjects, 
+		AActor* DroppedUponActor, int32 DroppedUponSlot, EObjectFlags ObjectFlags, 
+		TArray<AActor*>& OutNewActors, bool bCreateDropPreview = false, 
+		bool bSelectActors = true, class UActorFactory* FactoryToUse = NULL);
 
 	/**
 	 * Called when an asset is dropped upon a BSP surface.
@@ -794,6 +827,11 @@ private:
 	 *
 	 * @return	true if the drop operation was successfully handled; false otherwise
 	 */
+	UNREALED_API bool DropObjectsOnBSPSurface(FSceneView* View, struct FViewportCursorLocation& Cursor, 
+		const TArray<UObject*>& DroppedObjects, HModel* TargetProxy, EObjectFlags ObjectFlags,
+		TArray<FTypedElementHandle>& OutNewItems, const FDropObjectOptions& Options);
+
+	//UE_DEPRECATED(5.4, "Use the overload that uses FDropObjectOptions instead.")
 	UNREALED_API bool DropObjectsOnBSPSurface(FSceneView* View, struct FViewportCursorLocation& Cursor, const TArray<UObject*>& DroppedObjects, HModel* TargetProxy, EObjectFlags ObjectFlags, TArray<AActor*>& OutNewActors, bool bCreateDropPreview = false, bool bSelectActors = true, UActorFactory* FactoryToUse = NULL);
 
 	/**
@@ -886,15 +924,26 @@ public:
 	bool bAlwaysShowModeWidgetAfterSelectionChanges;
 
 private:
-	/** The actors that are currently being placed in the viewport via dragging */
+	//TODO: UE_DEPRECATED(5.4, TEXT("Use DropPreviewElements instead"))
 	static UNREALED_API TArray< TWeakObjectPtr< AActor > > DropPreviewActors;
+		
+	/** The viewport clients share a list of drop preview elements.However we still want the elements
+	  to be destroyed if all the viewports are destroyed, hence a static weak pointer and private shared
+	  pointers. */
+	static UNREALED_API TWeakPtr<FTypedElementList> StaticDropPreviewElements;
+	TSharedPtr<FTypedElementList> DropPreviewElements;
 
 	/** If currently creating a preview actor. */
 	static UNREALED_API bool bIsDroppingPreviewActor;
 
+	// TODO: Remove this to always use PreDragElementTransforms. That requires modifications to ProjectActorsIntoWorld
 	/** A map of actor locations before a drag operation */
 	mutable TMap<TWeakObjectPtr<const AActor>, FTransform> PreDragActorTransforms;
 
+	/** Map of element locations keyed by their handle. Currently used for the preview drop elements. Should
+	 eventually replace PreDragActorTransforms entirely. */
+	TMap<FTypedElementHandle, FTransform> PreDragElementTransforms;
+
 	/** The elements (from the current selection set) that this viewport can manipulate (eg, via the transform gizmo) */
 	bool bHasCachedElementsToManipulate = false;
 	FTypedElementListRef CachedElementsToManipulate;
