nop clean up and notes in EditorFactories texture creation
this is the not-much-used non-Interchange path

#rnx

diff --git a/Engine/Source/Editor/UnrealEd/Classes/Factories/TextureFactory.h b/Engine/Source/Editor/UnrealEd/Classes/Factories/TextureFactory.h
--- a/Engine/Source/Editor/UnrealEd/Classes/Factories/TextureFactory.h
+++ b/Engine/Source/Editor/UnrealEd/Classes/Factories/TextureFactory.h
@@ -124,7 +124,7 @@ class UTextureFactory : public UFactory, public IImportSettingsParser
 
 	/** Channel values to compare to when preserving alpha coverage from a mask. */
 	UPROPERTY(EditAnywhere, Category=PreserveAlphaCoverage, meta=(ToolTip="Channel values to compare to when preserving alpha coverage from a mask for mips"))
-	FVector4 AlphaCoverageThresholds = FVector4(0,0,0,0);
+	FVector4 AlphaCoverageThresholds = FVector4(0,0,0,0.75f);
 
 	/** If enabled, preserve the value of border pixels when creating mip-maps */
 	UPROPERTY(EditAnywhere, Category=PreserveBorder, meta=(ToolTip="If enabled, preserve the value of border pixels when creating mip-maps"))
diff --git a/Engine/Source/Editor/UnrealEd/Private/Factories/EditorFactories.cpp b/Engine/Source/Editor/UnrealEd/Private/Factories/EditorFactories.cpp
--- a/Engine/Source/Editor/UnrealEd/Private/Factories/EditorFactories.cpp
+++ b/Engine/Source/Editor/UnrealEd/Private/Factories/EditorFactories.cpp
@@ -4044,6 +4044,8 @@ UObject* UTextureFactory::FactoryCreateBinary
 	if (ExistingTexture2D)
 	{
 		// Update with new settings, which should disable streaming...
+		// -> pretty sure this is unnecessary
+		//	the combination of WaitForPendingInitOrStreaming and PreEditChange do all the waiting necessary
 		ExistingTexture2D->UpdateResource();
 	}
 	if(ExistingTexture)
@@ -4051,11 +4053,8 @@ UObject* UTextureFactory::FactoryCreateBinary
 		// Wait for InitRHI() to complete before the FTextureReferenceReplacer calls ReleaseRHI() to follow the workflow.
 		// Static texture needs to avoid having pending InitRHI() before enqueuing ReleaseRHI() to safely track access of the PlatformData on the renderthread.
 		ExistingTexture->WaitForPendingInitOrStreaming();
-	}
 
-	// Make sure the changes are part of the transaction when reimporting over an existing texture
-	if (ExistingTexture)
-	{
+		// Make sure the changes are part of the transaction when reimporting over an existing texture
 		ExistingTexture->PreEditChange(nullptr);
 	}
 
@@ -4082,7 +4081,7 @@ UObject* UTextureFactory::FactoryCreateBinary
 		if (ExistingTexture)
 		{
 			// We failed to import over the existing texture. Make sure the resource is ready in the existing texture.
-			ExistingTexture->UpdateResource();
+			ExistingTexture->PostEditChange();
 		}
 
 		Warn->Logf(ELogVerbosity::Error, TEXT("Texture import failed") );
@@ -4090,6 +4089,8 @@ UObject* UTextureFactory::FactoryCreateBinary
 		return nullptr;
 	}
 
+	// this is automatic now and redundant;
+	//	Source.Init does UseHashAsGuid
 	if (bUseHashAsGuid)
 	{
 		Texture->Source.UseHashAsGuid();
@@ -4146,6 +4147,7 @@ UObject* UTextureFactory::FactoryCreateBinary
 	Texture->bDoScaleMipsForAlphaCoverage	= bDoScaleMipsForAlphaCoverage;
 	Texture->AlphaCoverageThresholds		= AlphaCoverageThresholds;
 	Texture->bUseNewMipFilter				= bUseNewMipFilter;
+	// these get changed by SetModernSettingsForNewOrChangedTexture anyway
 
 	if(Texture->MipGenSettings == TMGS_FromTextureGroup)
 	{
