CoreRedirects: GameplayCueManager and ComponentTypeRegistry need to handle CoreRedirects in the classnames stored in Blueprint tags.
#rnx

diff --git a/Engine/Plugins/Runtime/GameplayAbilities/Source/GameplayAbilities/Private/GameplayCueManager.cpp b/Engine/Plugins/Runtime/GameplayAbilities/Source/GameplayAbilities/Private/GameplayCueManager.cpp
--- a/Engine/Plugins/Runtime/GameplayAbilities/Source/GameplayAbilities/Private/GameplayCueManager.cpp
+++ b/Engine/Plugins/Runtime/GameplayAbilities/Source/GameplayAbilities/Private/GameplayCueManager.cpp
@@ -898,6 +898,11 @@ void UGameplayCueManager::BuildCuesToAddToGlobalSet(const TArray<FAssetData>& As
 				// Add a new NotifyData entry to our flat list for this one
 				FSoftObjectPath StringRef;
 				StringRef.SetPath(FPackageName::ExportTextPathToObjectPath(GeneratedClassTag));
+				bool bFixedUpRef = StringRef.FixupCoreRedirects();
+				if (bFixedUpRef)
+				{
+					ABILITY_LOG(Log, TEXT("GameplayCueManager Redirected: %s -> %s"), *GeneratedClassTag, *StringRef.ToString());
+				}
 
 				OutCuesToAdd.Add(FGameplayCueReferencePair(GameplayCueTag, StringRef));
 
diff --git a/Engine/Source/Editor/UnrealEd/Private/ComponentTypeRegistry.cpp b/Engine/Source/Editor/UnrealEd/Private/ComponentTypeRegistry.cpp
--- a/Engine/Source/Editor/UnrealEd/Private/ComponentTypeRegistry.cpp
+++ b/Engine/Source/Editor/UnrealEd/Private/ComponentTypeRegistry.cpp
@@ -614,7 +614,12 @@ void FComponentTypeRegistryData::ForceRefreshComponentList()
 		AssetRegistry.GetAssetsByClass(UBlueprint::StaticClass()->GetClassPathName(), Assets, true);
 		for (FAssetData& BPAsset : Assets)
 		{
-			const FTopLevelAssetPath ClassPath(FEditorClassUtils::GetClassPathNameFromAssetTag(BPAsset));
+			FTopLevelAssetPath ClassPath(FEditorClassUtils::GetClassPathNameFromAssetTag(BPAsset));
+			FSoftObjectPath SoftClassPath(ClassPath);
+			if (SoftClassPath.FixupCoreRedirects() && SoftClassPath.GetSubPathString().IsEmpty())
+			{
+				ClassPath = FTopLevelAssetPath(SoftClassPath.GetLongPackageFName(), SoftClassPath.GetAssetFName());
+			}
 			if (!ClassPath.IsNull())
 			{
 				if (!InMemoryClassPaths.Contains(ClassPath) && 
