fix Thumbnail renderer for NormalMaps was doing gamma incorrectly
NormalMapPreview and SimpleElementNormalMap shaders are deprecated and should not be used
thumbnail seemed to be the only thing using them, it no longer does
use the standard texture preview shader which has a normal map flag instead

#rnx

diff --git a/Engine/Shaders/Private/SimpleElementNormalMapPixelShader.usf b/Engine/Shaders/Private/SimpleElementNormalMapPixelShader.usf
--- a/Engine/Shaders/Private/SimpleElementNormalMapPixelShader.usf
+++ b/Engine/Shaders/Private/SimpleElementNormalMapPixelShader.usf
@@ -3,6 +3,11 @@
 /*=============================================================================
 	SimpleElementNormalMapPixelShader.hlsl: Pixel shader for previewing normal
 		map textures rendered as simple elements.
+
+Deprecated, do not use.
+
+Use SimpleElementTexture2DPreviewPixelShader.usf instead, with the normal map flag.
+
 =============================================================================*/
 
 #include "Common.ush"
@@ -25,6 +30,16 @@ void Main(
 	const float4 RescaledNormal = float4(Normal.xyz * 0.5 + 0.5, 1);
 	OutColor = RETURN_COLOR(RescaledNormal);	
 
+/*
+	// this shader does no gamma correction
+	if( Gamma != 1.0 )
+	{
+		// typically Gamma = 1.0/2.2
+		// then this does LinearToSrgb
+		OutColor.rgb = ApplyGammaCorrection(saturate(OutColor.rgb), Gamma * 2.2);
+	}
+*/
+
 #if WRITE_TO_GBUFFER
 	// Set the G buffer bits that indicate that deferred lighting and image reflections are not enabled
 	OutWorldNormal = 0;
diff --git a/Engine/Shaders/Private/SimpleElementTexture2DPreviewPixelShader.usf b/Engine/Shaders/Private/SimpleElementTexture2DPreviewPixelShader.usf
--- a/Engine/Shaders/Private/SimpleElementTexture2DPreviewPixelShader.usf
+++ b/Engine/Shaders/Private/SimpleElementTexture2DPreviewPixelShader.usf
@@ -116,7 +116,8 @@ void Main(
 	{
 		const float4 Normal = UnpackNormalMap(Sample);
 		const float4 RescaledNormal = float4(Normal.xyz * 0.5 + 0.5, 1);
-		Sample = RETURN_COLOR(RescaledNormal);	
+		// RETURN_COLOR is typically a nop
+		Sample = RETURN_COLOR(RescaledNormal);
 	}
 
 	// Seperate the Color weights and use against the Base colour to detrmine the actual colour from our filter
@@ -129,7 +130,9 @@ void Main(
 
 	if( Gamma != 1.0 )
 	{
-		FinalColor.rgb = ApplyGammaCorrection(saturate(FinalColor.rgb), 2.2 / (1.0 / Gamma));
+		// typically Gamma = 1.0/2.2
+		// then this does LinearToSrgb
+		FinalColor.rgb = ApplyGammaCorrection(saturate(FinalColor.rgb), Gamma * 2.2);
 	}
 	
 	FinalColor = RETURN_COLOR(FinalColor);	
diff --git a/Engine/Source/Editor/UnrealEd/Private/NormalMapPreview.cpp b/Engine/Source/Editor/UnrealEd/Private/NormalMapPreview.cpp
--- a/Engine/Source/Editor/UnrealEd/Private/NormalMapPreview.cpp
+++ b/Engine/Source/Editor/UnrealEd/Private/NormalMapPreview.cpp
@@ -15,7 +15,9 @@
 
 /*------------------------------------------------------------------------------
 	Batched element shaders for previewing normal maps.
-------------------------------------------------------------------------------*/
+
+ * Deprecated, do not use.  Use FBatchedElementTexture2DPreviewParameters instead with normal map flag. 
+-----------------------------------------------------------------------------*/
 
 /**
  * Simple pixel shader that reconstructs a normal for the purposes of visualization.
@@ -74,6 +76,8 @@ void FNormalMapBatchedElementParameters::BindShaders(
 	TShaderMapRef<FSimpleElementVS> VertexShader(GetGlobalShaderMap(InFeatureLevel));
 	TShaderMapRef<FSimpleElementNormalMapPS> PixelShader(GetGlobalShaderMap(InFeatureLevel));
 
+	// bad : this does not pass Gamma
+
 	GraphicsPSOInit.BoundShaderState.VertexDeclarationRHI = GSimpleElementVertexDeclaration.VertexDeclarationRHI;
 	GraphicsPSOInit.BoundShaderState.VertexShaderRHI = VertexShader.GetVertexShader();
 	GraphicsPSOInit.BoundShaderState.PixelShaderRHI = PixelShader.GetPixelShader();
@@ -87,4 +91,3 @@ void FNormalMapBatchedElementParameters::BindShaders(
 	SetShaderParametersLegacyVS(RHICmdList, VertexShader, InTransform);
 	SetShaderParametersLegacyPS(RHICmdList, PixelShader, Texture);
 }
-
diff --git a/Engine/Source/Editor/UnrealEd/Private/ThumbnailRendering/TextureThumbnailRenderer.cpp b/Engine/Source/Editor/UnrealEd/Private/ThumbnailRendering/TextureThumbnailRenderer.cpp
--- a/Engine/Source/Editor/UnrealEd/Private/ThumbnailRendering/TextureThumbnailRenderer.cpp
+++ b/Engine/Source/Editor/UnrealEd/Private/ThumbnailRendering/TextureThumbnailRenderer.cpp
@@ -16,7 +16,6 @@
 #include "Engine/TextureRenderTargetCube.h"
 
 #include "CubemapUnwrapUtils.h"
-#include "NormalMapPreview.h"
 #include "CanvasItem.h"
 #include "CanvasTypes.h"
 #include "TextureResource.h"
@@ -92,11 +91,15 @@ void UTextureThumbnailRenderer::Draw(UObject* Object, int32 X, int32 Y, uint32 W
 			const bool bIsVirtualTexture = Texture->IsCurrentlyVirtualTextured();
 
 			TRefCountPtr<FBatchedElementParameters> BatchedElementParameters;
+			// BatchedElementParameters is released by the render thread when it was rendered
 
 			if(TextureCube || TextureCubeArray || RTTextureCube)
 			{
-				// is released by the render thread when it was rendered
-				BatchedElementParameters = new FMipLevelBatchedElementParameters((float)0, (float)-1, TextureCubeArray != nullptr, FMatrix44f::Identity, true, false, false);
+				// show LongLat Unwrap
+				float MipLevel = -1.f;
+				float SliceIndex = -1.f;
+				bool bShowLongLatUnwrap = true;
+				BatchedElementParameters = new FMipLevelBatchedElementParameters(MipLevel, SliceIndex, TextureCubeArray != nullptr, FMatrix44f::Identity, bShowLongLatUnwrap, false, false);
 			
 				// If the thumbnail is square then make it 2:1 for cubes.
 				if(Width == Height)
@@ -105,19 +108,35 @@ void UTextureThumbnailRenderer::Draw(UObject* Object, int32 X, int32 Y, uint32 W
 					Y += Height / 2;
 				}
 			}
-			else if (Texture2DArray) 
-			{
-				bool bIsNormalMap = Texture2DArray->IsNormalMap();
-				bool bIsSingleChannel = true;
-				BatchedElementParameters = new FBatchedElementTexture2DPreviewParameters((float)0, (float)0, (float)-1, bIsNormalMap, bIsSingleChannel, false, false, true, false);
-			}
 			else if (TextureLightProfile)
 			{
 				BatchedElementParameters = new FIESLightProfileBatchedElementParameters(TextureLightProfile->Brightness);
 			}
-			else if (Texture2D && Texture2D->IsNormalMap())
+			else if (Texture2D || Texture2DArray)
 			{
-				BatchedElementParameters = new FNormalMapBatchedElementParameters();
+				bool bIsNormalMap = Texture->IsNormalMap();
+				bool bIsSingleChannel = Texture->CompressionSettings == TC_Grayscale || Texture->CompressionSettings == TC_Alpha;
+				bool bSingleVTPhysicalSpace = Texture2D && Texture2D->IsVirtualTexturedWithSinglePhysicalSpace();
+				//bool bIsVirtualTexture = Texture->IsCurrentlyVirtualTextured();
+				float MipLevel = -1.f;
+				float LayerIndex = 0;
+				float SliceIndex = -1.f;
+				bool bIsTextureArray = (Texture2DArray != nullptr);
+				bool bUsePointSampling = false;
+				BatchedElementParameters = new FBatchedElementTexture2DPreviewParameters(MipLevel, LayerIndex, SliceIndex, bIsNormalMap, bIsSingleChannel, bSingleVTPhysicalSpace, bIsVirtualTexture, bIsTextureArray, bUsePointSampling);
+			
+				//FNormalMapBatchedElementParameters is broken, do not use
+				//BatchedElementParameters = new FNormalMapBatchedElementParameters();
+			}
+			else
+			{
+				// BatchedElementParameters is not set
+				// default Canvas will be used, which just shows the texture
+
+				// some UTexture types can hit this
+				// UTextureRenderTarget
+
+				// ?? maybe get rid of this and just always use the Texture2D branch above
 			}
 
 			if (bUseTranslucentBlend)
diff --git a/Engine/Source/Editor/UnrealEd/Public/NormalMapPreview.h b/Engine/Source/Editor/UnrealEd/Public/NormalMapPreview.h
--- a/Engine/Source/Editor/UnrealEd/Public/NormalMapPreview.h
+++ b/Engine/Source/Editor/UnrealEd/Public/NormalMapPreview.h
@@ -11,7 +11,10 @@
 
 /**
  * Batched element parameters for previewing normal maps.
+ *
+ * Deprecated, do not use.  Use FBatchedElementTexture2DPreviewParameters instead with normal map flag. 
  */
+//UE_DEPRECATED(5.5, "Use FBatchedElementTexture2DPreviewParameters instead with normal map flag")
 class FNormalMapBatchedElementParameters : public FBatchedElementParameters
 {
 	/** Binds vertex and pixel shaders for this element */
diff --git a/Engine/Source/Editor/UnrealEd/Public/Texture2DPreview.h b/Engine/Source/Editor/UnrealEd/Public/Texture2DPreview.h
--- a/Engine/Source/Editor/UnrealEd/Public/Texture2DPreview.h
+++ b/Engine/Source/Editor/UnrealEd/Public/Texture2DPreview.h
@@ -15,6 +15,10 @@
 class FBatchedElementTexture2DPreviewParameters : public FBatchedElementParameters
 {
 public:
+
+	// MipLevel = -1 to auto-mip
+	//	Layer is for VT, must be an integer
+	// SliceIndex = -1 for all
 	FBatchedElementTexture2DPreviewParameters(float InMipLevel, float InLayerIndex, float InSliceIndex, bool bInIsNormalMap, bool bInIsSingleChannel, bool bInIsSingleVTPhysicalSpace, bool bInIsVirtualTexture, bool bInIsTextureArray, bool bInUsePointSampling)
 		: MipLevel(InMipLevel)
 		, LayerIndex(InLayerIndex)
