fix Thumbnail VT bug introduced in previous commit

#rnx

diff --git a/Engine/Source/Editor/UnrealEd/Private/ThumbnailRendering/TextureThumbnailRenderer.cpp b/Engine/Source/Editor/UnrealEd/Private/ThumbnailRendering/TextureThumbnailRenderer.cpp
--- a/Engine/Source/Editor/UnrealEd/Private/ThumbnailRendering/TextureThumbnailRenderer.cpp
+++ b/Engine/Source/Editor/UnrealEd/Private/ThumbnailRendering/TextureThumbnailRenderer.cpp
@@ -117,16 +117,21 @@ void UTextureThumbnailRenderer::Draw(UObject* Object, int32 X, int32 Y, uint32 W
 				bool bIsNormalMap = Texture->IsNormalMap();
 				bool bIsSingleChannel = Texture->CompressionSettings == TC_Grayscale || Texture->CompressionSettings == TC_Alpha;
 				bool bSingleVTPhysicalSpace = Texture2D && Texture2D->IsVirtualTexturedWithSinglePhysicalSpace();
-				//bool bIsVirtualTexture = Texture->IsCurrentlyVirtualTextured();
 				float MipLevel = -1.f;
 				float LayerIndex = 0;
 				float SliceIndex = -1.f;
 				bool bIsTextureArray = (Texture2DArray != nullptr);
 				bool bUsePointSampling = false;
-				BatchedElementParameters = new FBatchedElementTexture2DPreviewParameters(MipLevel, LayerIndex, SliceIndex, bIsNormalMap, bIsSingleChannel, bSingleVTPhysicalSpace, bIsVirtualTexture, bIsTextureArray, bUsePointSampling);
-			
-				//FNormalMapBatchedElementParameters is broken, do not use
-				//BatchedElementParameters = new FNormalMapBatchedElementParameters();
+
+				// if you correctly tell FBatchedElementTexture2DPreviewParameters that the texture is a VT
+				//	then you only get a solid color thumbnail for the initial render
+				//	but if it is refreshes or the VT is examined, the thumbnail will become correct
+				//bool bSampleAsVirtualTexture = bIsVirtualTexture;
+
+				// if you just always lie and say you are not a VT, you get a good thumbnail :
+				bool bSampleAsVirtualTexture = false;
+
+				BatchedElementParameters = new FBatchedElementTexture2DPreviewParameters(MipLevel, LayerIndex, SliceIndex, bIsNormalMap, bIsSingleChannel, bSingleVTPhysicalSpace, bSampleAsVirtualTexture, bIsTextureArray, bUsePointSampling);
 			}
 			else
 			{
