Cooker: Suppress spurious warning in bSkipSave case: Trying to update asset package flags in package '%s' that does not exist
#rnx

diff --git a/Engine/Source/Editor/UnrealEd/Classes/CookOnTheSide/CookOnTheFlyServer.h b/Engine/Source/Editor/UnrealEd/Classes/CookOnTheSide/CookOnTheFlyServer.h
--- a/Engine/Source/Editor/UnrealEd/Classes/CookOnTheSide/CookOnTheFlyServer.h
+++ b/Engine/Source/Editor/UnrealEd/Classes/CookOnTheSide/CookOnTheFlyServer.h
@@ -1579,6 +1579,7 @@ private:
 	float DisplayUpdatePeriodSeconds = 0.0f;
 	EIdleStatus IdleStatus = EIdleStatus::Done;
 
+	friend FAssetRegistryGenerator;
 	friend UE::Cook::FAssetRegistryMPCollector;
 	friend UE::Cook::FBeginCookConfigSettings;
 	friend UE::Cook::FCookDirector;
diff --git a/Engine/Source/Editor/UnrealEd/Private/Commandlets/AssetRegistryGenerator.cpp b/Engine/Source/Editor/UnrealEd/Private/Commandlets/AssetRegistryGenerator.cpp
--- a/Engine/Source/Editor/UnrealEd/Private/Commandlets/AssetRegistryGenerator.cpp
+++ b/Engine/Source/Editor/UnrealEd/Private/Commandlets/AssetRegistryGenerator.cpp
@@ -2551,7 +2551,7 @@ void FAssetRegistryGenerator::GetChunkAssignments(TArray<TSet<FName>>& OutAssign
 void FAssetRegistryGenerator::UpdateAssetRegistryData(FName PackageName, const UPackage* Package,
 	UE::Cook::ECookResult CookResult, FSavePackageResultStruct* SavePackageResult,
 	TOptional<TArray<FAssetData>>&& AssetDatasFromSave, TOptional<FAssetPackageData>&& OverrideAssetPackageData,
-	TOptional<TArray<FAssetDependency>>&& OverridePackageDependencies)
+	TOptional<TArray<FAssetDependency>>&& OverridePackageDependencies, UCookOnTheFlyServer& COTFS)
 {
 	LLM_SCOPE_BYTAG(Cooker_GeneratedAssetRegistry);
 	PreviousPackagesToUpdate.Remove(PackageName);
@@ -2619,7 +2619,7 @@ void FAssetRegistryGenerator::UpdateAssetRegistryData(FName PackageName, const U
 	}
 
 	const bool bUpdated = UpdateAssetPackageFlags(PackageName, NewPackageFlags);
-	UE_CLOG(!bUpdated && bSaveSucceeded, LogAssetRegistryGenerator, Warning,
+	UE_CLOG(!bUpdated && bSaveSucceeded && !COTFS.bSkipSave, LogAssetRegistryGenerator, Warning,
 		TEXT("Trying to update asset package flags in package '%s' that does not exist"), *PackageName.ToString());
 }
 
@@ -2635,7 +2635,7 @@ void FAssetRegistryGenerator::SetOverridePackageDependencies(FName PackageName,
 }
 
 void FAssetRegistryGenerator::UpdateAssetRegistryData(UE::Cook::FMPCollectorServerMessageContext& Context,
-	UE::Cook::FAssetRegistryPackageMessage&& Message)
+	UE::Cook::FAssetRegistryPackageMessage&& Message, UCookOnTheFlyServer& COTFS)
 {
 	LLM_SCOPE_BYTAG(Cooker_GeneratedAssetRegistry);
 	const FName PackageName = Context.GetPackageName();
@@ -2672,7 +2672,7 @@ FAssetRegistryReporterRemote::FAssetRegistryReporterRemote(FCookWorkerClient& In
 void FAssetRegistryReporterRemote::UpdateAssetRegistryData(FName PackageName, const UPackage* Package,
 	UE::Cook::ECookResult CookResult, FSavePackageResultStruct* SavePackageResult,
 	TOptional<TArray<FAssetData>>&& AssetDatasFromSave, TOptional<FAssetPackageData>&& OverrideAssetPackageData,
-	TOptional<TArray<FAssetDependency>>&& OverridePackageDependencies)
+	TOptional<TArray<FAssetDependency>>&& OverridePackageDependencies, UCookOnTheFlyServer& COTFS)
 {
 	uint32 NewPackageFlags = 0;
 	int64 DiskSize = -1;
@@ -2869,7 +2869,7 @@ void FAssetRegistryMPCollector::ServerReceiveMessage(FMPCollectorServerMessageCo
 	{
 		FAssetRegistryGenerator* RegistryGenerator = COTFS.PlatformManager->GetPlatformData(TargetPlatform)->RegistryGenerator.Get();
 		check(RegistryGenerator); // The TargetPlatform came from OrderedSessionPlatforms, and the RegistryGenerator should exist for any of those platforms
-		RegistryGenerator->UpdateAssetRegistryData(Context, MoveTemp(ARMessage));
+		RegistryGenerator->UpdateAssetRegistryData(Context, MoveTemp(ARMessage), COTFS);
 	}
 }
 
diff --git a/Engine/Source/Editor/UnrealEd/Private/CookOnTheFlyServer.cpp b/Engine/Source/Editor/UnrealEd/Private/CookOnTheFlyServer.cpp
--- a/Engine/Source/Editor/UnrealEd/Private/CookOnTheFlyServer.cpp
+++ b/Engine/Source/Editor/UnrealEd/Private/CookOnTheFlyServer.cpp
@@ -6305,7 +6305,8 @@ void FSaveCookedPackageContext::FinishPlatform()
 			AssetDatasFromSave.Emplace(MoveTemp(SavePackageResult.SavedAssets));
 		}
 		Reporter.UpdateAssetRegistryData(Package->GetFName(), Package, CookResult, &SavePackageResult,
-			MoveTemp(AssetDatasFromSave), MoveTemp(OverrideAssetPackageData), MoveTemp(OverridePackageDependencies));
+			MoveTemp(AssetDatasFromSave), MoveTemp(OverrideAssetPackageData), MoveTemp(OverridePackageDependencies),
+			COTFS);
 	}
 
 	if (bSuccessful && COTFS.bSkipOnlyEditorOnly)
@@ -6501,7 +6502,8 @@ void UCookOnTheFlyServer::RecordExternalActorDependencies(TConstArrayView<FName>
 				DependencyData->SetPlatformCooked(TargetPlatform, ECookResult::NeverCookPlaceholder);
 				Reporter.UpdateAssetRegistryData(DependencyName, nullptr /* Package */,
 					ECookResult::NeverCookPlaceholder, nullptr /* SavePackageResult */,
-					TOptional<TArray<FAssetData>>(), TOptional<FAssetPackageData>(), TOptional<TArray<FAssetDependency>>());
+					TOptional<TArray<FAssetData>>(), TOptional<FAssetPackageData>(), TOptional<TArray<FAssetDependency>>(),
+					*this);
 			}
 		}
 	}
diff --git a/Engine/Source/Editor/UnrealEd/Public/Commandlets/AssetRegistryGenerator.h b/Engine/Source/Editor/UnrealEd/Public/Commandlets/AssetRegistryGenerator.h
--- a/Engine/Source/Editor/UnrealEd/Public/Commandlets/AssetRegistryGenerator.h
+++ b/Engine/Source/Editor/UnrealEd/Public/Commandlets/AssetRegistryGenerator.h
@@ -254,9 +254,10 @@ public:
 		UE::Cook::ECookResult CookResult, FSavePackageResultStruct* SavePackageResult,
 		TOptional<TArray<FAssetData>>&& AssetDatasFromSave,
 		TOptional<FAssetPackageData>&& OverrideAssetPackageData,
-		TOptional<TArray<FAssetDependency>>&& OverridePackageDependencies);
+		TOptional<TArray<FAssetDependency>>&& OverridePackageDependencies,
+		UCookOnTheFlyServer& COTFS);
 	void UpdateAssetRegistryData(UE::Cook::FMPCollectorServerMessageContext& Context,
-		UE::Cook::FAssetRegistryPackageMessage&& Message);
+		UE::Cook::FAssetRegistryPackageMessage&& Message, UCookOnTheFlyServer& COTFS);
 
 	/**
 	 * Check config to see whether chunk assignments use the AssetManager. If so, run the once-per-process construction
@@ -478,7 +479,8 @@ public:
 	virtual void UpdateAssetRegistryData(FName PackageName, const UPackage* Package, UE::Cook::ECookResult CookResult,
 		FSavePackageResultStruct* SavePackageResult,
 		TOptional<TArray<FAssetData>>&& AssetDatasFromSave, TOptional<FAssetPackageData>&& OverrideAssetPackageData,
-		TOptional<TArray<FAssetDependency>>&& OverridePackageDependencies) = 0;
+		TOptional<TArray<FAssetDependency>>&& OverridePackageDependencies,
+		UCookOnTheFlyServer& COTFS) = 0;
 
 };
 
@@ -493,10 +495,11 @@ public:
 	virtual void UpdateAssetRegistryData(FName PackageName, const UPackage* Package, UE::Cook::ECookResult CookResult,
 		FSavePackageResultStruct* SavePackageResult,
 		TOptional<TArray<FAssetData>>&& AssetDatasFromSave, TOptional<FAssetPackageData>&& OverrideAssetPackageData,
-		TOptional<TArray<FAssetDependency>>&& OverridePackageDependencies) override
+		TOptional<TArray<FAssetDependency>>&& OverridePackageDependencies, UCookOnTheFlyServer& COTFS) override
 	{
 		Generator.UpdateAssetRegistryData(PackageName, Package, CookResult, SavePackageResult,
-			MoveTemp(AssetDatasFromSave), MoveTemp(OverrideAssetPackageData), MoveTemp(OverridePackageDependencies));
+			MoveTemp(AssetDatasFromSave), MoveTemp(OverrideAssetPackageData), MoveTemp(OverridePackageDependencies),
+			COTFS);
 	}
 
 private:
@@ -511,7 +514,7 @@ public:
 	virtual void UpdateAssetRegistryData(FName PackageName, const UPackage* Package, UE::Cook::ECookResult CookResult,
 		FSavePackageResultStruct* SavePackageResult,
 		TOptional<TArray<FAssetData>>&& AssetDatasFromSave, TOptional<FAssetPackageData>&& OverrideAssetPackageData,
-		TOptional<TArray<FAssetDependency>>&& OverridePackageDependencies) override;
+		TOptional<TArray<FAssetDependency>>&& OverridePackageDependencies, UCookOnTheFlyServer& COTFS) override;
 
 private:
 	FCookWorkerClient& Client;
