Level Instance:
- Make sure we have the Parent container info when creating a new Level Instance (proper fix for previous change CL 30291083)
- AddLevelToWorld and CreateNewStreamingLevelForWorld have a Params version that allows passing in bEnableWorldPartitionStreaming and a new LevelStreamingCreatedCallback
- bEnableWorldPartitionStreaming allows removal of the global callback hook in ULevelInstanceSubsystem (FWorldDelegates::OnPreWorldInitialization)
- LevelStreamingCreatedCallback allows removal of the static EditLevelInstanceID and properly set it in the callback, it also allows setting the ParentContainerInstance / ParentActorGuid

#rnx

diff --git a/Engine/Source/Editor/UnrealEd/Classes/Factories/WorldFactory.h b/Engine/Source/Editor/UnrealEd/Classes/Factories/WorldFactory.h
--- a/Engine/Source/Editor/UnrealEd/Classes/Factories/WorldFactory.h
+++ b/Engine/Source/Editor/UnrealEd/Classes/Factories/WorldFactory.h
@@ -25,6 +25,7 @@ class UWorldFactory : public UFactory
 	TEnumAsByte<EWorldType::Type> WorldType;
 	bool bInformEngineOfWorld;
 	bool bCreateWorldPartition;
+	bool bEnableWorldPartitionStreaming;
 	ERHIFeatureLevel::Type FeatureLevel;
 
 	//~ Begin UFactory Interface
diff --git a/Engine/Source/Editor/UnrealEd/Private/EditorLevelUtils.cpp b/Engine/Source/Editor/UnrealEd/Private/EditorLevelUtils.cpp
--- a/Engine/Source/Editor/UnrealEd/Private/EditorLevelUtils.cpp
+++ b/Engine/Source/Editor/UnrealEd/Private/EditorLevelUtils.cpp
@@ -336,7 +336,7 @@ ULevel* UEditorLevelUtils::AddLevelsToWorld(UWorld* InWorld, TArray<FString> Pac
 	{
 		SlowTask.EnterProgressFrame();
 
-		if (ULevelStreaming* NewStreamingLevel = AddLevelToWorld_Internal(InWorld, *PackageName, LevelStreamingClass))
+		if (ULevelStreaming* NewStreamingLevel = AddLevelToWorld_Internal(InWorld, FAddLevelToWorldParams(LevelStreamingClass, *PackageName)))
 		{
 			NewLevel = NewStreamingLevel->GetLoadedLevel();
 			if (NewLevel)
@@ -378,6 +378,13 @@ ULevel* UEditorLevelUtils::AddLevelsToWorld(UWorld* InWorld, TArray<FString> Pac
 }
 
 ULevelStreaming* UEditorLevelUtils::AddLevelToWorld(UWorld* InWorld, const TCHAR* LevelPackageName, TSubclassOf<ULevelStreaming> LevelStreamingClass, const FTransform& LevelTransform)
+{
+	FAddLevelToWorldParams Params(LevelStreamingClass, LevelPackageName);
+	Params.Transform = LevelTransform;
+	return AddLevelToWorld(InWorld, Params);
+}
+
+ULevelStreaming* UEditorLevelUtils::AddLevelToWorld(UWorld* InWorld, const FAddLevelToWorldParams& InParams)
 {
 	if (!ensure(InWorld))
 	{
@@ -393,7 +400,7 @@ ULevelStreaming* UEditorLevelUtils::AddLevelToWorld(UWorld* InWorld, const TCHAR
 	// Try to add the levels that were specified in the dialog.
 	ULevel* NewLevel = nullptr;
 
-	ULevelStreaming* NewStreamingLevel = AddLevelToWorld_Internal(InWorld, LevelPackageName, LevelStreamingClass, LevelTransform);
+	ULevelStreaming* NewStreamingLevel = AddLevelToWorld_Internal(InWorld, InParams);
 
 	// Broadcast the levels have changed (new style)
 	InWorld->BroadcastLevelsChanged();
@@ -432,17 +439,16 @@ ULevelStreaming* UEditorLevelUtils::AddLevelToWorld(UWorld* InWorld, const TCHAR
 	return NewStreamingLevel;
 }
 
-ULevelStreaming* UEditorLevelUtils::AddLevelToWorld_Internal(UWorld* InWorld, const TCHAR* LevelPackageName, TSubclassOf<ULevelStreaming> LevelStreamingClass, const FTransform& LevelTransform)
+ULevelStreaming* UEditorLevelUtils::AddLevelToWorld_Internal(UWorld* InWorld, const FAddLevelToWorldParams& InParams)
 {
 	ULevel* NewLevel = nullptr;
 	ULevelStreaming* StreamingLevel = nullptr;
-	bool bIsPersistentLevel = (InWorld->PersistentLevel->GetOutermost()->GetName() == FString(LevelPackageName));
+	bool bIsPersistentLevel = (InWorld->PersistentLevel->GetOutermost()->GetFName() == InParams.PackageName);
 
-	if (bIsPersistentLevel || FLevelUtils::FindStreamingLevel(InWorld, LevelPackageName))
+	if (bIsPersistentLevel || FLevelUtils::FindStreamingLevel(InWorld, InParams.PackageName))
 	{
 		// Do nothing if the level already exists in the world.
-		const FString LevelName(LevelPackageName);
-		const FText MessageText = FText::Format(NSLOCTEXT("UnrealEd", "LevelAlreadyExistsInWorld", "A level with that name ({0}) already exists in the world."), FText::FromString(LevelName));
+		const FText MessageText = FText::Format(NSLOCTEXT("UnrealEd", "LevelAlreadyExistsInWorld", "A level with that name ({0}) already exists in the world."), FText::FromString(InParams.PackageName.ToString()));
 
 		FSuppressableWarningDialog::FSetupInfo Info(MessageText, LOCTEXT("AddLevelToWorld_Title", "Add Level"), "LevelAlreadyExistsInWorldWarning");
 		Info.ConfirmText = LOCTEXT("AlreadyExist_Ok", "Ok");
@@ -452,23 +458,29 @@ ULevelStreaming* UEditorLevelUtils::AddLevelToWorld_Internal(UWorld* InWorld, co
 	else
 	{
 		// If the selected class is still NULL or the selected class is abstract, abort the operation.
-		if (LevelStreamingClass == nullptr || LevelStreamingClass->HasAnyClassFlags(CLASS_Abstract))
+		if (InParams.LevelStreamingClass == nullptr || InParams.LevelStreamingClass->HasAnyClassFlags(CLASS_Abstract))
 		{
 			return nullptr;
 		}
 
 		const FScopedBusyCursor BusyCursor;
 
-		StreamingLevel = NewObject<ULevelStreaming>(InWorld, LevelStreamingClass, NAME_None, RF_NoFlags, NULL);
+		StreamingLevel = NewObject<ULevelStreaming>(InWorld, InParams.LevelStreamingClass, NAME_None, RF_NoFlags, NULL);
 
 		// Associate a package name.
-		StreamingLevel->SetWorldAssetByPackageName(LevelPackageName);
+		StreamingLevel->SetWorldAssetByPackageName(InParams.PackageName);
 
-		StreamingLevel->LevelTransform = LevelTransform;
+		StreamingLevel->LevelTransform = InParams.Transform;
 
 		// Seed the level's draw color.
 		StreamingLevel->LevelColor = FLinearColor::MakeRandomColor();
 
+		// Callback to allow initialization
+		if (InParams.LevelStreamingCreatedCallback)
+		{
+			InParams.LevelStreamingCreatedCallback(StreamingLevel);
+		}
+
 		// Add the new level to world.
 		InWorld->AddStreamingLevel(StreamingLevel);
 
@@ -676,13 +688,14 @@ ULevel* GetPersistentLevelForNewStreamingLevel(UWorld* NewLevelWorld, UWorld* In
 	return NewLevelWorld ? NewLevelWorld->PersistentLevel : InTemplateWorld->PersistentLevel;
 }
 
-UWorld* GetWorldForNewStreamingLevel(UWorld* InTemplateWorld, bool bIsPartitioned = false)
+UWorld* GetWorldForNewStreamingLevel(UWorld* InTemplateWorld, bool bCreateWorldPartition, bool bEnableWorldPartitionStreaming)
 {
 	if (!InTemplateWorld)
 	{
 		// Create a new world
 		UWorldFactory* Factory = NewObject<UWorldFactory>();
-		Factory->bCreateWorldPartition = bIsPartitioned;
+		Factory->bCreateWorldPartition = bCreateWorldPartition;
+		Factory->bEnableWorldPartitionStreaming = bEnableWorldPartitionStreaming;
 		Factory->WorldType = EWorldType::Inactive;
 		UPackage* Pkg = CreatePackage( NULL);
 		FName WorldName(TEXT("Untitled"));
@@ -717,6 +730,19 @@ ULevelStreaming* UEditorLevelUtils::CreateNewStreamingLevelForWorld(UWorld& InWo
 }
 
 ULevelStreaming* UEditorLevelUtils::CreateNewStreamingLevelForWorld(UWorld& InWorld, TSubclassOf<ULevelStreaming> LevelStreamingClass, bool bUseExternalActors, const FString& DefaultFilename /* = TEXT("") */, const TArray<AActor*>* ActorsToMove /* = nullptr */, UWorld* InTemplateWorld /* = nullptr */, bool bInUseSaveAs /*= true*/, bool bIsPartitioned /*= false*/, TFunction<void(ULevel*)> InPreSaveLevelOperation /* = TFunction<void(ULevel*)>()*/, const FTransform& InTransform /* = FTransform::Identity */)
+{
+	FCreateNewStreamingLevelForWorldParams CreateParams(LevelStreamingClass, DefaultFilename);
+	CreateParams.bUseExternalActors = bUseExternalActors;
+	CreateParams.ActorsToMove = ActorsToMove;
+	CreateParams.TemplateWorld = InTemplateWorld;
+	CreateParams.bUseSaveAs = bInUseSaveAs;
+	CreateParams.bCreateWorldPartition = bIsPartitioned;
+	CreateParams.PreSaveLevelCallback = InPreSaveLevelOperation;
+	CreateParams.Transform = InTransform;
+	return CreateNewStreamingLevelForWorld(InWorld, CreateParams);
+}
+
+ULevelStreaming* UEditorLevelUtils::CreateNewStreamingLevelForWorld(UWorld& InWorld,  const FCreateNewStreamingLevelForWorldParams& InCreateParams)
 {
 	// Editor modes cannot be active when any level saving occurs.
 	if (!IsRunningCommandlet())
@@ -738,32 +764,32 @@ ULevelStreaming* UEditorLevelUtils::CreateNewStreamingLevelForWorld(UWorld& InWo
 	UWorld* WorldToAddLevelTo = &InWorld;
 
 	// This is the new streaming level's world not the persistent level world
-	UWorld* NewLevelWorld = GetWorldForNewStreamingLevel(InTemplateWorld, bIsPartitioned);
-	ULevel* PersistentLevel = GetPersistentLevelForNewStreamingLevel(NewLevelWorld, InTemplateWorld);
+	UWorld* NewLevelWorld = GetWorldForNewStreamingLevel(InCreateParams.TemplateWorld, InCreateParams.bCreateWorldPartition, InCreateParams.bEnableWorldPartitionStreaming);
+	ULevel* PersistentLevel = GetPersistentLevelForNewStreamingLevel(NewLevelWorld, InCreateParams.TemplateWorld);
 	// No need to convert actors in partitioned worlds as they are already external
-	if (bUseExternalActors && !bIsPartitioned)
+	if (InCreateParams.bUseExternalActors && !InCreateParams.bCreateWorldPartition)
 	{
 		PersistentLevel->ConvertAllActorsToPackaging(true);
 		PersistentLevel->bUseExternalActors = true;
 	}
-	check(bIsPartitioned == PersistentLevel->bIsPartitioned);
+	check(InCreateParams.bCreateWorldPartition == PersistentLevel->bIsPartitioned);
 
-	if (InPreSaveLevelOperation)
+	if (InCreateParams.PreSaveLevelCallback)
 	{
 		// Call lambda before saving level
-		InPreSaveLevelOperation(PersistentLevel);
+		InCreateParams.PreSaveLevelCallback(PersistentLevel);
 	}
 
 	bool bNewWorldSaved = false;
-	FString NewPackageName = DefaultFilename;
+	FString NewPackageName = InCreateParams.DefaultFilename;
 
-	if (bInUseSaveAs)
+	if (InCreateParams.bUseSaveAs)
 	{
 		bNewWorldSaved = FEditorFileUtils::SaveLevelAs(PersistentLevel, &NewPackageName);
 	}
 	else
 	{
-		bNewWorldSaved = FEditorFileUtils::SaveLevel(PersistentLevel, DefaultFilename, &NewPackageName);
+		bNewWorldSaved = FEditorFileUtils::SaveLevel(PersistentLevel, InCreateParams.DefaultFilename, &NewPackageName);
 	}
 	
 	if (bNewWorldSaved && !NewPackageName.IsEmpty())
@@ -791,14 +817,18 @@ ULevelStreaming* UEditorLevelUtils::CreateNewStreamingLevelForWorld(UWorld& InWo
 			NewLevelWorld->CleanupWorld();
 		}
 
-		NewStreamingLevel = AddLevelToWorld(WorldToAddLevelTo, *NewPackageName, LevelStreamingClass, InTransform);
+		FAddLevelToWorldParams AddLevelToWorldParams(InCreateParams.LevelStreamingClass, *NewPackageName);
+		AddLevelToWorldParams.Transform = InCreateParams.Transform;
+		AddLevelToWorldParams.LevelStreamingCreatedCallback = InCreateParams.LevelStreamingCreatedCallback;
+
+		NewStreamingLevel = AddLevelToWorld(WorldToAddLevelTo, AddLevelToWorldParams);
 		if (NewStreamingLevel != nullptr)
 		{
 			NewLevel = NewStreamingLevel->GetLoadedLevel();
 			// If we are moving the selected actors to the new level move them now
-			if (ActorsToMove)
+			if (InCreateParams.ActorsToMove)
 			{
-				MoveActorsToLevel(*ActorsToMove, NewLevel);
+				MoveActorsToLevel(*InCreateParams.ActorsToMove, NewLevel);
 			}
 
 			// Finally make the new level the current one
diff --git a/Engine/Source/Editor/UnrealEd/Private/Factories/WorldFactory.cpp b/Engine/Source/Editor/UnrealEd/Private/Factories/WorldFactory.cpp
--- a/Engine/Source/Editor/UnrealEd/Private/Factories/WorldFactory.cpp
+++ b/Engine/Source/Editor/UnrealEd/Private/Factories/WorldFactory.cpp
@@ -18,6 +18,8 @@ UWorldFactory::UWorldFactory(const FObjectInitializer& ObjectInitializer)
 	WorldType = EWorldType::Inactive;
 	bInformEngineOfWorld = false;
 	bCreateWorldPartition = false;
+	// default to true to preserve previous behavior
+	bEnableWorldPartitionStreaming = true;
 	FeatureLevel = ERHIFeatureLevel::Num;
 }
 
@@ -37,7 +39,8 @@ UObject* UWorldFactory::FactoryCreateNew(UClass* Class, UObject* InParent, FName
 		.EnableTraceCollision(true)
 		.CreateNavigation(WorldType == EWorldType::Editor)
 		.CreateAISystem(WorldType == EWorldType::Editor)
-		.CreateWorldPartition(bCreateWorldPartition);
+		.CreateWorldPartition(bCreateWorldPartition)
+		.EnableWorldPartitionStreaming(bEnableWorldPartitionStreaming);
 
 	UWorld* NewWorld = UWorld::CreateWorld(WorldType, bInformEngineOfWorld, Name, Cast<UPackage>(InParent), bAddToRoot, FeatureLevel, &InitValues);
 	GEditor->InitBuilderBrush(NewWorld);
diff --git a/Engine/Source/Editor/UnrealEd/Public/EditorLevelUtils.h b/Engine/Source/Editor/UnrealEd/Public/EditorLevelUtils.h
--- a/Engine/Source/Editor/UnrealEd/Public/EditorLevelUtils.h
+++ b/Engine/Source/Editor/UnrealEd/Public/EditorLevelUtils.h
@@ -169,6 +169,55 @@ public:
 	 */
 	static UNREALED_API ULevelStreaming* CreateNewStreamingLevelForWorld(UWorld& World, TSubclassOf<ULevelStreaming> LevelStreamingClass, bool bUseExternalActors, const FString& DefaultFilename, const TArray<AActor*>* ActorsToMove = nullptr, UWorld* InTemplateWorld = nullptr, bool bInUseSaveAs = true, bool bIsPartitioned = false, TFunction<void(ULevel*)> InPreSaveLevelOperation = TFunction<void(ULevel*)>(), const FTransform& InTransform = FTransform::Identity);
 
+	struct FCreateNewStreamingLevelForWorldParams
+	{
+		FCreateNewStreamingLevelForWorldParams(TSubclassOf<ULevelStreaming> InLevelStreamingClass, FString InDefaultFilename)
+			: LevelStreamingClass(InLevelStreamingClass), DefaultFilename(InDefaultFilename) {}
+		
+		/** The streaming class type instead to use for the level */
+		TSubclassOf<ULevelStreaming> LevelStreamingClass;
+
+		/** file name for level.  If empty, the user will be prompted during the save process. */
+		FString DefaultFilename;
+
+		/** If valid, the new level will be a copy of the template world. */
+		UWorld* TemplateWorld = nullptr;
+
+		/** The transform to apply to the streaming level. */
+		FTransform Transform = FTransform::Identity;
+		
+		/** Optional, move provided actors into the new level. */
+		const TArray<AActor*>* ActorsToMove = nullptr;
+
+		/** If level should use external actors. */
+		bool bUseExternalActors = false;
+		
+		/** If level should be Partitioned (has precedence over bUseExternalActors). */
+		bool bCreateWorldPartition = false;
+		
+		/** If level WorldPartition should have streamng enabled (only valid if bCreateWorldPartition is true) */
+		bool bEnableWorldPartitionStreaming = true;
+		
+		/** If true, show SaveAs dialog instead of Save with DefaultFilename */
+		bool bUseSaveAs = true;
+
+		/** Optional function to call before saving the created level */
+		TFunction<void(ULevel*)> PreSaveLevelCallback = nullptr;
+
+		/** Optional function to call after the ULevelStreaming gets created. */
+		TFunction<void(ULevelStreaming*)> LevelStreamingCreatedCallback = nullptr;
+	};
+
+	/**
+	 * Creates a new streaming level and adds it to a world
+	 *
+	 * @param	InWorld			The world to add the streaming level to
+	 * @param	InCreateParams	Parameters used to create this new streaming world
+	 *
+	 * @return	Returns the newly created level, or NULL on failure
+	 */
+	static UNREALED_API ULevelStreaming* CreateNewStreamingLevelForWorld(UWorld& InWorld, const FCreateNewStreamingLevelForWorldParams& InCreateParams);
+
 	/**
 	 * Adds the named level packages to the world.  Does nothing if all the levels already exist in the world.
 	 *
@@ -232,14 +281,36 @@ public:
 	 * @param	InWorld				World in which to add the level
 	 * @param	LevelPackageName	The base filename of the level package to add.
 	 * @param	LevelStreamingClass	The streaming class type instead to use for the level.
+	 * @param	InTransform			The transform to apply to the streaming level.
 	 *
 	 * @return								The new level, or NULL if the level couldn't added.
 	 */
 	static UNREALED_API ULevelStreaming* AddLevelToWorld(UWorld* InWorld, const TCHAR* LevelPackageName, TSubclassOf<ULevelStreaming> LevelStreamingClass, const FTransform& LevelTransform = FTransform::Identity);
 
+	struct FAddLevelToWorldParams
+	{
+		FAddLevelToWorldParams(TSubclassOf<ULevelStreaming> InLevelStreamingClass, FName InLevelPackageName)
+			: LevelStreamingClass(InLevelStreamingClass), PackageName(InLevelPackageName)
+		{}
+
+		/** The streaming class type instead to use for the level. */
+		TSubclassOf<ULevelStreaming> LevelStreamingClass;
+		
+		/** The base filename of the level package to add. */
+		FName PackageName;
+		
+		/** The transform to apply to the streaming level. */
+		FTransform Transform = FTransform::Identity;
+
+		/** Optional function to call after the ULevelStreaming gets created. */
+		TFunction<void(ULevelStreaming*)> LevelStreamingCreatedCallback = nullptr;
+	};
+
+	static UNREALED_API ULevelStreaming* AddLevelToWorld(UWorld* InWorld, const FAddLevelToWorldParams& InParams);
+
 private:
 
-	static UNREALED_API ULevelStreaming* AddLevelToWorld_Internal(UWorld* InWorld, const TCHAR* LevelPackageName, TSubclassOf<ULevelStreaming> LevelStreamingClass, const FTransform& LevelTransform = FTransform::Identity);
+	static UNREALED_API ULevelStreaming* AddLevelToWorld_Internal(UWorld* InWorld, const FAddLevelToWorldParams& InParams);
 
 	static UNREALED_API int32 CopyOrMoveActorsToLevel(const TArray<AActor*>& ActorsToMove, ULevel* DestLevel, bool bMoveActors, bool bWarnAboutReferences = true, bool bWarnAboutRenaming = true, bool bMoveAllOrFail = false, TArray<AActor*>* OutActors = nullptr);
 
diff --git a/Engine/Source/Runtime/Engine/Classes/Engine/World.h b/Engine/Source/Runtime/Engine/Classes/Engine/World.h
--- a/Engine/Source/Runtime/Engine/Classes/Engine/World.h
+++ b/Engine/Source/Runtime/Engine/Classes/Engine/World.h
@@ -881,6 +881,7 @@ struct FWorldInitializationValues
 		, bTransactional(true)
 		, bCreateFXSystem(true)
 		, bCreateWorldPartition(false)
+		, bEnableWorldPartitionStreaming(true)
 	{
 	}
 
@@ -920,6 +921,9 @@ struct FWorldInitializationValues
 	/** Should the world be partitioned */
 	uint32 bCreateWorldPartition:1;
 
+	/** If bCreateWorldPartition is set to true, this flag will init the streaming mode for the WorldPartition object (default to true to preserve previous behavior) */
+	uint32 bEnableWorldPartitionStreaming:1;
+
 	/** The default game mode for this world (if any) */
 	TSubclassOf<class AGameModeBase> DefaultGameMode;
 
@@ -935,6 +939,7 @@ struct FWorldInitializationValues
 	FWorldInitializationValues& SetTransactional(const bool bInTransactional) { bTransactional = bInTransactional; return *this; }
 	FWorldInitializationValues& CreateFXSystem(const bool bCreate) { bCreateFXSystem = bCreate; return *this; }
 	FWorldInitializationValues& CreateWorldPartition(const bool bCreate) { bCreateWorldPartition = bCreate; return *this; }
+	FWorldInitializationValues& EnableWorldPartitionStreaming(const bool bEnableStreaming) { bEnableWorldPartitionStreaming = bEnableStreaming; return *this; }
 	FWorldInitializationValues& SetDefaultGameMode(TSubclassOf<class AGameModeBase> GameMode) { DefaultGameMode = GameMode; return *this; }
 };
 
diff --git a/Engine/Source/Runtime/Engine/Private/LevelInstance/LevelInstanceEditorLevelStreaming.cpp b/Engine/Source/Runtime/Engine/Private/LevelInstance/LevelInstanceEditorLevelStreaming.cpp
--- a/Engine/Source/Runtime/Engine/Private/LevelInstance/LevelInstanceEditorLevelStreaming.cpp
+++ b/Engine/Source/Runtime/Engine/Private/LevelInstance/LevelInstanceEditorLevelStreaming.cpp
@@ -20,15 +20,8 @@
 #include "WorldPartition/WorldPartition.h"
 #endif
 
-#if WITH_EDITOR
-static FLevelInstanceID EditLevelInstanceID;
-#endif
-
 ULevelStreamingLevelInstanceEditor::ULevelStreamingLevelInstanceEditor(const FObjectInitializer& ObjectInitializer)
 	: Super(ObjectInitializer)
-#if WITH_EDITOR
-	, LevelInstanceID(EditLevelInstanceID)
-#endif
 {
 #if WITH_EDITOR
 	SetShouldBeVisibleInEditor(true);
@@ -86,12 +79,26 @@ ULevelStreamingLevelInstanceEditor* ULevelStreamingLevelInstanceEditor::Load(ILe
 		UPackageTools::UnloadPackages({ LevelPackage });
 	}
 
-	TGuardValue<FLevelInstanceID> GuardEditLevelInstanceID(EditLevelInstanceID, LevelInstance->GetLevelInstanceID());
 	ULevelStreamingLevelInstanceEditor* LevelStreaming = nullptr;
+
+	// Set LevelInstanceID early
+	auto OnLevelStreamingCreated = [LevelInstanceID = LevelInstance->GetLevelInstanceID()](ULevelStreaming* InLevelStreaming)
+	{
+		ULevelStreamingLevelInstanceEditor* LevelInstanceLevelStreaming = Cast<ULevelStreamingLevelInstanceEditor>(InLevelStreaming);
+		check(LevelInstanceLevelStreaming);
+		check(!LevelInstanceLevelStreaming->LevelInstanceID.IsValid());
+		LevelInstanceLevelStreaming->LevelInstanceID = LevelInstanceID;
+	};
+
 	// If Asset is null we can create a level here
 	if (LevelInstance->GetWorldAsset().IsNull())
 	{
-		LevelStreaming = Cast<ULevelStreamingLevelInstanceEditor>(EditorLevelUtils::CreateNewStreamingLevelForWorld(*CurrentWorld, ULevelStreamingLevelInstanceEditor::StaticClass(), TEXT(""), false, nullptr, true, TFunction<void(ULevel*)>(), LevelInstanceActor->GetTransform()));
+		EditorLevelUtils::FCreateNewStreamingLevelForWorldParams CreateNewStreamingLevelParams(ULevelStreamingLevelInstanceEditor::StaticClass(), TEXT(""));
+		CreateNewStreamingLevelParams.Transform = LevelInstanceActor->GetTransform();
+		CreateNewStreamingLevelParams.LevelStreamingCreatedCallback = OnLevelStreamingCreated;
+		CreateNewStreamingLevelParams.bCreateWorldPartition = CurrentWorld->IsPartitionedWorld();
+
+		LevelStreaming = LevelInstance->GetLevelInstanceSubsystem()->CreateNewStreamingLevelForWorld(*CurrentWorld, CreateNewStreamingLevelParams);
 		if (LevelStreaming)
 		{
 			LevelInstanceActor->Modify();
@@ -100,7 +107,11 @@ ULevelStreamingLevelInstanceEditor* ULevelStreamingLevelInstanceEditor::Load(ILe
 	}
 	else
 	{
-		LevelStreaming = Cast<ULevelStreamingLevelInstanceEditor>(EditorLevelUtils::AddLevelToWorld(CurrentWorld, *LevelInstance->GetWorldAssetPackage(), ULevelStreamingLevelInstanceEditor::StaticClass(), LevelInstanceActor->GetTransform()));
+		EditorLevelUtils::FAddLevelToWorldParams AddLevelToWorldParams(ULevelStreamingLevelInstanceEditor::StaticClass(), *LevelInstance->GetWorldAssetPackage());
+		AddLevelToWorldParams.Transform = LevelInstanceActor->GetTransform();
+		AddLevelToWorldParams.LevelStreamingCreatedCallback = OnLevelStreamingCreated;
+
+		LevelStreaming = Cast<ULevelStreamingLevelInstanceEditor>(EditorLevelUtils::AddLevelToWorld(CurrentWorld, AddLevelToWorldParams));
 	}
 		
 	if (LevelStreaming)
@@ -154,13 +165,18 @@ void ULevelStreamingLevelInstanceEditor::OnLevelActorAdded(AActor* InActor)
 
 void ULevelStreamingLevelInstanceEditor::OnPreInitializeContainerInstance(UActorDescContainerInstance::FInitializeParams& InInitParams, UActorDescContainerInstance* InContainerInstance)
 {
-	// @todo_ow: Make sure LevelInstance exists (it is created after when creating a new Level Instance). This needs to be fixed properly but for now since the created LI is then unloaded it doesn't matter that the parent isn't set.
 	if (AActor* LevelInstanceActor = Cast<AActor>(GetLevelInstance()))
 	{
 		UWorldPartition* OwningWorldPartition = LevelInstanceActor->GetLevel()->GetWorldPartition();
 		// Add parenting info to init param
 		InInitParams.SetParent(OwningWorldPartition ? OwningWorldPartition->GetActorDescContainerInstance() : nullptr, LevelInstanceActor->GetActorGuid());
 	}
+	else
+	{
+		// When creating a new level instance the Level Instance actor doesn't exist yet
+		check(ParentContainerInstance && ParentContainerGuid.IsValid());
+		InInitParams.SetParent(ParentContainerInstance, ParentContainerGuid);
+	}
 }
 
 void ULevelStreamingLevelInstanceEditor::OnLevelLoadedChanged(ULevel* InLevel)
diff --git a/Engine/Source/Runtime/Engine/Private/LevelInstance/LevelInstanceSubsystem.cpp b/Engine/Source/Runtime/Engine/Private/LevelInstance/LevelInstanceSubsystem.cpp
--- a/Engine/Source/Runtime/Engine/Private/LevelInstance/LevelInstanceSubsystem.cpp
+++ b/Engine/Source/Runtime/Engine/Private/LevelInstance/LevelInstanceSubsystem.cpp
@@ -1019,6 +1019,75 @@ bool ULevelInstanceSubsystem::MoveActorsTo(ILevelInstanceInterface* LevelInstanc
 	return MoveActorsToLevel(ActorsToMove, LevelInstanceLevel, OutActors);
 }
 
+ULevelStreamingLevelInstanceEditor* ULevelInstanceSubsystem::CreateNewStreamingLevelForWorld(UWorld& InWorld, const EditorLevelUtils::FCreateNewStreamingLevelForWorldParams& InParams)
+{
+	EditorLevelUtils::FCreateNewStreamingLevelForWorldParams CreateNewStreamingLevelParamsCopy(InParams);
+	check(CreateNewStreamingLevelParamsCopy.LevelStreamingClass && CreateNewStreamingLevelParamsCopy.LevelStreamingClass->IsChildOf<ULevelStreamingLevelInstanceEditor>());
+		
+	CreateNewStreamingLevelParamsCopy.PreSaveLevelCallback = [ActorsToMove = InParams.ActorsToMove, PreSaveLevelCallback = InParams.PreSaveLevelCallback](ULevel* InLevel)
+	{
+		if (InLevel->IsUsingExternalActors())
+		{
+			// UWorldFactory::FactoryCreateNew will modify the default brush to be in global space (see GEditor->InitBuilderBrush(NewWorld)).
+			// The level is about to be saved, it doesn't have a transform and it is not yet added to the world levels.
+			// Since, no logic will remove the transform on the actor, force its transform to identity here.
+			if (ABrush* Brush = InLevel->GetDefaultBrush())
+			{
+				Brush->GetRootComponent()->SetRelativeTransform(FTransform::Identity);
+			}
+		}
+
+		if (UWorldPartition* WorldPartition = InLevel->GetWorldPartition())
+		{
+			// Validations
+			check(InLevel->IsUsingActorFolders());
+
+			// Reset HLOD Layer (no defaults needed for Level Instances)
+			WorldPartition->SetDefaultHLODLayer(nullptr);
+
+			// Make sure new level's AWorldDataLayers contains all the necessary Data Layer Instances before moving actors
+			TSet<TObjectPtr<const UDataLayerAsset>> SourceDataLayerAssets;
+			if (ActorsToMove)
+			{
+				for (AActor* ActorToMove : *ActorsToMove)
+				{
+					if (UDataLayerManager* DataLayerManager = UDataLayerManager::GetDataLayerManager(ActorToMove))
+					{
+						// Use the raw asset list as we don't want parent DataLayers
+						for (const UDataLayerAsset* DataLayerAsset : ActorToMove->GetDataLayerAssets())
+						{
+							if (const UDataLayerInstance* DataLayerInstance = DataLayerManager->GetDataLayerInstanceFromAsset(DataLayerAsset))
+							{
+								// Validate that there's a valid Data Layer Instance for this asset in the source level and that this isn't a private Data Layer
+								if (!DataLayerAsset->IsPrivate())
+								{
+									SourceDataLayerAssets.Add(DataLayerAsset);
+								}
+							}
+						}
+					}
+				}
+			}
+
+			if (!SourceDataLayerAssets.IsEmpty())
+			{
+				AWorldDataLayers* WorldDataLayers = InLevel->GetWorldDataLayers();
+				check(WorldDataLayers);
+				for (const UDataLayerAsset* SourceDataLayerAsset : SourceDataLayerAssets)
+				{
+					WorldDataLayers->CreateDataLayer<UDataLayerInstanceWithAsset>(SourceDataLayerAsset);
+				}
+			}
+		}
+
+		if (PreSaveLevelCallback)
+		{
+			PreSaveLevelCallback(InLevel);
+		}
+	};
+	return Cast<ULevelStreamingLevelInstanceEditor>(EditorLevelUtils::CreateNewStreamingLevelForWorld(*GetWorld(), CreateNewStreamingLevelParamsCopy));
+}
+
 ILevelInstanceInterface* ULevelInstanceSubsystem::CreateLevelInstanceFrom(const TArray<AActor*>& ActorsToMove, const FNewLevelInstanceParams& CreationParams)
 {
 	check(!bIsCreatingLevelInstance);
@@ -1102,82 +1171,28 @@ ILevelInstanceInterface* ULevelInstanceSubsystem::CreateLevelInstanceFrom(const
 		}
 	}
 
+	// Predetermine New Level Instance Actor Guid here and ContainerInstance so that we can feed them to the LevelStreaming object
+	FGuid LevelInstanceActorGuid = FGuid::NewGuid();
+	UWorldPartition* CurrentWorldPartition = CurrentLevel->GetWorldPartition();
+	UActorDescContainerInstance* ParentContainerInstance = CurrentWorldPartition ? CurrentWorldPartition->GetActorDescContainerInstance() : nullptr;
+
 	ULevelStreamingLevelInstanceEditor* LevelStreaming = nullptr;
 	{
-		const bool bIsPartitioned = GetWorld()->IsPartitionedWorld();
-		const bool bEnableStreaming = CreationParams.bEnableStreaming;
-
-		// We want to properly setup the world partition prior to its initialization
-		FDelegateHandle PreWorldInit = FWorldDelegates::OnPreWorldInitialization.AddLambda([bIsPartitioned, bEnableStreaming](UWorld* World, const UWorld::InitializationValues IVS)
-		{
-			if (bIsPartitioned)
-			{
-				UWorldPartition* WorldPartition = World->GetWorldPartition();
-				if (ensure(WorldPartition))
-				{
-					WorldPartition->bEnableStreaming = bEnableStreaming;
-				}
-			}
-		});
-
-		LevelStreaming = StaticCast<ULevelStreamingLevelInstanceEditor*>(EditorLevelUtils::CreateNewStreamingLevelForWorld(
-		*GetWorld(), ULevelStreamingLevelInstanceEditor::StaticClass(), CreationParams.UseExternalActors(), LevelFilename, &ActorsToMove, CreationParams.TemplateWorld, /*bUseSaveAs*/true, bIsPartitioned, [this, bIsPartitioned, bEnableStreaming, &ActorsToMove](ULevel* InLevel)
-		{
-			if (InLevel->IsUsingExternalActors())
-			{
-				// UWorldFactory::FactoryCreateNew will modify the default brush to be in global space (see GEditor->InitBuilderBrush(NewWorld)).
-				// The level is about to be saved, it doesn't have a transform and it is not yet added to the world levels.
-				// Since, no logic will remove the transform on the actor, force its transform to identity here.
-				if (ABrush* Brush = InLevel->GetDefaultBrush())
-				{
-					Brush->GetRootComponent()->SetRelativeTransform(FTransform::Identity);
-				}
-			}
-			if (bIsPartitioned)
-			{
-				UWorldPartition* WorldPartition = InLevel->GetWorldPartition();
-				// Validations
-				check(WorldPartition);
-				check(WorldPartition->IsStreamingEnabled() == bEnableStreaming);
-				check(InLevel->IsUsingActorFolders());
-
-				// Reset HLOD Layer (no defaults needed for Level Instances)
-				WorldPartition->SetDefaultHLODLayer(nullptr);
-
-				// Make sure new level's AWorldDataLayers contains all the necessary Data Layer Instances before moving actors
-				TSet<TObjectPtr<const UDataLayerAsset>> SourceDataLayerAssets;
-				for (AActor* ActorToMove : ActorsToMove)
-				{
-					if (UDataLayerManager* DataLayerManager = UDataLayerManager::GetDataLayerManager(ActorToMove))
-					{
-						// Use the raw asset list as we don't want parent DataLayers
-						for (const UDataLayerAsset* DataLayerAsset : ActorToMove->GetDataLayerAssets())
-						{
-							if (const UDataLayerInstance* DataLayerInstance = DataLayerManager->GetDataLayerInstanceFromAsset(DataLayerAsset))
-							{
-								// Validate that there's a valid Data Layer Instance for this asset in the source level and that this isn't a private Data Layer
-								if (!DataLayerAsset->IsPrivate())
-								{
-									SourceDataLayerAssets.Add(DataLayerAsset);
-								}
-							}
-						}
-					}
-				}
-
-				if (!SourceDataLayerAssets.IsEmpty())
-				{
-					AWorldDataLayers* WorldDataLayers = InLevel->GetWorldDataLayers();
-					check(WorldDataLayers);
-					for (const UDataLayerAsset* SourceDataLayerAsset : SourceDataLayerAssets)
-					{
-						WorldDataLayers->CreateDataLayer<UDataLayerInstanceWithAsset>(SourceDataLayerAsset);
-					}
-				}
-			}
-		}));
+		EditorLevelUtils::FCreateNewStreamingLevelForWorldParams CreateNewStreamingLevelParams(ULevelStreamingLevelInstanceEditor::StaticClass(), LevelFilename);
+		CreateNewStreamingLevelParams.bUseExternalActors = CreationParams.UseExternalActors();
+		CreateNewStreamingLevelParams.bUseSaveAs = true;
+		CreateNewStreamingLevelParams.bCreateWorldPartition = GetWorld()->IsPartitionedWorld();
+		CreateNewStreamingLevelParams.bEnableWorldPartitionStreaming = CreationParams.bEnableStreaming;
+		CreateNewStreamingLevelParams.ActorsToMove = &ActorsToMove;
+		CreateNewStreamingLevelParams.TemplateWorld = CreationParams.TemplateWorld,
+		CreateNewStreamingLevelParams.LevelStreamingCreatedCallback = [LevelInstanceActorGuid, ParentContainerInstance](ULevelStreaming* InLevelStreaming)
+		{
+			ULevelStreamingLevelInstanceEditor* LevelInstanceLevelStreaming = CastChecked<ULevelStreamingLevelInstanceEditor>(InLevelStreaming);
+			LevelInstanceLevelStreaming->ParentContainerInstance = ParentContainerInstance;
+			LevelInstanceLevelStreaming->ParentContainerGuid = LevelInstanceActorGuid;
+		};
 
-		FWorldDelegates::OnPreWorldInitialization.Remove(PreWorldInit);
+		LevelStreaming = CreateNewStreamingLevelForWorld(*GetWorld(), CreateNewStreamingLevelParams);
 	}
 
 	if (!LevelStreaming)
@@ -1199,6 +1214,7 @@ ILevelInstanceInterface* ULevelInstanceSubsystem::CreateLevelInstanceFrom(const
 	}
 
 	FActorSpawnParameters SpawnParams;
+	SpawnParams.OverrideActorGuid = LevelInstanceActorGuid;
 	SpawnParams.OverrideLevel = CurrentLevel;
 	AActor* NewLevelInstanceActor = nullptr;
 	TSoftObjectPtr<UWorld> WorldPtr(LoadedLevel->GetTypedOuter<UWorld>());
@@ -1249,6 +1265,7 @@ ILevelInstanceInterface* ULevelInstanceSubsystem::CreateLevelInstanceFrom(const
 	}
 	
 	check(NewLevelInstanceActor);
+	check(NewLevelInstanceActor->GetActorGuid() == LevelInstanceActorGuid);
 
 	ILevelInstanceInterface* NewLevelInstance = CastChecked<ILevelInstanceInterface>(NewLevelInstanceActor);
 	NewLevelInstance->SetWorldAsset(WorldPtr);
diff --git a/Engine/Source/Runtime/Engine/Private/World.cpp b/Engine/Source/Runtime/Engine/Private/World.cpp
--- a/Engine/Source/Runtime/Engine/Private/World.cpp
+++ b/Engine/Source/Runtime/Engine/Private/World.cpp
@@ -2366,7 +2366,8 @@ void UWorld::InitializeNewWorld(const InitializationValues IVS, bool bInSkipInit
 		
 		check(!GetStreamingLevels().Num());
 		
-		UWorldPartition::CreateOrRepairWorldPartition(WorldSettings);
+		UWorldPartition* WorldPartition = UWorldPartition::CreateOrRepairWorldPartition(WorldSettings);
+		WorldPartition->bEnableStreaming = IVS.bEnableWorldPartitionStreaming;
 	}
 #endif
 
diff --git a/Engine/Source/Runtime/Engine/Public/LevelInstance/LevelInstanceEditorLevelStreaming.h b/Engine/Source/Runtime/Engine/Public/LevelInstance/LevelInstanceEditorLevelStreaming.h
--- a/Engine/Source/Runtime/Engine/Public/LevelInstance/LevelInstanceEditorLevelStreaming.h
+++ b/Engine/Source/Runtime/Engine/Public/LevelInstance/LevelInstanceEditorLevelStreaming.h
@@ -38,5 +38,9 @@ protected:
 	virtual void OnLevelLoadedChanged(ULevel* Level) override;
 private:
 	FLevelInstanceID LevelInstanceID;
+
+	// When creating a new LevelInstance initialize UActorDescContainerInstance using those values
+	UActorDescContainerInstance* ParentContainerInstance = nullptr;
+	FGuid ParentContainerGuid;
 #endif
 };
\ No newline at end of file
diff --git a/Engine/Source/Runtime/Engine/Public/LevelInstance/LevelInstanceSubsystem.h b/Engine/Source/Runtime/Engine/Public/LevelInstance/LevelInstanceSubsystem.h
--- a/Engine/Source/Runtime/Engine/Public/LevelInstance/LevelInstanceSubsystem.h
+++ b/Engine/Source/Runtime/Engine/Public/LevelInstance/LevelInstanceSubsystem.h
@@ -12,6 +12,10 @@
 #include "WorldPartition/WorldPartitionHandle.h"
 #include "WorldPartition/WorldPartitionActorContainerID.h"
 
+#if WITH_EDITOR
+#include "EditorLevelUtils.h"
+#endif
+
 #include "LevelInstanceSubsystem.generated.h"
 
 class ILevelInstanceInterface;
@@ -151,6 +155,8 @@ private:
 	ENGINE_API void RegisterLoadedLevelStreamingLevelInstance(ULevelStreamingLevelInstance* LevelStreaming);
 
 #if WITH_EDITOR
+	ENGINE_API ULevelStreamingLevelInstanceEditor* CreateNewStreamingLevelForWorld(UWorld& InWorld, const EditorLevelUtils::FCreateNewStreamingLevelForWorldParams& InParams);
+	
 	ENGINE_API FWorldPartitionActorFilter GetLevelInstanceFilterInternal(const FString& LevelPackage, TSet<FString>& VisitedPackages) const;
 
 	ENGINE_API void ResetLoadersForWorldAssetInternal(const FString& WorldAsset);
