Cooker: Add -cookincludeassetclass argument.
#rnx

diff --git a/Engine/Source/Editor/UnrealEd/Classes/CookOnTheSide/CookOnTheFlyServer.h b/Engine/Source/Editor/UnrealEd/Classes/CookOnTheSide/CookOnTheFlyServer.h
--- a/Engine/Source/Editor/UnrealEd/Classes/CookOnTheSide/CookOnTheFlyServer.h
+++ b/Engine/Source/Editor/UnrealEd/Classes/CookOnTheSide/CookOnTheFlyServer.h
@@ -1088,6 +1088,7 @@ private:
 	UNREALED_API void LoadInitializeConfigSettings(const FString& InOutputDirectoryOverride);
 	UNREALED_API void SetInitializeConfigSettings(UE::Cook::FInitializeConfigSettings&& Settings);
 	UNREALED_API void ParseCookFilters();
+	UNREALED_API void ParseCookFilters(const TCHAR* Parameter, const TCHAR* Message, TSet<FName>& OutFilterClasses);
 
 	/** Set parameters that rely on config and CookByTheBook settings. */
 	UNREALED_API void LoadBeginCookConfigSettings(FBeginCookContext& BeginContext);
@@ -1453,6 +1454,8 @@ private:
 	TMap<FName, TArray<FName>> DiscoveredDependencies;
 	/** Classes (and all subclasses) that were listed as the only classes that should be cooked in the filter settings. */
 	TSet<FName> CookFilterIncludedClasses;
+	/** Classes (and all subclasses) that were listed as the only asset classes that should be cooked in the filter settings. */
+	TSet<FName> CookFilterIncludedAssetClasses;
 
 	ELogVerbosity::Type CookerIdleWarningSeverity = ELogVerbosity::Warning;
 	/** True when PumpLoads has detected it is blocked on async work and CookOnTheFlyServer should do work elsewhere. */
diff --git a/Engine/Source/Editor/UnrealEd/Private/CookOnTheFlyServer.cpp b/Engine/Source/Editor/UnrealEd/Private/CookOnTheFlyServer.cpp
--- a/Engine/Source/Editor/UnrealEd/Private/CookOnTheFlyServer.cpp
+++ b/Engine/Source/Editor/UnrealEd/Private/CookOnTheFlyServer.cpp
@@ -6927,14 +6927,22 @@ void UCookOnTheFlyServer::SetInitializeConfigSettings(UE::Cook::FInitializeConfi
 void UCookOnTheFlyServer::ParseCookFilters()
 {
 	CookFilterIncludedClasses.Empty();
+	CookFilterIncludedAssetClasses.Empty();
 	bCookFilter = false;
 	if (!IsCookByTheBookMode() || IsCookingInEditor())
 	{
 		return;
 	}
 
+	ParseCookFilters(TEXT("cookincludeclass"), TEXT("contain an object"), CookFilterIncludedClasses);
+	ParseCookFilters(TEXT("cookincludeassetclass"), TEXT("contain an asset"), CookFilterIncludedAssetClasses);
+}
+
+void UCookOnTheFlyServer::ParseCookFilters(const TCHAR* Parameter, const TCHAR* Message, TSet<FName>& OutFilterClasses)
+{
 	FString IncludeClassesString;
-	if (FParse::Value(FCommandLine::Get(), TEXT("-cookincludeclass="), IncludeClassesString))
+	TStringBuilder<256> FullParameter(InPlace, TEXT("-"), Parameter, TEXT("="));
+	if (FParse::Value(FCommandLine::Get(), *FullParameter, IncludeClassesString))
 	{
 		TArray<FString> IncludeClasses;
 		const TCHAR* Delimiters[] = { TEXT(","), TEXT("+"), TEXT(";")};
@@ -6954,20 +6962,20 @@ void UCookOnTheFlyServer::ParseCookFilters()
 			}
 			if (!ClassPath.IsValid())
 			{
-				UE_LOG(LogCook, Error, TEXT("CookIncludeClass: Could not convert string '%s' into a class path. Ignoring it."),
-					*IncludeClassString);
+				UE_LOG(LogCook, Error, TEXT("%s: Could not convert string '%s' into a class path. Ignoring it."),
+					Parameter, *IncludeClassString);
 				continue;
 			}
 			UClass* IncludedClass = FindObject<UClass>(nullptr, *ClassPath.ToString());
 			if (!IncludedClass)
 			{
-				UE_LOG(LogCook, Error, TEXT("CookIncludeClass: Could not find class with ClassPath '%s'. Ignoring it."),
-					*IncludeClassString);
+				UE_LOG(LogCook, Error, TEXT("%s: Could not find class with ClassPath '%s'. Ignoring it."),
+					Parameter, *IncludeClassString);
 				continue;
 			}
 			FTopLevelAssetPath NormalizedClassPath(IncludedClass);
 			RootNames.Add(NormalizedClassPath);
-			CookFilterIncludedClasses.Add(FName(*NormalizedClassPath.ToString()));
+			OutFilterClasses.Add(FName(*NormalizedClassPath.ToString()));
 		}
 		if (!RootNames.IsEmpty())
 		{
@@ -6976,11 +6984,11 @@ void UCookOnTheFlyServer::ParseCookFilters()
 				DerivedClassNames);
 			for (const FTopLevelAssetPath& NormalizedClassPath : DerivedClassNames)
 			{
-				CookFilterIncludedClasses.Add(FName(*NormalizedClassPath.ToString()));
+				OutFilterClasses.Add(FName(*NormalizedClassPath.ToString()));
 			}
 
-			UE_LOG(LogCook, Display, TEXT("CookIncludeClass: Only cooking packages that contains objects with class in { %s }"),
-				*FString::JoinBy(RootNames, TEXT(", "), [](const FTopLevelAssetPath& P) { return P.ToString(); }));
+			UE_LOG(LogCook, Display, TEXT("%s: Only cooking packages that %s with class in { %s }"),
+				Parameter, Message, *TStringBuilder<256>().Join(RootNames, TEXTVIEW(", ")));
 			bCookFilter = true;
 		}
 	}
diff --git a/Engine/Source/Editor/UnrealEd/Private/Cooker/CookRequestCluster.cpp b/Engine/Source/Editor/UnrealEd/Private/Cooker/CookRequestCluster.cpp
--- a/Engine/Source/Editor/UnrealEd/Private/Cooker/CookRequestCluster.cpp
+++ b/Engine/Source/Editor/UnrealEd/Private/Cooker/CookRequestCluster.cpp
@@ -1637,15 +1637,39 @@ void FRequestCluster::IsRequestCookable(const ITargetPlatform* Platform, FName P
 	// The package is ordinarily cookable and explorable. In some cases we filter out for testing
 	// packages that are ordinarily cookable; set bOutCookable to false if so.
 	bOutExplorable = true;
-	if (InCOTFS.bCookFilter && !InCOTFS.CookFilterIncludedClasses.IsEmpty())
+	if (InCOTFS.bCookFilter)
 	{
-		TOptional<FAssetPackageData> AssetData = IAssetRegistry::GetChecked().GetAssetPackageDataCopy(PackageName);
-		if (AssetData)
+		IAssetRegistry& AssetRegistry = IAssetRegistry::GetChecked();
+		if (!InCOTFS.CookFilterIncludedClasses.IsEmpty())
 		{
+			TOptional<FAssetPackageData> AssetData = AssetRegistry.GetAssetPackageDataCopy(PackageName);
 			bool bIncluded = false;
-			for (FName ClassName : AssetData->ImportedClasses)
+			if (AssetData)
 			{
-				if (InCOTFS.CookFilterIncludedClasses.Contains(ClassName))
+				for (FName ClassName : AssetData->ImportedClasses)
+				{
+					if (InCOTFS.CookFilterIncludedClasses.Contains(ClassName))
+					{
+						bIncluded = true;
+						break;
+					}
+				}
+			}
+			if (!bIncluded)
+			{
+				OutReason = ESuppressCookReason::CookFilter;
+				bOutCookable = false;
+				return;
+			}
+		}
+		if (!InCOTFS.CookFilterIncludedAssetClasses.IsEmpty())
+		{
+			TArray<FAssetData> AssetDatas;
+			AssetRegistry.GetAssetsByPackageName(PackageName, AssetDatas, true /* bIncludeOnlyDiskAssets */);
+			bool bIncluded = false;
+			for (FAssetData& AssetData : AssetDatas)
+			{
+				if (InCOTFS.CookFilterIncludedAssetClasses.Contains(FName(*AssetData.AssetClassPath.ToString())))
 				{
 					bIncluded = true;
 					break;
