Fix legacy fbx parser bone name limits.

#rnx

diff --git a/Engine/Source/Editor/UnrealEd/Private/SkeletalMeshEdit.cpp b/Engine/Source/Editor/UnrealEd/Private/SkeletalMeshEdit.cpp
--- a/Engine/Source/Editor/UnrealEd/Private/SkeletalMeshEdit.cpp
+++ b/Engine/Source/Editor/UnrealEd/Private/SkeletalMeshEdit.cpp
@@ -2309,6 +2309,14 @@ void UnFbx::FFbxImporter::ImportBoneTracks(USkeleton* Skeleton, FAnimCurveImport
 				if (BoneTreeIndex != INDEX_NONE)
 				{
 					//add new track
+					if (BoneName.GetStringLength() > 92)
+					{
+						//The bone name exceed the maximum length supported by the animation system
+						//The animation system is adding _CONTROL to the bone name to name the animation controller and
+						//the maximum total length is cap at 100, so user should not import bone name longer then 92 characters
+						AddTokenizedErrorMessage(FTokenizedMessage::Create(EMessageSeverity::Warning, FText::Format(LOCTEXT("Error_BoneNameExceed92Characters", "Bone with animation cannot have a name exceeding 92 characters: {0}"), FText::FromName(BoneName))), FFbxErrors::Animation_InvalidData);
+						continue;
+					}
 					Controller.AddBoneCurve(BoneName, bShouldTransact);
 					Controller.SetBoneTrackKeys(BoneName, RawTrack.PosKeys, RawTrack.RotKeys, RawTrack.ScaleKeys, bShouldTransact);
 					NewDebugData.SetTrackData(BoneTreeIndex, BoneName);
