[UnrealEd] Added toast notifications for viewport bookmark actions

diff --git a/Engine/Source/Editor/UnrealEd/Private/Bookmarks/BookmarkTypeToolsImpl.cpp b/Engine/Source/Editor/UnrealEd/Private/Bookmarks/BookmarkTypeToolsImpl.cpp
--- a/Engine/Source/Editor/UnrealEd/Private/Bookmarks/BookmarkTypeToolsImpl.cpp
+++ b/Engine/Source/Editor/UnrealEd/Private/Bookmarks/BookmarkTypeToolsImpl.cpp
@@ -8,6 +8,8 @@
 #include "EditorViewportClient.h"
 #include "Engine/BookmarkBase.h"
 #include "ScopedTransaction.h"
+#include "Framework/Notifications/NotificationManager.h"
+#include "Widgets/Notifications/SNotificationList.h"
 
 #include "Logging/LogMacros.h"
 
@@ -23,6 +25,15 @@ private:
 
 private:
 
+	void Notify(const FText& InText, const bool bInSuccess = true) const
+	{
+		FNotificationInfo Info(InText);
+		Info.bUseLargeFont = false;
+		Info.bUseSuccessFailIcons = false;
+			
+		FSlateNotificationManager::Get().AddNotification(Info);	
+	}
+
 	static FORCEINLINE AWorldSettings* GetWorldSettings(FEditorViewportClient* InViewportClient)
 	{
 		return const_cast<AWorldSettings*>(GetWorldSettings(const_cast<const FEditorViewportClient*>(InViewportClient)));
@@ -119,10 +130,16 @@ public:
 				{
 					Actions->InitFromViewport(Bookmark, *InViewportClient);
 				}
+
+				const FText NotifyMessage = LOCTEXT("CreateOrSetBookmarkSuccessFormat", "Saved bookmark {0}");
+				Notify(FText::Format(NotifyMessage, InIndex), true);
 			}
 			else
 			{
 				UE_LOG(LogEditorBookmarks, Error, TEXT("FBookmarkTypeToolsImpl::CreateOrSetBookmark - Failed to create bookmark at Index %d"), InIndex);
+
+				const FText NotifyMessage = LOCTEXT("CreateOrSetBookmarkFailedFormat", "Failed to save bookmark {0}");
+				Notify(FText::Format(NotifyMessage, InIndex), false);
 			}
 		}
 	}
@@ -154,6 +171,9 @@ public:
 	{
 		if (AWorldSettings* WorldSettings = GetWorldSettings(InViewportClient))
 		{
+			bool bFailed = false;
+			const FText FailNotifyMessage = LOCTEXT("JumpToBookmarkFailedFormat", "No bookmark at {0} to jump to");
+			
 			const TArray<UBookmarkBase*>& Bookmarks = WorldSettings->GetBookmarks();
 			if (Bookmarks.IsValidIndex(InIndex))
 			{
@@ -164,6 +184,9 @@ public:
 					{
 						Actions->JumpToBookmark(Bookmark, InSettings, *InViewportClient);
 					}
+
+					const FText NotifyMessage = LOCTEXT("JumpToBookmarkSuccessFormat", "Jumped to bookmark {0}");
+					Notify(FText::Format(NotifyMessage, InIndex), true);
 				}
 				else
 				{
@@ -174,6 +197,11 @@ public:
 			{
 				UE_LOG(LogEditorBookmarks, Warning, TEXT("FBookmarkTypeToolsImpl::JumpToBookmark - Invalid bookmark index %d"), InIndex);
 			}
+
+			if (bFailed)
+			{
+				Notify(FText::Format(FailNotifyMessage, InIndex), false);
+			}
 		}
 	}
 
@@ -190,6 +218,9 @@ public:
 		if (AWorldSettings* WorldSettings = GetWorldSettings(InViewportClient))
 		{
 			WorldSettings->ClearBookmark(InIndex);
+
+			const FText NotifyMessage = LOCTEXT("ClearBookmarkSuccessFormat", "Cleared bookmark {0}");
+			Notify(FText::Format(NotifyMessage, InIndex), true);
 		}
 	}
 
@@ -205,6 +236,9 @@ public:
 		if (AWorldSettings* WorldSettings = GetWorldSettings(InViewportClient))
 		{
 			WorldSettings->ClearAllBookmarks();
+
+			const FText NotifyMessage = LOCTEXT("ClearBookmarkSuccessFormat", "Cleared all bookmarks");
+			Notify(NotifyMessage, true);
 		}
 	}
 
