PR #11216: Fix Overlapping section name in FMainMRUFavoritesList by allowing an override for the favorites section ini
This fixes a bug where favorite levels were not persisting

diff --git a/Engine/Source/Editor/UnrealEd/Private/MRUFavoritesList.cpp b/Engine/Source/Editor/UnrealEd/Private/MRUFavoritesList.cpp
--- a/Engine/Source/Editor/UnrealEd/Private/MRUFavoritesList.cpp
+++ b/Engine/Source/Editor/UnrealEd/Private/MRUFavoritesList.cpp
@@ -14,15 +14,20 @@
 #include "Templates/SharedPointer.h"
 #include "Widgets/Notifications/SNotificationList.h"
 
-const FString FMainMRUFavoritesList::FAVORITES_INI_SECTION = TEXT("FavoriteFiles");
-
 FMainMRUFavoritesList::FMainMRUFavoritesList()
 	: FMRUList( TEXT("MRU") )
+	, INIFavoritesSection(TEXT("FavoriteFiles"))
 {
 }
-
 FMainMRUFavoritesList::FMainMRUFavoritesList(const FString& IniSectionOverride, const int32 InitMaxItems)
-	:FMRUList(IniSectionOverride, InitMaxItems)
+	: FMRUList(IniSectionOverride, InitMaxItems)
+	, INIFavoritesSection(TEXT("FavoriteFiles"))
+{
+}
+
+FMainMRUFavoritesList::FMainMRUFavoritesList(const FString& IniSectionOverride, const FString& IniFavoritesSectionOverride, const int32 InitMaxItems)
+	: FMRUList(IniSectionOverride, InitMaxItems)
+	, INIFavoritesSection(IniFavoritesSectionOverride)
 {
 }
 
@@ -39,14 +44,14 @@ void FMainMRUFavoritesList::ReadFromINI()
 	InternalReadINI( Items, INISection, TEXT("MRUItem"), GetMaxItems() );
 
 	// Read in the Favorite items
-	InternalReadINI( FavoriteItems, FAVORITES_INI_SECTION, TEXT("FavoritesItem"), MaxItems );	
+	InternalReadINI( FavoriteItems, INIFavoritesSection, TEXT("FavoritesItem"), MaxItems );
 }
 
 /** Save off the state of the MRU and favorites lists to the relevant INI file */
 void FMainMRUFavoritesList::WriteToINI() const
 {
 	InternalWriteINI( Items, INISection, TEXT("MRUItem") );
-	InternalWriteINI( FavoriteItems, FAVORITES_INI_SECTION, TEXT("FavoritesItem") );
+	InternalWriteINI( FavoriteItems, INIFavoritesSection, TEXT("FavoritesItem") );
 }
 
 /**
diff --git a/Engine/Source/Editor/UnrealEd/Private/Subsystems/AssetEditorSubsystem.cpp b/Engine/Source/Editor/UnrealEd/Private/Subsystems/AssetEditorSubsystem.cpp
--- a/Engine/Source/Editor/UnrealEd/Private/Subsystems/AssetEditorSubsystem.cpp
+++ b/Engine/Source/Editor/UnrealEd/Private/Subsystems/AssetEditorSubsystem.cpp
@@ -111,7 +111,7 @@ void UAssetEditorSubsystem::Deinitialize()
 void UAssetEditorSubsystem::InitializeRecentAssets()
 {
 	// The current max allowed is 30 assets
-	RecentAssetsList = MakeUnique<FMainMRUFavoritesList>(TEXT("AssetEditorSubsystemRecents"), 30);
+	RecentAssetsList = MakeUnique<FMainMRUFavoritesList>(TEXT("AssetEditorSubsystemRecents"), TEXT("AssetEditorSubsystemFavorites"), 30);
 	RecentAssetsList->ReadFromINI();
 	
 	TArray<FString> RecentAssetEditors;
diff --git a/Engine/Source/Editor/UnrealEd/Public/MRUFavoritesList.h b/Engine/Source/Editor/UnrealEd/Public/MRUFavoritesList.h
--- a/Engine/Source/Editor/UnrealEd/Public/MRUFavoritesList.h
+++ b/Engine/Source/Editor/UnrealEd/Public/MRUFavoritesList.h
@@ -21,7 +21,10 @@ public:
 	/** Constructor */
 	UNREALED_API FMainMRUFavoritesList();
 
+	UE_DEPRECATED(5.4, "Use the constructor that takes in an override section for the favorites as well")
 	UNREALED_API FMainMRUFavoritesList(const FString& IniSectionOverride, const int32 InitMaxItems = 12);
+
+	UNREALED_API FMainMRUFavoritesList(const FString& IniSectionOverride, const FString& IniFavoritesSectionOverride, const int32 InitMaxItems = 12);
 	
 	/** Destructor */
 	UNREALED_API ~FMainMRUFavoritesList();
@@ -132,5 +135,5 @@ private:
 	TArray<FString> FavoriteItems;
 
 	/** INI section to read/write favorite items to */
-	static UNREALED_API const FString FAVORITES_INI_SECTION;
+	FString INIFavoritesSection;
 };
