Localization:
Refactor for GatherTextFromAssetsCommandlet. Refactor is purely structural, nothing behavioral yet.
- Broke down UGatherTextFromAssetsCommandlet::Main() into a series of much finer functions for easier readability and understandability.
- Added comments to improve understandability of the gather text from assets algorithm.
- Added profiling and logs to measure the time various parts of the GatherTextFromAssets pipeline take.

#test Ran a localization gather locally to get the resultant .po files and compared them against a run with the previous implementation.
#rnx

diff --git a/Engine/Source/Editor/UnrealEd/Classes/Commandlets/GatherTextFromAssetsCommandlet.h b/Engine/Source/Editor/UnrealEd/Classes/Commandlets/GatherTextFromAssetsCommandlet.h
--- a/Engine/Source/Editor/UnrealEd/Classes/Commandlets/GatherTextFromAssetsCommandlet.h
+++ b/Engine/Source/Editor/UnrealEd/Classes/Commandlets/GatherTextFromAssetsCommandlet.h
@@ -6,8 +6,12 @@
 #include "UObject/ObjectMacros.h"
 #include "Internationalization/GatherableTextData.h"
 #include "Commandlets/GatherTextCommandletBase.h"
+#include "AssetRegistry/AssetData.h"
 #include "GatherTextFromAssetsCommandlet.generated.h"
 
+struct FARFilter;
+struct FPackageFileSummary;
+
 /**
  *	UGatherTextFromAssetsCommandlet: Localization commandlet that collects all text to be localized from the game assets.
  */
@@ -43,6 +47,35 @@ public:
 	};
 
 private:
+	/** Parses the command line for the commandlet. Returns true if all required parameters are provided and are correct.*/
+	bool ParseCommandLineHelper(const FString& InCommandLine);
+
+// Filtering of asset registry elements
+	// Broadly, there is the first pass filter,the exact class filter and the include/exclude path filter that can be applied to filter out asset registry elements.
+	// Look at Main() to see how the functions are applied to understand the logic.
+	bool PerformFirstPassFilter(TArray<FAssetData>& OutAssetDataArray) const;
+	void ApplyFirstPassFilter(const FARFilter& InFilter, TArray<FAssetData>& InOutAssetDataArray) const;
+	bool BuildFirstPassFilter(FARFilter& InOutFilter) const;
+	bool BuildCollectionFilter(FARFilter& InOutFilter) const;
+	bool BuildExcludeDerivedClassesFilter(FARFilter& InOutFilter) const;
+	bool PerformExcludeExactClassesFilter(TArray<FAssetData>& InOutAssetDataArray) const;
+	bool BuildExcludeExactClassesFilter(FARFilter& InOutFilter) const;
+	void ApplyExcludeExactClassesFilter(const FARFilter& InFilter, TArray<FAssetData>& InOutAssetDataArray) const;
+	void FilterAssetsBasedOnIncludeExcludePaths(TArray<FAssetData>& InOutAssetDataArray) const;
+	
+	void DiscoverExternalActors(TArray<FAssetData>& InOutAssetDataArray) const;
+	void RemoveExistingExternalActors(TArray<FAssetData>& InOutAssetDataArray, TArray<FName>& OutExternalActorsSearchPaths) const;
+
+	TSet<FName> GetPackageNamesToGather(const TArray<FAssetData>& InAssetDataArray) const;
+	void PopulatePackagesPendingGather(TSet<FName> PackageNamesToGather);
+	void ProcessAndRemoveCachedPackages(TMap<FName, TSet<FGuid>>& OutExternalActorsWithStaleOrMissingCaches);
+	void MergeInExternalActorsWithStaleOrMissingCaches(TMap<FName, TSet<FGuid>>& ExternalActorsWithStaleOrMissingCaches);
+
+	void LoadAndProcessUncachedPackages(TArray<FName>& OutPackagesWithStaleGatherCache);
+
+	void ReportStaleGatherCache(TArray<FName>& InPackagesWithStaleGatherCache) const;
+	/** Determines the loc cache state for a package. This determines whether the package should be fully loaded for gathering.*/
+	EPackageLocCacheState CalculatePackageLocCacheState(const FPackageFileSummary& PackageFileSummary, const FName PackageName, bool bIsExternalActorPackage) const;
 	/** Struct containing the data needed by a pending package that we will gather text from */
 	struct FPackagePendingGather
 	{
@@ -65,6 +98,9 @@ private:
 		TArray<FGatherableTextData> GatherableTextDataArray;
 	};
 
+	/** Adds a package to PackagesPendingGather and returns a pointer to the appended package.*/
+	FPackagePendingGather* AppendPackagePendingGather(const FName PackageNameToGather);
+
 	static const FString UsageText;
 
 	TArray<FString> ModulesToPreload;
@@ -85,11 +121,12 @@ private:
 
 	/** Array of objects that should be kept alive during the next call to CollectGarbage (used by PurgeGarbage and AddReferencedObjects) */
 	TSet<TObjectPtr<UObject>> ObjectsToKeepAlive;
-
+	/** Path to the directory where output reports etc will be saved.*/
+	FString DestinationPath;
 	bool bSkipGatherCache;
 	bool bReportStaleGatherCache;
 	bool bFixStaleGatherCache;
 	bool bFixMissingGatherCache;
-	bool ShouldGatherFromEditorOnlyData;
-	bool ShouldExcludeDerivedClasses;
+	bool bShouldGatherFromEditorOnlyData;
+	bool bShouldExcludeDerivedClasses;
 };
diff --git a/Engine/Source/Editor/UnrealEd/Private/Commandlets/GatherTextFromAssetsCommandlet.cpp b/Engine/Source/Editor/UnrealEd/Private/Commandlets/GatherTextFromAssetsCommandlet.cpp
--- a/Engine/Source/Editor/UnrealEd/Private/Commandlets/GatherTextFromAssetsCommandlet.cpp
+++ b/Engine/Source/Editor/UnrealEd/Private/Commandlets/GatherTextFromAssetsCommandlet.cpp
@@ -27,7 +27,6 @@
 #include "WorldPartition/WorldPartition.h"
 #include "WorldPartition/WorldPartitionHelpers.h"
 #include "WorldPartition/WorldPartitionActorDescUtils.h"
-#include "WorldPartition/WorldPartitionActorDescInstance.h"
 #include "Engine/Level.h"
 #include "Engine/World.h"
 #include "Editor.h"
@@ -40,6 +39,7 @@
 #include "Templates/UniquePtr.h"
 #include "CollectionManagerModule.h"
 #include "ICollectionManager.h"
+#include "WorldPartition/WorldPartitionActorDescInstance.h"
 
 DEFINE_LOG_CATEGORY_STATIC(LogGatherTextFromAssetsCommandlet, Log, All);
 
@@ -224,6 +224,32 @@ private:
 	int32 UncachedAssetBreakdown[(int32)UGatherTextFromAssetsCommandlet::EPackageLocCacheState::Cached];
 };
 
+namespace UE::Private::GatherTextFromAssetsCommandlet
+{
+	static FAssetGatherCacheMetrics AssetGatherCacheMetrics;
+	/**
+	 * Commandlets don't tick, but loading assets can queue async building work to various systems.
+	 * We tick these systems periodically during a gather to prevent us from running out of memory due to the queued pending tasks.
+	 * Refer to the cooker to determine if this function needs to be expanded to cover more systems.
+	 */
+	void TickBackgroundTasks()
+	{
+		if (GShaderCompilingManager)
+		{
+			GShaderCompilingManager->ProcessAsyncResults(true, false);
+		}
+		if (GDistanceFieldAsyncQueue)
+		{
+			GDistanceFieldAsyncQueue->ProcessAsyncTasks();
+		}
+		if (GCardRepresentationAsyncQueue)
+		{
+			GCardRepresentationAsyncQueue->ProcessAsyncTasks();
+		}
+	}
+
+}
+
 #define LOC_DEFINE_REGION
 
 //////////////////////////////////////////////////////////////////////////
@@ -243,8 +269,8 @@ UGatherTextFromAssetsCommandlet::UGatherTextFromAssetsCommandlet(const FObjectIn
 	, MinFreeMemoryBytes(0)
 	, MaxUsedMemoryBytes(0)
 	, bSkipGatherCache(false)
-	, ShouldGatherFromEditorOnlyData(false)
-	, ShouldExcludeDerivedClasses(false)
+	, bShouldGatherFromEditorOnlyData(false)
+	, bShouldExcludeDerivedClasses(false)
 {
 }
 
@@ -254,7 +280,7 @@ void UGatherTextFromAssetsCommandlet::ProcessGatherableTextDataArray(const TArra
 	{
 		for (const FTextSourceSiteContext& TextSourceSiteContext : GatherableTextData.SourceSiteContexts)
 		{
-			if (!TextSourceSiteContext.IsEditorOnly || ShouldGatherFromEditorOnlyData)
+			if (!TextSourceSiteContext.IsEditorOnly || bShouldGatherFromEditorOnlyData)
 			{
 				if (TextSourceSiteContext.KeyName.IsEmpty())
 				{
@@ -477,264 +503,151 @@ bool UGatherTextFromAssetsCommandlet::ShouldRunInPreview(const TArray<FString>&
 	// If the param is not specified, it is assumed that both source and assets are to be gathered 
 	return !GatherType || *GatherType == TEXT("Asset") || *GatherType == TEXT("All");
 }
-
-int32 UGatherTextFromAssetsCommandlet::Main(const FString& Params)
+/**
+ * Builds the first pass filter which currently consists of the collection filter and the optional derived class filter.
+ * See BuildCollectionFilter and BuildExcludeDerivedClassesFilter
+ */
+bool UGatherTextFromAssetsCommandlet::BuildFirstPassFilter(FARFilter& InOutFilter) const
 {
-	// Parse command line.
-	TArray<FString> Tokens;
-	TArray<FString> Switches;
-	TMap<FString, FString> ParamVals;
-	UCommandlet::ParseCommandLine(*Params, Tokens, Switches, ParamVals);
-
-	FString GatherTextConfigPath;
-	FString SectionName;
-	if (!GetConfigurationScript(ParamVals, GatherTextConfigPath, SectionName))
-	{
-		return -1;
-	}
-
-	if (!ConfigureFromScript(GatherTextConfigPath, SectionName))
-	{
-		return -1;
-	}
-
-	FGatherTextDelegates::GetAdditionalGatherPaths.Broadcast(GatherManifestHelper->GetTargetName(), IncludePathFilters, ExcludePathFilters);
-
-	// Get destination path
-	FString DestinationPath;
-	if (!GetPathFromConfig(*SectionName, TEXT("DestinationPath"), DestinationPath, GatherTextConfigPath))
-	{
-		UE_LOG(LogGatherTextFromAssetsCommandlet, Error, TEXT("No destination path specified."));
-		return -1;
-	}
-
-	// Add any manifest dependencies if they were provided
+	// Filter object paths to only those in any of the specified collections.
+	if (!BuildCollectionFilter(InOutFilter))
 	{
-		bool HasFailedToAddManifestDependency = false;
-		for (const FString& ManifestDependency : ManifestDependenciesList)
-		{
-			FText OutError;
-			if (!GatherManifestHelper->AddDependency(ManifestDependency, &OutError))
-			{
-				UE_LOG(LogGatherTextFromAssetsCommandlet, Error, TEXT("The GatherTextFromAssets commandlet couldn't load the specified manifest dependency: '%s'. %s"), *ManifestDependency, *OutError.ToString());
-				HasFailedToAddManifestDependency = true;
-			}
-		}
-		if (HasFailedToAddManifestDependency)
-		{
-			return -1;
-		}
+		return false;
 	}
 
-	// Preload necessary modules.
+	// Filter out any objects of the specified classes and their children at this point.
+	if (bShouldExcludeDerivedClasses)
 	{
-		bool HasFailedToPreloadAnyModules = false;
-		for (const FString& ModuleName : ModulesToPreload)
-		{
-			EModuleLoadResult ModuleLoadResult;
-			FModuleManager::Get().LoadModuleWithFailureReason(*ModuleName, ModuleLoadResult);
-
-			if (ModuleLoadResult != EModuleLoadResult::Success)
-			{
-				UE_LOG(LogGatherTextFromAssetsCommandlet, Warning, TEXT("Failed to preload dependent module %s. Please check if the modules have been renamed or moved to another folder."), *ModuleName);
-				HasFailedToPreloadAnyModules = true;
-				continue;
-			}
-		}
-
-		if (HasFailedToPreloadAnyModules)
+		if (!BuildExcludeDerivedClassesFilter(InOutFilter))
 		{
-			return -1;
+			return false;
 		}
 	}
+	return true;
+}
 
-	// If the editor has loaded a persistent world then create an empty world prior to starting the asset gather
-	// This avoids any issues when loading and initializing worlds during the gather, as WP needs to re-initialize the world
-	if (GEditor)
+/** Builds a filter based on the specified collections to be used for gathering.*/
+bool UGatherTextFromAssetsCommandlet::BuildCollectionFilter(FARFilter& InOutFilter) const
+{
+	bool bHasFailedToGetACollection = false;
+	FCollectionManagerModule& CollectionManagerModule = FCollectionManagerModule::GetModule();
+	ICollectionManager& CollectionManager = CollectionManagerModule.Get();
+	for (const FString& CollectionName : CollectionFilters)
 	{
-		if (UWorld* EditorWorld = GEditor->GetEditorWorldContext().World();
-			EditorWorld && !FPackageName::IsTempPackage(FNameBuilder(EditorWorld->GetPackage()->GetFName()).ToView()))
+		if (!CollectionManager.GetObjectsInCollection(FName(*CollectionName), ECollectionShareType::CST_All, InOutFilter.SoftObjectPaths, ECollectionRecursionFlags::SelfAndChildren))
 		{
-			GEditor->CreateNewMapForEditing(/*bPromptForSave*/false);
+			UE_LOG(LogGatherTextFromAssetsCommandlet, Error, TEXT("Failed get objects in specified collection: %s"), *CollectionName);
+			bHasFailedToGetACollection = true;
 		}
 	}
 
-	UE_LOG(LogGatherTextFromAssetsCommandlet, Display, TEXT("Discovering assets to gather..."));
-	const double DiscoveringAssetsStartTime = FPlatformTime::Seconds();
-	const double SearchAssetRegistryStartTime = FPlatformTime::Seconds();
-	FAssetRegistryModule& AssetRegistryModule = FModuleManager::LoadModuleChecked<FAssetRegistryModule>(TEXT("AssetRegistry"));
-	IAssetRegistry& AssetRegistry = AssetRegistryModule.Get();
-	AssetRegistry.SearchAllAssets(true);
-	UE_LOG(LogGatherTextFromAssetsCommandlet, Display, TEXT("Searching all assets took %.2f seconds."), FPlatformTime::Seconds() - SearchAssetRegistryStartTime)			;
+	return !bHasFailedToGetACollection;
+}
 
-	TArray<FAssetData> AssetDataArray;
-	
+/** Builds a filter to remove classes and derived classes of ExactClassFilter.*/
+bool UGatherTextFromAssetsCommandlet::BuildExcludeDerivedClassesFilter(FARFilter& InOutFilter) const
+{
+	InOutFilter.bRecursiveClasses = true;
+	InOutFilter.ClassPaths.Add(UObject::StaticClass()->GetClassPathName());
+	for (const FString& ExcludeClassName : ExcludeClassNames)
 	{
-		FARFilter FirstPassFilter;
-
-		// Filter object paths to only those in any of the specified collections.
-		const double LoadCollectionObjectsStartTime = FPlatformTime::Seconds();
-		{
-			bool HasFailedToGetACollection = false;
-			FCollectionManagerModule& CollectionManagerModule = FCollectionManagerModule::GetModule();
-			ICollectionManager& CollectionManager = CollectionManagerModule.Get();
-			for (const FString& CollectionName : CollectionFilters)
-			{
-				if (!CollectionManager.GetObjectsInCollection(FName(*CollectionName), ECollectionShareType::CST_All, FirstPassFilter.SoftObjectPaths, ECollectionRecursionFlags::SelfAndChildren))
-				{
-					UE_LOG(LogGatherTextFromAssetsCommandlet, Error, TEXT("Failed get objects in specified collection: %s"), *CollectionName);
-					HasFailedToGetACollection = true;
-				}
-			}
-			if (HasFailedToGetACollection)
-			{
-				return -1;
-			}
-		}
-		UE_LOG(LogGatherTextFromAssetsCommandlet, Display, TEXT("Loading all collection objects took %.2f seconds."), FPlatformTime::Seconds() - LoadCollectionObjectsStartTime);
-
-		const double FilterExcludedClassesStartTime = FPlatformTime::Seconds();
-		// Filter out any objects of the specified classes and their children at this point.
-		if (ShouldExcludeDerivedClasses)
-		{
-			FirstPassFilter.bRecursiveClasses = true;
-			FirstPassFilter.ClassPaths.Add(UObject::StaticClass()->GetClassPathName());
-			for (const FString& ExcludeClassName : ExcludeClassNames)
-			{
-				FTopLevelAssetPath ExcludedClassPathName = UClass::TryConvertShortTypeNameToPathName<UClass>(ExcludeClassName, ELogVerbosity::Warning, TEXT("GatherTextFromAssetsCommandlet"));
-				if (!ExcludedClassPathName.IsNull())
-				{
-					// Note: Can't necessarily validate these class names here, as the class may be a generated blueprint class that hasn't been loaded yet.
-					FirstPassFilter.RecursiveClassPathsExclusionSet.Add(FTopLevelAssetPath(ExcludeClassName));
-				}
-				else
-				{
-					UE_CLOG(!ExcludeClassName.IsEmpty(), LogGatherTextFromAssetsCommandlet, Error, TEXT("Unable to convert short class name \"%s\" to path name. Please use path names fo ExcludeClassNames"), *ExcludeClassName);
-				}
-			}
-		}
-		UE_LOG(LogGatherTextFromAssetsCommandlet, Display, TEXT("Filtering excluded classes took %.2f seconds."), FPlatformTime::Seconds() - FilterExcludedClassesStartTime);
-
-		// Apply filter if valid to do so, get all assets otherwise.
-		if (FirstPassFilter.IsEmpty())
+		FTopLevelAssetPath ExcludedClassPathName = UClass::TryConvertShortTypeNameToPathName<UClass>(ExcludeClassName, ELogVerbosity::Warning, TEXT("GatherTextFromAssetsCommandlet"));
+		if (!ExcludedClassPathName.IsNull())
 		{
-			// @TODOLocalization: Logging that the first path filter is empty resulting in all assets being gathered can confuse users who generally rely on the second pass.
-			// Figure out a good way to still convey the information in a log or clog.
-			const double GetAllAssetsStartTime = FPlatformTime::Seconds();
-			AssetRegistry.GetAllAssets(AssetDataArray);
-			UE_LOG(LogGatherTextFromAssetsCommandlet, Display, TEXT("Loading all assets from asset registry took %.2f seconds."), FPlatformTime::Seconds() - GetAllAssetsStartTime);
+			// Note: Can't necessarily validate these class names here, as the class may be a generated blueprint class that hasn't been loaded yet.
+			InOutFilter.RecursiveClassPathsExclusionSet.Add(FTopLevelAssetPath(ExcludeClassName));
 		}
 		else
 		{
-			const double GetAllAssetsWithFirstPassFilterStartTime = FPlatformTime::Seconds();
-			AssetRegistry.GetAssets(FirstPassFilter, AssetDataArray);
-			UE_LOG(LogGatherTextFromAssetsCommandlet, Display, TEXT("Getting all assets with first pass filter from asset registry took %.2f seconds."), FPlatformTime::Seconds() - GetAllAssetsWithFirstPassFilterStartTime);
+			UE_CLOG(!ExcludeClassName.IsEmpty(), LogGatherTextFromAssetsCommandlet, Error, TEXT("Unable to convert short class name \"%s\" to path name. Please use path names fo ExcludeClassNames"), *ExcludeClassName);
 		}
 	}
+	
+	return true;
+}
 
-	if (!ShouldExcludeDerivedClasses)
+/** Builds a filter to exclude exactly the specified classes. This will retrieve the exact assets from the asset registry to exclude. */
+bool UGatherTextFromAssetsCommandlet::BuildExcludeExactClassesFilter(FARFilter& InOutFilter) const
+{
+	InOutFilter.bRecursiveClasses = false;
+	for (const FString& ExcludeClassName : ExcludeClassNames)
 	{
-		const double ExcludeDerivedClassesStartTime = FPlatformTime::Seconds();
-		// Filter out any objects of the specified classes.
-		FARFilter ExcludeExactClassesFilter;
-		ExcludeExactClassesFilter.bRecursiveClasses = false;
-		for (const FString& ExcludeClassName : ExcludeClassNames)
+		FTopLevelAssetPath ExcludedClassPathName = UClass::TryConvertShortTypeNameToPathName<UClass>(ExcludeClassName, ELogVerbosity::Warning, TEXT("GatherTextFromAssetsCommandlet"));
+		if (!ExcludedClassPathName.IsNull())
 		{
-			FTopLevelAssetPath ExcludedClassPathName = UClass::TryConvertShortTypeNameToPathName<UClass>(ExcludeClassName, ELogVerbosity::Warning, TEXT("GatherTextFromAssetsCommandlet"));
-			if (!ExcludedClassPathName.IsNull())
-			{
-				// Note: Can't necessarily validate these class names here, as the class may be a generated blueprint class that hasn't been loaded yet.
-				ExcludeExactClassesFilter.ClassPaths.Add(FTopLevelAssetPath(ExcludeClassName));
-			}
-			else
-			{
-				UE_CLOG(!ExcludeClassName.IsEmpty(), LogGatherTextFromAssetsCommandlet, Error, TEXT("Unable to convert short class name \"%s\" to path name. Please use path names fo ExcludeClassNames"), *ExcludeClassName);
-			}
-
+			// Note: Can't necessarily validate these class names here, as the class may be a generated blueprint class that hasn't been loaded yet.
+			InOutFilter.ClassPaths.Add(FTopLevelAssetPath(ExcludeClassName));
 		}
-
-		// Reapply filter over the current set of assets.
-		if (!ExcludeExactClassesFilter.IsEmpty())
+		else
 		{
-			// NOTE: The filter applied is actually the inverse, due to API limitations, so the resultant set must be removed from the current set.
-			TArray<FAssetData> AssetsToExclude = AssetDataArray;
-			AssetRegistry.RunAssetsThroughFilter(AssetsToExclude, ExcludeExactClassesFilter);
-			AssetDataArray.RemoveAll([&](const FAssetData& AssetData)
-			{
-				return AssetsToExclude.Contains(AssetData);
-			});
-			UE_LOG(LogGatherTextFromAssetsCommandlet, Display, TEXT("Filtering otu derived classes took %.2f seconds."), FPlatformTime::Seconds() - ExcludeDerivedClassesStartTime);
+			UE_CLOG(!ExcludeClassName.IsEmpty(), LogGatherTextFromAssetsCommandlet, Error, TEXT("Unable to convert short class name \"%s\" to path name. Please use path names fo ExcludeClassNames"), *ExcludeClassName);
 		}
 	}
+	return true;
+}
 
-	// Note: AssetDataArray now contains all assets in the specified collections that are not instances of the specified excluded classes.
-
+/** Filters out assets that fail the IncludePath and ExcludePath wildcard filters. */
+void UGatherTextFromAssetsCommandlet::FilterAssetsBasedOnIncludeExcludePaths(TArray<FAssetData>& InOutAssetDataArray) const
+{
 	const FFuzzyPathMatcher FuzzyPathMatcher = FFuzzyPathMatcher(IncludePathFilters, ExcludePathFilters);
 	const double FilteringAssetsByIncludeExcludePathsStartTime = FPlatformTime::Seconds();
-	AssetDataArray.RemoveAll([&](const FAssetData& PartiallyFilteredAssetData) -> bool
-	{
-		if (PartiallyFilteredAssetData.IsRedirector())
+	InOutAssetDataArray.RemoveAll([&](const FAssetData& PartiallyFilteredAssetData) -> bool
 		{
-			// Redirectors never have localization
-			return true;
-		}
+			if (PartiallyFilteredAssetData.IsRedirector())
+			{
+				// Redirectors never have localization
+				return true;
+			}
 
-		FString PackageFilePathWithoutExtension;
-		if (!FPackageName::TryConvertLongPackageNameToFilename(PartiallyFilteredAssetData.PackageName.ToString(), PackageFilePathWithoutExtension))
-		{
-			// This means the asset data is for content that isn't mounted - this can happen when using a cooked asset registry
-			return true;
-		}
+			FString PackageFilePathWithoutExtension;
+			if (!FPackageName::TryConvertLongPackageNameToFilename(PartiallyFilteredAssetData.PackageName.ToString(), PackageFilePathWithoutExtension))
+			{
+				// This means the asset data is for content that isn't mounted - this can happen when using a cooked asset registry
+				return true;
+			}
 
-		FString PackageFilePathWithExtension;
-		if (!FPackageName::FindPackageFileWithoutExtension(PackageFilePathWithoutExtension, PackageFilePathWithExtension))
-		{
-			// This means the package file doesn't exist on disk, which means we cannot gather it
-			return true;
-		}
+			FString PackageFilePathWithExtension;
+			if (!FPackageName::FindPackageFileWithoutExtension(PackageFilePathWithoutExtension, PackageFilePathWithExtension))
+			{
+				// This means the package file doesn't exist on disk, which means we cannot gather it
+				return true;
+			}
 
-		PackageFilePathWithExtension = FPaths::ConvertRelativePathToFull(PackageFilePathWithExtension);
-		const FString PackageFileName = FPaths::GetCleanFilename(PackageFilePathWithExtension);
+			PackageFilePathWithExtension = FPaths::ConvertRelativePathToFull(PackageFilePathWithExtension);
+			const FString PackageFileName = FPaths::GetCleanFilename(PackageFilePathWithExtension);
 
-		// Filter out assets whose package file names DO NOT match any of the package file name filters.
-		{
-			bool HasPassedAnyFileNameFilter = false;
-			for (const FString& PackageFileNameFilter : PackageFileNameFilters)
+			// Filter out assets whose package file names DO NOT match any of the package file name filters.
 			{
-				if (PackageFileName.MatchesWildcard(PackageFileNameFilter))
+				bool HasPassedAnyFileNameFilter = false;
+				for (const FString& PackageFileNameFilter : PackageFileNameFilters)
+				{
+					if (PackageFileName.MatchesWildcard(PackageFileNameFilter))
+					{
+						HasPassedAnyFileNameFilter = true;
+						break;
+					}
+				}
+				if (!HasPassedAnyFileNameFilter)
 				{
-					HasPassedAnyFileNameFilter = true;
-					break;
+					return true;
 				}
 			}
-			if (!HasPassedAnyFileNameFilter)
+
+			// Filter out assets whose package file paths do not pass the "fuzzy path" filters.
+			if (FuzzyPathMatcher.TestPath(PackageFilePathWithExtension) != FFuzzyPathMatcher::EPathMatch::Included)
 			{
 				return true;
 			}
-		}
 
-		// Filter out assets whose package file paths do not pass the "fuzzy path" filters.
-		if (FuzzyPathMatcher.TestPath(PackageFilePathWithExtension) != FFuzzyPathMatcher::EPathMatch::Included)
-		{
-			return true;
-		}
-
-		return false;
-	});
+			return false;
+		});
 	UE_LOG(LogGatherTextFromAssetsCommandlet, Display, TEXT("Filtering assets by include exclude paths took %.2f seconds."), FPlatformTime::Seconds() - FilteringAssetsByIncludeExcludePathsStartTime);
+}
 
-	if (AssetDataArray.Num() == 0)
-	{
-		UE_LOG(LogGatherTextFromAssetsCommandlet, Display, TEXT("No assets matched the specified criteria."));
-		return 0;
-	}
-
-	// Discover the external actors for any worlds that are pending gather
-	{
-		UE_LOG(LogGatherTextFromAssetsCommandlet, Display, TEXT("Discovering external actors to gather..."));
-		const double DiscoveringExternalActorsStartTime = FPlatformTime::Seconds();
-		TArray<FName> ExternalActorsSearchPaths;
-		AssetDataArray.RemoveAll([&ExternalActorsSearchPaths](const FAssetData& AssetData)
+/** Remove any external actors that currently exist in InOutAssetDataArray. OutExternalActorsSearchPaths is populated for the correct paths to search for external actors.*/
+void UGatherTextFromAssetsCommandlet::RemoveExistingExternalActors(TArray<FAssetData>& InOutAssetDataArray, TArray<FName>& OutExternalActorsSearchPaths) const
+{
+	InOutAssetDataArray.RemoveAll([&OutExternalActorsSearchPaths](const FAssetData& AssetData)
 		{
 			const FNameBuilder PackageNameStr(AssetData.PackageName);
 
@@ -743,219 +656,319 @@ int32 UGatherTextFromAssetsCommandlet::Main(const FString& Params)
 				if (ULevel::GetIsLevelPartitionedFromAsset(AssetData))
 				{
 					FString ExternalActorsPathForWorld = ULevel::GetExternalActorsPath(*PackageNameStr);
-					ExternalActorsSearchPaths.Add(*ExternalActorsPathForWorld);
+					// External actors will only be processed if/when their world is gathered.
+					// Adding this as a search path will ensure that all external actors within the world will be added for gather later.
+					OutExternalActorsSearchPaths.Add(*ExternalActorsPathForWorld);
 				}
 			}
 			else if (PackageNameStr.ToView().Contains(FPackagePath::GetExternalActorsFolderName()))
 			{
 				// Remove any external actors that are already in the list, as they will be re-added providing their owner world passed the gather criteria
+				// It is possible for an external actor to be directly specified for gather in the configs but have their world not pass the gather criteria.
 				return true;
 			}
 
 			return false;
 		});
-		
-		if (ExternalActorsSearchPaths.Num() > 0)
-		{
-			AssetRegistry.GetAssetsByPaths(ExternalActorsSearchPaths, AssetDataArray, /*bRecursive*/true);
-		}
-		UE_LOG(LogGatherTextFromAssetsCommandlet, Display, TEXT("Discovering external actors took %.2f seconds."), FPlatformTime::Seconds() - DiscoveringExternalActorsStartTime);
+}
+
+/** Appends any external actors that also need to be gathered to the InOutAssetDataArray. */
+void UGatherTextFromAssetsCommandlet::DiscoverExternalActors(TArray<FAssetData>& InOutAssetDataArray) const
+{
+	UE_LOG(LogGatherTextFromAssetsCommandlet, Display, TEXT("Discovering external actors to gather..."));
+	const double DiscoveringExternalActorsStartTime = FPlatformTime::Seconds();
+	TArray<FName> ExternalActorsSearchPaths;
+	RemoveExistingExternalActors(InOutAssetDataArray, ExternalActorsSearchPaths);
+
+	if (ExternalActorsSearchPaths.Num() > 0)
+	{
+		IAssetRegistry::GetChecked().GetAssetsByPaths(ExternalActorsSearchPaths, InOutAssetDataArray, /*bRecursive*/true);
 	}
+	UE_LOG(LogGatherTextFromAssetsCommandlet, Display, TEXT("Discovering external actors took %.2f seconds."), FPlatformTime::Seconds() - DiscoveringExternalActorsStartTime);
+}
 
-	auto AppendPackagePendingGather = [this](const FName PackageNameToGather) -> FPackagePendingGather*
+/** Applies the passed in filter to the asset registry. If the filter is empty, the entire asset registry will be returned in InOutAssetDataArray. Else assets that pass the filter will be in InOutAssetDataArray.*/
+void UGatherTextFromAssetsCommandlet::ApplyFirstPassFilter(const FARFilter& InFilter, TArray<FAssetData>& InOutAssetDataArray) const
+{
+	// Apply filter if valid to do so, get all assets otherwise.
+	if (InFilter.IsEmpty())
 	{
-		FString PackageFilename;
-		if (!FPackageName::FindPackageFileWithoutExtension(FPackageName::LongPackageNameToFilename(PackageNameToGather.ToString()), PackageFilename))
-		{
-			return nullptr;
-		}
-		PackageFilename = FPaths::ConvertRelativePathToFull(PackageFilename);
+		// @TODOLocalization: Logging that the first path filter is empty resulting in all assets being gathered can confuse users who generally rely on the second pass.
+		// Figure out a good way to still convey the information in a log or clog.
+		const double GetAllAssetsStartTime = FPlatformTime::Seconds();
+		IAssetRegistry::GetChecked().GetAllAssets(InOutAssetDataArray);
+		UE_LOG(LogGatherTextFromAssetsCommandlet, Display, TEXT("Loading all assets from asset registry took %.2f seconds."), FPlatformTime::Seconds() - GetAllAssetsStartTime);
+	}
+	else
+	{
+		const double GetAllAssetsWithFirstPassFilterStartTime = FPlatformTime::Seconds();
+		IAssetRegistry::GetChecked().GetAssets(InFilter, InOutAssetDataArray);
+		UE_LOG(LogGatherTextFromAssetsCommandlet, Display, TEXT("Getting all assets with first pass filter from asset registry took %.2f seconds."), FPlatformTime::Seconds() - GetAllAssetsWithFirstPassFilterStartTime);
+	}
+}
 
-		FPackagePendingGather& PackagePendingGather = PackagesPendingGather.AddDefaulted_GetRef();
-		PackagePendingGather.PackageName = PackageNameToGather;
-		PackagePendingGather.PackageFilename = MoveTemp(PackageFilename);
-		PackagePendingGather.PackageLocCacheState = EPackageLocCacheState::Cached;
-		return &PackagePendingGather;
-	};
+/**
+ * Builds and applies the first pass filter to the entire AssetRegistry. OutAssetDataArray will hold all the assets that pass the first pass filter.
+ * See BuildFirstPassFilter and ApplyFirstPassFilter
+ */
+bool UGatherTextFromAssetsCommandlet::PerformFirstPassFilter(TArray<FAssetData>& OutAssetDataArray) const
+{
+	FARFilter FirstPassFilter;
 
-	// Collect the basic information about the packages that we're going to gather from
+	if (!BuildFirstPassFilter(FirstPassFilter))
 	{
-		// Collapse the assets down to a set of packages
-		TSet<FName> PackageNamesToGather;
-		PackageNamesToGather.Reserve(AssetDataArray.Num());
-		for (const FAssetData& AssetData : AssetDataArray)
-		{
-			PackageNamesToGather.Add(AssetData.PackageName);
-		}
-		AssetDataArray.Empty();
+		return false;
+	}
+	ApplyFirstPassFilter(FirstPassFilter, OutAssetDataArray);
+	return true;
+}
 
-		// Build the basic information for the packages to gather (dependencies are filled in later once we've processed cached packages)
-		PackagesPendingGather.Reserve(PackageNamesToGather.Num());
-		for (const FName& PackageNameToGather : PackageNamesToGather)
+void UGatherTextFromAssetsCommandlet::ApplyExcludeExactClassesFilter(const FARFilter& InFilter, TArray<FAssetData>& InOutAssetDataArray) const
+{
+	// NOTE: The filter applied is actually the inverse, due to API limitations, so the resultant set must be removed from the current set.
+	TArray<FAssetData> AssetsToExclude = InOutAssetDataArray;
+	IAssetRegistry::GetChecked().RunAssetsThroughFilter(AssetsToExclude, InFilter);
+	InOutAssetDataArray.RemoveAll([&](const FAssetData& AssetData)
 		{
-			AppendPackagePendingGather(PackageNameToGather);
-		}
+			return AssetsToExclude.Contains(AssetData);
+		});
+}
+
+/**
+ * Filters out assets from the exact specified classes . Assets that pass the filter will be in InOutAssetDataArray.
+ * See BuildExactClassesFilter and ApplyExactClassesFilter
+ */
+bool UGatherTextFromAssetsCommandlet::PerformExcludeExactClassesFilter(TArray<FAssetData>& InOutAssetDataArray) const
+{
+	const double ExcludeDerivedClassesStartTime = FPlatformTime::Seconds();
+	// Filter out any objects of the specified classes.
+	FARFilter ExcludeExactClassesFilter;
+	if (!BuildExcludeExactClassesFilter(ExcludeExactClassesFilter))
+	{
+		return false;
 	}
-	UE_LOG(LogGatherTextFromAssetsCommandlet, Display, TEXT("Discovering assests to gather took %.2f seconds."), FPlatformTime::Seconds() - DiscoveringAssetsStartTime);
 
-	FAssetGatherCacheMetrics AssetGatherCacheMetrics;
-	TMap<FString, FName> AssignedPackageLocalizationIds;
-	TMap<FName, TSet<FGuid>> StaleExternalActors;
+	// Reapply filter over the current set of assets.
+	if (!ExcludeExactClassesFilter.IsEmpty())
+	{
+		ApplyExcludeExactClassesFilter(ExcludeExactClassesFilter, InOutAssetDataArray);
+		UE_LOG(LogGatherTextFromAssetsCommandlet, Display, TEXT("Filtering otu derived classes took %.2f seconds."), FPlatformTime::Seconds() - ExcludeDerivedClassesStartTime);
+	}
+	return true;
+}
 
-	int32 NumPackagesProcessed = 0;
-	int32 PackageCount = PackagesPendingGather.Num();
+bool UGatherTextFromAssetsCommandlet::ParseCommandLineHelper(const FString& InCommandLine)
+{
+	TArray<FString> Tokens;
+	TArray<FString> Switches;
+	TMap<FString, FString> ParamVals;
+	UCommandlet::ParseCommandLine(*InCommandLine, Tokens, Switches, ParamVals);
 
-	// Process all packages that do not need to be loaded. Remove processed packages from the list.
-	UE_LOG(LogGatherTextFromAssetsCommandlet, Display, TEXT("Processing assets to gather..."));
-	PackagesPendingGather.RemoveAll([&](FPackagePendingGather& PackagePendingGather) -> bool
+	FString GatherTextConfigPath;
+	FString SectionName;
+	if (!GetConfigurationScript(ParamVals, GatherTextConfigPath, SectionName))
 	{
-		const FNameBuilder PackageNameStr(PackagePendingGather.PackageName);
-
-		const int32 CurrentPackageNum = ++NumPackagesProcessed;
-		const float PercentageComplete = (((float)CurrentPackageNum / (float)PackageCount) * 100.0f);
+		return false;
+	}
 
-		TUniquePtr<FArchive> FileReader(IFileManager::Get().CreateFileReader(*PackagePendingGather.PackageFilename));
-		if (!FileReader)
-		{
-			return false;
-		}
+	if (!ConfigureFromScript(GatherTextConfigPath, SectionName))
+	{
+		return false;
+	}
 
-		const bool bIsExternalActorPackage = PackageNameStr.ToView().Contains(FPackagePath::GetExternalActorsFolderName());
+	FGatherTextDelegates::GetAdditionalGatherPaths.Broadcast(GatherManifestHelper->GetTargetName(), IncludePathFilters, ExcludePathFilters);
 
-		// Read package file summary from the file.
-		FPackageFileSummary PackageFileSummary;
-		*FileReader << PackageFileSummary;
+	// Get destination path
+	if (!GetPathFromConfig(*SectionName, TEXT("DestinationPath"), DestinationPath, GatherTextConfigPath))
+	{
+		UE_LOG(LogGatherTextFromAssetsCommandlet, Error, TEXT("No destination path specified."));
+		return false;
+	}
 
-		// Track the package localization ID of this package (if known) and detect duplicates
-		if (!PackageFileSummary.LocalizationId.IsEmpty())
+	// Add any manifest dependencies if they were provided
+	{
+		bool bHasFailedToAddManifestDependency = false;
+		for (const FString& ManifestDependency : ManifestDependenciesList)
 		{
-			if (const FName* ExistingLongPackageName = AssignedPackageLocalizationIds.Find(PackageFileSummary.LocalizationId))
-			{
-				UE_LOG(LogGatherTextFromAssetsCommandlet, Warning, TEXT("Package '%s' and '%s' have the same localization ID (%s). Please reset one of these (Asset Localization -> Reset Localization ID) to avoid conflicts."), *PackageNameStr, *ExistingLongPackageName->ToString(), *PackageFileSummary.LocalizationId);
-			}
-			else
+			FText OutError;
+			if (!GatherManifestHelper->AddDependency(ManifestDependency, &OutError))
 			{
-				AssignedPackageLocalizationIds.Add(PackageFileSummary.LocalizationId, PackagePendingGather.PackageName);
+				UE_LOG(LogGatherTextFromAssetsCommandlet, Error, TEXT("The GatherTextFromAssets commandlet couldn't load the specified manifest dependency: '%s'. %s"), *ManifestDependency, *OutError.ToString());
+				bHasFailedToAddManifestDependency = true;
 			}
 		}
+		if (bHasFailedToAddManifestDependency)
+		{
+			return false;
+		}
+	}
 
-		PackagePendingGather.PackageLocCacheState = EPackageLocCacheState::Cached;
-
-		// Have we been asked to skip the cache of text that exists in the header of newer packages?
-		if (bSkipGatherCache && PackageFileSummary.GetFileVersionUE() >= VER_UE4_SERIALIZE_TEXT_IN_PACKAGES)
+	// Preload necessary modules.
+	{
+		bool bHasFailedToPreloadAnyModules = false;
+		for (const FString& ModuleName : ModulesToPreload)
 		{
-			// Fallback on the old package flag check.
-			if (PackageFileSummary.GetPackageFlags() & PKG_RequiresLocalizationGather)
+			EModuleLoadResult ModuleLoadResult;
+			FModuleManager::Get().LoadModuleWithFailureReason(*ModuleName, ModuleLoadResult);
+
+			if (ModuleLoadResult != EModuleLoadResult::Success)
 			{
-				PackagePendingGather.PackageLocCacheState = EPackageLocCacheState::Uncached_NoCache;
+				UE_LOG(LogGatherTextFromAssetsCommandlet, Warning, TEXT("Failed to preload dependent module %s. Please check if the modules have been renamed or moved to another folder."), *ModuleName);
+				bHasFailedToPreloadAnyModules = true;
+				continue;
 			}
 		}
 
-		const FCustomVersion* const EditorVersion = PackageFileSummary.GetCustomVersionContainer().GetVersion(FEditorObjectVersion::GUID);
-		const FCustomVersion* const FNMainVersion = PackageFileSummary.GetCustomVersionContainer().GetVersion(FFortniteMainBranchObjectVersion::GUID);
-
-		// Packages not resaved since localization gathering flagging was added to packages must be loaded.
-		if (PackageFileSummary.GetFileVersionUE() < VER_UE4_PACKAGE_REQUIRES_LOCALIZATION_GATHER_FLAGGING)
+		if (bHasFailedToPreloadAnyModules)
 		{
-			PackagePendingGather.PackageLocCacheState = EPackageLocCacheState::Uncached_TooOld;
+			return false;
 		}
-		// Package not resaved since gatherable text data was added to package headers must be loaded, since their package header won't contain pregathered text data.
-		else if (PackageFileSummary.GetFileVersionUE() < VER_UE4_SERIALIZE_TEXT_IN_PACKAGES || (!EditorVersion || EditorVersion->Version < FEditorObjectVersion::GatheredTextEditorOnlyPackageLocId))
+	}
+	return true;
+}
+
+UGatherTextFromAssetsCommandlet::FPackagePendingGather* UGatherTextFromAssetsCommandlet::AppendPackagePendingGather(const FName PackageNameToGather)
+{
+	FString PackageFilename;
+	if (!FPackageName::FindPackageFileWithoutExtension(FPackageName::LongPackageNameToFilename(PackageNameToGather.ToString()), PackageFilename))
+	{
+		return nullptr;
+	}
+	PackageFilename = FPaths::ConvertRelativePathToFull(PackageFilename);
+
+	FPackagePendingGather& PackagePendingGather = PackagesPendingGather.AddDefaulted_GetRef();
+	PackagePendingGather.PackageName = PackageNameToGather;
+	PackagePendingGather.PackageFilename = MoveTemp(PackageFilename);
+	PackagePendingGather.PackageLocCacheState = EPackageLocCacheState::Cached;
+	return &PackagePendingGather;
+}
+
+TSet<FName> UGatherTextFromAssetsCommandlet::GetPackageNamesToGather(const TArray<FAssetData>& InAssetDataArray) const
+{
+	// Collapse the assets down to a set of packages
+	TSet<FName> PackageNamesToGather;
+	PackageNamesToGather.Reserve(InAssetDataArray.Num());
+	for (const FAssetData& AssetData : InAssetDataArray)
+	{
+		PackageNamesToGather.Add(AssetData.PackageName);
+	}
+	return PackageNamesToGather;
+}
+
+void UGatherTextFromAssetsCommandlet::PopulatePackagesPendingGather(TSet<FName> PackageNamesToGather)
+{
+	const double PopulationStartTime = FPlatformTime::Seconds();
+	// Build the basic information for the packages to gather (dependencies are filled in later once we've processed cached packages)
+	PackagesPendingGather.Reserve(PackageNamesToGather.Num());
+	for (const FName& PackageNameToGather : PackageNamesToGather)
+	{
+		AppendPackagePendingGather(PackageNameToGather);
+	}
+	UE_LOG(LogGatherTextFromAssetsCommandlet, Display, TEXT("Populating pending packages took %.2f seconds."), FPlatformTime::Seconds() - PopulationStartTime);
+}
+
+/** Process packages with loc data cached in its header and removes them from the pending packages.*/
+void UGatherTextFromAssetsCommandlet::ProcessAndRemoveCachedPackages(TMap<FName, TSet<FGuid>>& OutExternalActorsWithStaleOrMissingCaches)
+{
+	const double ProcessingStartTime = FPlatformTime::Seconds();
+	int32 NumPackagesProcessed = 0;
+	int32 PackageCount = PackagesPendingGather.Num();
+	TMap<FString, FName> AssignedPackageLocalizationIds;
+	PackagesPendingGather.RemoveAll([&](FPackagePendingGather& PackagePendingGather) -> bool
 		{
-			// Fallback on the old package flag check.
-			if (PackageFileSummary.GetPackageFlags() & PKG_RequiresLocalizationGather)
+			const FNameBuilder PackageNameStr(PackagePendingGather.PackageName);
+			const int32 CurrentPackageNum = ++NumPackagesProcessed;
+			const float PercentageComplete = static_cast<float>(CurrentPackageNum) / static_cast<float>(PackageCount) * 100.0f;
+
+			TUniquePtr<FArchive> FileReader(IFileManager::Get().CreateFileReader(*PackagePendingGather.PackageFilename));
+			if (!FileReader)
 			{
-				PackagePendingGather.PackageLocCacheState = EPackageLocCacheState::Uncached_TooOld;
+				return false;
 			}
-		}
-		else if (PackageFileSummary.GetFileVersionUE() < VER_UE4_DIALOGUE_WAVE_NAMESPACE_AND_CONTEXT_CHANGES)
-		{
-			TArray<FAssetData> AllAssetDataInSamePackage;
-			AssetRegistry.GetAssetsByPackageName(PackagePendingGather.PackageName, AllAssetDataInSamePackage);
-			for (const FAssetData& AssetData : AllAssetDataInSamePackage)
+
+			const bool bIsExternalActorPackage = PackageNameStr.ToView().Contains(FPackagePath::GetExternalActorsFolderName());
+
+			// Read package file summary from the file.
+			FPackageFileSummary PackageFileSummary;
+			*FileReader << PackageFileSummary;
+
+			// Track the package localization ID of this package (if known) and detect duplicates
+			if (!PackageFileSummary.LocalizationId.IsEmpty())
 			{
-				if (AssetData.AssetClassPath == UDialogueWave::StaticClass()->GetClassPathName())
+				if (const FName* ExistingLongPackageName = AssignedPackageLocalizationIds.Find(PackageFileSummary.LocalizationId))
 				{
-					PackagePendingGather.PackageLocCacheState = EPackageLocCacheState::Uncached_TooOld;
+					UE_LOG(LogGatherTextFromAssetsCommandlet, Warning, TEXT("Package '%s' and '%s' have the same localization ID (%s). Please reset one of these (Asset Localization -> Reset Localization ID) to avoid conflicts."), *PackageNameStr, *ExistingLongPackageName->ToString(), *PackageFileSummary.LocalizationId);
+				}
+				else
+				{
+					AssignedPackageLocalizationIds.Add(PackageFileSummary.LocalizationId, PackagePendingGather.PackageName);
 				}
 			}
-		}
-		else if (bIsExternalActorPackage && (!FNMainVersion || FNMainVersion->Version < FFortniteMainBranchObjectVersion::FixedLocalizationGatherForExternalActorPackage))
-		{
-			// Fallback on the old package flag check.
-			if (PackageFileSummary.GetPackageFlags() & PKG_RequiresLocalizationGather)
-			{
-				PackagePendingGather.PackageLocCacheState = EPackageLocCacheState::Uncached_TooOld;
-			}
-		}
 
-		// If this package doesn't have any cached data, then we have to load it for gather
-		if (PackageFileSummary.GetFileVersionUE() >= VER_UE4_SERIALIZE_TEXT_IN_PACKAGES && PackageFileSummary.GatherableTextDataOffset == 0 && (PackageFileSummary.GetPackageFlags() & PKG_RequiresLocalizationGather))
-		{
-			PackagePendingGather.PackageLocCacheState = EPackageLocCacheState::Uncached_NoCache;
-		}
+			PackagePendingGather.PackageLocCacheState = CalculatePackageLocCacheState(PackageFileSummary, PackagePendingGather.PackageName, bIsExternalActorPackage);
 
-		if (PackagePendingGather.PackageLocCacheState != EPackageLocCacheState::Cached)
-		{
-			// External actors actors must be gathered via their owner world rather than via a raw LoadPackage call
-			// Remove them from PackagesToGather as the owner world is merged back in below
-			if (bIsExternalActorPackage)
+			if (PackagePendingGather.PackageLocCacheState != EPackageLocCacheState::Cached)
 			{
-				TArray<FAssetData> ActorsInPackage;
-				AssetRegistry.GetAssetsByPackageName(PackagePendingGather.PackageName, ActorsInPackage);
-				for (const FAssetData& ActorInPackage : ActorsInPackage)
+				// External actors actors must be gathered via their owner world rather than via a raw LoadPackage call
+				// Remove them from PackagesToGather as the owner world is merged back in below
+				if (bIsExternalActorPackage)
 				{
-					if (TUniquePtr<FWorldPartitionActorDesc> ActorDesc = FWorldPartitionActorDescUtils::GetActorDescriptorFromAssetData(ActorInPackage))
+					TArray<FAssetData> ActorsInPackage;
+					IAssetRegistry::GetChecked().GetAssetsByPackageName(PackagePendingGather.PackageName, ActorsInPackage);
+					for (const FAssetData& ActorInPackage : ActorsInPackage)
 					{
-						FName WorldPackageName = *FPackageName::ObjectPathToPackageName(ActorDesc->GetActorSoftPath().ToString());
-						StaleExternalActors.FindOrAdd(WorldPackageName).Add(ActorDesc->GetGuid());
+						if (TUniquePtr<FWorldPartitionActorDesc> ActorDesc = FWorldPartitionActorDescUtils::GetActorDescriptorFromAssetData(ActorInPackage))
+						{
+							FName WorldPackageName = *FPackageName::ObjectPathToPackageName(ActorDesc->GetActorSoftPath().ToString());
+							OutExternalActorsWithStaleOrMissingCaches.FindOrAdd(WorldPackageName).Add(ActorDesc->GetGuid());
+						}
 					}
+					return true;
 				}
-				return true;
+
+				UE::Private::GatherTextFromAssetsCommandlet::AssetGatherCacheMetrics.CountUncachedAsset(PackagePendingGather.PackageLocCacheState);
+				return false;
 			}
 
-			AssetGatherCacheMetrics.CountUncachedAsset(PackagePendingGather.PackageLocCacheState);
-			return false;
-		}
+			// Process packages that don't require loading to process.
+			if (PackageFileSummary.GatherableTextDataOffset > 0)
+				{
+				UE_LOG(LogGatherTextFromAssetsCommandlet, Display, TEXT("[%6.2f%%] Gathering package: '%s'..."), PercentageComplete, *PackageNameStr);
 
-		// Process packages that don't require loading to process.
-		if (PackageFileSummary.GatherableTextDataOffset > 0)
-		{
-			UE_LOG(LogGatherTextFromAssetsCommandlet, Display, TEXT("[%6.2f%%] Gathering package: '%s'..."), PercentageComplete, *PackageNameStr);
+				UE::Private::GatherTextFromAssetsCommandlet::AssetGatherCacheMetrics.CountCachedAsset();
 
-			AssetGatherCacheMetrics.CountCachedAsset();
+				FileReader->Seek(PackageFileSummary.GatherableTextDataOffset);
 
-			FileReader->Seek(PackageFileSummary.GatherableTextDataOffset);
+				PackagePendingGather.GatherableTextDataArray.SetNum(PackageFileSummary.GatherableTextDataCount);
+				for (int32 GatherableTextDataIndex = 0; GatherableTextDataIndex < PackageFileSummary.GatherableTextDataCount; ++GatherableTextDataIndex)
+				{
+					*FileReader << PackagePendingGather.GatherableTextDataArray[GatherableTextDataIndex];
+				}
 
-			PackagePendingGather.GatherableTextDataArray.SetNum(PackageFileSummary.GatherableTextDataCount);
-			for (int32 GatherableTextDataIndex = 0; GatherableTextDataIndex < PackageFileSummary.GatherableTextDataCount; ++GatherableTextDataIndex)
-			{
-				(*FileReader) << PackagePendingGather.GatherableTextDataArray[GatherableTextDataIndex];
+				ProcessGatherableTextDataArray(PackagePendingGather.GatherableTextDataArray);
 			}
 
-			ProcessGatherableTextDataArray(PackagePendingGather.GatherableTextDataArray);
-		}
-
-		// If we're reporting or fixing assets with a stale gather cache then we still need to load this 
-		// package in order to do that, but the PackageLocCacheState prevents it being gathered again
-		if (bReportStaleGatherCache || bFixStaleGatherCache)
-		{
-			check(PackagePendingGather.PackageLocCacheState == EPackageLocCacheState::Cached);
-			return false;
-		}
-
-		return true;
-	});
+			// If we're reporting or fixing assets with a stale gather cache then we still need to load this 
+			// package in order to do that, but the PackageLocCacheState prevents it being gathered again
+			if (bReportStaleGatherCache || bFixStaleGatherCache)
+			{
+				check(PackagePendingGather.PackageLocCacheState == EPackageLocCacheState::Cached);
+				return false;
+			}
 
-	AssetGatherCacheMetrics.LogMetrics();
+			return true;
+		});
+	UE_LOG(LogGatherTextFromAssetsCommandlet, Display, TEXT("Processing and removing packages with cached localization data took %.2f"), FPlatformTime::Seconds() - ProcessingStartTime);
+}
 
-	// Merge any pending WP map requests back into PackagesPendingGather
+void UGatherTextFromAssetsCommandlet::MergeInExternalActorsWithStaleOrMissingCaches(TMap<FName, TSet<FGuid>>& StaleExternalActors)
+{
 	for (TTuple<FName, TSet<FGuid>>& StaleExternalActorsPair : StaleExternalActors)
 	{
 		FPackagePendingGather* WorldPackagePendingGather = PackagesPendingGather.FindByPredicate([&StaleExternalActorsPair](const FPackagePendingGather& PotentialPackagePendingGather)
-		{
-			return PotentialPackagePendingGather.PackageName == StaleExternalActorsPair.Key;
-		});
+			{
+				return PotentialPackagePendingGather.PackageName == StaleExternalActorsPair.Key;
+			});
 		if (!WorldPackagePendingGather)
 		{
 			WorldPackagePendingGather = AppendPackagePendingGather(StaleExternalActorsPair.Key);
@@ -971,45 +984,24 @@ int32 UGatherTextFromAssetsCommandlet::Main(const FString& Params)
 			UE_LOG(LogGatherTextFromAssetsCommandlet, Warning, TEXT("Failed to queue world package '%s' for %d external actor(s)."), *StaleExternalActorsPair.Key.ToString(), StaleExternalActorsPair.Value.Num());
 		}
 	}
-	StaleExternalActors.Reset();
-
-	NumPackagesProcessed = 0;
-	PackageCount = PackagesPendingGather.Num();
-	if (PackageCount == 0)
-	{
-		// Nothing more to do!
-		return 0;
-	}
-
-	const double PackageLoadingStartTime = FPlatformTime::Seconds();
-	UE_LOG(LogGatherTextFromAssetsCommandlet, Display, TEXT("Preparing to load %d packages..."), PackageCount);
+}
 
+/** Load the remaining pending packages for gather.*/
+void UGatherTextFromAssetsCommandlet::LoadAndProcessUncachedPackages(TArray<FName>& OutPackagesWithStaleGatherCache)
+{
 	FLoadPackageLogOutputRedirector LogOutputRedirector;
-
-	CalculateDependenciesForPackagesPendingGather();
-
-	// Collect garbage before beginning to load packages
-	// This also sorts the list of packages into the best processing order
-	PurgeGarbage(/*bPurgeReferencedPackages*/false);
-
-	// We don't need to have compiled shaders to gather text
-	bool bWasShaderCompilationEnabled = false;
-	if (GShaderCompilingManager)
-	{
-		bWasShaderCompilationEnabled = !GShaderCompilingManager->IsShaderCompilationSkipped();
-		GShaderCompilingManager->SkipShaderCompilation(true);
-	}
-
-	int32 NumPackagesFailed = 0;
-	TArray<FName> PackagesWithStaleGatherCache;
 	TArray<FGatherableTextData> GatherableTextDataArray;
+	int32 NumPackagesProcessed = 0;
+	int32 PackageCount = PackagesPendingGather.Num();
+	int32 NumPackagesFailedLoading = 0;
+	const double LoadingStartTime = FPlatformTime::Seconds();
 	while (PackagesPendingGather.Num() > 0)
 	{
 		const FPackagePendingGather PackagePendingGather = PackagesPendingGather.Pop(/*bAllowShrinking*/false);
 		const FNameBuilder PackageNameStr(PackagePendingGather.PackageName);
 
-		const int32 CurrentPackageNum = ++NumPackagesProcessed + NumPackagesFailed;
-		const float PercentageComplete = (((float)CurrentPackageNum / (float)PackageCount) * 100.0f);
+		const int32 CurrentPackageNum = ++NumPackagesProcessed;
+		const float PercentageComplete = static_cast<float>(CurrentPackageNum) / static_cast<float>(PackageCount) * 100.0f;
 		UE_LOG(LogGatherTextFromAssetsCommandlet, Display, TEXT("[%6.2f%%] Loading package: '%s'..."), PercentageComplete, *PackageNameStr);
 
 		UPackage* Package = nullptr;
@@ -1028,23 +1020,12 @@ int32 UGatherTextFromAssetsCommandlet::Main(const FString& Params)
 		if (!Package)
 		{
 			UE_LOG(LogGatherTextFromAssetsCommandlet, Warning, TEXT("Failed to load package: '%s'."), *PackageNameStr);
-			++NumPackagesFailed;
+			++NumPackagesFailedLoading;
 			continue;
 		}
 
 		// Tick background tasks
-		if (GShaderCompilingManager)
-		{
-			GShaderCompilingManager->ProcessAsyncResults(true, false);
-		}
-		if (GDistanceFieldAsyncQueue)
-		{
-			GDistanceFieldAsyncQueue->ProcessAsyncTasks();
-		}
-		if (GCardRepresentationAsyncQueue)
-		{
-			GCardRepresentationAsyncQueue->ProcessAsyncTasks();
-		}
+		UE::Private::GatherTextFromAssetsCommandlet::TickBackgroundTasks();
 
 		// Because packages may not have been resaved after this flagging was implemented, we may have added packages to load that weren't flagged - potential false positives.
 		// The loading process should have reflagged said packages so that only true positives will have this flag.
@@ -1063,11 +1044,11 @@ int32 UGatherTextFromAssetsCommandlet::Main(const FString& Params)
 			{
 				// Look for any structurally significant changes (missing, added, or changed texts) in the cache
 				// Ignore insignificant things (like source changes caused by assets moving or being renamed)
-				if (EnumHasAnyFlags(GatherableTextResultFlags, EPropertyLocalizationGathererResultFlags::HasTextWithInvalidPackageLocalizationID) 
+				if (EnumHasAnyFlags(GatherableTextResultFlags, EPropertyLocalizationGathererResultFlags::HasTextWithInvalidPackageLocalizationID)
 					|| !IsGatherableTextDataIdentical(GatherableTextDataArray, PackagePendingGather.GatherableTextDataArray))
 				{
-					PackagesWithStaleGatherCache.Add(PackagePendingGather.PackageName);
-						
+					OutPackagesWithStaleGatherCache.Add(PackagePendingGather.PackageName);
+
 					if (bFixStaleGatherCache)
 					{
 						bSavePackage = true;
@@ -1127,14 +1108,14 @@ int32 UGatherTextFromAssetsCommandlet::Main(const FString& Params)
 						ForEachActorParams.ActorGuids = PackagePendingGather.ExternalActors.Array();
 
 						FWorldPartitionHelpers::ForEachActorWithLoading(WorldPartition, [&GatherableTextDataArray](const FWorldPartitionActorDescInstance* ActorDescInstance)
-						{
-							if (const AActor* Actor = ActorDescInstance->GetActor())
 							{
-								EPropertyLocalizationGathererResultFlags ActorGatherableTextResultFlags = EPropertyLocalizationGathererResultFlags::Empty;
-								FPropertyLocalizationDataGatherer(GatherableTextDataArray, Actor->GetExternalPackage(), ActorGatherableTextResultFlags);
-							}
-							return true;
-						}, ForEachActorParams);
+								if (const AActor* Actor = ActorDescInstance->GetActor())
+								{
+									EPropertyLocalizationGathererResultFlags ActorGatherableTextResultFlags = EPropertyLocalizationGathererResultFlags::Empty;
+									FPropertyLocalizationDataGatherer(GatherableTextDataArray, Actor->GetExternalPackage(), ActorGatherableTextResultFlags);
+								}
+								return true;
+							}, ForEachActorParams);
 					}
 				}
 			}
@@ -1164,7 +1145,192 @@ int32 UGatherTextFromAssetsCommandlet::Main(const FString& Params)
 			}
 		}
 	}
-	UE_LOG(LogGatherTextFromAssetsCommandlet, Display, TEXT("Loaded %d packages in %.2f seconds. %d failed."), NumPackagesProcessed, FPlatformTime::Seconds() - PackageLoadingStartTime, NumPackagesFailed);
+	UE_LOG(LogGatherTextFromAssetsCommandlet, Display, TEXT("Loaded %d packages. %d failed."), NumPackagesProcessed, NumPackagesFailedLoading);
+	UE_LOG(LogGatherTextFromAssetsCommandlet, Display, TEXT("Loading all uncached packages took %.2f seconds."), FPlatformTime::Seconds() - LoadingStartTime);
+}
+
+void UGatherTextFromAssetsCommandlet::ReportStaleGatherCache(TArray<FName>& InPackagesWithStaleGatherCache) const
+{
+	InPackagesWithStaleGatherCache.Sort(FNameLexicalLess());
+
+	FString StaleGatherCacheReport;
+	for (const FName& PackageWithStaleGatherCache : InPackagesWithStaleGatherCache)
+	{
+		StaleGatherCacheReport += PackageWithStaleGatherCache.ToString();
+		StaleGatherCacheReport += TEXT("\n");
+	}
+
+	const FString StaleGatherCacheReportFilename = DestinationPath / TEXT("StaleGatherCacheReport.txt");
+	const bool bStaleGatherCacheReportSaved = FLocalizedAssetSCCUtil::SaveFileWithSCC(SourceControlInfo, StaleGatherCacheReportFilename, [&StaleGatherCacheReport](const FString& InSaveFileName) -> bool
+		{
+			return FFileHelper::SaveStringToFile(StaleGatherCacheReport, *InSaveFileName, FFileHelper::EEncodingOptions::ForceUTF8);
+		});
+
+	if (!bStaleGatherCacheReportSaved)
+	{
+		UE_LOG(LogGatherTextFromAssetsCommandlet, Error, TEXT("Failed to save report: '%s'."), *StaleGatherCacheReportFilename);
+	}
+}
+
+UGatherTextFromAssetsCommandlet::EPackageLocCacheState UGatherTextFromAssetsCommandlet::CalculatePackageLocCacheState(const FPackageFileSummary& PackageFileSummary, const FName PackageName, bool bIsExternalActorPackage) const
+{
+	// Have we been asked to skip the cache of text that exists in the header of newer packages?
+	if (bSkipGatherCache && PackageFileSummary.GetFileVersionUE() >= VER_UE4_SERIALIZE_TEXT_IN_PACKAGES)
+	{
+		// Fallback on the old package flag check.
+		if (PackageFileSummary.GetPackageFlags() & PKG_RequiresLocalizationGather)
+		{
+			return EPackageLocCacheState::Uncached_NoCache;
+		}
+	}
+
+	const FCustomVersion* const EditorVersion = PackageFileSummary.GetCustomVersionContainer().GetVersion(FEditorObjectVersion::GUID);
+	const FCustomVersion* const FNMainVersion = PackageFileSummary.GetCustomVersionContainer().GetVersion(FFortniteMainBranchObjectVersion::GUID);
+
+	// Packages not resaved since localization gathering flagging was added to packages must be loaded.
+	if (PackageFileSummary.GetFileVersionUE() < VER_UE4_PACKAGE_REQUIRES_LOCALIZATION_GATHER_FLAGGING)
+	{
+		return EPackageLocCacheState::Uncached_TooOld;
+	}
+	// Package not resaved since gatherable text data was added to package headers must be loaded, since their package header won't contain pregathered text data.
+	else if (PackageFileSummary.GetFileVersionUE() < VER_UE4_SERIALIZE_TEXT_IN_PACKAGES || (!EditorVersion || EditorVersion->Version < FEditorObjectVersion::GatheredTextEditorOnlyPackageLocId))
+	{
+		// Fallback on the old package flag check.
+		if (PackageFileSummary.GetPackageFlags() & PKG_RequiresLocalizationGather)
+		{
+			return EPackageLocCacheState::Uncached_TooOld;
+		}
+	}
+	else if (PackageFileSummary.GetFileVersionUE() < VER_UE4_DIALOGUE_WAVE_NAMESPACE_AND_CONTEXT_CHANGES)
+	{
+		TArray<FAssetData> AllAssetDataInSamePackage;
+		IAssetRegistry::GetChecked().GetAssetsByPackageName(PackageName, AllAssetDataInSamePackage);
+		for (const FAssetData& AssetData : AllAssetDataInSamePackage)
+		{
+			if (AssetData.AssetClassPath == UDialogueWave::StaticClass()->GetClassPathName())
+			{
+				return EPackageLocCacheState::Uncached_TooOld;
+			}
+		}
+	}
+	else if (bIsExternalActorPackage && (!FNMainVersion || FNMainVersion->Version < FFortniteMainBranchObjectVersion::FixedLocalizationGatherForExternalActorPackage))
+	{
+		// Fallback on the old package flag check.
+		if (PackageFileSummary.GetPackageFlags() & PKG_RequiresLocalizationGather)
+		{
+			return EPackageLocCacheState::Uncached_TooOld;
+		}
+	}
+
+	// If this package doesn't have any cached data, then we have to load it for gather
+	if (PackageFileSummary.GetFileVersionUE() >= VER_UE4_SERIALIZE_TEXT_IN_PACKAGES && PackageFileSummary.GatherableTextDataOffset == 0 && (PackageFileSummary.GetPackageFlags() & PKG_RequiresLocalizationGather))
+	{
+		return EPackageLocCacheState::Uncached_NoCache;
+	}
+	return EPackageLocCacheState::Cached;
+}
+
+int32 UGatherTextFromAssetsCommandlet::Main(const FString& Params)
+{
+	// Parse command line.
+	if (!ParseCommandLineHelper(Params))
+	{
+		return -1;
+	}
+
+	// If the editor has loaded a persistent world then create an empty world prior to starting the asset gather
+	// This avoids any issues when loading and initializing worlds during the gather, as WP needs to re-initialize the world
+	if (GEditor)
+	{
+		if (UWorld* EditorWorld = GEditor->GetEditorWorldContext().World();
+			EditorWorld && !FPackageName::IsTempPackage(FNameBuilder(EditorWorld->GetPackage()->GetFName()).ToView()))
+		{
+			GEditor->CreateNewMapForEditing(/*bPromptForSave*/false);
+		}
+	}
+
+	UE_LOG(LogGatherTextFromAssetsCommandlet, Display, TEXT("Discovering assets to gather..."));
+	const double DiscoveringAssetsStartTime = FPlatformTime::Seconds();
+	const double SearchAssetRegistryStartTime = FPlatformTime::Seconds();
+	FAssetRegistryModule& AssetRegistryModule = FModuleManager::LoadModuleChecked<FAssetRegistryModule>(TEXT("AssetRegistry"));
+	IAssetRegistry& AssetRegistry = AssetRegistryModule.Get();
+	AssetRegistry.SearchAllAssets(true);
+	UE_LOG(LogGatherTextFromAssetsCommandlet, Display, TEXT("Searching all assets took %.2f seconds."), FPlatformTime::Seconds() - SearchAssetRegistryStartTime);
+
+	TArray<FAssetData> AssetDataArray;
+	if (!PerformFirstPassFilter(AssetDataArray))
+	{
+		return -1;
+	}
+
+	if (!bShouldExcludeDerivedClasses)
+	{
+		if (!PerformExcludeExactClassesFilter(AssetDataArray))
+		{
+			return -1;
+		}
+	}
+
+	// Note: AssetDataArray now contains all assets in the specified collections that are not instances of the specified excluded classes.
+	FilterAssetsBasedOnIncludeExcludePaths(AssetDataArray);
+
+	if (AssetDataArray.Num() == 0)
+	{
+		UE_LOG(LogGatherTextFromAssetsCommandlet, Display, TEXT("No assets matched the specified criteria."));
+		return 0;
+	}
+
+	// Discover the external actors for any worlds that are pending gather
+	DiscoverExternalActors(AssetDataArray);
+
+	// Collect the basic information about the packages that we're going to gather from
+	TSet<FName> PackageNamesToGather = GetPackageNamesToGather(AssetDataArray);
+	AssetDataArray.Empty();
+	PopulatePackagesPendingGather(MoveTemp(PackageNamesToGather));
+
+	UE_LOG(LogGatherTextFromAssetsCommandlet, Display, TEXT("Discovering assets to gather took %.2f seconds."), FPlatformTime::Seconds() - DiscoveringAssetsStartTime);
+
+	// These are external actor packages that are stale or are missing a gather cache from their package
+	// Map of world package name -> external actor Ids in the world 
+	TMap<FName, TSet<FGuid>> ExternalActorsWithStaleOrMissingCaches;
+	// Process all packages that do not need to be loaded. Remove processed packages from the list.
+	UE_LOG(LogGatherTextFromAssetsCommandlet, Display, TEXT("Processing assets to gather..."));
+	ProcessAndRemoveCachedPackages(ExternalActorsWithStaleOrMissingCaches);
+
+	UE::Private::GatherTextFromAssetsCommandlet::AssetGatherCacheMetrics.LogMetrics();
+
+	// Merge any pending WP map requests back into PackagesPendingGather
+	MergeInExternalActorsWithStaleOrMissingCaches(ExternalActorsWithStaleOrMissingCaches);
+	ExternalActorsWithStaleOrMissingCaches.Reset();
+
+	// All packages left in PackagesPendingGather should now have to be loaded 
+	if (PackagesPendingGather.Num() == 0)
+	{
+		// Nothing more to do!
+		return 0;
+	}
+
+	const double PackageLoadingStartTime = FPlatformTime::Seconds();
+	UE_LOG(LogGatherTextFromAssetsCommandlet, Display, TEXT("Preparing to load %d packages..."), PackagesPendingGather.Num());
+
+	CalculateDependenciesForPackagesPendingGather();
+
+	// Collect garbage before beginning to load packages
+	// This also sorts the list of packages into the best processing order
+	PurgeGarbage(/*bPurgeReferencedPackages*/false);
+
+	// We don't need to have compiled shaders to gather text
+	bool bWasShaderCompilationEnabled = false;
+	if (GShaderCompilingManager)
+	{
+		bWasShaderCompilationEnabled = !GShaderCompilingManager->IsShaderCompilationSkipped();
+		GShaderCompilingManager->SkipShaderCompilation(true);
+	}
+
+	TArray<FName> PackagesWithStaleGatherCache;
+	LoadAndProcessUncachedPackages(PackagesWithStaleGatherCache);
+
+	UE_LOG(LogGatherTextFromAssetsCommandlet, Display, TEXT("Loading packages took %.2f seconds."), FPlatformTime::Seconds() - PackageLoadingStartTime);
 
 	// Collect garbage after loading all packages
 	// This reclaims as much memory as possible for the rest of the gather pipeline
@@ -1177,25 +1343,7 @@ int32 UGatherTextFromAssetsCommandlet::Main(const FString& Params)
 	
 	if (bReportStaleGatherCache)
 	{
-		PackagesWithStaleGatherCache.Sort(FNameLexicalLess());
-
-		FString StaleGatherCacheReport;
-		for (const FName& PackageWithStaleGatherCache : PackagesWithStaleGatherCache)
-		{
-			StaleGatherCacheReport += PackageWithStaleGatherCache.ToString();
-			StaleGatherCacheReport += TEXT("\n");
-		}
-
-		const FString StaleGatherCacheReportFilename = DestinationPath / TEXT("StaleGatherCacheReport.txt");
-		const bool bStaleGatherCacheReportSaved = FLocalizedAssetSCCUtil::SaveFileWithSCC(SourceControlInfo, StaleGatherCacheReportFilename, [&StaleGatherCacheReport](const FString& InSaveFileName) -> bool
-		{
-			return FFileHelper::SaveStringToFile(StaleGatherCacheReport, *InSaveFileName, FFileHelper::EEncodingOptions::ForceUTF8);
-		});
-
-		if (!bStaleGatherCacheReportSaved)
-		{
-			UE_LOG(LogGatherTextFromAssetsCommandlet, Error, TEXT("Failed to save report: '%s'."), *StaleGatherCacheReportFilename);
-		}
+		ReportStaleGatherCache(PackagesWithStaleGatherCache);
 	}
 
 	return 0;
@@ -1232,7 +1380,7 @@ bool UGatherTextFromAssetsCommandlet::GetConfigurationScript(const TMap<FString,
 
 bool UGatherTextFromAssetsCommandlet::ConfigureFromScript(const FString& GatherTextConfigPath, const FString& SectionName)
 {
-	bool HasFatalError = false;
+	bool bHasFatalError = false;
 
 	// Modules to Preload
 	GetStringArrayFromConfig(*SectionName, TEXT("ModulesToPreload"), ModulesToPreload, GatherTextConfigPath);
@@ -1254,7 +1402,7 @@ bool UGatherTextFromAssetsCommandlet::ConfigureFromScript(const FString& GatherT
 	if (IncludePathFilters.Num() == 0)
 	{
 		UE_LOG(LogGatherTextFromAssetsCommandlet, Error, TEXT("No include path filters in section %s."), *SectionName);
-		HasFatalError = true;
+		bHasFatalError = true;
 	}
 
 	// Collections
@@ -1264,11 +1412,11 @@ bool UGatherTextFromAssetsCommandlet::ConfigureFromScript(const FString& GatherT
 		FCollectionManagerModule& CollectionManagerModule = FCollectionManagerModule::GetModule();
 		ICollectionManager& CollectionManager = CollectionManagerModule.Get();
 
-		const bool DoesCollectionExist = CollectionManager.CollectionExists(FName(*CollectionName), ECollectionShareType::CST_All);
-		if (!DoesCollectionExist)
+		const bool bDoesCollectionExist = CollectionManager.CollectionExists(FName(*CollectionName), ECollectionShareType::CST_All);
+		if (!bDoesCollectionExist)
 		{
 			UE_LOG(LogGatherTextFromAssetsCommandlet, Error, TEXT("Failed to find a collection with name \"%s\", collection does not exist."), *CollectionName);
-			HasFatalError = true;
+			bHasFatalError = true;
 		}
 	}
 
@@ -1303,13 +1451,13 @@ bool UGatherTextFromAssetsCommandlet::ConfigureFromScript(const FString& GatherT
 	if (PackageFileNameFilters.Num() == 0)
 	{
 		UE_LOG(LogGatherTextFromAssetsCommandlet, Error, TEXT("No package file name filters in section %s."), *SectionName);
-		HasFatalError = true;
+		bHasFatalError = true;
 	}
 
 	// Recursive asset class exclusion
-	if (!GetBoolFromConfig(*SectionName, TEXT("ShouldExcludeDerivedClasses"), ShouldExcludeDerivedClasses, GatherTextConfigPath))
+	if (!GetBoolFromConfig(*SectionName, TEXT("ShouldExcludeDerivedClasses"), bShouldExcludeDerivedClasses, GatherTextConfigPath))
 	{
-		ShouldExcludeDerivedClasses = false;
+		bShouldExcludeDerivedClasses = false;
 	}
 
 	// Asset class exclude
@@ -1318,9 +1466,9 @@ bool UGatherTextFromAssetsCommandlet::ConfigureFromScript(const FString& GatherT
 	GetPathArrayFromConfig(*SectionName, TEXT("ManifestDependencies"), ManifestDependenciesList, GatherTextConfigPath);
 
 	// Get whether we should gather editor-only data. Typically only useful for the localization of UE itself.
-	if (!GetBoolFromConfig(*SectionName, TEXT("ShouldGatherFromEditorOnlyData"), ShouldGatherFromEditorOnlyData, GatherTextConfigPath))
+	if (!GetBoolFromConfig(*SectionName, TEXT("ShouldGatherFromEditorOnlyData"), bShouldGatherFromEditorOnlyData, GatherTextConfigPath))
 	{
-		ShouldGatherFromEditorOnlyData = false;
+		bShouldGatherFromEditorOnlyData = false;
 	}
 
 	auto ReadBoolFlagWithFallback = [this, &SectionName, &GatherTextConfigPath](const TCHAR* FlagName, bool& OutValue)
@@ -1358,7 +1506,7 @@ bool UGatherTextFromAssetsCommandlet::ConfigureFromScript(const FString& GatherT
 		MaxUsedMemoryBytes = MaxUsedMemoryMB * 1024LL * 1024LL;
 	}
 
-	return !HasFatalError;
+	return !bHasFatalError;
 }
 
 #undef LOC_DEFINE_REGION
