Adapt EditorAssetSubsystem's Checkout/Save/Delete Assets for external objects (saved in an external package).

diff --git a/Engine/Source/Editor/UnrealEd/Private/Subsystems/EditorAssetSubsystem.cpp b/Engine/Source/Editor/UnrealEd/Private/Subsystems/EditorAssetSubsystem.cpp
--- a/Engine/Source/Editor/UnrealEd/Private/Subsystems/EditorAssetSubsystem.cpp
+++ b/Engine/Source/Editor/UnrealEd/Private/Subsystems/EditorAssetSubsystem.cpp
@@ -105,14 +105,15 @@ namespace UE::EditorAssetUtils
 		return LoadAssetFromData(AssetDataResult.GetValue());
 	}
 
-	TError<FString> IsARegisteredAsset(UObject* Object)
+	TError<FString> IsARegisteredAsset(UObject* Object, bool bAllowSkipBrowsableTestForExternalObject = false)
 	{
 		if (!IsValid(Object))
 		{
 			return MakeError(TEXT("The Asset is not valid."));
 		}
 
-		if (!ObjectTools::IsObjectBrowsable(Object))
+		const bool bCanSkipIsBrowsable = bAllowSkipBrowsableTestForExternalObject && Object->IsPackageExternal();
+		if (!bCanSkipIsBrowsable && !ObjectTools::IsObjectBrowsable(Object))
 		{
 			return MakeError(FString::Printf(TEXT("The object '%s' is not an asset."), *Object->GetName()));
 		}
@@ -413,7 +414,7 @@ FString UEditorAssetSubsystem::GetPathNameForLoadedAsset(UObject* LoadedAsset)
 		return FString();
 	}
 	
-	if (TError<FString> Result = UE::EditorAssetUtils::IsARegisteredAsset(LoadedAsset); Result.HasError())
+	if (TError<FString> Result = UE::EditorAssetUtils::IsARegisteredAsset(LoadedAsset, true/*bAllowSkipBrowsableTestForExternalObject*/); Result.HasError())
 	{
 		UE_LOG(LogEditorAssetSubsystem, Error, TEXT("GetPathNameForLoadedAsset failed: %s"), *Result.GetError());
 	}
@@ -690,7 +691,7 @@ bool UEditorAssetSubsystem::DeleteLoadedAsset(UObject* AssetToDelete)
 		return false;
 	}
 	
-	if (TError<FString> Result = UE::EditorAssetUtils::IsARegisteredAsset(AssetToDelete); Result.HasError())
+	if (TError<FString> Result = UE::EditorAssetUtils::IsARegisteredAsset(AssetToDelete, true/*bAllowSkipBrowsableTestForExternalObject*/); Result.HasError())
 	{
 		UE_LOG(LogEditorAssetSubsystem, Error, TEXT("DeleteLoadedAsset failed: Asset is not registered. %s"), *Result.GetError());
 		return false;
@@ -714,7 +715,7 @@ bool UEditorAssetSubsystem::DeleteLoadedAssets(const TArray<UObject*>& AssetsToD
 	FString FailureReason;
 	for (UObject* Asset : AssetsToDelete)
 	{
-		if (TError<FString> Result = UE::EditorAssetUtils::IsARegisteredAsset(Asset); Result.HasError())
+		if (TError<FString> Result = UE::EditorAssetUtils::IsARegisteredAsset(Asset, true/*bAllowSkipBrowsableTestForExternalObject*/); Result.HasError())
 		{
 			UE_LOG(LogEditorAssetSubsystem, Warning, TEXT("DeleteLoadedAssets failed: An asset is not registered. %s"), *Result.GetError());
 			return false;
@@ -1119,7 +1120,7 @@ bool UEditorAssetSubsystem::CheckoutLoadedAsset(UObject* AssetToCheckout)
 		return false;
 	}
 
-	if (TError<FString> Result = UE::EditorAssetUtils::IsARegisteredAsset(AssetToCheckout); Result.HasError())
+	if (TError<FString> Result = UE::EditorAssetUtils::IsARegisteredAsset(AssetToCheckout, true/*bAllowSkipBrowsableTestForExternalObject*/); Result.HasError())
 	{
 		UE_LOG(LogEditorAssetSubsystem, Error, TEXT("CheckoutLoadedAsset failed: Asset is not registered. %s"), *Result.GetError());
 		return false;
@@ -1144,7 +1145,7 @@ bool UEditorAssetSubsystem::CheckoutLoadedAssets(const TArray<UObject*>& AssetsT
 	Packages.Reserve(AssetsToCheckout.Num());
 	for (UObject* AssetToCheckout: AssetsToCheckout)
 	{
-		if (TError<FString> Result = UE::EditorAssetUtils::IsARegisteredAsset(AssetToCheckout); Result.HasError())
+		if (TError<FString> Result = UE::EditorAssetUtils::IsARegisteredAsset(AssetToCheckout, true/*bAllowSkipBrowsableTestForExternalObject*/); Result.HasError())
 		{
 			UE_LOG(LogEditorAssetSubsystem, Warning, TEXT("CheckoutLoadedAssets: An asset is not registered. %s"), *Result.GetError());
 			continue;
@@ -1208,7 +1209,7 @@ bool UEditorAssetSubsystem::SaveLoadedAsset(UObject* AssetToSave, bool bOnlyIfIs
 		return false;
 	}
 	
-	if (TError<FString> Result = UE::EditorAssetUtils::IsARegisteredAsset(AssetToSave); Result.HasError())
+	if (TError<FString> Result = UE::EditorAssetUtils::IsARegisteredAsset(AssetToSave, true/*bAllowSkipBrowsableTestForExternalObject*/); Result.HasError())
 	{
 		UE_LOG(LogEditorAssetSubsystem, Error, TEXT("SaveLoadedAsset failed: Asset is not registered. %s"), *Result.GetError());
 		return false;
@@ -1235,7 +1236,7 @@ bool UEditorAssetSubsystem::SaveLoadedAssets(const TArray<UObject*>& AssetsToSav
 	Packages.Reserve(AssetsToSave.Num());
 	for (UObject* AssetToSave : AssetsToSave)
 	{
-		if (TError<FString> Result = UE::EditorAssetUtils::IsARegisteredAsset(AssetToSave); Result.HasError())
+		if (TError<FString> Result = UE::EditorAssetUtils::IsARegisteredAsset(AssetToSave, true/*bAllowSkipBrowsableTestForExternalObject*/); Result.HasError())
 		{
 			UE_LOG(LogEditorAssetSubsystem, Warning, TEXT("SaveLoadedAsset: An asset is not registered. %s"), *Result.GetError());
 		}
