Prevent crash when attempting to create animation asset with invalid mesh/skeleton

diff --git a/Engine/Source/Editor/UnrealEd/Private/AnimationEditorUtils.cpp b/Engine/Source/Editor/UnrealEd/Private/AnimationEditorUtils.cpp
--- a/Engine/Source/Editor/UnrealEd/Private/AnimationEditorUtils.cpp
+++ b/Engine/Source/Editor/UnrealEd/Private/AnimationEditorUtils.cpp
@@ -472,8 +472,15 @@ namespace AnimationEditorUtils
 			USkeleton* Skeleton = Cast<USkeleton>(SkeletonOrSkeletalMeshObject);
 			if (Skeleton == nullptr)
 			{
-				SkeletalMesh = CastChecked<USkeletalMesh>(SkeletonOrSkeletalMeshObject);
-				Skeleton = SkeletalMesh->GetSkeleton();
+				SkeletalMesh = Cast<USkeletalMesh>(SkeletonOrSkeletalMeshObject);
+				if (SkeletalMesh)
+				{
+					Skeleton = SkeletalMesh->GetSkeleton();				
+				}
+				else
+				{
+					UE_LOG(LogAnimation, Warning, TEXT("Invalid skeleton or skeletal mesh passed to CreateAnimationAssets. No asset will be generated."));
+				}
 			}
 
 			if (Skeleton)
