New TRS Gizmo: handle transform transaction using the ITF

    - when created, the gizmo asks the mode manager to create a state target by calling FEditorModeTools::GetGizmoStateTarget(), which will create one if needed
    - UEditorGizmoStateTarget::BeginUpdate() creates a new transaction + calls FEditorModeTools::BeginTransform() which will iterate thru the modes and ask them if they have anything to handle
    - if they don't, the focused VPC still has the option of doing something (it doesn't by default, but some, like the level VPC, handle what they need to)

diff --git a/Engine/Plugins/2D/Paper2D/Source/Paper2DEditor/Private/TileSetEditor/SingleTileEditorViewportClient.cpp b/Engine/Plugins/2D/Paper2D/Source/Paper2DEditor/Private/TileSetEditor/SingleTileEditorViewportClient.cpp
--- a/Engine/Plugins/2D/Paper2D/Source/Paper2DEditor/Private/TileSetEditor/SingleTileEditorViewportClient.cpp
+++ b/Engine/Plugins/2D/Paper2D/Source/Paper2DEditor/Private/TileSetEditor/SingleTileEditorViewportClient.cpp
@@ -80,9 +80,7 @@ void FSingleTileEditorViewportClient::TrackingStarted(const struct FInputEventSt
 
 	if (!bManipulating && bIsDragging && !bTrackingHandledExternally)
 	{
-		BeginTransaction(LOCTEXT("ModificationInViewportTransaction", "Modification in Viewport"));
-		bManipulating = true;
-		bManipulationDirtiedSomething = false;
+		HandleBeginTransform();
 	}
 }
 
@@ -92,10 +90,42 @@ void FSingleTileEditorViewportClient::TrackingStopped()
 	const bool bTransactingHandledByEditorMode = ModeTools->EndTracking(this, Viewport);
 
 	if (bManipulating && !bTransactingHandledByEditorMode)
+	{
+		HandleEndTransform();
+	}
+}
+
+bool FSingleTileEditorViewportClient::BeginTransform(const FGizmoState& InState)
+{
+	return HandleBeginTransform();
+}
+
+bool FSingleTileEditorViewportClient::EndTransform(const FGizmoState& InState)
+{
+	return HandleEndTransform();
+}
+
+bool FSingleTileEditorViewportClient::HandleBeginTransform()
+{
+	if (!bManipulating)
+	{
+		BeginTransaction(LOCTEXT("ModificationInViewportTransaction", "Modification in Viewport"));
+		bManipulating = true;
+		bManipulationDirtiedSomething = false;
+		return true;
+	}
+	return false;
+}
+
+bool FSingleTileEditorViewportClient::HandleEndTransform()
+{
+	if (bManipulating)
 	{
 		EndTransaction();
 		bManipulating = false;
+		return true;
 	}
+	return false;
 }
 
 void FSingleTileEditorViewportClient::DrawCanvas(FViewport& InViewport, FSceneView& View, FCanvas& Canvas)
diff --git a/Engine/Plugins/2D/Paper2D/Source/Paper2DEditor/Private/TileSetEditor/SingleTileEditorViewportClient.h b/Engine/Plugins/2D/Paper2D/Source/Paper2DEditor/Private/TileSetEditor/SingleTileEditorViewportClient.h
--- a/Engine/Plugins/2D/Paper2D/Source/Paper2DEditor/Private/TileSetEditor/SingleTileEditorViewportClient.h
+++ b/Engine/Plugins/2D/Paper2D/Source/Paper2DEditor/Private/TileSetEditor/SingleTileEditorViewportClient.h
@@ -13,6 +13,7 @@ class FScopedTransaction;
 class FUICommandList;
 class UPaperSpriteComponent;
 class UPaperTileSet;
+struct FGizmoState;
 
 DECLARE_MULTICAST_DELEGATE_TwoParams(FOnSingleTileIndexChanged, int32 /*NewIndex*/, int32 /*OldIndex*/);
 
@@ -32,6 +33,8 @@ public:
 	virtual FLinearColor GetBackgroundColor() const override;
 	virtual void TrackingStarted(const FInputEventState& InInputState, bool bIsDragging, bool bNudge) override;
 	virtual void TrackingStopped() override;
+	virtual bool BeginTransform(const FGizmoState& InState) override;
+	virtual bool EndTransform(const FGizmoState& InState) override;
 	virtual void DrawCanvas(FViewport& Viewport, FSceneView& View, FCanvas& Canvas) override;
 	// End of FEditorViewportClient interface
 
@@ -98,4 +101,7 @@ private:
 
 	// Called when TileBeingEditedIndex changes
 	FOnSingleTileIndexChanged OnSingleTileIndexChanged;
+	
+	bool HandleBeginTransform();
+	bool HandleEndTransform();
 };
diff --git a/Engine/Plugins/Animation/ControlRig/Source/ControlRigEditor/Private/EditMode/ControlRigEditMode.cpp b/Engine/Plugins/Animation/ControlRig/Source/ControlRigEditor/Private/EditMode/ControlRigEditMode.cpp
--- a/Engine/Plugins/Animation/ControlRig/Source/ControlRigEditor/Private/EditMode/ControlRigEditMode.cpp
+++ b/Engine/Plugins/Animation/ControlRig/Source/ControlRigEditor/Private/EditMode/ControlRigEditMode.cpp
@@ -1021,17 +1021,15 @@
 	return false;
 }
 
-bool FControlRigEditMode::EndTracking(FEditorViewportClient* InViewportClient, FViewport* InViewport)
+bool FControlRigEditMode::StartTracking(FEditorViewportClient* InViewportClient, FViewport* InViewport)
 {
 	if(RuntimeControlRigs.IsEmpty())
 	{
 		return false;
 	}
+	
+	//break right out if doing anim slider scrub
 
-	if (bisTrackingAnimToolDrag)
-	{
-		ResetAnimSlider();
-	}
 	if (IsDragAnimSliderToolPressed(InViewport))
 	{
 		return true;
@@ -1039,64 +1037,29 @@
 
 	if (IsMovingCamera(InViewport))
 	{
+		InViewportClient->SetCurrentWidgetAxis(EAxisList::None);
 		return true;
 	}
-	if (DragToolHandler.EndTracking(InViewportClient, InViewport))
+	if (IsDoingDrag(InViewport))
 	{
-		return true;
+		DragToolHandler.MakeDragTool(InViewportClient);
+		return DragToolHandler.StartTracking(InViewportClient, InViewport);
 	}
 
-	const bool bWasInteracting = bManipulatorMadeChange && InteractionType != (uint8)EControlRigInteractionType::None;
-	
-	InteractionType = (uint8)EControlRigInteractionType::None;
-	bIsTracking = false;
-	
-	if (InteractionScopes.Num() > 0)
-	{		
-		if (bManipulatorMadeChange)
-		{
-			bManipulatorMadeChange = false;
-			GEditor->EndTransaction();
-		}
-
-		for (TPair<UControlRig*, FControlRigInteractionScope*>& InteractionScope : InteractionScopes)
-		{
-			if (InteractionScope.Value)
-			{
-				delete InteractionScope.Value; 
-			}
-		}
-		InteractionScopes.Reset();
-
-		if (bWasInteracting && !AreEditingControlRigDirectly())
-		{
-			// We invalidate the hit proxies when in level editor to ensure that the gizmo's hit proxy is up to date.
-			// The invalidation is called here to avoid useless viewport update in the FControlRigEditMode::Tick
-			// function (that does an update when not in level editor)
-			TickManipulatableObjects(0.f);
-			
-			static constexpr bool bInvalidateChildViews = false;
-			static constexpr bool bInvalidateHitProxies = true;
-			InViewportClient->Invalidate(bInvalidateChildViews, bInvalidateHitProxies);
-		}
-		
-		return true;
-	}
-
-	bManipulatorMadeChange = false;
-	
-	return false;
+	return HandleBeginTransform(InViewportClient);
 }
 
-bool FControlRigEditMode::StartTracking(FEditorViewportClient* InViewportClient, FViewport* InViewport)
+bool FControlRigEditMode::EndTracking(FEditorViewportClient* InViewportClient, FViewport* InViewport)
 {
 	if(RuntimeControlRigs.IsEmpty())
 	{
 		return false;
 	}
-	
-	//break right out if doing anim slider scrub
 
+	if (bisTrackingAnimToolDrag)
+	{
+		ResetAnimSlider();
+	}
 	if (IsDragAnimSliderToolPressed(InViewport))
 	{
 		return true;
@@ -1104,15 +1067,33 @@
 
 	if (IsMovingCamera(InViewport))
 	{
-		InViewportClient->SetCurrentWidgetAxis(EAxisList::None);
 		return true;
 	}
-	if (IsDoingDrag(InViewport))
+	if (DragToolHandler.EndTracking(InViewportClient, InViewport))
 	{
-		DragToolHandler.MakeDragTool(InViewportClient);
-		return DragToolHandler.StartTracking(InViewportClient, InViewport);
+		return true;
 	}
 
+	return HandleEndTransform(InViewportClient);
+}
+
+bool FControlRigEditMode::BeginTransform(const FGizmoState& InState)
+{
+	return HandleBeginTransform(Owner->GetFocusedViewportClient());
+}
+
+bool FControlRigEditMode::EndTransform(const FGizmoState& InState)
+{
+	return HandleEndTransform(Owner->GetFocusedViewportClient());
+}
+
+bool FControlRigEditMode::HandleBeginTransform(const FEditorViewportClient* InViewportClient)
+{
+	if (!InViewportClient)
+	{
+		return false;
+	}
+	
 	InteractionType = GetInteractionType(InViewportClient);
 	bIsTracking = true;
 	
@@ -1202,6 +1183,55 @@
 	return InteractionScopes.Num() != 0;
 }
 
+bool FControlRigEditMode::HandleEndTransform(FEditorViewportClient* InViewportClient)
+{
+	if (!InViewportClient)
+	{
+		return false;
+	}
+	
+	const bool bWasInteracting = bManipulatorMadeChange && InteractionType != (uint8)EControlRigInteractionType::None;
+	
+	InteractionType = (uint8)EControlRigInteractionType::None;
+	bIsTracking = false;
+	
+	if (InteractionScopes.Num() > 0)
+	{		
+		if (bManipulatorMadeChange)
+		{
+			bManipulatorMadeChange = false;
+			GEditor->EndTransaction();
+		}
+
+		for (TPair<UControlRig*, FControlRigInteractionScope*>& InteractionScope : InteractionScopes)
+		{
+			if (InteractionScope.Value)
+			{
+				delete InteractionScope.Value; 
+			}
+		}
+		InteractionScopes.Reset();
+
+		if (bWasInteracting && !AreEditingControlRigDirectly())
+		{
+			// We invalidate the hit proxies when in level editor to ensure that the gizmo's hit proxy is up to date.
+			// The invalidation is called here to avoid useless viewport update in the FControlRigEditMode::Tick
+			// function (that does an update when not in level editor)
+			TickManipulatableObjects(0.f);
+			
+			static constexpr bool bInvalidateChildViews = false;
+			static constexpr bool bInvalidateHitProxies = true;
+			InViewportClient->Invalidate(bInvalidateChildViews, bInvalidateHitProxies);
+		}
+		
+		return true;
+	}
+
+	bManipulatorMadeChange = false;
+	
+	return false;
+}
+
 bool FControlRigEditMode::UsesTransformWidget() const
 {
 	for (const auto& Pairs : ControlRigShapeActors)
@@ -2830,28 +2860,33 @@
 	GetModeManager()->SetWidgetScale(PreviousGizmoScale);
 }
 
-uint8 FControlRigEditMode::GetInteractionType(FEditorViewportClient* InViewportClient)
+uint8 FControlRigEditMode::GetInteractionType(const FEditorViewportClient* InViewportClient)
 {
-	uint8 Result = (uint8)EControlRigInteractionType::None;
-	if(InViewportClient->IsMovingCamera())
+	EControlRigInteractionType Result = EControlRigInteractionType::None;
+	if (InViewportClient->IsMovingCamera())
 	{
-		return Result;
+		return static_cast<uint8>(Result);
 	}
 	
-	const UE::Widget::EWidgetMode WidgetMode = InViewportClient->GetWidgetMode();
-	if(WidgetMode == UE::Widget::WM_Translate || WidgetMode == UE::Widget::WM_TranslateRotateZ)
+	switch (InViewportClient->GetWidgetMode())
 	{
-		Result |= (uint8)EControlRigInteractionType::Translate;
+		case UE::Widget::WM_Translate:
+			EnumAddFlags(Result, EControlRigInteractionType::Translate);
+			break;
+		case UE::Widget::WM_TranslateRotateZ:
+			EnumAddFlags(Result, EControlRigInteractionType::Translate);
+			EnumAddFlags(Result, EControlRigInteractionType::Rotate);
+			break;
+		case UE::Widget::WM_Rotate:
+			EnumAddFlags(Result, EControlRigInteractionType::Rotate);
+			break;
+		case UE::Widget::WM_Scale:
+			EnumAddFlags(Result, EControlRigInteractionType::Scale);
+			break;
+		default:
+			break;
 	}
-	if(WidgetMode == UE::Widget::WM_Rotate || WidgetMode == UE::Widget::WM_TranslateRotateZ)
-	{
-		Result |= (uint8)EControlRigInteractionType::Rotate;
-	}
-	if(WidgetMode == UE::Widget::WM_Scale)
-	{
-		Result |= (uint8)EControlRigInteractionType::Scale;
-	}
-	return Result;	
+	return static_cast<uint8>(Result);
 }
 
 void FControlRigEditMode::ToggleControlShapeTransformEdit()
diff --git a/Engine/Plugins/Animation/ControlRig/Source/ControlRigEditor/Public/EditMode/ControlRigEditMode.h b/Engine/Plugins/Animation/ControlRig/Source/ControlRigEditor/Public/EditMode/ControlRigEditMode.h
--- a/Engine/Plugins/Animation/ControlRig/Source/ControlRigEditor/Public/EditMode/ControlRigEditMode.h
+++ b/Engine/Plugins/Animation/ControlRig/Source/ControlRigEditor/Public/EditMode/ControlRigEditMode.h
@@ -39,6 +39,7 @@ class IControlRigManipulatable;
 class ISequencer;
 enum class EControlRigSetKey : uint8;
 class UToolMenu;
+struct FGizmoState;
 
 DECLARE_DELEGATE_RetVal_ThreeParams(FTransform, FOnGetRigElementTransform, const FRigElementKey& /*RigElementKey*/, bool /*bLocal*/, bool /*bOnDebugInstance*/);
 DECLARE_DELEGATE_ThreeParams(FOnSetRigElementTransform, const FRigElementKey& /*RigElementKey*/, const FTransform& /*Transform*/, bool /*bLocal*/);
@@ -116,6 +117,8 @@ public:
 	virtual bool InputKey(FEditorViewportClient* InViewportClient, FViewport* InViewport, FKey InKey, EInputEvent InEvent) override;
 	virtual bool EndTracking(FEditorViewportClient* InViewportClient, FViewport* InViewport) override;
 	virtual bool StartTracking(FEditorViewportClient* InViewportClient, FViewport* InViewport) override;
+	virtual bool BeginTransform(const FGizmoState& InState) override;
+	virtual bool EndTransform(const FGizmoState& InState) override;
 	virtual bool ProcessCapturedMouseMoves(FEditorViewportClient* InViewportClient, FViewport* InViewport, const TArrayView<FIntPoint>& CapturedMouseMoves) override;
 	virtual bool HandleClick(FEditorViewportClient* InViewportClient, HHitProxy *HitProxy, const FViewportClick &Click) override;
 	virtual bool BoxSelect(FBox& InBox, bool InSelect = true) override;
@@ -234,6 +237,9 @@ protected:
 	/** If the Drag Anim Slider Tool is pressed*/
 	bool IsDragAnimSliderToolPressed(FViewport* InViewport);
 
+	bool HandleBeginTransform(const FEditorViewportClient* InViewportClient);
+	bool HandleEndTransform(FEditorViewportClient* InViewportClient);
+
 public:
 	
 	/** Clear Selection*/
@@ -518,7 +524,7 @@ private:
 	TArray<FRigElementKey> DeferredItemsToFrame;
 
 	/** Computes the current interaction types based on the widget mode */
-	static uint8 GetInteractionType(FEditorViewportClient* InViewportClient);
+	static uint8 GetInteractionType(const FEditorViewportClient* InViewportClient);
 	uint8 InteractionType;
 	bool bShowControlsAsOverlay;
 
diff --git a/Engine/Plugins/Animation/IKRig/Source/IKRigEditor/Private/RetargetEditor/IKRetargetDefaultMode.cpp b/Engine/Plugins/Animation/IKRig/Source/IKRigEditor/Private/RetargetEditor/IKRetargetDefaultMode.cpp
--- a/Engine/Plugins/Animation/IKRig/Source/IKRigEditor/Private/RetargetEditor/IKRetargetDefaultMode.cpp
+++ b/Engine/Plugins/Animation/IKRig/Source/IKRigEditor/Private/RetargetEditor/IKRetargetDefaultMode.cpp
@@ -329,6 +329,31 @@ bool FIKRetargetDefaultMode::HandleClick(FEditorViewportClient* InViewportClient
 
 bool FIKRetargetDefaultMode::StartTracking(FEditorViewportClient* InViewportClient, FViewport* InViewport)
 {
+	return HandleBeginTransform(InViewportClient);
+}
+
+bool FIKRetargetDefaultMode::EndTracking(FEditorViewportClient* InViewportClient, FViewport* InViewport)
+{
+	return HandleEndTransform();
+}
+
+bool FIKRetargetDefaultMode::BeginTransform(const FGizmoState& InState)
+{
+	return HandleBeginTransform(Owner->GetFocusedViewportClient());
+}
+
+bool FIKRetargetDefaultMode::EndTransform(const FGizmoState& InState)
+{
+	return HandleEndTransform();
+}
+
+bool FIKRetargetDefaultMode::HandleBeginTransform(const FEditorViewportClient* InViewportClient)
+{
+	if (!InViewportClient)
+	{
+		return false;
+	}
+	
 	bIsTranslating = false;
 
 	// not manipulating any widget axes, so stop tracking
@@ -356,7 +381,7 @@ bool FIKRetargetDefaultMode::StartTracking(FEditorViewportClient* InViewportClie
 	return false;
 }
 
-bool FIKRetargetDefaultMode::EndTracking(FEditorViewportClient* InViewportClient, FViewport* InViewport)
+bool FIKRetargetDefaultMode::HandleEndTransform()
 {
 	GEditor->EndTransaction();
 	bIsTranslating = false;
diff --git a/Engine/Plugins/Animation/IKRig/Source/IKRigEditor/Private/RetargetEditor/IKRetargetEditPoseMode.cpp b/Engine/Plugins/Animation/IKRig/Source/IKRigEditor/Private/RetargetEditor/IKRetargetEditPoseMode.cpp
--- a/Engine/Plugins/Animation/IKRig/Source/IKRigEditor/Private/RetargetEditor/IKRetargetEditPoseMode.cpp
+++ b/Engine/Plugins/Animation/IKRig/Source/IKRigEditor/Private/RetargetEditor/IKRetargetEditPoseMode.cpp
@@ -164,6 +164,31 @@ bool FIKRetargetEditPoseMode::HandleClick(FEditorViewportClient* InViewportClien
 
 bool FIKRetargetEditPoseMode::StartTracking(FEditorViewportClient* InViewportClient, FViewport* InViewport)
 {
+	return HandleBeginTransform(InViewportClient);
+}
+
+bool FIKRetargetEditPoseMode::EndTracking(FEditorViewportClient* InViewportClient, FViewport* InViewport)
+{
+	return HandleEndTransform(InViewportClient);
+}
+
+bool FIKRetargetEditPoseMode::BeginTransform(const FGizmoState& InState)
+{
+	return HandleBeginTransform(Owner->GetFocusedViewportClient());
+}
+
+bool FIKRetargetEditPoseMode::EndTransform(const FGizmoState& InState)
+{
+	return HandleEndTransform(Owner->GetFocusedViewportClient());
+}
+
+bool FIKRetargetEditPoseMode::HandleBeginTransform(const FEditorViewportClient* InViewportClient)
+{
+	if (!InViewportClient)
+	{
+		return false;
+	}
+	
 	TrackingState = FIKRetargetTrackingState::None;
 
 	// not manipulating any widget axes, so stop tracking
@@ -211,8 +236,13 @@ bool FIKRetargetEditPoseMode::StartTracking(FEditorViewportClient* InViewportCli
 	return false;
 }
 
-bool FIKRetargetEditPoseMode::EndTracking(FEditorViewportClient* InViewportClient, FViewport* InViewport)
-{	
+bool FIKRetargetEditPoseMode::HandleEndTransform(FEditorViewportClient* InViewportClient)
+{
+	if (!InViewportClient)
+	{
+		return false;
+	}
+	
 	if (TrackingState == FIKRetargetTrackingState::None)
 	{
 		const bool bIsRootSelected = IsRootSelected();
diff --git a/Engine/Plugins/Animation/IKRig/Source/IKRigEditor/Private/RigEditor/IKRigEditMode.cpp b/Engine/Plugins/Animation/IKRig/Source/IKRigEditor/Private/RigEditor/IKRigEditMode.cpp
--- a/Engine/Plugins/Animation/IKRig/Source/IKRigEditor/Private/RigEditor/IKRigEditMode.cpp
+++ b/Engine/Plugins/Animation/IKRig/Source/IKRigEditor/Private/RigEditor/IKRigEditMode.cpp
@@ -388,50 +388,18 @@ bool FIKRigEditMode::HandleClick(FEditorViewportClient* InViewportClient, HHitPr
 
 bool FIKRigEditMode::StartTracking(FEditorViewportClient* InViewportClient, FViewport* InViewport)
 {
-	const TSharedPtr<FIKRigEditorController> Controller = EditorController.Pin();
-	if (!Controller.IsValid())
-	{
-		return false; 
-	}
-	
-	TArray<FName> SelectedGoalNames;
-	Controller->GetSelectedGoalNames(SelectedGoalNames);
-	if (SelectedGoalNames.IsEmpty())
-	{
-		return false; // no goals selected to manipulate
-	}
-
 	const EAxisList::Type CurrentAxis = InViewportClient->GetCurrentWidgetAxis();
 	if (CurrentAxis == EAxisList::None)
 	{
 		return false; // not manipulating a required axis
 	}
 
-	GEditor->BeginTransaction(LOCTEXT("ManipulateGoal", "Manipulate IK Rig Goal"));
-	for (const FName& SelectedGoal : SelectedGoalNames)
-	{
-		Controller->AssetController->ModifyGoal(SelectedGoal);
-	}
-	Controller->bManipulatingGoals = true;
-	return true;
+	return HandleBeginTransform();
 }
 
 bool FIKRigEditMode::EndTracking(FEditorViewportClient* InViewportClient, FViewport* InViewport)
 {
-	const TSharedPtr<FIKRigEditorController> Controller = EditorController.Pin();
-	if (!Controller.IsValid())
-	{
-		return false; 
-	}
-	
-	if (!Controller->bManipulatingGoals)
-	{
-		return false; // not handled
-	}
-
-	GEditor->EndTransaction();
-	Controller->bManipulatingGoals = false;
-	return true;
+	return HandleEndTransform();
 }
 
 bool FIKRigEditMode::InputDelta(FEditorViewportClient* InViewportClient, FViewport* InViewport, FVector& InDrag, FRotator& InRot, FVector& InScale)
@@ -478,6 +446,58 @@ bool FIKRigEditMode::InputDelta(FEditorViewportClient* InViewportClient, FViewpo
 	return true;
 }
 
+bool FIKRigEditMode::BeginTransform(const FGizmoState& InState)
+{
+	return HandleBeginTransform();
+}
+
+bool FIKRigEditMode::EndTransform(const FGizmoState& InState)
+{
+	return HandleEndTransform();
+}
+
+bool FIKRigEditMode::HandleBeginTransform() const
+{
+	const TSharedPtr<FIKRigEditorController> Controller = EditorController.Pin();
+	if (!Controller.IsValid())
+	{
+		return false; 
+	}
+	
+	TArray<FName> SelectedGoalNames;
+	Controller->GetSelectedGoalNames(SelectedGoalNames);
+	if (SelectedGoalNames.IsEmpty())
+	{
+		return false; // no goals selected to manipulate
+	}
+
+	for (const FName& SelectedGoal : SelectedGoalNames)
+	{
+		(void)Controller->AssetController->ModifyGoal(SelectedGoal);
+	}
+	Controller->bManipulatingGoals = true;
+
+	return true;
+}
+
+bool FIKRigEditMode::HandleEndTransform() const
+{
+	const TSharedPtr<FIKRigEditorController> Controller = EditorController.Pin();
+	if (!Controller.IsValid())
+	{
+		return false; 
+	}
+	
+	if (!Controller->bManipulatingGoals)
+	{
+		return false; // not handled
+	}
+
+	Controller->bManipulatingGoals = false;
+	return true;
+}
+
+
 bool FIKRigEditMode::GetCustomDrawingCoordinateSystem(FMatrix& InMatrix, void* InData)
 {
 	const TSharedPtr<FIKRigEditorController> Controller = EditorController.Pin();
diff --git a/Engine/Plugins/Animation/IKRig/Source/IKRigEditor/Public/RetargetEditor/IKRetargetDefaultMode.h b/Engine/Plugins/Animation/IKRig/Source/IKRigEditor/Public/RetargetEditor/IKRetargetDefaultMode.h
--- a/Engine/Plugins/Animation/IKRig/Source/IKRigEditor/Public/RetargetEditor/IKRetargetDefaultMode.h
+++ b/Engine/Plugins/Animation/IKRig/Source/IKRigEditor/Public/RetargetEditor/IKRetargetDefaultMode.h
@@ -12,7 +12,7 @@ class UIKRigProcessor;
 class FIKRetargetEditorController;
 class FIKRetargetEditor;
 class FIKRetargetPreviewScene;
-
+struct FGizmoState;
 
 class FIKRetargetDefaultMode : public IPersonaEditMode
 {
@@ -42,6 +42,8 @@ public:
 	virtual bool HandleClick(FEditorViewportClient* InViewportClient, HHitProxy* HitProxy, const FViewportClick& Click) override;
 	virtual bool StartTracking(FEditorViewportClient* InViewportClient, FViewport* InViewport) override;
 	virtual bool EndTracking(FEditorViewportClient* InViewportClient, FViewport* InViewport) override;
+	virtual bool BeginTransform(const FGizmoState& InState) override;
+	virtual bool EndTransform(const FGizmoState& InState) override;
 	virtual bool InputDelta(FEditorViewportClient* InViewportClient, FViewport* InViewport, FVector& InDrag, FRotator& InRot, FVector& InScale) override;
 	virtual bool GetCustomDrawingCoordinateSystem(FMatrix& InMatrix, void* InData) override;
 	virtual bool GetCustomInputCoordinateSystem(FMatrix& InMatrix, void* InData) override;
@@ -56,6 +58,9 @@ private:
 	void RenderDebugProxies(FPrimitiveDrawInterface* PDI, const FIKRetargetEditorController* Controller) const;
 	static void ApplyOffsetToMeshTransform(const FVector& Offset, USceneComponent* Component);
 
+	bool HandleBeginTransform(const FEditorViewportClient* InViewportClient);
+	bool HandleEndTransform();
+	
 	// the skeleton currently being edited
 	UDebugSkelMeshComponent* GetCurrentlyEditedMesh() const;
 	ERetargetSourceOrTarget SkeletonMode;
diff --git a/Engine/Plugins/Animation/IKRig/Source/IKRigEditor/Public/RetargetEditor/IKRetargetEditPoseMode.h b/Engine/Plugins/Animation/IKRig/Source/IKRigEditor/Public/RetargetEditor/IKRetargetEditPoseMode.h
--- a/Engine/Plugins/Animation/IKRig/Source/IKRigEditor/Public/RetargetEditor/IKRetargetEditPoseMode.h
+++ b/Engine/Plugins/Animation/IKRig/Source/IKRigEditor/Public/RetargetEditor/IKRetargetEditPoseMode.h
@@ -13,6 +13,7 @@ class UIKRigProcessor;
 class FIKRetargetEditorController;
 class FIKRetargetEditor;
 class FIKRetargetPreviewScene;
+struct FGizmoState;
 
 enum class FIKRetargetTrackingState : int8
 {
@@ -70,6 +71,8 @@ public:
 	virtual bool HandleClick(FEditorViewportClient* InViewportClient, HHitProxy* HitProxy, const FViewportClick& Click) override;
 	virtual bool StartTracking(FEditorViewportClient* InViewportClient, FViewport* InViewport) override;
 	virtual bool EndTracking(FEditorViewportClient* InViewportClient, FViewport* InViewport) override;
+	virtual bool BeginTransform(const FGizmoState& InState) override;
+	virtual bool EndTransform(const FGizmoState& InState) override;
 	virtual bool InputDelta(FEditorViewportClient* InViewportClient, FViewport* InViewport, FVector& InDrag, FRotator& InRot, FVector& InScale) override;
 	virtual bool GetCustomDrawingCoordinateSystem(FMatrix& InMatrix, void* InData) override;
 	virtual bool GetCustomInputCoordinateSystem(FMatrix& InMatrix, void* InData) override;
@@ -88,6 +91,9 @@ private:
 	bool IsOnlyRootSelected() const;
 	void UpdateWidgetTransform();
 
+	bool HandleBeginTransform(const FEditorViewportClient* InViewportClient);
+	bool HandleEndTransform(FEditorViewportClient* InViewportClient);
+
 	// the bone(s) currently being edited
 	FBoneEdit BoneEdit;
 
diff --git a/Engine/Plugins/Animation/IKRig/Source/IKRigEditor/Public/RigEditor/IKRigEditMode.h b/Engine/Plugins/Animation/IKRig/Source/IKRigEditor/Public/RigEditor/IKRigEditMode.h
--- a/Engine/Plugins/Animation/IKRig/Source/IKRigEditor/Public/RigEditor/IKRigEditMode.h
+++ b/Engine/Plugins/Animation/IKRig/Source/IKRigEditor/Public/RigEditor/IKRigEditMode.h
@@ -13,6 +13,7 @@ class UIKRigEffectorGoal;
 class FIKRigEditorToolkit;
 class FIKRigPreviewScene;
 class UIKRigProcessor;
+struct FGizmoState;
 
 class FIKRigEditMode : public IPersonaEditMode
 {
@@ -48,7 +49,9 @@ public:
 	virtual bool InputDelta(FEditorViewportClient* InViewportClient, FViewport* InViewport, FVector& InDrag, FRotator& InRot, FVector& InScale) override;
 	virtual bool GetCustomDrawingCoordinateSystem(FMatrix& InMatrix, void* InData) override;
 	virtual bool GetCustomInputCoordinateSystem(FMatrix& InMatrix, void* InData) override;
-	virtual bool InputKey(FEditorViewportClient* ViewportClient,FViewport* Viewport, FKey Key,EInputEvent Event);
+	virtual bool InputKey(FEditorViewportClient* ViewportClient,FViewport* Viewport, FKey Key,EInputEvent Event) override;
+	virtual bool BeginTransform(const FGizmoState& InState) override;
+	virtual bool EndTransform(const FGizmoState& InState) override;
 	/** END FEdMode interface */
 
 private:
@@ -60,6 +63,9 @@ private:
 		const UIKRigProcessor* Processor,
 		const FReferenceSkeleton& RefSkeleton,
 		TArray<FLinearColor>& OutBoneColors) const;
+
+	bool HandleBeginTransform() const;
+	bool HandleEndTransform() const;
 	
 	/** The hosting app */
 	TWeakPtr<FIKRigEditorController> EditorController;
diff --git a/Engine/Plugins/Experimental/Mutable/Source/CustomizableObjectEditor/Private/MuCOE/CustomizableObjectEditorViewportClient.cpp b/Engine/Plugins/Experimental/Mutable/Source/CustomizableObjectEditor/Private/MuCOE/CustomizableObjectEditorViewportClient.cpp
--- a/Engine/Plugins/Experimental/Mutable/Source/CustomizableObjectEditor/Private/MuCOE/CustomizableObjectEditorViewportClient.cpp
+++ b/Engine/Plugins/Experimental/Mutable/Source/CustomizableObjectEditor/Private/MuCOE/CustomizableObjectEditorViewportClient.cpp
@@ -1219,18 +1219,37 @@ bool FCustomizableObjectEditorViewportClient::InputWidgetDelta(FViewport* InView
 
 
 void FCustomizableObjectEditorViewportClient::TrackingStarted(const FInputEventState& InInputState, bool bIsDraggingWidget, bool bNudge)
+{
+	if (!bIsDraggingWidget || !InInputState.IsLeftMouseButtonPressed() || (Widget->GetCurrentAxis() & EAxisList::All) == 0)
+	{
+		return;
+	}
+	
+	(void)HandleBeginTransform();
+}
+
+
+void FCustomizableObjectEditorViewportClient::TrackingStopped()
+{
+	(void)HandleEndTransform();
+}
+
+bool FCustomizableObjectEditorViewportClient::BeginTransform(const FGizmoState& InState)
+{
+	return HandleBeginTransform();
+}
+
+bool FCustomizableObjectEditorViewportClient::EndTransform(const FGizmoState& InState)
+{
+	return HandleEndTransform();
+}
+
+bool FCustomizableObjectEditorViewportClient::HandleBeginTransform()
 {
 	switch (WidgetType)
 	{
 	case EWidgetType::Projector:
 		{
-			if (!bIsDraggingWidget ||
-				!InInputState.IsLeftMouseButtonPressed() ||
-				(Widget->GetCurrentAxis() & EAxisList::All) == 0)
-			{
-				return;
-			}
-
 			bManipulating = true;
 
 			const UE::Widget::EWidgetMode WidgetMode = GetWidgetMode();
@@ -1248,7 +1267,8 @@ void FCustomizableObjectEditorViewportClient::TrackingStarted(const FInputEventS
 				GEditor->BeginTransaction(LOCTEXT("CustomizableObjectEditor_ScaleProjector", "Scale Projector"));
 			}
 
-			WidgetTrackingStartedDelegate.ExecuteIfBound();
+			(void)WidgetTrackingStartedDelegate.ExecuteIfBound();
+			return true;
 		}
 
 	// The following cases are missing Undo/Redo functionality MTBL-391.
@@ -1256,14 +1276,15 @@ void FCustomizableObjectEditorViewportClient::TrackingStarted(const FInputEventS
 	case EWidgetType::ClipMesh:
 	case EWidgetType::Light:
 	case EWidgetType::Hidden:
+		return true;
 		break;
 	default:
 		unimplemented();
 	}
+	return false;
 }
 
-
-void FCustomizableObjectEditorViewportClient::TrackingStopped()
+bool FCustomizableObjectEditorViewportClient::HandleEndTransform()
 {
 	switch (WidgetType)
 	{
@@ -1272,19 +1293,21 @@ void FCustomizableObjectEditorViewportClient::TrackingStopped()
 		{
 			bManipulating = false;
 			GEditor->EndTransaction();
+			return true;
 		}
 
 	case EWidgetType::Hidden:
 	case EWidgetType::ClipMorph:
 	case EWidgetType::ClipMesh:
 	case EWidgetType::Light:
+		return true;
 		break;
 	default:
 		unimplemented();
 	}
+	return false;
 }
 
-
 FVector FCustomizableObjectEditorViewportClient::GetWidgetLocation() const
 {
 	switch (WidgetType)
diff --git a/Engine/Plugins/Experimental/Mutable/Source/CustomizableObjectEditor/Public/MuCOE/CustomizableObjectEditorViewportClient.h b/Engine/Plugins/Experimental/Mutable/Source/CustomizableObjectEditor/Public/MuCOE/CustomizableObjectEditorViewportClient.h
--- a/Engine/Plugins/Experimental/Mutable/Source/CustomizableObjectEditor/Public/MuCOE/CustomizableObjectEditorViewportClient.h
+++ b/Engine/Plugins/Experimental/Mutable/Source/CustomizableObjectEditor/Public/MuCOE/CustomizableObjectEditorViewportClient.h
@@ -28,7 +28,7 @@ class UObject;
 class UStaticMesh;
 class UStaticMeshComponent;
 struct FInputKeyEventArgs;
-
+struct FGizmoState;
 
 DECLARE_DELEGATE_RetVal(FVector, FWidgetLocationDelegate);
 DECLARE_DELEGATE_OneParam(FOnWidgetLocationChangedDelegate, const FVector&)
@@ -85,6 +85,8 @@ public:
 	virtual bool InputWidgetDelta(FViewport* Viewport, EAxisList::Type CurrentAxis, FVector& Drag, FRotator& Rot, FVector& Scale) override;
 	virtual void TrackingStarted(const FInputEventState& InInputState, bool bIsDragging, bool bNudge) override;
 	virtual void TrackingStopped() override;
+	virtual bool BeginTransform(const FGizmoState& InState) override;
+	virtual bool EndTransform(const FGizmoState& InState) override;
 	virtual bool CanSetWidgetMode(UE::Widget::EWidgetMode NewMode) const override;
 	virtual FVector GetWidgetLocation() const override;
 	virtual FMatrix GetWidgetCoordSystem() const override;
@@ -272,6 +274,9 @@ private:
 	void DrawMeshBones(UDebugSkelMeshComponent* MeshComponent, FPrimitiveDrawInterface* PDI);
 
 	void SetWidgetType(EWidgetType Type);
+
+	bool HandleBeginTransform();
+	bool HandleEndTransform();
 	
 	/** Component for the static/skeletal mesh. */
 	TWeakObjectPtr<UStaticMeshComponent> StaticMeshComponent;
diff --git a/Engine/Plugins/Runtime/nDisplay/Source/DisplayClusterConfigurator/Private/Views/Viewport/DisplayClusterConfiguratorSCSEditorViewportClient.cpp b/Engine/Plugins/Runtime/nDisplay/Source/DisplayClusterConfigurator/Private/Views/Viewport/DisplayClusterConfiguratorSCSEditorViewportClient.cpp
--- a/Engine/Plugins/Runtime/nDisplay/Source/DisplayClusterConfigurator/Private/Views/Viewport/DisplayClusterConfiguratorSCSEditorViewportClient.cpp
+++ b/Engine/Plugins/Runtime/nDisplay/Source/DisplayClusterConfigurator/Private/Views/Viewport/DisplayClusterConfiguratorSCSEditorViewportClient.cpp
@@ -710,6 +710,40 @@ void FDisplayClusterConfiguratorSCSEditorViewportClient::TrackingStopped()
 	}
 }
 
+bool FDisplayClusterConfiguratorSCSEditorViewportClient::BeginTransform(const FGizmoState& InState)
+{
+	// Suspend component modification during each delta step to avoid recording unnecessary overhead into the transaction buffer
+	GEditor->DisableDeltaModification(true);
+
+	// Begin transaction
+	BeginTransaction(NSLOCTEXT("UnrealEd", "ModifyComponents", "Modify Component(s)"));
+	bIsManipulating = true;
+	return true;
+}
+
+bool FDisplayClusterConfiguratorSCSEditorViewportClient::EndTransform(const FGizmoState& InState)
+{
+	// Re-run construction scripts if we haven't done so yet (so that the components in the preview actor can update their transforms)
+	AActor* PreviewActor = GetPreviewActor();
+	if (PreviewActor != nullptr)
+	{
+		UBlueprint* PreviewBlueprint = UBlueprint::GetBlueprintFromClass(PreviewActor->GetClass());
+		if (PreviewBlueprint != nullptr && !PreviewBlueprint->bRunConstructionScriptOnDrag)
+		{
+			PreviewActor->RerunConstructionScripts();
+		}
+	}
+
+	// End transaction
+	bIsManipulating = false;
+	EndTransaction();
+
+	// Restore component delta modification
+	GEditor->DisableDeltaModification(false);
+	
+	return true;
+}
+
 UE::Widget::EWidgetMode FDisplayClusterConfiguratorSCSEditorViewportClient::GetWidgetMode() const
 {
 	// Default to not drawing the widget
diff --git a/Engine/Plugins/Runtime/nDisplay/Source/DisplayClusterConfigurator/Private/Views/Viewport/DisplayClusterConfiguratorSCSEditorViewportClient.h b/Engine/Plugins/Runtime/nDisplay/Source/DisplayClusterConfigurator/Private/Views/Viewport/DisplayClusterConfiguratorSCSEditorViewportClient.h
--- a/Engine/Plugins/Runtime/nDisplay/Source/DisplayClusterConfigurator/Private/Views/Viewport/DisplayClusterConfiguratorSCSEditorViewportClient.h
+++ b/Engine/Plugins/Runtime/nDisplay/Source/DisplayClusterConfigurator/Private/Views/Viewport/DisplayClusterConfiguratorSCSEditorViewportClient.h
@@ -11,6 +11,7 @@ class UPostProcessComponent;
 class UAssetViewerSettings;
 class UDisplayClusterWorldOriginComponent;
 class SDisplayClusterConfiguratorSCSEditorViewport;
+struct FGizmoState;
 
 class FDisplayClusterConfiguratorSCSEditorViewportClient : public FEditorViewportClient, public TSharedFromThis<FDisplayClusterConfiguratorSCSEditorViewportClient>
 {
@@ -38,6 +39,8 @@ public:
 	virtual bool InputWidgetDelta(FViewport* Viewport, EAxisList::Type CurrentAxis, FVector& Drag, FRotator& Rot, FVector& Scale) override;
 	virtual void TrackingStarted(const struct FInputEventState& InInputState, bool bIsDragging, bool bNudge) override;
 	virtual void TrackingStopped() override;
+	virtual bool BeginTransform(const FGizmoState& InState) override;
+	virtual bool EndTransform(const FGizmoState& InState) override;
 	virtual UE::Widget::EWidgetMode GetWidgetMode() const override;
 	virtual void SetWidgetMode(UE::Widget::EWidgetMode NewMode) override;
 	virtual void SetWidgetCoordSystemSpace(ECoordSystem NewCoordSystem) override;
diff --git a/Engine/Source/Editor/Experimental/EditorInteractiveToolsFramework/Private/EditorGizmos/EditorGizmoStateTarget.cpp b/Engine/Source/Editor/Experimental/EditorInteractiveToolsFramework/Private/EditorGizmos/EditorGizmoStateTarget.cpp
new file mode 100644
--- /dev/null
+++ b/Engine/Source/Editor/Experimental/EditorInteractiveToolsFramework/Private/EditorGizmos/EditorGizmoStateTarget.cpp
@@ -0,0 +1,60 @@
+﻿// Copyright Epic Games, Inc. All Rights Reserved.
+
+#include "EditorGizmos/EditorGizmoStateTarget.h"
+#include "EditorGizmos/TransformGizmo.h"
+#include "EditorModeManager.h"
+#include "GizmoEdModeInterface.h"
+
+void UEditorGizmoStateTarget::BeginUpdate()
+{
+	const TSharedPtr<FEditorModeTools> ModeTools = WeakModeTools.IsValid() ? WeakModeTools.Pin() : nullptr;
+	if (ensure(ModeTools.IsValid()))
+	{
+		if (TransactionManager)
+		{
+			TransactionManager->BeginUndoTransaction(TransactionDescription);
+		}
+
+		// empty state for now
+		constexpr FGizmoState State;
+		(void)ModeTools->BeginTransform(State);
+	}
+}
+
+void UEditorGizmoStateTarget::EndUpdate()
+{
+	const TSharedPtr<FEditorModeTools> ModeTools = WeakModeTools.IsValid() ? WeakModeTools.Pin() : nullptr;
+	if (ensure(ModeTools))
+	{
+		// empty state for now
+		constexpr FGizmoState State;
+		(void)ModeTools->EndTransform(State);
+
+		if (TransactionManager)
+		{
+			TransactionManager->EndUndoTransaction();
+		}
+	}
+}
+
+UEditorGizmoStateTarget* UEditorGizmoStateTarget::Construct(
+	FEditorModeTools* InModeManager,
+	const FText& InDescription,
+	IToolContextTransactionProvider* TransactionManagerIn,
+	UObject* Outer)
+{
+	UEditorGizmoStateTarget* NewTarget = NewObject<UEditorGizmoStateTarget>(Outer);
+	NewTarget->WeakModeTools = InModeManager->AsShared();
+	NewTarget->TransactionDescription = InDescription;
+
+	// have to explicitly configure this because we only have IToolContextTransactionProvider pointer
+	NewTarget->TransactionManager.SetInterface(TransactionManagerIn);
+	NewTarget->TransactionManager.SetObject(CastChecked<UObject>(TransactionManagerIn));
+		
+	return NewTarget;
+}
+
+void UEditorGizmoStateTarget::SetTransformGizmo(UTransformGizmo* InGizmo)
+{
+	TransformGizmo = InGizmo;
+}
\ No newline at end of file
diff --git a/Engine/Source/Editor/Experimental/EditorInteractiveToolsFramework/Private/EditorGizmos/EditorTransformGizmoBuilder.cpp b/Engine/Source/Editor/Experimental/EditorInteractiveToolsFramework/Private/EditorGizmos/EditorTransformGizmoBuilder.cpp
--- a/Engine/Source/Editor/Experimental/EditorInteractiveToolsFramework/Private/EditorGizmos/EditorTransformGizmoBuilder.cpp
+++ b/Engine/Source/Editor/Experimental/EditorInteractiveToolsFramework/Private/EditorGizmos/EditorTransformGizmoBuilder.cpp
@@ -6,6 +6,7 @@
 #include "BaseGizmos/GizmoElementHitTargets.h"
 #include "BaseGizmos/GizmoViewContext.h"
 #include "ContextObjectStore.h"
+#include "EditorModeManager.h"
 #include "EdModeInteractiveToolsContext.h"
 #include "EditorGizmos/EditorTransformGizmo.h"
 #include "EditorGizmos/EditorTransformGizmoSource.h"
@@ -43,8 +44,11 @@ void UEditorTransformGizmoBuilder::UpdateGizmoForSelection(UInteractiveGizmo* Gi
 {
 	if (UTransformGizmo* TransformGizmo = Cast<UTransformGizmo>(Gizmo))
 	{
-		UEditorTransformProxy* TransformProxy = UEditorTransformProxy::CreateNew(GetTransformGizmoContext(SceneState));
-		TransformGizmo->SetActiveTarget(TransformProxy);
+		const UEditorTransformGizmoContextObject* GizmoContextObject = GetTransformGizmoContext(SceneState);
+		FEditorModeTools* ModeTools = GizmoContextObject ? GizmoContextObject->GetModeTools() : nullptr;
+		ensure(ModeTools);
+		UEditorTransformProxy* TransformProxy = UEditorTransformProxy::CreateNew(GizmoContextObject);
+		TransformGizmo->SetActiveTarget(TransformProxy, nullptr, ModeTools->GetGizmoStateTarget());
 		TransformGizmo->SetVisibility(true);
 		
 		if (UGizmoElementHitMultiTarget* HitMultiTarget = Cast< UGizmoElementHitMultiTarget>(TransformGizmo->HitTarget))
diff --git a/Engine/Source/Editor/Experimental/EditorInteractiveToolsFramework/Private/EditorGizmos/EditorTransformGizmoDataBinder.cpp b/Engine/Source/Editor/Experimental/EditorInteractiveToolsFramework/Private/EditorGizmos/EditorTransformGizmoDataBinder.cpp
--- a/Engine/Source/Editor/Experimental/EditorInteractiveToolsFramework/Private/EditorGizmos/EditorTransformGizmoDataBinder.cpp
+++ b/Engine/Source/Editor/Experimental/EditorInteractiveToolsFramework/Private/EditorGizmos/EditorTransformGizmoDataBinder.cpp
@@ -6,6 +6,7 @@
 #include "EditorModeManager.h"
 #include "EditorGizmos/EditorTransformGizmoUtil.h"
 #include "EditorGizmos/TransformGizmo.h"
+#include "EditorViewportClient.h"
 
 FEditorTransformGizmoDataBinder::~FEditorTransformGizmoDataBinder()
 {
@@ -56,20 +57,38 @@ void FEditorTransformGizmoDataBinder::BindToUninitializedGizmo(UTransformGizmo*
 	if ( ensure(!bSetBound && !bClearBound))
 	{
 		InGizmo->OnSetActiveTarget.AddSP(this, &FEditorTransformGizmoDataBinder::BindToInitializedGizmo);
+		if (InGizmo->ActiveTarget)
+		{
+			BindToInitializedGizmo(InGizmo, InGizmo->ActiveTarget);
+		}
 		InGizmo->OnAboutToClearActiveTarget.AddSP(this, &FEditorTransformGizmoDataBinder::UnbindFromGizmo);
 	}
 }
 
 void FEditorTransformGizmoDataBinder::BindToInitializedGizmo(UTransformGizmo* InGizmo, UTransformProxy* InProxy)
 {
-	const bool bHasTarget = InGizmo && InGizmo->ActiveTarget && InGizmo->ActiveTarget == InProxy; 
-	if (ensure(!bHasTarget && !BoundGizmos.Contains(InGizmo)))
+	const bool bHaveValidInput = InGizmo && InGizmo->ActiveTarget && InGizmo->ActiveTarget == InProxy;
+	if (!ensure(bHaveValidInput && !BoundGizmos.Contains(InGizmo)))
 	{
 		return;
 	}
 
 	// bind InProxy OnBeginTransformEdit/OnTransformChanged/OnEndTransformEdit etc here if needed
+	if (!InProxy->OnBeginTransformEdit.IsBoundToObject(this))
+	{
+		InProxy->OnBeginTransformEdit.AddSP(this, &FEditorTransformGizmoDataBinder::OnProxyBeginTransformEdit);
+	}
 
+	if (!InProxy->OnTransformChanged.IsBoundToObject(this))
+	{
+		InProxy->OnTransformChanged.AddSP(this, &FEditorTransformGizmoDataBinder::OnProxyTransformChanged);
+	}
+	
+	if (!InProxy->OnEndTransformEdit.IsBoundToObject(this))
+	{
+		InProxy->OnEndTransformEdit.AddSP(this, &FEditorTransformGizmoDataBinder::OnProxyEndTransformEdit);
+	}
+	
 	BoundGizmos.Add(InGizmo);
 }
 
@@ -87,6 +106,33 @@ void FEditorTransformGizmoDataBinder::UnbindFromGizmo(UTransformGizmo* InGizmo,
 	}
 
 	// unbind InProxy OnBeginTransformEdit/OnTransformChanged/OnEndTransformEdit etc here if needed
-
+	InProxy->OnBeginTransformEdit.RemoveAll(this);
+	InProxy->OnEndTransformEdit.RemoveAll(this);
+	
 	BoundGizmos.Remove(InGizmo);
 }
+
+FEditorViewportClient* FEditorTransformGizmoDataBinder::GetViewportClient() const
+{
+	if (WeakContext.IsValid() && WeakContext->GetModeTools())
+	{
+		return WeakContext->GetModeTools()->GetFocusedViewportClient();	
+	}
+	return GLevelEditorModeTools().GetFocusedViewportClient();
+}
+
+void FEditorTransformGizmoDataBinder::OnProxyBeginTransformEdit(UTransformProxy* InTransformProxy)
+{
+	bHasTransformChanged = false;
+}
+
+void FEditorTransformGizmoDataBinder::OnProxyTransformChanged(UTransformProxy* InTransformProxy, FTransform InTransform)
+{
+	// TODO track if we actually did drag something to fallback to the viewport if we didn't
+	bHasTransformChanged = true;
+}
+
+void FEditorTransformGizmoDataBinder::OnProxyEndTransformEdit(UTransformProxy* InTransformProxy)
+{
+	bHasTransformChanged = false;
+}
\ No newline at end of file
diff --git a/Engine/Source/Editor/Experimental/EditorInteractiveToolsFramework/Public/EditorGizmos/EditorGizmoStateTarget.h b/Engine/Source/Editor/Experimental/EditorInteractiveToolsFramework/Public/EditorGizmos/EditorGizmoStateTarget.h
new file mode 100644
--- /dev/null
+++ b/Engine/Source/Editor/Experimental/EditorInteractiveToolsFramework/Public/EditorGizmos/EditorGizmoStateTarget.h
@@ -0,0 +1,79 @@
+﻿// Copyright Epic Games, Inc. All Rights Reserved.
+
+#pragma once
+
+#include "BaseGizmos/StateTargets.h"
+
+#include "Templates/Function.h"
+#include "Internationalization/Text.h"
+
+#include "UObject/Object.h"
+#include "UObject/WeakObjectPtr.h"
+#include "UObject/WeakObjectPtrTemplates.h"
+#include "UObject/Package.h"
+#include "UObject/UObjectGlobals.h"
+
+#include "EditorGizmoStateTarget.generated.h"
+
+class FEditorModeTools;
+class IToolContextTransactionProvider;
+class UTransformGizmo;
+
+/**
+ * UEditorGizmoStateTarget
+ */
+UCLASS(MinimalAPI)
+class UEditorGizmoStateTarget : public UObject, public IGizmoStateTarget
+{
+	GENERATED_BODY()
+public:
+
+	/**
+	 * BeginUpdate is called before the gizmo starts modifying the transform.  
+	 */
+	virtual void BeginUpdate() override;
+
+	/**
+	 * EndUpdate is called after the gizmo finished modifying the transform.
+	 */
+	virtual void EndUpdate() override;
+
+	/**
+	 * TWeakPtr to mode manager to interface with
+	 */
+	TWeakPtr<FEditorModeTools> WeakModeTools;
+
+	/**
+	* Localized text description of the transaction (used if TransactionManager is not null) 
+	*/
+	FText TransactionDescription;
+	
+	/**
+	 * Pointer to the GizmoManager or ToolManager that is used to open/close the transaction (if no Begin/End functions is provided by BeginUpdate)
+	 */
+	UPROPERTY()
+	TScriptInterface<IToolContextTransactionProvider> TransactionManager;
+
+	/**
+	 * Create and initialize an standard instance of UGizmoEditorStateTarget
+	 * @param InModeManager the context object this StateTarget will retrieve the mode manager from
+	 * @param InDescription Localized text description of the transaction (if no Begin/End functions is provided by BeginUpdate)
+	 * @param TransactionManagerIn pointer to the object to manage transactions (if no Begin/End functions is provided by BeginUpdate)
+	 */
+	EDITORINTERACTIVETOOLSFRAMEWORK_API static UEditorGizmoStateTarget* Construct(
+		FEditorModeTools* InModeManager,
+		const FText& InDescription,
+		IToolContextTransactionProvider* TransactionManagerIn,
+		UObject* Outer = (UObject*)GetTransientPackage());
+
+	/**
+	 * Sets the transform gizmo from which we can get data if required.
+	 */
+	void SetTransformGizmo(UTransformGizmo* InGizmo);
+	
+protected:
+	/**
+	 * The transform gizmo from which we can get data if required.
+	 */
+	TWeakObjectPtr<UTransformGizmo> TransformGizmo;
+};
diff --git a/Engine/Source/Editor/Experimental/EditorInteractiveToolsFramework/Public/EditorGizmos/EditorTransformGizmoDataBinder.h b/Engine/Source/Editor/Experimental/EditorInteractiveToolsFramework/Public/EditorGizmos/EditorTransformGizmoDataBinder.h
--- a/Engine/Source/Editor/Experimental/EditorInteractiveToolsFramework/Public/EditorGizmos/EditorTransformGizmoDataBinder.h
+++ b/Engine/Source/Editor/Experimental/EditorInteractiveToolsFramework/Public/EditorGizmos/EditorTransformGizmoDataBinder.h
@@ -7,9 +7,11 @@
 #include "UObject/WeakObjectPtrTemplates.h"
 #include "Containers/Set.h"
 
+class UGizmoEditorStateTarget;
 class UEditorTransformGizmoContextObject;
 class UTransformGizmo;
 class UTransformProxy;
+class FEditorViewportClient;
 
 /**
  * FEditorTransformGizmoDataBinder is a helper class for binding a UTransformGizmo to a FEditorModeTools
@@ -46,4 +48,13 @@ private:
 
 	/** Weak ptr to the context so we can interface with the current mode manager */
 	TWeakObjectPtr<UEditorTransformGizmoContextObject> WeakContext;
-};
\ No newline at end of file
+
+	FEditorViewportClient* GetViewportClient() const;
+
+	void OnProxyBeginTransformEdit(UTransformProxy* InTransformProxy);
+	void OnProxyTransformChanged(UTransformProxy* InTransformProxy, FTransform InTransform);
+	void OnProxyEndTransformEdit(UTransformProxy* InTransformProxy);
+
+	bool bHasTransformChanged = false;
+};
+
diff --git a/Engine/Source/Editor/Experimental/EditorInteractiveToolsFramework/Public/GizmoEdModeInterface.h b/Engine/Source/Editor/Experimental/EditorInteractiveToolsFramework/Public/GizmoEdModeInterface.h
new file mode 100644
--- /dev/null
+++ b/Engine/Source/Editor/Experimental/EditorInteractiveToolsFramework/Public/GizmoEdModeInterface.h
@@ -0,0 +1,35 @@
+﻿// Copyright Epic Games, Inc. All Rights Reserved.
+
+#pragma once
+
+#include "UObject/Interface.h"
+
+#include "GizmoEdModeInterface.generated.h"
+
+/**
+ * Gizmo state structure used to pass gizmo data if needed
+ */
+
+struct FGizmoState
+{
+};
+
+/**
+ * Interface for the new editor TRS gizmos
+ */
+
+UINTERFACE(NotBlueprintable, MinimalAPI)
+class UGizmoEdModeInterface : public UInterface
+{
+	GENERATED_BODY()
+};
+
+class EDITORINTERACTIVETOOLSFRAMEWORK_API IGizmoEdModeInterface
+{
+	GENERATED_BODY()
+
+public:
+	
+	virtual bool BeginTransform(const FGizmoState& InState) = 0;
+	virtual bool EndTransform(const FGizmoState& InState) = 0;
+};
\ No newline at end of file
diff --git a/Engine/Source/Editor/Kismet/Private/SCSEditorViewportClient.cpp b/Engine/Source/Editor/Kismet/Private/SCSEditorViewportClient.cpp
--- a/Engine/Source/Editor/Kismet/Private/SCSEditorViewportClient.cpp
+++ b/Engine/Source/Editor/Kismet/Private/SCSEditorViewportClient.cpp
@@ -711,6 +711,32 @@ bool FSCSEditorViewportClient::InputWidgetDelta( FViewport* InViewport, EAxisLis
 void FSCSEditorViewportClient::TrackingStarted( const struct FInputEventState& InInputState, bool bIsDraggingWidget, bool bNudge )
 {
 	if( !bIsManipulating && bIsDraggingWidget )
+	{
+		HandleBeginTransform();
+	}
+}
+
+void FSCSEditorViewportClient::TrackingStopped() 
+{
+	if( bIsManipulating )
+	{
+		HandleEndTransform();
+	}
+}
+
+bool FSCSEditorViewportClient::BeginTransform(const FGizmoState& InState)
+{
+	return HandleBeginTransform();
+}
+
+bool FSCSEditorViewportClient::EndTransform(const FGizmoState& InState)
+{
+	return HandleEndTransform();
+}
+
+bool FSCSEditorViewportClient::HandleBeginTransform()
+{
+	if (!bIsManipulating)
 	{
 		// Suspend component modification during each delta step to avoid recording unnecessary overhead into the transaction buffer
 		GEditor->DisableDeltaModification(true);
@@ -718,12 +744,15 @@ void FSCSEditorViewportClient::TrackingStarted( const struct FInputEventState& I
 		// Begin transaction
 		BeginTransaction( NSLOCTEXT("UnrealEd", "ModifyComponents", "Modify Component(s)") );
 		bIsManipulating = true;
+		return true;
 	}
+	
+	return false;
 }
 
-void FSCSEditorViewportClient::TrackingStopped() 
+bool FSCSEditorViewportClient::HandleEndTransform()
 {
-	if( bIsManipulating )
+	if (bIsManipulating)
 	{
 		// Re-run construction scripts if we haven't done so yet (so that the components in the preview actor can update their transforms)
 		AActor* PreviewActor = GetPreviewActor();
@@ -742,7 +771,11 @@ void FSCSEditorViewportClient::TrackingStopped()
 
 		// Restore component delta modification
 		GEditor->DisableDeltaModification(false);
+
+		return true;
 	}
+
+	return false;
 }
 
 UE::Widget::EWidgetMode FSCSEditorViewportClient::GetWidgetMode() const
diff --git a/Engine/Source/Editor/Kismet/Private/SCSEditorViewportClient.h b/Engine/Source/Editor/Kismet/Private/SCSEditorViewportClient.h
--- a/Engine/Source/Editor/Kismet/Private/SCSEditorViewportClient.h
+++ b/Engine/Source/Editor/Kismet/Private/SCSEditorViewportClient.h
@@ -24,6 +24,7 @@ class FViewport;
 class SSCSEditorViewport;
 class UStaticMeshComponent;
 struct FInputKeyEventArgs;
+struct FGizmoState;
 
 /**
  * An editor viewport client subclass for the SCS editor viewport.
@@ -53,6 +54,8 @@ public:
 	virtual bool InputWidgetDelta( FViewport* Viewport, EAxisList::Type CurrentAxis, FVector& Drag, FRotator& Rot, FVector& Scale ) override;
 	virtual void TrackingStarted( const struct FInputEventState& InInputState, bool bIsDragging, bool bNudge ) override;
 	virtual void TrackingStopped() override;
+	virtual bool BeginTransform(const FGizmoState& InState) override;
+	virtual bool EndTransform(const FGizmoState& InState) override;
 	virtual UE::Widget::EWidgetMode GetWidgetMode() const override;
 	virtual void SetWidgetMode( UE::Widget::EWidgetMode NewMode ) override;
 	virtual void SetWidgetCoordSystemSpace( ECoordSystem NewCoordSystem ) override;
@@ -167,4 +170,7 @@ private:
 
 	/** If true, the physics simulation gets ticked */
 	bool bIsSimulateEnabled;
+
+	bool HandleBeginTransform();
+	bool HandleEndTransform();
 };
diff --git a/Engine/Source/Editor/Persona/Private/SkeletonSelectionEditMode.cpp b/Engine/Source/Editor/Persona/Private/SkeletonSelectionEditMode.cpp
--- a/Engine/Source/Editor/Persona/Private/SkeletonSelectionEditMode.cpp
+++ b/Engine/Source/Editor/Persona/Private/SkeletonSelectionEditMode.cpp
@@ -95,12 +95,72 @@ FSelectedSocketInfo FSkeletonSelectionEditMode::DuplicateAndSelectSocket(const F
 	return NewSocketInfo;
 }
 
+bool FSkeletonSelectionEditMode::BeginTransform(const FGizmoState& InState)
+{
+	bManipulating = bInTransaction = false;
+	
+	const UDebugSkelMeshComponent* PreviewMeshComponent = GetAnimPreviewScene().GetPreviewMeshComponent();
+	const USkeletalMesh* SkeletalMeshAsset = PreviewMeshComponent ? PreviewMeshComponent->GetSkeletalMeshAsset() : nullptr;
+	if (!SkeletalMeshAsset)
+	{
+		return false;
+	}
+
+	// transact bone transform?
+	const int32 BoneIndex = GetAnimPreviewScene().GetSelectedBoneIndex();
+	const FReferenceSkeleton& ReferenceSkeleton = GetReferenceSkeletonForComponent(PreviewMeshComponent);
+	if (BoneIndex >= INDEX_NONE && ReferenceSkeleton.IsValidIndex(BoneIndex))
+	{
+		PreviewMeshComponent->PreviewInstance->SetFlags(RF_Transactional);	// Undo doesn't work without this!
+		PreviewMeshComponent->PreviewInstance->Modify();
+
+		// now modify the bone array
+		const FName BoneName = ReferenceSkeleton.GetBoneName(BoneIndex);
+		PreviewMeshComponent->PreviewInstance->ModifyBone(BoneName);
+
+		bManipulating = bInTransaction = true;
+		return true;
+	}
+
+	FSelectedSocketInfo SelectedSocketInfo = GetAnimPreviewScene().GetSelectedSocket();
+	if (SelectedSocketInfo.IsValid())
+	{
+		if (GetModeManager()->GetFocusedViewportClient())
+		{
+			const bool bAltDown = GetModeManager()->GetFocusedViewportClient()->IsAltPressed();
+			if (bAltDown)
+			{
+				// Rather than moving/rotating the selected socket, copy it and move the copy instead
+				SelectedSocketInfo = DuplicateAndSelectSocket(SelectedSocketInfo);
+			}
+		}
+
+		// Socket movement is transactional - we want undo/redo and saving of it
+		if (USkeletalMeshSocket* Socket = SelectedSocketInfo.Socket)
+		{
+			Socket->SetFlags(RF_Transactional);	// Undo doesn't work without this!
+			Socket->Modify();
+			bManipulating = bInTransaction = true;
+			return true;
+		}
+	}
+	
+	return false;
+}
+
+bool FSkeletonSelectionEditMode::EndTransform(const FGizmoState& InState)
+{
+	const bool bWasManipulating = bManipulating;
+	bManipulating = bInTransaction = false;
+	return bWasManipulating;
+}
+
 bool FSkeletonSelectionEditMode::StartTracking(FEditorViewportClient* InViewportClient, FViewport* InViewport)
 {
-	EAxisList::Type CurrentAxis = InViewportClient->GetCurrentWidgetAxis();
-	UE::Widget::EWidgetMode WidgetMode = InViewportClient->GetWidgetMode();
+	const EAxisList::Type CurrentAxis = InViewportClient->GetCurrentWidgetAxis();
+	const UE::Widget::EWidgetMode WidgetMode = InViewportClient->GetWidgetMode();
 
-	UDebugSkelMeshComponent* PreviewMeshComponent = GetAnimPreviewScene().GetPreviewMeshComponent();
+	const UDebugSkelMeshComponent* PreviewMeshComponent = GetAnimPreviewScene().GetPreviewMeshComponent();
 	if(PreviewMeshComponent != nullptr && PreviewMeshComponent->GetSkeletalMeshAsset() != nullptr)
 	{
 		const int32 BoneIndex = GetAnimPreviewScene().GetSelectedBoneIndex();
@@ -111,9 +171,6 @@ bool FSkeletonSelectionEditMode::StartTracking(FEditorViewportClient* InViewport
 		const FReferenceSkeleton& ReferenceSkeleton = GetReferenceSkeletonForComponent(PreviewMeshComponent);
 		if ((BoneIndex >= 0 && ReferenceSkeleton.IsValidIndex(BoneIndex)) || SelectedSocket != nullptr || SelectedActor != nullptr)
 		{
-			bool bValidAxis = false;
-			FVector WorldAxisDir;
-
 			if ( ((CurrentAxis & EAxisList::XYZ) | (CurrentAxis & EAxisList::Screen)) != 0)
 			{
 				FSelectedSocketInfo SelectedSocketInfo = GetAnimPreviewScene().GetSelectedSocket();
@@ -203,11 +260,6 @@ bool FSkeletonSelectionEditMode::InputDelta(FEditorViewportClient* InViewportCli
 	const UE::Widget::EWidgetMode WidgetMode = InViewportClient->GetWidgetMode();
 	const ECoordSystem CoordSystem = InViewportClient->GetWidgetCoordSystemSpace();
 
-	// Get some useful info about buttons being held down
-	const bool bCtrlDown = InViewport->KeyState(EKeys::LeftControl) || InViewport->KeyState(EKeys::RightControl);
-	const bool bShiftDown = InViewport->KeyState(EKeys::LeftShift) || InViewport->KeyState(EKeys::RightShift);
-	const bool bMouseButtonDown = InViewport->KeyState( EKeys::LeftMouseButton ) || InViewport->KeyState( EKeys::MiddleMouseButton ) || InViewport->KeyState( EKeys::RightMouseButton );
-
 	bool bHandled = false;
 
 	UDebugSkelMeshComponent* PreviewMeshComponent = GetAnimPreviewScene().GetPreviewMeshComponent();
diff --git a/Engine/Source/Editor/Persona/Public/SkeletonSelectionEditMode.h b/Engine/Source/Editor/Persona/Public/SkeletonSelectionEditMode.h
--- a/Engine/Source/Editor/Persona/Public/SkeletonSelectionEditMode.h
+++ b/Engine/Source/Editor/Persona/Public/SkeletonSelectionEditMode.h
@@ -17,6 +17,7 @@ struct FViewportClick;
 struct FSelectedSocketInfo;
 class FViewportClient;
 struct FReferenceSkeleton;
+struct FGizmoState;
 
 class FSkeletonSelectionEditMode : public IPersonaEditMode
 {
@@ -43,6 +44,9 @@ public:
 	virtual bool HandleClick(FEditorViewportClient* InViewportClient, HHitProxy *HitProxy, const FViewportClick &Click) override;
 	virtual bool CanCycleWidgetMode() const override;
 
+	virtual bool BeginTransform(const FGizmoState& InState) override;
+	virtual bool EndTransform(const FGizmoState& InState) override;
+	
 protected:
 	FTransform GetBoneTransform(const int32 BoneIndex) const;
 	FTransform GetSocketTransform(const USkeletalMeshSocket* Socket) const;
diff --git a/Engine/Source/Editor/UnrealEd/Private/EditorModeManager.cpp b/Engine/Source/Editor/UnrealEd/Private/EditorModeManager.cpp
--- a/Engine/Source/Editor/UnrealEd/Private/EditorModeManager.cpp
+++ b/Engine/Source/Editor/UnrealEd/Private/EditorModeManager.cpp
@@ -37,12 +37,14 @@
 #include "Tools/AssetEditorContextObject.h"
 #include "ContextObjectStore.h"
 #include "EditorInteractiveGizmoManager.h"
+#include "GizmoEdModeInterface.h"
 #include "UObject/GCObjectScopeGuard.h"
 #include "Settings/LevelEditorViewportSettings.h"
 #include "Subsystems/EditorElementSubsystem.h"
 
 #include "Elements/Interfaces/TypedElementWorldInterface.h"
 #include "TextureResource.h"
+#include "EditorGizmos/EditorGizmoStateTarget.h"
 #include "EditorGizmos/EditorTransformGizmoUtil.h"
 #include "Toolkits/ToolkitManager.h"
 
@@ -118,6 +120,7 @@ FEditorModeTools::~FEditorModeTools()
 		InteractiveToolsContext->Deactivate();
 		InteractiveToolsContext->ShutdownContext();
 		InteractiveToolsContext = nullptr;
+		GizmoStateTarget = nullptr;
 	}
 }
 
@@ -1099,35 +1102,50 @@ bool FEditorModeTools::StartTracking(FEditorViewportClient* InViewportClient, FV
 	bIsTracking = true;
 	CachedLocation = PivotLocation;	// Cache the pivot location
 
-	bool bTransactionHandled = InteractiveToolsContext->StartTracking(InViewportClient, InViewport);
-	ForEachEdMode<ILegacyEdModeViewportInterface>([&bTransactionHandled, InViewportClient, InViewport](ILegacyEdModeViewportInterface* ViewportInterface)
+	bool bTrackingHandled = InteractiveToolsContext->StartTracking(InViewportClient, InViewport);
+
+	// no need to go further if bHasOngoingTransform is true
+	if (!bHasOngoingTransform)
+	{
+		ForEachEdMode<ILegacyEdModeViewportInterface>([&bTrackingHandled, InViewportClient, InViewport](ILegacyEdModeViewportInterface* ViewportInterface)
 		{
-			bTransactionHandled |= ViewportInterface->StartTracking(InViewportClient, InViewport);
+			bTrackingHandled |= ViewportInterface->StartTracking(InViewportClient, InViewport);
 			return true;
 		});
+	}
 
-	return bTransactionHandled;
+	return bTrackingHandled;
 }
 
 /** Mouse tracking interface.  Passes tracking messages to all active modes */
 bool FEditorModeTools::EndTracking(FEditorViewportClient* InViewportClient, FViewport* InViewport)
 {
 	bIsTracking = false;
-	bool bTransactionHandled = InteractiveToolsContext->EndTracking(InViewportClient, InViewport);
-
-	ForEachEdMode<ILegacyEdModeViewportInterface>([&bTransactionHandled, InViewportClient, InViewport](ILegacyEdModeViewportInterface* ViewportInterface)
-		{
-			bTransactionHandled |= ViewportInterface->EndTracking(InViewportClient, InViewport);
-			return true;
-		});
 
+	bool bTrackingHandled = InteractiveToolsContext->EndTracking(InViewportClient, InViewport);
+	// no need to go further if bHasOngoingTransform is true
+	if (!bHasOngoingTransform)
+	{
+		ForEachEdMode<ILegacyEdModeViewportInterface>([&bTrackingHandled, InViewportClient, InViewport](ILegacyEdModeViewportInterface* ViewportInterface)
+		   {
+			   bTrackingHandled |= ViewportInterface->EndTracking(InViewportClient, InViewport);
+			   return true;
+		   });
+	}
+	
 	CachedLocation = PivotLocation;	// Clear the pivot location
+	bHasOngoingTransform = false;
 	
-	return bTransactionHandled;
+	return bTrackingHandled;
 }
 
 bool FEditorModeTools::AllowsViewportDragTool() const
 {
+	if (bHasOngoingTransform)
+	{
+		return false;
+	}
+	
 	bool bCanUseDragTool = false;
 	ForEachEdMode<const ILegacyEdModeViewportInterface>([&bCanUseDragTool](const ILegacyEdModeViewportInterface* LegacyMode)
 		{
@@ -1338,6 +1356,8 @@ bool FEditorModeTools::ProcessCapturedMouseMoves( FEditorViewportClient* InViewp
 /** Notifies all active modes of keyboard input via a viewport client */
 bool FEditorModeTools::InputKey(FEditorViewportClient* InViewportClient, FViewport* Viewport, FKey Key, EInputEvent Event, bool bRouteToToolsContext)
 {
+	const bool bHadOngoingTransform = bHasOngoingTransform;
+
 	bool bWasHandledByToolsContext = false;
 	if (bRouteToToolsContext)
 	{
@@ -1359,6 +1379,15 @@ bool FEditorModeTools::InputKey(FEditorViewportClient* InViewportClient, FViewpo
 		EndTracking(InViewportClient, Viewport);
 	}
 
+	// no need to go further if bHasOngoingTransform state has changed 
+	// NOTE, this should probably be done comparing HasActiveMouseCapture changes instead as it means that the ITF handled the event
+	// however, as with StartTracking/EndTracking and bTrackingHandled, testing bWasHandledByToolsContext is not reliable as
+	// InteractiveToolsContext->InputKey will return false when the mouse is released, even if there was an ongoing capture before the release event ended it.
+	if (bHasOngoingTransform != bHadOngoingTransform)
+	{
+		return true;
+	}
+	
 	// If the toolkit should process the command, it should not have been handled by ITF, or be tracked elsewhere.
 	const bool bPassToToolkitCommands = bRouteToToolsContext && !bWasHandledByToolsContext;
 	bool bHandled = bWasHandledByToolsContext;
@@ -1716,14 +1745,19 @@ bool FEditorModeTools::PostConvertMouseMovement(FEditorViewportClient* InViewpor
 
 bool FEditorModeTools::GetShowWidget() const
 {
+	if (!bShowWidget)
+	{
+		return false;
+	}
+	
 	bool bDrawModeSupportsWidgetDrawing = false;
 	// Check to see of any active modes support widget drawing
 	ForEachEdMode<ILegacyEdModeWidgetInterface>([&bDrawModeSupportsWidgetDrawing](ILegacyEdModeWidgetInterface* LegacyMode)
 		{
 			bDrawModeSupportsWidgetDrawing |= LegacyMode->ShouldDrawWidget();
-			return true;
+			return !bDrawModeSupportsWidgetDrawing;
 		});
-	return bDrawModeSupportsWidgetDrawing && bShowWidget;
+	return bDrawModeSupportsWidgetDrawing;
 }
 
 /**
@@ -1937,3 +1971,56 @@ UModeManagerInteractiveToolsContext* FEditorModeTools::GetInteractiveToolsContex
 {
 	return InteractiveToolsContext;
 }
+
+IGizmoStateTarget* FEditorModeTools::GetGizmoStateTarget()
+{
+	if (!GizmoStateTarget)
+	{
+		GizmoStateTarget = UEditorGizmoStateTarget::Construct(
+			this,
+			NSLOCTEXT("UTransformGizmo", "UTransformGizmoTransaction", "Transform"),
+			InteractiveToolsContext->GizmoManager);
+	}
+	return GizmoStateTarget;
+}
+
+bool FEditorModeTools::BeginTransform(const FGizmoState& InState)
+{
+	bool bHandled = false;
+	ForEachEdMode<IGizmoEdModeInterface>([&bHandled, &InState](IGizmoEdModeInterface* GizmoInterface)
+	{
+		bHandled |= GizmoInterface->BeginTransform(InState);
+		return true;
+	});
+
+	// Give the focused VPC an opportunity to open the transform transacting if it has not been handled before 
+	if (!bHandled && FocusedViewportClient)
+	{
+		bHandled = FocusedViewportClient->BeginTransform(InState);
+	}
+
+	bHasOngoingTransform = bHandled;
+	
+	return bHandled;
+}
+
+bool FEditorModeTools::EndTransform(const FGizmoState& InState) const
+{
+	bool bHandled = false;
+	ForEachEdMode<IGizmoEdModeInterface>([&bHandled, &InState](IGizmoEdModeInterface* GizmoInterface)
+	{
+		bHandled |= GizmoInterface->EndTransform(InState);
+		return true;
+	});
+
+	// Give the focused VPC an opportunity to close the transform transacting if it has not been handled before 
+	if (!bHandled && FocusedViewportClient)
+	{
+		bHandled = FocusedViewportClient->EndTransform(InState);
+	}
+
+	// NOTE bHasOngoingTransform is not set to false here but in EndTracking as its needed there.
+	// See header file more more explanations
+	
+	return bHandled;
+}
\ No newline at end of file
diff --git a/Engine/Source/Editor/UnrealEd/Private/EditorViewportClient.cpp b/Engine/Source/Editor/UnrealEd/Private/EditorViewportClient.cpp
--- a/Engine/Source/Editor/UnrealEd/Private/EditorViewportClient.cpp
+++ b/Engine/Source/Editor/UnrealEd/Private/EditorViewportClient.cpp
@@ -2983,9 +2983,9 @@ bool FEditorViewportClient::Internal_InputKey(const FInputKeyEventArgs& EventArg
 		return true;
 	}
 	
-	FKey Key = EventArgs.Key;
-	EInputEvent Event = EventArgs.Event;
-	FViewport* InViewport = EventArgs.Viewport;
+	const FKey& Key = EventArgs.Key;
+	const EInputEvent& Event = EventArgs.Event;
+	const FViewport* InViewport = EventArgs.Viewport;
 
 	// Let the current mode have a look at the input before reacting to it. 
 	// Note that bIsTracking tells us whether the viewport client is capturing mouse behavior to fly around,
@@ -3000,7 +3000,7 @@ bool FEditorViewportClient::Internal_InputKey(const FInputKeyEventArgs& EventArg
 		return true;
 	}
 
-	FInputEventState InputState(EventArgs.Viewport, EventArgs.Key, EventArgs.Event);
+	const FInputEventState InputState(EventArgs.Viewport, Key, Event);
 	
 	bool bHandled = false;
 
@@ -3032,11 +3032,7 @@ bool FEditorViewportClient::Internal_InputKey(const FInputKeyEventArgs& EventArg
 		ModeTools->SetWidgetModeOverride(UE::Widget::WM_None);
 	}
 
-	const int32	HitX = InViewport->GetMouseX();
-	const int32	HitY = InViewport->GetMouseY();
-
-	FCachedJoystickState* JoystickState = GetJoystickState(EventArgs.InputDevice.GetId());
-	if (JoystickState)
+	if (FCachedJoystickState* JoystickState = GetJoystickState(EventArgs.InputDevice.GetId()))
 	{
 		JoystickState->KeyEventValues.Add(Key, Event);
 	}
@@ -3044,10 +3040,6 @@ bool FEditorViewportClient::Internal_InputKey(const FInputKeyEventArgs& EventArg
 	const bool bWasCursorVisible = InViewport->IsCursorVisible();
 	const bool bWasSoftwareCursorVisible = InViewport->IsSoftwareCursorVisible();
 
-	const bool AltDown = InputState.IsAltButtonPressed();
-	const bool ShiftDown = InputState.IsShiftButtonPressed();
-	const bool ControlDown = InputState.IsCtrlButtonPressed();
-
 	RequiredCursorVisibiltyAndAppearance.bDontResetCursor = false;
 	UpdateRequiredCursorVisibility();
 
diff --git a/Engine/Source/Editor/UnrealEd/Private/LevelEditorViewport.cpp b/Engine/Source/Editor/UnrealEd/Private/LevelEditorViewport.cpp
--- a/Engine/Source/Editor/UnrealEd/Private/LevelEditorViewport.cpp
+++ b/Engine/Source/Editor/UnrealEd/Private/LevelEditorViewport.cpp
@@ -3677,12 +3677,6 @@ void FLevelEditorViewportClient::TrackingStarted( const FInputEventState& InInpu
 
 void FLevelEditorViewportClient::TrackingStopped()
 {
-	const bool AltDown = IsAltPressed();
-	const bool ShiftDown = IsShiftPressed();
-	const bool ControlDown = IsCtrlPressed();
-	const bool LeftMouseButtonDown = Viewport->KeyState(EKeys::LeftMouseButton);
-	const bool RightMouseButtonDown = Viewport->KeyState(EKeys::RightMouseButton);
-	const bool MiddleMouseButtonDown = Viewport->KeyState(EKeys::MiddleMouseButton);
 
 	// Only disable the duplicate on next drag flag if we actually dragged the mouse.
 	bDuplicateOnNextDrag = false;
@@ -5405,6 +5399,134 @@ bool FLevelEditorViewportClient::OverrideHighResScreenshotCaptureRegion(FIntRect
 	return false;
 }
 
+bool FLevelEditorViewportClient::BeginTransform(const FGizmoState& InState)
+{
+	TrackingTransaction.End();
+	bDuplicateOnNextDrag = false;
+	bOnlyMovedPivot = false;
+	bNeedToRestoreComponentBeingMovedFlag = false;
+	MouseDeltaTracker->SetExternalMovement(true);
+
+	PreDragActorTransforms.Empty();
+
+	Widget->SetSnapEnabled(true);
+
+	const FTypedElementListConstRef ElementsToManipulate = GetElementsToManipulate();
+	ViewportInteraction->BeginGizmoManipulation(ElementsToManipulate, GetWidgetMode());
+	bHasBegunGizmoManipulation = true;
+
+	if (!bDuplicateActorsInProgress)
+	{
+		bNeedToRestoreComponentBeingMovedFlag = true;
+		TypedElementListObjectUtil::ForEachObject<AActor>(ElementsToManipulate, [this](AActor* InActor)
+		{
+			SetActorBeingMovedByEditor(InActor, true);
+			return true;
+		});
+	}
+
+	TrackingTransaction.TransCount++;
+	const FText Description = LOCTEXT("TransformTransaction", "Transform Elements");
+	TrackingTransaction.BeginPending(Description);
+
+	if (TrackingTransaction.IsActive() || TrackingTransaction.IsPending())
+	{
+		// Suspend actor/component modification during each delta step to avoid recording unnecessary overhead into the transaction buffer
+		GEditor->DisableDeltaModification(true);
+	}
+
+	GUnrealEd->ComponentVisManager.TrackingStarted(this);
+	
+	return true;
+}
+
+bool FLevelEditorViewportClient::EndTransform(const FGizmoState& InState)
+{
+	if (!bHasBegunGizmoManipulation)
+	{
+		return false;
+	}
+	
+	bDuplicateOnNextDrag = false;
+
+	// here we check to see if anything of worth actually changed when ending our MouseMovement
+	// If the TransCount > 0 (we changed something of value) so we need to call PostEditMove() on stuff
+	// if we didn't change anything then don't call PostEditMove()
+	bool bDidAnythingActuallyChange = false;
+
+	if( TrackingTransaction.TransCount > 0 )
+	{
+		bDidAnythingActuallyChange = true;
+		TrackingTransaction.TransCount--;
+	}
+
+	const bool bDidMove = bDidAnythingActuallyChange && MouseDeltaTracker->HasReceivedDelta();
+	const FTypedElementListConstRef ElementsToManipulate = GetElementsToManipulate();
+
+	if (bHasBegunGizmoManipulation)
+	{
+		auto GetManipType = [bDidMove]()
+		{
+			if (bDidMove)
+			{
+				return ETypedElementViewportInteractionGizmoManipulationType::Drag; 
+			}
+		   return ETypedElementViewportInteractionGizmoManipulationType::Click;
+		};
+		ViewportInteraction->EndGizmoManipulation(ElementsToManipulate, GetWidgetMode(), GetManipType());
+		bHasBegunGizmoManipulation = false;
+	}
+
+	if (bDidMove && !GUnrealEd->IsPivotMovedIndependently())
+	{
+		GUnrealEd->UpdatePivotLocationForSelection();
+	}
+
+	GUnrealEd->ComponentVisManager.TrackingStopped(this, bDidMove);
+
+	if (bNeedToRestoreComponentBeingMovedFlag)
+	{
+		TypedElementListObjectUtil::ForEachObject<AActor>(ElementsToManipulate, [this](AActor* InActor)
+		{
+			SetActorBeingMovedByEditor(InActor, false);
+			return true;
+		});
+
+		bNeedToRestoreComponentBeingMovedFlag = false;
+	}
+
+	// End the transaction here if one was started in StartTransaction()
+	if( TrackingTransaction.IsActive() || TrackingTransaction.IsPending() )
+	{
+		if (!HaveSelectedObjectsBeenChanged())
+		{
+			TrackingTransaction.Cancel();
+		}
+		else
+		{
+			TrackingTransaction.End();
+		}
+	
+		// Restore actor/component delta modification
+		GEditor->DisableDeltaModification(false);
+	}
+
+	ModeTools->ActorMoveNotify();
+
+	if (bDidAnythingActuallyChange)
+	{
+		FScopedLevelDirtied LevelDirtyCallback;
+		LevelDirtyCallback.Request();
+
+		RedrawAllViewportsIntoThisScene();
+	}
+
+	PreDragActorTransforms.Empty();
+	MouseDeltaTracker->SetExternalMovement(false);
+
+	return true;
+}
+
 /**
  * Static: Adds a hover effect to the specified object
  *
diff --git a/Engine/Source/Editor/UnrealEd/Private/Tools/LegacyEdMode.cpp b/Engine/Source/Editor/UnrealEd/Private/Tools/LegacyEdMode.cpp
--- a/Engine/Source/Editor/UnrealEd/Private/Tools/LegacyEdMode.cpp
+++ b/Engine/Source/Editor/UnrealEd/Private/Tools/LegacyEdMode.cpp
@@ -204,6 +204,16 @@ bool ULegacyEdModeWrapper::InputDelta(FEditorViewportClient* InViewportClient, F
 	return LegacyEditorMode->InputDelta(InViewportClient, InViewport, InDrag, InRot, InScale);
 }
 
+bool ULegacyEdModeWrapper::BeginTransform(const FGizmoState& InState)
+{
+	return LegacyEditorMode->BeginTransform(InState);
+}
+
+bool ULegacyEdModeWrapper::EndTransform(const FGizmoState& InState)
+{
+	return LegacyEditorMode->EndTransform(InState);
+}
+
 bool ULegacyEdModeWrapper::StartTracking(FEditorViewportClient* InViewportClient, FViewport* InViewport)
 {
 	return LegacyEditorMode->StartTracking(InViewportClient, InViewport);
diff --git a/Engine/Source/Editor/UnrealEd/Private/Tools/LegacyEdMode.h b/Engine/Source/Editor/UnrealEd/Private/Tools/LegacyEdMode.h
--- a/Engine/Source/Editor/UnrealEd/Private/Tools/LegacyEdMode.h
+++ b/Engine/Source/Editor/UnrealEd/Private/Tools/LegacyEdMode.h
@@ -7,6 +7,7 @@
 #include "Engine/EngineBaseTypes.h"
 #include "Templates/SharedPointer.h"
 #include "Tools/LegacyEdModeInterfaces.h"
+#include "GizmoEdModeInterface.h"
 
 #include "LegacyEdMode.generated.h"
 
@@ -19,7 +20,7 @@ class FViewport;
 class UPrimitiveComponent;
 
 UCLASS(MinimalAPI)
-class ULegacyEdModeWrapper final : public UEdMode, public ILegacyEdModeWidgetInterface, public ILegacyEdModeToolInterface, public ILegacyEdModeSelectInterface, public ILegacyEdModeDrawHelperInterface, public ILegacyEdModeViewportInterface
+class ULegacyEdModeWrapper final : public UEdMode, public ILegacyEdModeWidgetInterface, public ILegacyEdModeToolInterface, public ILegacyEdModeSelectInterface, public ILegacyEdModeDrawHelperInterface, public ILegacyEdModeViewportInterface, public IGizmoEdModeInterface
 {
 	GENERATED_BODY()
 
@@ -121,6 +122,11 @@ public:
 	UNREALED_API virtual void Draw(const FSceneView* View, FPrimitiveDrawInterface* PDI) override;
 	// End FCommonDrawHelper overrides
 
+	// Start IGizmoEdModeInterface overrides
+	UNREALED_API virtual bool BeginTransform(const FGizmoState& InState) override;
+	UNREALED_API virtual bool EndTransform(const FGizmoState& InState) override;
+	// End IGizmoEdModeInterface overrides
+
 private:
 	TSharedPtr<FEdMode> LegacyEditorMode;
 };
diff --git a/Engine/Source/Editor/UnrealEd/Public/EdMode.h b/Engine/Source/Editor/UnrealEd/Public/EdMode.h
--- a/Engine/Source/Editor/UnrealEd/Public/EdMode.h
+++ b/Engine/Source/Editor/UnrealEd/Public/EdMode.h
@@ -32,6 +32,7 @@ struct FViewportClick;
 class FModeTool;
 class FEditorViewportClient;
 struct FViewportClick;
+struct FGizmoState;
 
 /**
  * Base class for all editor modes.
@@ -89,6 +90,9 @@ public:
 	virtual bool ProcessEditCopy() { return false; }
 	virtual bool ProcessEditPaste() { return false; }
 
+	virtual bool BeginTransform(const FGizmoState& InState) { return false; }
+	virtual bool EndTransform(const FGizmoState& InState) { return false; }
+	
 	UNREALED_API virtual void Tick(FEditorViewportClient* ViewportClient,float DeltaTime);
 
 	virtual bool IsCompatibleWith(FEditorModeID OtherModeID) const { return false; }
diff --git a/Engine/Source/Editor/UnrealEd/Public/EditorModeManager.h b/Engine/Source/Editor/UnrealEd/Public/EditorModeManager.h
--- a/Engine/Source/Editor/UnrealEd/Public/EditorModeManager.h
+++ b/Engine/Source/Editor/UnrealEd/Public/EditorModeManager.h
@@ -30,6 +30,9 @@ class UInteractiveGizmoManager;
 class UInputRouter;
 class UModeManagerInteractiveToolsContext;
 class UTypedElementSelectionSet;
+class IGizmoStateTarget;
+class UEditorGizmoStateTarget;
+struct FGizmoState;
 
 /**
  * A helper class to store the state of the various editor modes.
@@ -532,7 +535,11 @@ public:
 	/** @return ToolsContext for this Mode Manager */
 	UNREALED_API UModeManagerInteractiveToolsContext* GetInteractiveToolsContext() const;
 
-
+	/** New TRS Gizmo interface */
+	UNREALED_API IGizmoStateTarget* GetGizmoStateTarget();
+	UNREALED_API bool BeginTransform(const FGizmoState& InState);
+	UNREALED_API bool EndTransform(const FGizmoState& InState) const;
+	
 protected:
 	/** 
 	 * Delegate handlers
@@ -623,6 +630,16 @@ private:
 	/** Guard to prevent modes from entering as part of their exit routine */
 	bool bIsExitingModesDuringTick = false;
 
+	/** GizmoStateTarget used to handle a new TRS gizmo transform begin/end sequence. */
+	TObjectPtr<UEditorGizmoStateTarget> GizmoStateTarget;
+
+	/** Flag to track if we started a new TRS gizmo transform.
+	 * NOTE: we have to use it for now as StartTracking / EndTracking iterates thru the modes even if the ITF context captured something.
+	 * This might not be needed anymore in the future but as some routing behavior had to be kept for legacy reasons, we use that extra flag.
+	 * If StartTracking / EndTracking were to be changed (which needs more testing), this could probably be removed.
+	 */
+	bool bHasOngoingTransform = false;
+
 	FEditorViewportClient* HoveredViewportClient = nullptr;
 	FEditorViewportClient* FocusedViewportClient = nullptr;
 
diff --git a/Engine/Source/Editor/UnrealEd/Public/EditorViewportClient.h b/Engine/Source/Editor/UnrealEd/Public/EditorViewportClient.h
--- a/Engine/Source/Editor/UnrealEd/Public/EditorViewportClient.h
+++ b/Engine/Source/Editor/UnrealEd/Public/EditorViewportClient.h
@@ -38,6 +38,7 @@ class SEditorViewport;
 class UActorFactory;
 class UTypedElementViewportInteraction;
 enum class EViewStatusForScreenPercentage;
+struct FGizmoState;
 
 /** Delegate called by FEditorViewportClient to check its visibility */
 DECLARE_DELEGATE_RetVal( bool, FViewportStateGetter );
@@ -616,6 +617,9 @@ public:
 	UNREALED_API virtual void SetEnabledStats(const TArray<FString>& InEnabledStats) override;
 	UNREALED_API virtual bool IsStatEnabled(const FString& InName) const override;
 
+	UNREALED_API virtual bool BeginTransform(const FGizmoState& InState) { return false; }
+	UNREALED_API virtual bool EndTransform(const FGizmoState& InState) { return false; }
+	
 protected:
 
 	UNREALED_API virtual bool Internal_InputKey(const FInputKeyEventArgs& EventArgs);
diff --git a/Engine/Source/Editor/UnrealEd/Public/LevelEditorViewport.h b/Engine/Source/Editor/UnrealEd/Public/LevelEditorViewport.h
--- a/Engine/Source/Editor/UnrealEd/Public/LevelEditorViewport.h
+++ b/Engine/Source/Editor/UnrealEd/Public/LevelEditorViewport.h
@@ -27,6 +27,7 @@ class UModel;
 class UTypedElementSelectionSet;
 struct FWorldContext;
 struct FTypedElementHandle;
+struct FGizmoState;
 
 /** Describes an object that's currently hovered over in the level viewport */
 struct FViewportHoverTarget
@@ -231,6 +232,10 @@ public:
 
 	UNREALED_API virtual bool OverrideHighResScreenshotCaptureRegion(FIntRect& OutCaptureRegion) override;
 
+	UNREALED_API virtual bool BeginTransform(const FGizmoState& InState) override;
+	UNREALED_API virtual bool EndTransform(const FGizmoState& InState) override;
+	
+
 	/** Sets a flag for this frame indicating that the camera has been cut, and temporal effects (such as motion blur) should be reset */
 	void SetIsCameraCut()
 	{
diff --git a/Engine/Source/Editor/UnrealEd/Public/MouseDeltaTracker.h b/Engine/Source/Editor/UnrealEd/Public/MouseDeltaTracker.h
--- a/Engine/Source/Editor/UnrealEd/Public/MouseDeltaTracker.h
+++ b/Engine/Source/Editor/UnrealEd/Public/MouseDeltaTracker.h
@@ -111,7 +111,7 @@ public:
 	/**
 	 * Marks that something caused the equivalent of mouse dragging, but with other means (keyboard short cut, mouse wheel).  Allows suppression of context menus with flight camera, etc
 	 */
-	void SetExternalMovement (void)			{ bExternalMovement = true;};
+	void SetExternalMovement (bool InExternalMovement = true)			{ bExternalMovement = InExternalMovement;};
 	/**
 	 * @return		true if something caused external movement of the mouse (keyboard, mouse wheel)
 	 */
