Editor: Fix some drop previews not being destroyed after actor factory refactor.

diff --git a/Engine/Source/Editor/UnrealEd/Private/LevelEditorViewport.cpp b/Engine/Source/Editor/UnrealEd/Private/LevelEditorViewport.cpp
--- a/Engine/Source/Editor/UnrealEd/Private/LevelEditorViewport.cpp
+++ b/Engine/Source/Editor/UnrealEd/Private/LevelEditorViewport.cpp
@@ -325,6 +325,13 @@ TArray<FTypedElementHandle> FLevelEditorViewportClient::TryPlacingAssetObject(UL
 		if (!bPlaced && !ObjectClass->HasAnyClassFlags(CLASS_NotPlaceable | CLASS_Abstract) )
 		{
 			// If no actor factory was found or failed, add the actor directly.
+	
+			// TODO: We might want to investigate using the above PlaceAssetInCurrentLevel path even when we
+			// don't have an actor factory, and only use this legacy one if that fails for whatever reason.
+			// Note that this path acts a bit differently. For instance it doesn't use IsDroppingPreviewActor()
+			// to set bIsEditorPrieviewActor on the output if we are dropping a preview. We fix this in
+			// DropObjectsAtCoordinates where we directly know whether we are adding to previews (instead of using
+			// the static IsDroppingPreviewActor).
 			const FTransform ActorTransform = FActorPositioning::GetCurrentViewportPlacementTransform(
 				*ObjectClass->GetDefaultObject<AActor>(), /*bSnap=*/true, CursorInformation);
 			AActor* Actor = GEditor->AddActor( InLevel, ObjectClass, ActorTransform, 
@@ -1695,7 +1702,21 @@ void FLevelEditorViewportClient::DestroyDropPreviewElements()
 		Options
 			.SetWarnAboutReferences(false)
 			.SetWarnAboutSoftReferences(false);
-		ensure(InElement.DeleteElement(InElement.GetOwnerWorld(), GetMutableSelectionSet(), Options));
+
+		if (!ensure(InElement.DeleteElement(InElement.GetOwnerWorld(), GetMutableSelectionSet(), Options)))
+		{
+			// We don't expect to fail, but try a legacy actor deletion path if we do.
+			// TODO: This might not help that much because for actors, the DeleteElements will frequently still 
+			// return true even if deletion does not go through because it calls UUnrealEdEngine::DeleteActors,
+			// and that function returns true for many failed deletions. It would be nice to make edits to make
+			// the check more reliable.
+			ITypedElementObjectInterface* ObjectInterface = UTypedElementRegistry::GetInstance()->GetElementInterface<ITypedElementObjectInterface>(InElement);
+			AActor* Actor = ObjectInterface ? ObjectInterface->GetObjectAs<AActor>(InElement) : nullptr;
+			if (Actor)
+			{
+				GetWorld()->DestroyActor(Actor);
+			}
+		}
 
 		return true;  // true means continue
 	});
@@ -2003,8 +2024,6 @@ bool FLevelEditorViewportClient::DropObjectsAtCoordinates(int32 MouseX, int32 Mo
 					}
 
 					DropPreviewElements->Add(Element);
-
-					// If we wanted 
 				};
 
 				// Collect the preview elements and save pre drag transforms.
@@ -2040,6 +2059,13 @@ bool FLevelEditorViewportClient::DropObjectsAtCoordinates(int32 MouseX, int32 Mo
 
 					NewActor->SetActorEnableCollision(false);
 
+					// This boolean already gets set via UActorFactory::PlaceAsset in code paths that go through
+					// PlaceAssetUsingFactory (PlacementOptions.bIsCreatingPreviewElements = FLevelEditorViewportClient::IsDroppingPreviewActor();)
+					// But bIsEditorPrieviewActor does not get set if we go through GEditor->AddActor, which can
+					// prevent preview actors from being properly destroyed later. If we fix/eliminate paths that
+					// do not set this bool properly, we can do an ensure here.
+					NewActor->bIsEditorPreviewActor = FLevelEditorViewportClient::IsDroppingPreviewActor();
+
 					// Prevent future selection. This also prevents the hit proxy from interfering with placement logic.
 					for (UActorComponent* Component : NewActor->GetComponents())
 					{
