ResavePackagesCommandlet:
- When fixing up redirectors for external actors, add the owning world to the resave list so it handles partitioned worlds.
- Fix NormalizePackageNames so it always outputs relative paths, because we are comparing relative paths later on.
- Only output HLODs warning when using the -BuildHLOD parameter.

diff --git a/Engine/Source/Editor/UnrealEd/Private/Commandlets/ContentCommandlets.cpp b/Engine/Source/Editor/UnrealEd/Private/Commandlets/ContentCommandlets.cpp
--- a/Engine/Source/Editor/UnrealEd/Private/Commandlets/ContentCommandlets.cpp
+++ b/Engine/Source/Editor/UnrealEd/Private/Commandlets/ContentCommandlets.cpp
@@ -5,6 +5,7 @@
 =============================================================================*/
 
 #include "Algo/RemoveIf.h"
+#include "Algo/Unique.h"
 #include "AssetCompilingManager.h"
 #include "AssetRegistry/AssetData.h"
 #include "CollectionManagerModule.h"
@@ -292,7 +293,7 @@ int32 UResavePackagesCommandlet::InitializeResaveParameters( const TArray<FStrin
 	// ... if not, load in all packages
 	if( !bExplicitPackages || Switches.Contains(TEXT("SaveAll")))
 	{
-		UE_LOG( LogContentCommandlet, Display, TEXT( "No maps found to save when building HLODs, checking Project Settings for Directory or Asset Path(s)" ) );
+		UE_CLOG( bShouldBuildHLOD, LogContentCommandlet, Display, TEXT( "No maps found to save when building HLODs, checking Project Settings for Directory or Asset Path(s)" ) );
 
 		uint8 PackageFilter = NORMALIZE_DefaultFlags;
 		if ( Switches.Contains(TEXT("SKIPMAPS")) )
@@ -429,6 +430,26 @@ int32 UResavePackagesCommandlet::InitializeResaveParameters( const TArray<FStrin
 				TArray<FName> Referencers;
 				AssetRegistry.GetReferencers(AssetData.PackageName, Referencers);
 
+				// For external objects referencers, also add the object's outer package as a referencer so it can be handled by PerformAdditionalOperations.
+				FARFilter Filter;
+				Filter.bIncludeOnlyOnDiskAssets = true;
+				Filter.PackageNames = Referencers;
+
+				TArray<FAssetData> AssetReferencers;
+				AssetRegistry.GetAssets(Filter, AssetReferencers);
+
+				TArray<FName> ReferencerOuters;
+				for (const FAssetData& AssetReferencer : AssetReferencers)
+				{
+					if (!AssetReferencer.GetOptionalOuterPathName().IsNone())
+					{
+						Referencers.Add(FSoftObjectPath(AssetReferencer.GetOptionalOuterPathName().ToString()).GetLongPackageFName());
+					}
+				}
+
+				Referencers.Sort(FNameFastLess());
+				Referencers.SetNum(Algo::Unique(Referencers));
+
 				for (FName Referencer : Referencers)
 				{
 					FString ReferencerFile;
@@ -3925,4 +3946,4 @@ int32 ULandscapeGrassTypeCommandlet::Main(const FString& Params)
 	FEditorFileUtils::PromptForCheckoutAndSave(PackagesToSave, false, false, nullptr, true, false);
 	
 	return 0;
-}
\ No newline at end of file
+}
diff --git a/Engine/Source/Editor/UnrealEd/Private/Commandlets/PackageUtilities.cpp b/Engine/Source/Editor/UnrealEd/Private/Commandlets/PackageUtilities.cpp
--- a/Engine/Source/Editor/UnrealEd/Private/Commandlets/PackageUtilities.cpp
+++ b/Engine/Source/Editor/UnrealEd/Private/Commandlets/PackageUtilities.cpp
@@ -124,13 +124,15 @@ bool NormalizePackageNames( TArray<FString> PackageNames, TArray<FString>& Packa
 			TStringBuilder<256> UnusedRelPath;
 			for ( const FString& Path : Paths)
 			{
-				if (!FPackageName::TryGetMountPointForPath(Path, UnusedPackagePath, UnusedFilePath, UnusedRelPath))
+				// Make sure paths are relative so SearchDirectoryRecursive will output relative paths
+				const FString RelativePath = FPaths::CreateStandardFilename(Path);
+				if (!FPackageName::TryGetMountPointForPath(RelativePath, UnusedPackagePath, UnusedFilePath, UnusedRelPath))
 				{
 					UE_LOG(LogPackageUtilities, Warning,
-						TEXT("Engine.ini:[Core.System]:Paths entry '%s' is not mounted. Skipping it."), *Path);
+						TEXT("Engine.ini:[Core.System]:Paths entry '%s' is not mounted. Skipping it."), *RelativePath);
 					continue;
 				}
-				FString SearchWildcard = Path / PackageWildcard;
+				FString SearchWildcard = RelativePath / PackageWildcard;
 				UE_LOG(LogPackageUtilities, Log, TEXT("Searching using wildcard: '%s'"), *SearchWildcard);
 				SearchDirectoryRecursive(SearchWildcard, PackageNames, PackagePathNames);
 			}
