Bookmark CIS fix

diff --git a/Engine/Source/Editor/UnrealEd/Private/Bookmarks/BookmarkTypeToolsImpl.cpp b/Engine/Source/Editor/UnrealEd/Private/Bookmarks/BookmarkTypeToolsImpl.cpp
--- a/Engine/Source/Editor/UnrealEd/Private/Bookmarks/BookmarkTypeToolsImpl.cpp
+++ b/Engine/Source/Editor/UnrealEd/Private/Bookmarks/BookmarkTypeToolsImpl.cpp
@@ -213,14 +213,14 @@ public:
 	 */
 	virtual void ClearBookmark(const uint32 InIndex, FEditorViewportClient* InViewportClient) const override
 	{
-		FScopedTransaction ScopedTransaction(FText::Format(LOCTEXT("ClearedBookmark", "Cleared Bookmark {0}"), InIndex));
+		const FText ClearBookmarkMessage = FText::Format(LOCTEXT("ClearBookmarkSuccessFormat", "Cleared bookmark {0}"), InIndex);		
+		FScopedTransaction ScopedTransaction(ClearBookmarkMessage);
 
 		if (AWorldSettings* WorldSettings = GetWorldSettings(InViewportClient))
 		{
 			WorldSettings->ClearBookmark(InIndex);
 
-			const FText NotifyMessage = LOCTEXT("ClearBookmarkSuccessFormat", "Cleared bookmark {0}");
-			Notify(FText::Format(NotifyMessage, InIndex), true);
+			Notify(ClearBookmarkMessage, true);
 		}
 	}
 
@@ -231,14 +231,14 @@ public:
 	 */
 	virtual void ClearAllBookmarks(FEditorViewportClient* InViewportClient) const override
 	{
-		FScopedTransaction ScopedTransaction(LOCTEXT("ClearedAllBookmarks", "Cleared All Bookmarks"));
+		const FText ClearAllBookmarksMessage = LOCTEXT("ClearAllBookmarksSuccessFormat", "Cleared all bookmarks");
+		FScopedTransaction ScopedTransaction(ClearAllBookmarksMessage);
 
 		if (AWorldSettings* WorldSettings = GetWorldSettings(InViewportClient))
 		{
 			WorldSettings->ClearAllBookmarks();
 
-			const FText NotifyMessage = LOCTEXT("ClearBookmarkSuccessFormat", "Cleared all bookmarks");
-			Notify(NotifyMessage, true);
+			Notify(ClearAllBookmarksMessage, true);
 		}
 	}
 
