Convert an asset registry generator log statement into a structured log

diff --git a/Engine/Source/Editor/UnrealEd/Private/Commandlets/AssetRegistryGenerator.cpp b/Engine/Source/Editor/UnrealEd/Private/Commandlets/AssetRegistryGenerator.cpp
--- a/Engine/Source/Editor/UnrealEd/Private/Commandlets/AssetRegistryGenerator.cpp
+++ b/Engine/Source/Editor/UnrealEd/Private/Commandlets/AssetRegistryGenerator.cpp
@@ -51,6 +51,7 @@
 #include "TargetDomain/TargetDomainUtils.h"
 #include "Templates/UniquePtr.h"
 #include "UObject/SoftObjectPath.h"
+#include "Logging/StructuredLog.h"
 
 #if WITH_EDITOR
 #include "HAL/ThreadHeartBeat.h"
@@ -789,10 +790,13 @@ bool FAssetRegistryGenerator::GenerateStreamingInstallManifest(int64 InOverrideC
 					Unit = 1;
 					UnitsText = TEXT("bytes");
 				}
-				UE_LOG(LogAssetRegistryGenerator, Error, TEXT("Failed to add file %s to paklist '%s'. The maximum size for a Pakfile is %d %s, but the file to add is %d %s."),
-					*NextFilename, *PakListFilename,
-					MaxChunkSize/Unit, UnitsText, // Round the limit down and round the value up, so that the display always shows that the value is greater than the limit
-					(NextFileSize+Unit-1)/Unit, UnitsText);
+				UE_LOGFMT(LogAssetRegistryGenerator, Error, "Failed to add file {NextFilename} to paklist '{PakListFilename}'. The maximum size for a Pakfile is {MaxChunkFileSize}{UnitsText}, but the file to add is {ActualChunkFileSize}{UnitsText}.", 
+					("NextFilename", NextFilename),
+					("PakListFilename", PakListFilename),
+					("MaxChunkFileSize", MaxChunkSize / Unit),
+					("UnitsText", UnitsText),
+					("ActualChunkFileSize", (NextFileSize + Unit - 1) / Unit)	// Round the limit down and round the value up, so that the display always shows that the value is greater than the limit
+				);
 				bSucceeded = false;
 				break;
 			}
