Move bCookAllByDefault into a CVar so it can be turned off when attempting to cook nothing.

diff --git a/Engine/Source/Editor/UnrealEd/Private/CookOnTheFlyServer.cpp b/Engine/Source/Editor/UnrealEd/Private/CookOnTheFlyServer.cpp
--- a/Engine/Source/Editor/UnrealEd/Private/CookOnTheFlyServer.cpp
+++ b/Engine/Source/Editor/UnrealEd/Private/CookOnTheFlyServer.cpp
@@ -227,6 +227,13 @@ namespace UE
 namespace Cook
 {
 const TCHAR* GeneratedPackageSubPath = TEXT("_Generated_");
+
+// Keep the old behavior of cooking all by default until we implement good feedback in the editor about the missing setting
+static bool bCookAllByDefault = true;
+static FAutoConsoleVariableRef CookAllByDefaultCVar(
+	TEXT("Cook.CookAllByDefault"),
+	bCookAllByDefault,
+	TEXT("When FilesInPath is empty. Cook all packages by default."));
 }
 }
 
@@ -8994,11 +9001,8 @@ void UCookOnTheFlyServer::CollectFilesToCook(TArray<FName>& FilesInPath, TMap<FN
 			}
 		}
 
-		// Keep the old behavior of cooking all by default until we implement good feedback in the editor about the missing setting
-		constexpr bool bCookAllByDefault = true;
-
 		// If no packages were explicitly added by command line or game callback, add all maps
-		if (bCookAll || (bCookAllByDefault && FilesInPath.Num() == InitialNum))
+		if (bCookAll || (UE::Cook::bCookAllByDefault && FilesInPath.Num() == InitialNum))
 		{
 			TArray<FString> Tokens;
 			Tokens.Empty(2);
