New TRS Gizmo: make the state target a weak ptr to make sure it becomes invalid after GC

diff --git a/Engine/Source/Editor/UnrealEd/Private/EditorModeManager.cpp b/Engine/Source/Editor/UnrealEd/Private/EditorModeManager.cpp
--- a/Engine/Source/Editor/UnrealEd/Private/EditorModeManager.cpp
+++ b/Engine/Source/Editor/UnrealEd/Private/EditorModeManager.cpp
@@ -1974,14 +1974,14 @@ UModeManagerInteractiveToolsContext* FEditorModeTools::GetInteractiveToolsContex
 
 IGizmoStateTarget* FEditorModeTools::GetGizmoStateTarget()
 {
-	if (!GizmoStateTarget)
+	if (!GizmoStateTarget.IsValid())
 	{
 		GizmoStateTarget = UEditorGizmoStateTarget::Construct(
 			this,
 			NSLOCTEXT("UTransformGizmo", "UTransformGizmoTransaction", "Transform"),
 			InteractiveToolsContext->GizmoManager);
 	}
-	return GizmoStateTarget;
+	return GizmoStateTarget.Get();
 }
 
 bool FEditorModeTools::BeginTransform(const FGizmoState& InState)
diff --git a/Engine/Source/Editor/UnrealEd/Public/EditorModeManager.h b/Engine/Source/Editor/UnrealEd/Public/EditorModeManager.h
--- a/Engine/Source/Editor/UnrealEd/Public/EditorModeManager.h
+++ b/Engine/Source/Editor/UnrealEd/Public/EditorModeManager.h
@@ -631,7 +631,7 @@ private:
 	bool bIsExitingModesDuringTick = false;
 
 	/** GizmoStateTarget used to handle a new TRS gizmo transform begin/end sequence. */
-	TObjectPtr<UEditorGizmoStateTarget> GizmoStateTarget;
+	TWeakObjectPtr<UEditorGizmoStateTarget> GizmoStateTarget;
 
 	/** Flag to track if we started a new TRS gizmo transform.
 	 * NOTE: we have to use it for now as StartTracking / EndTracking iterates thru the modes even if the ITF context captured something.
