Silence warnings

diff --git a/Engine/Plugins/FX/Niagara/Source/NiagaraEditor/Private/NiagaraScriptMergeManager.cpp b/Engine/Plugins/FX/Niagara/Source/NiagaraEditor/Private/NiagaraScriptMergeManager.cpp
--- a/Engine/Plugins/FX/Niagara/Source/NiagaraEditor/Private/NiagaraScriptMergeManager.cpp
+++ b/Engine/Plugins/FX/Niagara/Source/NiagaraEditor/Private/NiagaraScriptMergeManager.cpp
@@ -1602,7 +1602,7 @@ INiagaraMergeManager::FMergeEmitterResults FNiagaraScriptMergeManager::MergeEmit
 			ReportScriptStackDiffErrors(MergeResults, EventHandlerDiffResults.ScriptDiffResults, EventHandlerName);
 		}
 	}
-	else if (State == EmitterMergeState::CopyProperties)
+	else if (State == EmitterMergeState::CopyProperties) //-V547
 	{
 		UNiagaraEmitter* MergedInstance = Parent.Emitter->DuplicateWithoutMerging(GetTransientPackage());
 		MergedInstance->DisableVersioning(Parent.Version);
diff --git a/Engine/Plugins/Media/MediaPlayerEditor/Source/MediaPlayerEditor/Private/Widgets/SMediaPlayerEditorOverlay.cpp b/Engine/Plugins/Media/MediaPlayerEditor/Source/MediaPlayerEditor/Private/Widgets/SMediaPlayerEditorOverlay.cpp
--- a/Engine/Plugins/Media/MediaPlayerEditor/Source/MediaPlayerEditor/Private/Widgets/SMediaPlayerEditorOverlay.cpp
+++ b/Engine/Plugins/Media/MediaPlayerEditor/Source/MediaPlayerEditor/Private/Widgets/SMediaPlayerEditorOverlay.cpp
@@ -30,6 +30,7 @@ void SMediaPlayerEditorOverlay::Tick(const FGeometry& AllottedGeometry, const do
 {
 	Canvas->ClearChildren();
 
+#if 0
 	// fetch overlay samples
 	TArray<TSharedPtr<IMediaOverlaySample, ESPMode::ThreadSafe>> OverlaySamples;
 	// @todo gmp: MF3.0 implement overlay sample caching
@@ -75,4 +76,5 @@ void SMediaPlayerEditorOverlay::Tick(const FGeometry& AllottedGeometry, const do
 
 	// recalculate layout
 	SlatePrepass(AllottedGeometry.GetAccumulatedLayoutTransform().GetScale());
+#endif
 }
diff --git a/Engine/Source/Editor/UnrealEd/Private/Commandlets/TextAssetCommandlet.cpp b/Engine/Source/Editor/UnrealEd/Private/Commandlets/TextAssetCommandlet.cpp
--- a/Engine/Source/Editor/UnrealEd/Private/Commandlets/TextAssetCommandlet.cpp
+++ b/Engine/Source/Editor/UnrealEd/Private/Commandlets/TextAssetCommandlet.cpp
@@ -308,8 +308,6 @@ bool UTextAssetCommandlet::DoTextAssetProcessing(const FProcessingArgs& InArgs)
 
 	RepairDamagedFiles();
 
-	TArray<FString> DenyList;
-
 	switch (InArgs.ProcessingMode)
 	{
 	case ETextAssetCommandletMode::FindMismatchedSerializers:
@@ -401,15 +399,6 @@ bool UTextAssetCommandlet::DoTextAssetProcessing(const FProcessingArgs& InArgs)
 
 		bIgnore = bIgnore || (InputAssetFilename.Contains(TEXT("_BuiltData")));
 
-		for (const FString& DeniedItem : DenyList)
-		{
-			if (InputAssetFilename.Contains(DeniedItem))
-			{
-				bIgnore = true;
-				break;
-			}
-		}
-
 		if (bIgnore)
 		{
 			continue;
@@ -486,7 +475,6 @@ bool UTextAssetCommandlet::DoTextAssetProcessing(const FProcessingArgs& InArgs)
 		TArray<FString> PhaseSuccess;
 		TArray<TArray<FString>> PhaseFails;
 		PhaseFails.AddDefaulted(3);
-		TArray<FString> IntermediateFilenames;
 
 		for (const TTuple<FString, FString>& FileToProcess : FilesToProcess)
 		{
@@ -496,8 +484,6 @@ bool UTextAssetCommandlet::DoTextAssetProcessing(const FProcessingArgs& InArgs)
 
 			TRACE_CPUPROFILER_EVENT_SCOPE_TEXT(*SourceFilename);
 
-			IntermediateFilenames.Empty();
-
 			double StartTime = FPlatformTime::Seconds();
 
 			switch (InArgs.ProcessingMode)
@@ -711,14 +697,6 @@ bool UTextAssetCommandlet::DoTextAssetProcessing(const FProcessingArgs& InArgs)
 				IFileManager::Get().Delete(*SourceFilename, false, true, true);
 				IFileManager::Get().Move(*SourceFilename, *SourceBackupFilename);
 
-				if (!bDisableCleanup)
-				{
-					for (const FString& IntermediateFilename : IntermediateFilenames)
-					{
-						IFileManager::Get().Delete(*IntermediateFilename, false, true, true);
-					}
-				}
-
 				if (!bPhasesMatched[0])
 				{
 					UE_LOG(LogTextAsset, Display, TEXT("-----------------------------------------------------------------------------------------"));
diff --git a/Engine/Source/Editor/UnrealEd/Private/Instances/EditorPlacementSettings.cpp b/Engine/Source/Editor/UnrealEd/Private/Instances/EditorPlacementSettings.cpp
--- a/Engine/Source/Editor/UnrealEd/Private/Instances/EditorPlacementSettings.cpp
+++ b/Engine/Source/Editor/UnrealEd/Private/Instances/EditorPlacementSettings.cpp
@@ -39,6 +39,7 @@ void UEditorInstancedPlacementSettings::Serialize(FArchive& Ar)
 void UEditorInstancedPlacementSettings::RegisterISMDescriptors(AInstancedPlacementPartitionActor* ParentPartitionActor, TSortedMap<int32, TArray<FTransform>>& ISMDefinition) const
 {
 	ISMDefinition.Empty();
+#if 0
 	if (!ParentPartitionActor)
 	{
 		return;
@@ -107,4 +108,5 @@ void UEditorInstancedPlacementSettings::RegisterISMDescriptors(AInstancedPlaceme
 			}
 		}
 	}
+#endif
 }
diff --git a/Engine/Source/Editor/VirtualizationEditor/Private/RehydrationMenuEntry.cpp b/Engine/Source/Editor/VirtualizationEditor/Private/RehydrationMenuEntry.cpp
--- a/Engine/Source/Editor/VirtualizationEditor/Private/RehydrationMenuEntry.cpp
+++ b/Engine/Source/Editor/VirtualizationEditor/Private/RehydrationMenuEntry.cpp
@@ -25,7 +25,7 @@ namespace UE::Virtualization
  */
 void RehydratePackages(TArray<FString> SelectedFiles)
 {
-	TArray<FText> Errors;
+	//TArray<FText> Errors;
 
 	FRehydrationResult Result = IVirtualizationSystem::Get().TryRehydratePackages(SelectedFiles, ERehydrationOptions::Checkout);
 	if (Result.WasSuccessful())
@@ -47,10 +47,10 @@ void RehydratePackages(TArray<FString> SelectedFiles)
 
 		FMessageLog Log("LogVirtualization");
 
-		for (const FText& Error : Errors)
+		/*for (const FText& Error : Errors)
 		{
 			Log.Error(Error);
-		}
+		}*/
 
 		Log.Notify(LOCTEXT("RehydrationFailed", "Failed to rehydrate packages, see the message log for more info"), EMessageSeverity::Info, bForceNotification);
 	}
diff --git a/Engine/Source/Runtime/Engine/Private/Materials/HLSLMaterialTranslator.cpp b/Engine/Source/Runtime/Engine/Private/Materials/HLSLMaterialTranslator.cpp
--- a/Engine/Source/Runtime/Engine/Private/Materials/HLSLMaterialTranslator.cpp
+++ b/Engine/Source/Runtime/Engine/Private/Materials/HLSLMaterialTranslator.cpp
@@ -2369,6 +2369,7 @@ void FHLSLMaterialTranslator::DoTranslate()
 
 void FHLSLMaterialTranslator::ValidateShadingModelsForFeatureLevel(const FMaterialShadingModelField& ShadingModels)
 {
+#if 0
 	if (FeatureLevel <= ERHIFeatureLevel::ES3_1)
 	{
 		const TArray<EMaterialShadingModel>& InvalidShadingModels = {};
@@ -2390,6 +2391,7 @@ void FHLSLMaterialTranslator::ValidateShadingModelsForFeatureLevel(const FMateri
 			}
 		}
 	}
+#endif
 }
 
 void FHLSLMaterialTranslator::GetMaterialEnvironment(EShaderPlatform InPlatform, FShaderCompilerEnvironment& OutEnvironment)
diff --git a/Engine/Source/Runtime/Projects/Private/PluginManager.cpp b/Engine/Source/Runtime/Projects/Private/PluginManager.cpp
--- a/Engine/Source/Runtime/Projects/Private/PluginManager.cpp
+++ b/Engine/Source/Runtime/Projects/Private/PluginManager.cpp
@@ -768,9 +768,9 @@ void FPluginManager::ReadAllPlugins(FDiscoveredPluginMap& Plugins, const TSet<FS
 	TArray<FString> OptionalOutPluginRoots;
 	const FProjectDescriptor* Project = IProjectManager::Get().GetCurrentProject();
 
+#if !WITH_EDITOR
 	// Find any plugin manifest files. These give us the plugin list (and their descriptors) without needing to scour the directory tree.
 	TArray<FString> ManifestFileNames;
-#if !WITH_EDITOR
 	if (Project != nullptr)
 	{
 		FindPluginManifestsInDirectory(*FPaths::ProjectPluginsDir(), ManifestFileNames);
@@ -782,8 +782,10 @@ void FPluginManager::ReadAllPlugins(FDiscoveredPluginMap& Plugins, const TSet<FS
 	// track child plugins that don't want to go into main plugin set
 	TArray<TSharedRef<FPlugin>> ChildPlugins;
 
+#if !WITH_EDITOR
 	// If we didn't find any manifests, do a recursive search for plugins
 	if (ManifestFileNames.Num() == 0)
+#endif
 	{
 		UE_LOG(LogPluginManager, Verbose, TEXT("No *.upluginmanifest files found, looking for *.uplugin files instead."))
 
@@ -823,6 +825,7 @@ void FPluginManager::ReadAllPlugins(FDiscoveredPluginMap& Plugins, const TSet<FS
 			}
 		}
 	}
+#if !WITH_EDITOR
 	else
 	{
 		SlowTask_ReadAll.EnterProgressFrame();
@@ -880,6 +883,7 @@ void FPluginManager::ReadAllPlugins(FDiscoveredPluginMap& Plugins, const TSet<FS
 			}
 		}
 	}
+#endif
 
 	SlowTask_ReadAll.EnterProgressFrame();
 	if (Project != nullptr)
diff --git a/Engine/Source/Runtime/RenderCore/Public/RenderGraphPass.h b/Engine/Source/Runtime/RenderCore/Public/RenderGraphPass.h
--- a/Engine/Source/Runtime/RenderCore/Public/RenderGraphPass.h
+++ b/Engine/Source/Runtime/RenderCore/Public/RenderGraphPass.h
@@ -644,7 +644,7 @@ class FRDGSentinelPass final
 {
 public:
 	FRDGSentinelPass(FRDGEventName&& Name, ERDGPassFlags InPassFlagsToAdd = ERDGPassFlags::None)
-		: FRDGPass(MoveTemp(Name), FRDGParameterStruct(&EmptyShaderParameters, FEmptyShaderParameters::FTypeInfo::GetStructMetadata()), ERDGPassFlags::NeverCull | InPassFlagsToAdd)
+		: FRDGPass(MoveTemp(Name), FRDGParameterStruct(&EmptyShaderParameters, FEmptyShaderParameters::FTypeInfo::GetStructMetadata()), ERDGPassFlags::NeverCull | InPassFlagsToAdd) //-V1050
 	{
 		bSentinel = 1;
 	}
