Cooker: Improve performance of GenerateAssetChunkInformationCSV by calculating distance graph for the entire graph once rather than searching outwards once per vertex.
#rnx

diff --git a/Engine/Source/Editor/UnrealEd/Private/Commandlets/AssetRegistryGenerator.cpp b/Engine/Source/Editor/UnrealEd/Private/Commandlets/AssetRegistryGenerator.cpp
--- a/Engine/Source/Editor/UnrealEd/Private/Commandlets/AssetRegistryGenerator.cpp
+++ b/Engine/Source/Editor/UnrealEd/Private/Commandlets/AssetRegistryGenerator.cpp
@@ -17,6 +17,7 @@
 #include "CookMetadata.h"
 #include "Commandlets/ChunkDependencyInfo.h"
 #include "Commandlets/IChunkDataGenerator.h"
+#include "Containers/RingBuffer.h"
 #include "Engine/AssetManager.h"
 #include "Engine/Level.h"
 #include "Engine/World.h"
@@ -2098,6 +2099,30 @@ bool FAssetRegistryGenerator::GatherAllPackageDependencies(FName PackageName, TA
 	return true;
 }
 
+/**
+ * Helper struct to get the shortest reference chain in cook dependencies from one or more PackageNames
+ * to the set of packages that are hard-imported into a chunk. Constructs the distance graph for
+ * all packages reachable via cookdependencies from the input InSourceSet.
+ */
+class FAssetRegistryGenerator::FGetShortestReferenceChain
+{
+public:
+	void Initialize(const FAssetRegistryGenerator::FChunkPackageSet* InSourceSet);
+	FString Get(FName PackageName);
+
+private:
+	static constexpr int32 NoPath = MAX_int32;
+	struct FVertexData
+	{
+		FName NextTowardSource = NAME_None;
+		int32 SourceDistance = NoPath;
+	};
+
+private:
+	TMap<FName, FVertexData> Vertices;
+	bool bInitialized = false;
+};
+
 bool FAssetRegistryGenerator::GenerateAssetChunkInformationCSV(const FString& OutputPath, bool bWriteIndividualFiles)
 {
 	FString TmpString, TmpStringChunks;
@@ -2140,6 +2165,8 @@ bool FAssetRegistryGenerator::GenerateAssetChunkInformationCSV(const FString& Ou
 		}
 	}
 
+	TMap<int32, FGetShortestReferenceChain> ReferenceChainFinderForChunk;
+
 	for (const FAssetData* AssetDataPtr : AssetDataList)
 	{
 		const FAssetData& AssetData = *AssetDataPtr;
@@ -2160,13 +2187,11 @@ bool FAssetRegistryGenerator::GenerateAssetChunkInformationCSV(const FString& Ou
 
 					if (!bHardChunk)
 					{
-						SoftChain = GetShortestReferenceChain(AssetData.PackageName, PakchunkIndex);
+						FGetShortestReferenceChain& ChainFinder = ReferenceChainFinderForChunk.FindOrAdd(PakchunkIndex);
+						ChainFinder.Initialize(ChunkManifests[PakchunkIndex].Get());
+						SoftChain = ChainFinder.Get(AssetData.PackageName);
 					}
 				}
-				if (SoftChain.IsEmpty())
-				{
-					SoftChain = TEXT("Soft: Possibly Unassigned Asset");
-				}
 
 				// Build "other chunks" string or None if not part of
 				TmpStringChunks.Empty(64);
@@ -2404,67 +2429,84 @@ void FAssetRegistryGenerator::FixupPackageDependenciesForChunks(UE::Cook::FCookS
 	}
 }
 
-void FAssetRegistryGenerator::FindShortestReferenceChain(TArray<FReferencePair> PackageNames, int32 PakchunkIndex, uint32& OutParentIndex, FString& OutChainPath)
+void FAssetRegistryGenerator::FGetShortestReferenceChain::Initialize(const FAssetRegistryGenerator::FChunkPackageSet* SourceSet)
 {
-	TArray<FReferencePair> ReferencesToCheck;
-	uint32 Index = 0;
-	for (const auto& Pkg : PackageNames)
+	if (bInitialized)
 	{
-		if (ChunkManifests[PakchunkIndex] && ChunkManifests[PakchunkIndex]->Contains(Pkg.PackageName))
-		{
-			OutChainPath += TEXT("Soft: ");
-			OutChainPath += Pkg.PackageName.ToString();
-			OutParentIndex = Pkg.ParentNodeIndex;
-			return;
-		}
-		TArray<FName> AssetReferences;
-		AssetRegistry.GetReferencers(Pkg.PackageName, AssetReferences);
-		for (const auto& Ref : AssetReferences)
-		{
-			if (!InspectedNames.Contains(Ref))
-			{
-				ReferencesToCheck.Add(FReferencePair(Ref, Index));
-				InspectedNames.Add(Ref);
-			}
-		}
+		return;
+	}
+	bInitialized = true;
+	if (!SourceSet)
+	{
+		return;
+	}
+
+	IAssetRegistry& AssetRegistryShadowedVar = IAssetRegistry::GetChecked();
+	Vertices.Reset();
 
-		++Index;
+	// BFS the entire graph of dependencies outward from the SourceSet to construct the distance graph
+	TRingBuffer<FName> BFSQueue;
+	for (const TPair<FName, FString>& Pair : *SourceSet)
+	{
+		BFSQueue.Add(Pair.Key);
+		Vertices.FindOrAdd(Pair.Key).SourceDistance = 0;
 	}
 
-	if (ReferencesToCheck.Num() > 0)
+	TArray<FName> AssetDependencies;
+	while (!BFSQueue.IsEmpty())
 	{
-		uint32 ParentIndex = INDEX_NONE;
-		FindShortestReferenceChain(ReferencesToCheck, PakchunkIndex, ParentIndex, OutChainPath);
+		FName CurrentVertex = BFSQueue.PopFrontValue();
+		int32 NextDistance = Vertices[CurrentVertex].SourceDistance + 1;
 
-		if (ParentIndex < (uint32)PackageNames.Num())
+		AssetDependencies.Reset();
+		AssetRegistryShadowedVar.GetDependencies(CurrentVertex, AssetDependencies);
+		for (FName NextVertex : AssetDependencies)
 		{
-			OutChainPath += TEXT("->");
-			OutChainPath += PackageNames[ParentIndex].PackageName.ToString();
-			OutParentIndex = PackageNames[ParentIndex].ParentNodeIndex;
+			FVertexData& NextData = Vertices.FindOrAdd(NextVertex);
+			if (NextData.SourceDistance > NextDistance)
+			{
+				NextData.SourceDistance = NextDistance;
+				NextData.NextTowardSource = CurrentVertex;
+				BFSQueue.Add(NextVertex);
+			}
 		}
 	}
-	else if (PackageNames.Num() > 0)
-	{
-		//best guess
-		OutChainPath += TEXT("Soft From Unassigned Package? Best Guess: ");
-		OutChainPath += PackageNames[0].PackageName.ToString();
-		OutParentIndex = PackageNames[0].ParentNodeIndex;
-	}
 }
 
-FString FAssetRegistryGenerator::GetShortestReferenceChain(FName PackageName, int32 PakchunkIndex)
+FString FAssetRegistryGenerator::FGetShortestReferenceChain::Get(FName TargetPackageName)
 {
-	FString StringChain;
-	TArray<FReferencePair> ReferencesToCheck;
-	uint32 ParentIndex;
-	ReferencesToCheck.Add(FReferencePair(PackageName, 0));
-	InspectedNames.Empty();
-	InspectedNames.Add(PackageName);
-	FindShortestReferenceChain(ReferencesToCheck, PakchunkIndex, ParentIndex, StringChain);
+	FName CurrentVertex = TargetPackageName;
+	FVertexData* VertexData = &Vertices.FindOrAdd(CurrentVertex);
+	if (VertexData->SourceDistance == NoPath)
+	{
+		return TEXT("Soft: Unknown reference chain. Soft From Unassigned Package?");
+	}
+	if (VertexData->SourceDistance == 0)
+	{
+		return FString(TEXT("Hard"));
+	}
 
-	return StringChain;
-}
+	TArray<FName> Chain;
+	Chain.Add(TargetPackageName);
+	int32 MaxPossibleLength = Vertices.Num();
+	for (; VertexData->SourceDistance != 0;)
+	{
+		CurrentVertex = VertexData->NextTowardSource;
+		checkf(!CurrentVertex.IsNone(), TEXT("Distance graph has incomplete path; this should be impossible."));
+		Chain.Add(CurrentVertex);
+		checkf(Chain.Num() <= MaxPossibleLength, TEXT("Cycle in Distance graph; this should be impossible."));
+		VertexData = &Vertices.FindOrAdd(CurrentVertex);
+	}
+	check(Chain.Num() > 1); // Loop runs at least once
 
+	TStringBuilder<1024> Result;
+	Result << TEXT("Soft: ") << Chain[Chain.Num() - 1];
+	for (int32 Index = Chain.Num() - 2; Index >= 0; --Index)
+	{
+		Result << TEXT("->") << Chain[Index];
+	}
+	return FString(Result);
+}
 
 bool FAssetRegistryGenerator::CreateOrEmptyCollection(FName CollectionName)
 {
diff --git a/Engine/Source/Editor/UnrealEd/Public/Commandlets/AssetRegistryGenerator.h b/Engine/Source/Editor/UnrealEd/Public/Commandlets/AssetRegistryGenerator.h
--- a/Engine/Source/Editor/UnrealEd/Public/Commandlets/AssetRegistryGenerator.h
+++ b/Engine/Source/Editor/UnrealEd/Public/Commandlets/AssetRegistryGenerator.h
@@ -266,6 +266,8 @@ public:
 	static void UpdateAssetManagerDatabase();
 
 private:
+	class FGetShortestReferenceChain;
+
 	/**
 	 * Updates all asset package flags in the specified package
 	 *
@@ -332,8 +334,6 @@ private:
 	TArray<TUniquePtr<FChunkPackageSet>> FinalChunkManifests;
 	/** Additional data generators used when creating chunks */
 	TArray<TSharedRef<IChunkDataGenerator>> ChunkDataGenerators;
-	/** Lookup table of used package names used when searching references. */
-	TSet<FName> InspectedNames;
 	/** Source of the config-driven parent-child relationships between chunks. */
 	UChunkDependencyInfo& DependencyInfo;
 
@@ -343,24 +343,6 @@ private:
 	/** Mapping from chunk id to pakchunk file index. If not defined, Pakchunk index will be the same as chunk id by default */
 	TMap<int32, int32> ChunkIdPakchunkIndexMapping;
 
-	struct FReferencePair
-	{
-		FReferencePair() {}
-
-		FReferencePair(const FName& InName, uint32 InParentIndex)
-			: PackageName(InName)
-			, ParentNodeIndex(InParentIndex)
-		{}
-
-		bool operator == (const FReferencePair& RHS) const
-		{
-			return PackageName == RHS.PackageName;
-		}
-
-		FName		PackageName;
-		uint32		ParentNodeIndex;
-	};
-
 	/**
 	 * Updates AssetData with TagsAndValues corresponding to any collections 
 	 * flagged for inclusion as asset registry tags.
@@ -442,12 +424,6 @@ private:
 	/** Save a CSV dump of chunk asset information, if bWriteIndividualFiles is true it writes a CSV per chunk in addition to AllChunksInfo */
 	bool GenerateAssetChunkInformationCSV(const FString& OutputPath, bool bWriteIndividualFiles = false);
 
-	/** Finds the asset belonging to ChunkID with the smallest number of links to Packages In PackageNames */
-	void FindShortestReferenceChain(TArray<FReferencePair> PackageNames, int32 ChunkID, uint32& OutParentIndex, FString& OutChainPath);
-
-	/** Helper function for FindShortestReferenceChain */
-	FString	GetShortestReferenceChain(FName PackageName, int32 ChunkID);
-
 	/** Recursively remove redundant packages from child chunks based on the chunk dependency tree. */
 	void SubtractParentChunkPackagesFromChildChunks(const FChunkDependencyTreeNode& Node,
 		const TSet<FName>& CumulativeParentPackages, TArray<TArray<FName>>& OutPackagesMovedBetweenChunks);
