Fix debugging watched pins in child anim BPs

diff --git a/Engine/Source/Editor/UnrealEd/Private/Kismet2/KismetDebugUtilities.cpp b/Engine/Source/Editor/UnrealEd/Private/Kismet2/KismetDebugUtilities.cpp
--- a/Engine/Source/Editor/UnrealEd/Private/Kismet2/KismetDebugUtilities.cpp
+++ b/Engine/Source/Editor/UnrealEd/Private/Kismet2/KismetDebugUtilities.cpp
@@ -1915,6 +1915,16 @@ FKismetDebugUtilities::EWatchTextResult FKismetDebugUtilities::FindDebuggingData
 
 			// see if our WatchPin is on a animation node & if so try to get its property info
 			const UAnimBlueprintGeneratedClass* AnimBlueprintGeneratedClass = Cast<UAnimBlueprintGeneratedClass>(Blueprint->GeneratedClass);
+
+			// Use the root anim BP's debug data - derived anim BPs have empty debug data
+			if (UAnimBlueprint* AnimBlueprint = Cast<UAnimBlueprint>(AnimBlueprintGeneratedClass->ClassGeneratedBy))
+			{
+				if(UAnimBlueprint* RootAnimBP = UAnimBlueprint::FindRootAnimBlueprint(AnimBlueprint))
+				{
+					AnimBlueprintGeneratedClass = RootAnimBP->GetAnimBlueprintGeneratedClass();
+				}
+			}
+			
 			if (!PropertyBase && AnimBlueprintGeneratedClass)
 			{
 				// are we linked to an anim graph node?
