Fix crash when editing a blueprint introduced by previous CL 30500060

diff --git a/Engine/Source/Editor/UnrealEd/Private/Kismet2/KismetDebugUtilities.cpp b/Engine/Source/Editor/UnrealEd/Private/Kismet2/KismetDebugUtilities.cpp
--- a/Engine/Source/Editor/UnrealEd/Private/Kismet2/KismetDebugUtilities.cpp
+++ b/Engine/Source/Editor/UnrealEd/Private/Kismet2/KismetDebugUtilities.cpp
@@ -1917,14 +1917,17 @@ FKismetDebugUtilities::EWatchTextResult FKismetDebugUtilities::FindDebuggingData
 			const UAnimBlueprintGeneratedClass* AnimBlueprintGeneratedClass = Cast<UAnimBlueprintGeneratedClass>(Blueprint->GeneratedClass);
 
 			// Use the root anim BP's debug data - derived anim BPs have empty debug data
-			if (UAnimBlueprint* AnimBlueprint = Cast<UAnimBlueprint>(AnimBlueprintGeneratedClass->ClassGeneratedBy))
+			if(AnimBlueprintGeneratedClass)
 			{
-				if(UAnimBlueprint* RootAnimBP = UAnimBlueprint::FindRootAnimBlueprint(AnimBlueprint))
+				if (UAnimBlueprint* AnimBlueprint = Cast<UAnimBlueprint>(AnimBlueprintGeneratedClass->ClassGeneratedBy))
 				{
-					AnimBlueprintGeneratedClass = RootAnimBP->GetAnimBlueprintGeneratedClass();
+					if(UAnimBlueprint* RootAnimBP = UAnimBlueprint::FindRootAnimBlueprint(AnimBlueprint))
+					{
+						AnimBlueprintGeneratedClass = RootAnimBP->GetAnimBlueprintGeneratedClass();
+					}
 				}
 			}
-			
+
 			if (!PropertyBase && AnimBlueprintGeneratedClass)
 			{
 				// are we linked to an anim graph node?
