Assets (i.e. not umaps) with external packages now save correctly when saving from the content browser

diff --git a/Engine/Source/Editor/UnrealEd/Private/FileHelpers.cpp b/Engine/Source/Editor/UnrealEd/Private/FileHelpers.cpp
--- a/Engine/Source/Editor/UnrealEd/Private/FileHelpers.cpp
+++ b/Engine/Source/Editor/UnrealEd/Private/FileHelpers.cpp
@@ -583,6 +583,38 @@ FString FEditorFileUtils::GetFilterString(EFileInteraction Interaction)
 
 ///////////////////////////////////////////////////////////////////////////////
 
+// Gather external packages to save for maps an other package that use external packages
+// @param	InPackage				The package that is being saved
+// @param	InOutPackagesToSave		Any packages to save. Array will be modified with additional packages if they are found, then those packages will all be saved
+// @param	bInNewlyCreated			Whether the package was newly created
+// @param	bInAutosaving			Should be set to true if autosaving
+// @returns true if all the save operations completed successfully
+static bool SaveExternalPackages(UPackage* InPackage, TArray<UPackage*>& InOutPackagesToSave, bool bInNewlyCreated, bool bInAutosaving)
+{
+	bool bSuccess = true;
+
+	if (!bInAutosaving && (!FEditorFileUtils::ShouldSkipExternalObjectSave() || bInNewlyCreated))
+	{
+		for (UPackage* ExternalPackage : InPackage->GetExternalPackages())
+		{
+			if (!FPackageName::IsTempPackage(ExternalPackage->GetName()))
+			{
+				InOutPackagesToSave.Add(ExternalPackage);
+			}
+		}
+
+		if (InOutPackagesToSave.Num())
+		{
+			if (!UEditorLoadingAndSavingUtils::SavePackages(InOutPackagesToSave, /*bCheckDirty=*/ !bInNewlyCreated))
+			{
+				FMessageDialog::Open(EAppMsgType::Ok, NSLOCTEXT("UnrealEd", "Error_FailedToSaveExternalPackages", "Failed to save external packages"));
+				bSuccess = false;
+			}
+		}
+	}
+
+	return bSuccess;
+}
 
 /**
  * @param	World					The world to save.
@@ -908,28 +940,9 @@ static bool SaveWorld(UWorld* World,
 
 		SlowTask.EnterProgressFrame(50);
 
-		if (!bAutosaving && (!FEditorFileUtils::ShouldSkipExternalObjectSave() || bNewlyCreated))
+		if (bSuccess)
 		{
-			if (bSuccess)
-			{
-				// Gather external actors to save
-				for (UPackage* ExternalPackage : Package->GetExternalPackages())
-				{
-					if (!FPackageName::IsTempPackage(ExternalPackage->GetName()))
-					{
-						PackagesToSave.Add(ExternalPackage);
-					}
-				}
-
-				if (PackagesToSave.Num())
-				{
-					if (!UEditorLoadingAndSavingUtils::SavePackages(PackagesToSave, /*bCheckDirty=*/ !bNewlyCreated))
-					{
-						FMessageDialog::Open(EAppMsgType::Ok, NSLOCTEXT("UnrealEd", "Error_FailedToSaveHLODLayersPackages", "Failed to save dependant map packages"));
-						bSuccess = false;
-					}
-				}
-			}
+			bSuccess = SaveExternalPackages(Package, PackagesToSave, bNewlyCreated, bAutosaving);
 		}
 
 		if (bSuccess)
@@ -1018,6 +1031,19 @@ static bool SaveWorld(UWorld* World,
 	return bSuccess;
 }
 
+// Save an individual asset's package as well as any external packages too
+// @param	InPackage				The package to save
+// @param	PackageName				The name of the package to save
+// @param	FinalPackageSavePath	The save path of the package
+// @param	SaveOutput				Output device for error reporting
+// @returns true if all the save operations completed successfully
+static bool SaveAsset(UPackage* InPackage, const FString& PackageName, const FString& FinalPackageSavePath, FOutputDevice& SaveOutput)
+{
+	TArray<UPackage*> PackagesToSave;
+	return	SaveExternalPackages(InPackage, PackagesToSave, InPackage->HasAnyPackageFlags(PKG_NewlyCreated), false) &&
+			GEngine->Exec(nullptr, *FString::Printf( TEXT("OBJ SAVEPACKAGE PACKAGE=\"%s\" FILE=\"%s\" SILENT=true"), *PackageName, *FinalPackageSavePath));
+}
+
 FString FEditorFileUtils::GetAutoSaveFilename(UPackage* const Package, const FString& AbsoluteAutosaveDir, const int32 AutoSaveIndex, const FString& PackageExt)
 {
 	// Come up with a meaningful name for the auto-save file
@@ -3513,9 +3539,9 @@ static InternalSavePackageResult InternalSavePackage(UPackage* PackageToSave, bo
 		}
 		else
 		{
-			// normally, we just save the package
+			// normally, we just save the package (and its external packages)
 			SaveOutput.Log("LogFileHelpers", ELogVerbosity::Log, FString::Printf(TEXT("Saving Package: %s"), *PackageName));
-			bWasSuccessful = GEngine->Exec( NULL, *FString::Printf( TEXT("OBJ SAVEPACKAGE PACKAGE=\"%s\" FILE=\"%s\" SILENT=true"), *PackageName, *FinalPackageSavePath ), SaveOutput );
+			bWasSuccessful = SaveAsset(PackageToSave, PackageName, FinalPackageSavePath, SaveOutput);
 		}
 
 		ISourceControlProvider& SourceControlProvider = ISourceControlModule::Get().GetProvider();
