[CompileShadersTestBadCommandlet]
* Added wait for when list of outstanding jobs go over 20000. This solves memory issues when doing some of the more crazy testbeds that contains million of shaders and outstanding jobs can end up being 300k+

diff --git a/Engine/Source/Editor/UnrealEd/Private/Commandlets/CompileShadersTestBedCommandlet.cpp b/Engine/Source/Editor/UnrealEd/Private/Commandlets/CompileShadersTestBedCommandlet.cpp
--- a/Engine/Source/Editor/UnrealEd/Private/Commandlets/CompileShadersTestBedCommandlet.cpp
+++ b/Engine/Source/Editor/UnrealEd/Private/Commandlets/CompileShadersTestBedCommandlet.cpp
@@ -195,6 +195,8 @@ int32 UCompileShadersTestBedCommandlet::Main(const FString& Params)
 
 		int32 PreviousOutstandingJobs = 0;
 
+		constexpr int32 MaxOutstandingJobs = 20000; // Having a max is a way to try to reduce memory usage.. otherwise outstanding jobs can reach 100k+ and use up 300gb committed memory
+
 		// Submit all the jobs.
 		{
 			TRACE_CPUPROFILER_EVENT_SCOPE(SubmitJobs);
@@ -215,16 +217,25 @@ int32 UCompileShadersTestBedCommandlet::Main(const FString& Params)
 
 					GShaderCompilingManager->ProcessAsyncResults(bLimitExecutationTime, false /* bBlockOnGlobalShaderCompilation */);
 
-					const int32 CurrentOutstandingJobs = GShaderCompilingManager->GetNumOutstandingJobs();
-					if (CurrentOutstandingJobs != PreviousOutstandingJobs)
+					while (true)
 					{
-						UE_LOG(LogCompileShadersTestBedCommandlet, Display, TEXT("Outstanding Jobs: %d"), CurrentOutstandingJobs);
-						PreviousOutstandingJobs = CurrentOutstandingJobs;
+						const int32 CurrentOutstandingJobs = GShaderCompilingManager->GetNumOutstandingJobs();
+						if (CurrentOutstandingJobs != PreviousOutstandingJobs)
+						{
+							UE_LOG(LogCompileShadersTestBedCommandlet, Display, TEXT("Outstanding Jobs: %d"), CurrentOutstandingJobs);
+							PreviousOutstandingJobs = CurrentOutstandingJobs;
+						}
+
+						// Flush rendering commands to release any RHI resources (shaders and shader maps).
+						// Delete any FPendingCleanupObjects (shader maps).
+						FlushRenderingCommands();
+
+						if (CurrentOutstandingJobs < MaxOutstandingJobs)
+						{
+							break;
+						}
+						FPlatformProcess::Sleep(1);
 					}
-
-					// Flush rendering commands to release any RHI resources (shaders and shader maps).
-					// Delete any FPendingCleanupObjects (shader maps).
-					FlushRenderingCommands();
 				}
 			}
 		}
@@ -239,16 +250,25 @@ int32 UCompileShadersTestBedCommandlet::Main(const FString& Params)
 			{
 				GShaderCompilingManager->ProcessAsyncResults(bLimitExecutationTime, false /* bBlockOnGlobalShaderCompilation */);
 
-				const int32 CurrentOutstandingJobs = GShaderCompilingManager->GetNumOutstandingJobs();
-				if (CurrentOutstandingJobs != PreviousOutstandingJobs)
+				while (true)
 				{
-					UE_LOG(LogCompileShadersTestBedCommandlet, Display, TEXT("Outstanding Jobs: %d"), CurrentOutstandingJobs);
-					PreviousOutstandingJobs = CurrentOutstandingJobs;
+					const int32 CurrentOutstandingJobs = GShaderCompilingManager->GetNumOutstandingJobs();
+					if (CurrentOutstandingJobs != PreviousOutstandingJobs)
+					{
+						UE_LOG(LogCompileShadersTestBedCommandlet, Display, TEXT("Outstanding Jobs: %d"), CurrentOutstandingJobs);
+						PreviousOutstandingJobs = CurrentOutstandingJobs;
+					}
+
+					// Flush rendering commands to release any RHI resources (shaders and shader maps).
+					// Delete any FPendingCleanupObjects (shader maps).
+					FlushRenderingCommands();
+
+					if (CurrentOutstandingJobs < MaxOutstandingJobs)
+					{
+						break;
+					}
+					FPlatformProcess::Sleep(1);
 				}
-				
-				// Flush rendering commands to release any RHI resources (shaders and shader maps).
-				// Delete any FPendingCleanupObjects (shader maps).
-				FlushRenderingCommands();
 			}
 		}
 
