New TRS Gizmo: don't use the mouse delta tracker to check if we actually moved something

diff --git a/Engine/Source/Editor/UnrealEd/Private/LevelEditorViewport.cpp b/Engine/Source/Editor/UnrealEd/Private/LevelEditorViewport.cpp
--- a/Engine/Source/Editor/UnrealEd/Private/LevelEditorViewport.cpp
+++ b/Engine/Source/Editor/UnrealEd/Private/LevelEditorViewport.cpp
@@ -5486,7 +5486,8 @@ bool FLevelEditorViewportClient::EndTransform(const FGizmoState& InState)
 		TrackingTransaction.TransCount--;
 	}
 
-	const bool bDidMove = bDidAnythingActuallyChange && MouseDeltaTracker->HasReceivedDelta();
+	// TODO ensure that the gizmo actually moved
+	const bool bDidMove = bDidAnythingActuallyChange;
 	const FTypedElementListConstRef ElementsToManipulate = GetElementsToManipulate();
 
 	if (bHasBegunGizmoManipulation)
