Remove unnecessary UNREALED_API annotations on private methods

The UNREALED_API annotations were added to a bunch of methods some time ago when they were moved from being on a class to being on a function. This apparently included a bunch of false-positives as this now meant that all private methods in the CookOnTheFlyServer where marked as external API, which is in all likelihood an error.

diff --git a/Engine/Source/Editor/UnrealEd/Classes/CookOnTheSide/CookOnTheFlyServer.h b/Engine/Source/Editor/UnrealEd/Classes/CookOnTheSide/CookOnTheFlyServer.h
--- a/Engine/Source/Editor/UnrealEd/Classes/CookOnTheSide/CookOnTheFlyServer.h
+++ b/Engine/Source/Editor/UnrealEd/Classes/CookOnTheSide/CookOnTheFlyServer.h
@@ -401,7 +401,7 @@ private:
 	TArray<FWeakObjectPtr> CachedTexturesToCacheArray;
 	int32 LastUpdateTick = 0;
 	int32 MaxPrecacheShaderJobs = 0;
-	UNREALED_API void TickPrecacheObjectsForPlatforms(const float TimeSlice, const TArray<const ITargetPlatform*>& TargetPlatform);
+	void TickPrecacheObjectsForPlatforms(const float TimeSlice, const TArray<const ITargetPlatform*>& TargetPlatform);
 
 	//////////////////////////////////////////////////////////////////////////
 
@@ -427,46 +427,46 @@ private:
 	mutable UE::Cook::FIniSettingContainer AccessedIniStrings;
 	TArray<const FConfigFile*> OpenConfigFiles;
 	TArray<FString> ConfigSettingDenyList;
-	UNREALED_API void OnFConfigDeleted(const FConfigFile* Config);
-	UNREALED_API void OnFConfigCreated(const FConfigFile* Config);
+	void OnFConfigDeleted(const FConfigFile* Config);
+	void OnFConfigCreated(const FConfigFile* Config);
 
-	UNREALED_API void ProcessAccessedIniSettings(const FConfigFile* Config, UE::Cook::FIniSettingContainer& AccessedIniStrings) const;
+	void ProcessAccessedIniSettings(const FConfigFile* Config, UE::Cook::FIniSettingContainer& AccessedIniStrings) const;
 
-	UNREALED_API void OnRequestClusterCompleted(const UE::Cook::FRequestCluster& RequestCluster);
+	void OnRequestClusterCompleted(const UE::Cook::FRequestCluster& RequestCluster);
 
 	/**
 	* OnTargetPlatformChangedSupportedFormats
 	* called when target platform changes the return value of supports shader formats 
 	* used to reset the cached cooked shaders
 	*/
-	UNREALED_API void OnTargetPlatformChangedSupportedFormats(const ITargetPlatform* TargetPlatform);
+	void OnTargetPlatformChangedSupportedFormats(const ITargetPlatform* TargetPlatform);
 
 	/* Initializing Platforms must be done on the tickloop thread; Platform data is read only on other threads */
-	UNREALED_API void AddCookOnTheFlyPlatformFromGameThread(ITargetPlatform* TargetPlatform);
+	void AddCookOnTheFlyPlatformFromGameThread(ITargetPlatform* TargetPlatform);
 	/* Start the session for the given platform in response to the first client connection. */
-	UNREALED_API void StartCookOnTheFlySessionFromGameThread(ITargetPlatform* TargetPlatform);
+	void StartCookOnTheFlySessionFromGameThread(ITargetPlatform* TargetPlatform);
 
 	/* Callback to recalculate all ITargetPlatform* pointers when they change due to modules reloading */
-	UNREALED_API void OnTargetPlatformsInvalidated();
+	void OnTargetPlatformsInvalidated();
 
 	/* Update polled fields used by CookOnTheFly's network request handlers */
-	UNREALED_API void TickNetwork();
+	void TickNetwork();
 
 	/** Tick the cook for the current cook mode: request, load, save, finish. */
-	UNREALED_API void TickMainCookLoop(UE::Cook::FTickStackData& StackData);
+	void TickMainCookLoop(UE::Cook::FTickStackData& StackData);
 
 	/** Execute operations that need to be done after each Scheduler task, such as checking for new external requests. */
-	UNREALED_API void TickCookStatus(UE::Cook::FTickStackData& StackData);
-	UNREALED_API void SetSaveBusy(bool bInSaveBusy);
-	UNREALED_API void SetLoadBusy(bool bInLoadBusy);
+	void TickCookStatus(UE::Cook::FTickStackData& StackData);
+	void SetSaveBusy(bool bInSaveBusy);
+	void SetLoadBusy(bool bInLoadBusy);
 	enum class EIdleStatus
 	{
 		Active,
 		Idle,
 		Done
 	};
-	UNREALED_API void SetIdleStatus(UE::Cook::FTickStackData& StackData, EIdleStatus InStatus);
-	UNREALED_API void UpdateDisplay(UE::Cook::FTickStackData& StackData, bool bForceDisplay);
+	void SetIdleStatus(UE::Cook::FTickStackData& StackData, EIdleStatus InStatus);
+	void UpdateDisplay(UE::Cook::FTickStackData& StackData, bool bForceDisplay);
 	enum class ECookAction
 	{
 		Done,			// The cook is complete; no requests remain in any non-idle state
@@ -481,81 +481,81 @@ private:
 		YieldTick,		// Progress is blocked by an async result. Temporarily exit TickMainCookLoop.
 	};
 	/** Inspect all tasks the scheduler could do and return which one it should do. */
-	UNREALED_API ECookAction DecideNextCookAction(UE::Cook::FTickStackData& StackData);
+	ECookAction DecideNextCookAction(UE::Cook::FTickStackData& StackData);
 	/** How many packages are assigned to the CookOnTheFlyServer for saving/loading? Returns 0 if the current process is not a CookDirector. */
-	UNREALED_API int32 NumMultiprocessLocalWorkerAssignments() const;
+	int32 NumMultiprocessLocalWorkerAssignments() const;
 	/** Execute any existing external callbacks and push any existing external cook requests into new RequestClusters. */
-	UNREALED_API void PumpExternalRequests(const UE::Cook::FCookerTimer& CookerTimer);
+	void PumpExternalRequests(const UE::Cook::FCookerTimer& CookerTimer);
 	/** Send the PackageData back to request state to create a request cluster, if it has not yet been explored. */
-	UNREALED_API bool TryCreateRequestCluster(UE::Cook::FPackageData& PackageData);
+	bool TryCreateRequestCluster(UE::Cook::FPackageData& PackageData);
 	/** Inspect the next package in the RequestQueue and push it on to its next state. Report the number of PackageDatas that were pushed to load. */
-	UNREALED_API void PumpRequests(UE::Cook::FTickStackData& StackData, int32& OutNumPushed);
+	void PumpRequests(UE::Cook::FTickStackData& StackData, int32& OutNumPushed);
 	/**
 	 * Assign the requests found in PumpRequests; either pushing them to ReadyRequests if SingleProcess Cook or
 	 * to CookWorkers if MultiProcess. Input Requests have been sorted by leaf to root load order.
 	 */
-	UNREALED_API void AssignRequests(TArrayView<UE::Cook::FPackageData*> Requests, UE::Cook::FRequestQueue& RequestQueue,
+	void AssignRequests(TArrayView<UE::Cook::FPackageData*> Requests, UE::Cook::FRequestQueue& RequestQueue,
 		TMap<UE::Cook::FPackageData*, TArray<UE::Cook::FPackageData*>>&& RequestGraph);
 	/**
 	 * Multiprocess cook: Notify the CookDirector that a package it assigned to a worker was demoted
 	 * due to e.g. cancelled cook and should be removed from the CookWorker.
 	 */
-	UNREALED_API void NotifyRemovedFromWorker(UE::Cook::FPackageData& PackageData);
+	void NotifyRemovedFromWorker(UE::Cook::FPackageData& PackageData);
 
 	/** Load packages in the LoadQueue until it's time to break. Report the number of loads that were pushed to save. */
-	UNREALED_API void PumpLoads(UE::Cook::FTickStackData& StackData, uint32 DesiredQueueLength, int32& OutNumPushed, bool& bOutBusy);
+	void PumpLoads(UE::Cook::FTickStackData& StackData, uint32 DesiredQueueLength, int32& OutNumPushed, bool& bOutBusy);
 	/** Move packages from LoadPrepare's entry queue into the PreloadingQueue until we run out of Preload slots. */
-	UNREALED_API void PumpPreloadStarts();
+	void PumpPreloadStarts();
 	/** Move preload-completed packages from LoadPrepare->LoadReady until we find one that is not finished preloading. */
-	UNREALED_API void PumpPreloadCompletes();
+	void PumpPreloadCompletes();
 	/**
 	 * Load the given PackageData that was in the load queue and send it on to its next state.
 	 * Report the number of PackageDatas that were pushed to save (0 or 1)
 	 */
-	UNREALED_API void LoadPackageInQueue(UE::Cook::FPackageData& PackageData, uint32& ResultFlags, int32& OutNumPushed);
+	void LoadPackageInQueue(UE::Cook::FPackageData& PackageData, uint32& ResultFlags, int32& OutNumPushed);
 	/** Mark that the given PackageData failed to load and return it to idle. */
-	UNREALED_API void RejectPackageToLoad(UE::Cook::FPackageData& PackageData, const TCHAR* ReasonText,
+	void RejectPackageToLoad(UE::Cook::FPackageData& PackageData, const TCHAR* ReasonText,
 		UE::Cook::ESuppressCookReason Reason);
 
 	/**
 	 * Try to save all packages in the SaveQueue until it's time to break.
 	 * Report the number of requests that were completed (either skipped or successfully saved or failed to save)
 	 */
-	UNREALED_API void PumpSaves(UE::Cook::FTickStackData& StackData, uint32 DesiredQueueLength, int32& OutNumPushed, bool& bOutBusy);
+	void PumpSaves(UE::Cook::FTickStackData& StackData, uint32 DesiredQueueLength, int32& OutNumPushed, bool& bOutBusy);
 
 	/**
 	 * Inspect the given package and queue it for saving if necessary.
 	 *
 	 * @param PackageData			The PackageData to be considered for saving.
 	 */
-	UNREALED_API void QueueDiscoveredPackage(UE::Cook::FPackageData& PackageData, UE::Cook::FInstigator&& Instigator, 
+	void QueueDiscoveredPackage(UE::Cook::FPackageData& PackageData, UE::Cook::FInstigator&& Instigator, 
 		UE::Cook::FDiscoveredPlatformSet&& ReachablePlatforms, bool bUrgent=false);
-	UNREALED_API void QueueDiscoveredPackageOnDirector(UE::Cook::FPackageData& PackageData, UE::Cook::FInstigator&& Instigator,
+	void QueueDiscoveredPackageOnDirector(UE::Cook::FPackageData& PackageData, UE::Cook::FInstigator&& Instigator,
 		UE::Cook::FDiscoveredPlatformSet&& ReachablePlatforms, bool bUrgent);
 
 	/** Called when a package is cancelled and returned to idle. Notifies CookDirector when on a CookWorker. */
-	UNREALED_API void DemoteToIdle(UE::Cook::FPackageData& PackageData, UE::Cook::ESendFlags SendFlags, UE::Cook::ESuppressCookReason Reason);
+	void DemoteToIdle(UE::Cook::FPackageData& PackageData, UE::Cook::ESendFlags SendFlags, UE::Cook::ESuppressCookReason Reason);
 	/** Called when a package completes its save and returns to idle. Notifies the CookDirector when on a CookWorker. */
-	UNREALED_API void PromoteToSaveComplete(UE::Cook::FPackageData& PackageData, UE::Cook::ESendFlags SendFlags);
+	void PromoteToSaveComplete(UE::Cook::FPackageData& PackageData, UE::Cook::ESendFlags SendFlags);
 
 	/**
 	 * Remove all request data about the given platform from any data in the CookOnTheFlyServer.
 	 * Called when a platform is removed from the list of session platforms e.g. because it has not been recently
 	 * used by CookOnTheFly. Does not modify Cooked platforms.
 	 */
-	UNREALED_API void OnRemoveSessionPlatform(const ITargetPlatform* TargetPlatform, int32 RemovedIndex);
+	void OnRemoveSessionPlatform(const ITargetPlatform* TargetPlatform, int32 RemovedIndex);
 	/** Update structures that have data per platfrom when COTF adds a new SessionPlatform. */
-	UNREALED_API void OnPlatformAddedToSession(const ITargetPlatform* TargetPlatform);
-
-	UNREALED_API void InitializePollables();
-	UNREALED_API void PumpPollables(UE::Cook::FTickStackData& StackData, bool bIsIdle);
-	UNREALED_API void PollFlushRenderingCommands();
-	UNREALED_API TRefCountPtr<FPollable> CreatePollableLLM();
-	UNREALED_API TRefCountPtr<FPollable> CreatePollableTriggerGC();
-	UNREALED_API void PollGarbageCollection(UE::Cook::FTickStackData& StackData);
-	UNREALED_API void PollQueuedCancel(UE::Cook::FTickStackData& StackData);
-	UNREALED_API void WaitForAsync(UE::Cook::FTickStackData& StackData);
-	UNREALED_API void TickRecompileShaderRequestsPrivate();
+	void OnPlatformAddedToSession(const ITargetPlatform* TargetPlatform);
+
+	void InitializePollables();
+	void PumpPollables(UE::Cook::FTickStackData& StackData, bool bIsIdle);
+	void PollFlushRenderingCommands();
+	TRefCountPtr<FPollable> CreatePollableLLM();
+	TRefCountPtr<FPollable> CreatePollableTriggerGC();
+	void PollGarbageCollection(UE::Cook::FTickStackData& StackData);
+	void PollQueuedCancel(UE::Cook::FTickStackData& StackData);
+	void WaitForAsync(UE::Cook::FTickStackData& StackData);
+	void TickRecompileShaderRequestsPrivate();
 
 public:
 
@@ -893,55 +893,55 @@ protected:
 
 private:
 
-	UNREALED_API bool PumpHasExceededMaxMemory(uint32& OutResultFlags);
+	bool PumpHasExceededMaxMemory(uint32& OutResultFlags);
 
 	/**
 	 * Is the local CookOnTheFlyServer initialized to run in CookOnTheFly AND using the legacy scheduling method
 	 * that predates COTF2?
 	 */
-	UNREALED_API bool IsUsingLegacyCookOnTheFlyScheduling() const;
-	UNREALED_API bool IsDebugRecordUnsolicited() const;
+	bool IsUsingLegacyCookOnTheFlyScheduling() const;
+	bool IsDebugRecordUnsolicited() const;
 
 	/** Update accumulators of editor data and consume editor changes after last cook when starting cooks in the editor. */
-	UNREALED_API void BeginCookEditorSystems();
-	UNREALED_API void BeginCookPackageWriters(FBeginCookContext& BeginContext);
-	UNREALED_API void BeginCookDirector(FBeginCookContext& BeginContext);
+	void BeginCookEditorSystems();
+	void BeginCookPackageWriters(FBeginCookContext& BeginContext);
+	void BeginCookDirector(FBeginCookContext& BeginContext);
 	/**
 	 * Initialization for systems that persist across CookSessions but that are not available during Initialize(which
 	 * can occur during EngineStartup if !bDisableCookInEditor). We initialize these as late as possible: at the
 	 * beginning of the first session, or for CookOnTheFly at the first cook request.
 	 */
-	UNREALED_API void InitializeAtFirstSession();
+	void InitializeAtFirstSession();
 	/** Initialize steps that are reexecuted at the beginning of every cook session. */
-	UNREALED_API void InitializeSession();
+	void InitializeSession();
 
 	//////////////////////////////////////////////////////////////////////////
 	// cook by the book specific functions
-	UNREALED_API const FCookByTheBookStartupOptions& BlockOnPrebootCookGate(bool& bOutAbortCook,
+	const FCookByTheBookStartupOptions& BlockOnPrebootCookGate(bool& bOutAbortCook,
 		const FCookByTheBookStartupOptions& CookByTheBookStartupOptions,
 		TOptional<FCookByTheBookStartupOptions>& ModifiedStartupOptions);
 	/** Construct the on-stack data for StartCook functions, based on the arguments to StartCookByTheBook */
-	UNREALED_API FBeginCookContext CreateBeginCookByTheBookContext(const FCookByTheBookStartupOptions& StartupOptions);
+	FBeginCookContext CreateBeginCookByTheBookContext(const FCookByTheBookStartupOptions& StartupOptions);
 	/** Set the PlatformManager's session platforms and finish filling out the BeginContext for StartCookByTheBook */
-	UNREALED_API void SelectSessionPlatforms(FBeginCookContext& BeginContext);
+	void SelectSessionPlatforms(FBeginCookContext& BeginContext);
 	/** CollectFilesToCook and add them as external requests. */
-	UNREALED_API void GenerateInitialRequests(FBeginCookContext& BeginContext);
+	void GenerateInitialRequests(FBeginCookContext& BeginContext);
 	/** If cooking DLC, initialize cooked PackageDatas for all of the packages cooked by the base release. */
-	UNREALED_API void RecordDLCPackagesFromBaseGame(FBeginCookContext& BeginContext);
-	UNREALED_API void RegisterCookByTheBookDelegates();
-	UNREALED_API void UnregisterCookByTheBookDelegates();
+	void RecordDLCPackagesFromBaseGame(FBeginCookContext& BeginContext);
+	void RegisterCookByTheBookDelegates();
+	void UnregisterCookByTheBookDelegates();
 	/** Start the collection of EDL diagnostics for the cook, if applicable. */
-	UNREALED_API void BeginCookEDLCookInfo(FBeginCookContext& BeginContext);
+	void BeginCookEDLCookInfo(FBeginCookContext& BeginContext);
 
 	/** Collect all the files which need to be cooked for a cook by the book session */
-	UNREALED_API void CollectFilesToCook(TArray<FName>& FilesInPath, TMap<FName, UE::Cook::FInstigator>& Instigators,
+	void CollectFilesToCook(TArray<FName>& FilesInPath, TMap<FName, UE::Cook::FInstigator>& Instigators,
 		const TArray<FString>& CookMaps, const TArray<FString>& CookDirectories, 
 		const TArray<FString>& IniMapSections, ECookByTheBookOptions FilesToCookFlags,
 		const TArrayView<const ITargetPlatform* const>& TargetPlatforms,
 		const TMap<FName, TArray<FName>>& GameDefaultObjects);
 
 	/** Gets all game default objects for all platforms. */
-	static UNREALED_API void GetGameDefaultObjects(const TArray<ITargetPlatform*>& TargetPlatforms, TMap<FName,
+	static void GetGameDefaultObjects(const TArray<ITargetPlatform*>& TargetPlatforms, TMap<FName,
 		TArray<FName>>& GameDefaultObjectsOut);
 
 	/*
@@ -949,63 +949,63 @@ private:
 	 * Does not include checking UAssetManager, which has to be queried later
 	 * This function is const because it is not always called and should avoid side effects
 	 */
-	UNREALED_API TArray<FName> GetNeverCookPackageFileNames(TArrayView<const FString> ExtraNeverCookDirectories
+	TArray<FName> GetNeverCookPackageFileNames(TArrayView<const FString> ExtraNeverCookDirectories
 		= TArrayView<const FString>()) const;
 
 
 	/** AddFileToCook add file to cook list */
-	UNREALED_API void AddFileToCook( TArray<FName>& InOutFilesToCook, TMap<FName, UE::Cook::FInstigator>& InOutInstigators,
+	void AddFileToCook( TArray<FName>& InOutFilesToCook, TMap<FName, UE::Cook::FInstigator>& InOutInstigators,
 		const FString &InFilename, const UE::Cook::FInstigator& Instigator) const;
 
 	/** Return the name to use for the project's global shader library */
-	UNREALED_API FString GetProjectShaderLibraryName() const;
+	FString GetProjectShaderLibraryName() const;
 
 	/** Invokes the necessary FShaderCodeLibrary functions to start cooking the shader code library. */
-	UNREALED_API void BeginCookStartShaderCodeLibrary(FBeginCookContext& BeginContext);
+	void BeginCookStartShaderCodeLibrary(FBeginCookContext& BeginContext);
 	/**
 	 * Finishes async operations from ShaderCodeLibraryBeginCookStart, saves initial data to disk,
 	 * and opens the library for the rest of the cook
 	 */
-	UNREALED_API void BeginCookFinishShaderCodeLibrary(FBeginCookContext& BeginContext);
+	void BeginCookFinishShaderCodeLibrary(FBeginCookContext& BeginContext);
     
 	/** Opens Global shader library. Global shaderlib is a special case - no chunks and not connected to assets. */
-	UNREALED_API void OpenGlobalShaderLibrary();
+	void OpenGlobalShaderLibrary();
 
 	/** Saves Global shader library. Global shaderlib is a special case - no chunks and not connected to assets. */
-	UNREALED_API void SaveAndCloseGlobalShaderLibrary();
+	void SaveAndCloseGlobalShaderLibrary();
 
 	/** Invokes the necessary FShaderCodeLibrary functions to open a named code library. */
-	UNREALED_API void OpenShaderLibrary(FString const& Name);
+	void OpenShaderLibrary(FString const& Name);
     
 	/**
 	 * Finishes collection of data that should be in the named code library (e.g. load data from previous
 	 * iterative cook)
 	 */
-	UNREALED_API void FinishPopulateShaderLibrary(const ITargetPlatform* TargetPlatform, FString const& Name);
+	void FinishPopulateShaderLibrary(const ITargetPlatform* TargetPlatform, FString const& Name);
 
 	/** Invokes the necessary FShaderCodeLibrary functions to save and close a named code library. */
-	UNREALED_API void SaveShaderLibrary(const ITargetPlatform* TargetPlatform, FString const& Name);
+	void SaveShaderLibrary(const ITargetPlatform* TargetPlatform, FString const& Name);
 
 	/**
 	 * Calls the ShaderPipelineCacheToolsCommandlet to build a upipelinecache file from the stable pipeline cache
 	 * (.spc, in the past textual .stablepc.csv) file, if any
 	 */
-	UNREALED_API void CreatePipelineCache(const ITargetPlatform* TargetPlatform, const FString& LibraryName);
+	void CreatePipelineCache(const ITargetPlatform* TargetPlatform, const FString& LibraryName);
 
 	/** Invokes the necessary FShaderCodeLibrary functions to clean out all the temporary files. */
-	UNREALED_API void CleanShaderCodeLibraries();
+	void CleanShaderCodeLibraries();
 
-	UNREALED_API void RegisterShaderChunkDataGenerator();
+	void RegisterShaderChunkDataGenerator();
 
 	/** Called at the end of CookByTheBook to write aggregated data such as AssetRegistry and shaders. */
-	UNREALED_API void CookByTheBookFinished();
+	void CookByTheBookFinished();
 	void CookByTheBookFinishedInternal();
 
 	/** Clears session-lifetime data from COTFS and CookPackageDatas. */
-	UNREALED_API void ShutdownCookSession();
+	void ShutdownCookSession();
 
 	/** Print some stats when finishing or cancelling CookByTheBook */
-	UNREALED_API void PrintFinishStats();
+	void PrintFinishStats();
 
 	/**
 	 * Get all the packages which are listed in asset registry passed in.  
@@ -1016,49 +1016,49 @@ private:
 	 * @param OutPackageDatas out list of packagename/filenames for packages contained in the asset registry file
 	 * @return true if successfully read false otherwise
 	 */
-	UNREALED_API bool GetAllPackageFilenamesFromAssetRegistry( const FString& AssetRegistryPath, bool bVerifyPackagesExist,
+	bool GetAllPackageFilenamesFromAssetRegistry( const FString& AssetRegistryPath, bool bVerifyPackagesExist,
 		 bool bSkipUncookedPackages, TArray<UE::Cook::FConstructPackageData>& OutPackageDatas) const;
 
 	/** Build a map of the package dependencies used by each Map Package. */
-	UNREALED_API TMap<FName, TSet<FName>> BuildMapDependencyGraph(const ITargetPlatform* TargetPlatform);
+	TMap<FName, TSet<FName>> BuildMapDependencyGraph(const ITargetPlatform* TargetPlatform);
 
 	/** Write a MapDependencyGraph to a metadata file in the sandbox for the given platform. */
-	UNREALED_API void WriteMapDependencyGraph(const ITargetPlatform* TargetPlatform, TMap<FName, TSet<FName>>& MapDependencyGraph);
+	void WriteMapDependencyGraph(const ITargetPlatform* TargetPlatform, TMap<FName, TSet<FName>>& MapDependencyGraph);
 
 	/** Generates the CachedEditorThumbnails.bin file */
-	UNREALED_API void GenerateCachedEditorThumbnails();
+	void GenerateCachedEditorThumbnails();
 
-	UNREALED_API void InitializeAllCulturesToCook(TConstArrayView<FString> CookCultures);
-	UNREALED_API void CompileDLCLocalization(FBeginCookContext& BeginContext);
+	void InitializeAllCulturesToCook(TConstArrayView<FString> CookCultures);
+	void CompileDLCLocalization(FBeginCookContext& BeginContext);
 	/** Find localization dependencies for all packages, used to add required localization files as soft references. */
-	UNREALED_API void GenerateLocalizationReferences();
-	UNREALED_API void RegisterLocalizationChunkDataGenerator();
+	void GenerateLocalizationReferences();
+	void RegisterLocalizationChunkDataGenerator();
 
 	//////////////////////////////////////////////////////////////////////////
 	// cook on the fly specific functions
 
 	/** Construct the on-stack data for StartCook functions, based on the arguments to StartCookOnTheFly */
-	UNREALED_API FBeginCookContext CreateBeginCookOnTheFlyContext(const FCookOnTheFlyStartupOptions& Options);
+	FBeginCookContext CreateBeginCookOnTheFlyContext(const FCookOnTheFlyStartupOptions& Options);
 	/** Construct the on-stack data for StartCook functions, based on the arguments when adding a COTF platform */
-	UNREALED_API FBeginCookContext CreateAddPlatformContext(ITargetPlatform* TargetPlatform);
-	UNREALED_API void GetCookOnTheFlyUnsolicitedFiles(const ITargetPlatform* TargetPlatform, const FString& PlatformName,
+	FBeginCookContext CreateAddPlatformContext(ITargetPlatform* TargetPlatform);
+	void GetCookOnTheFlyUnsolicitedFiles(const ITargetPlatform* TargetPlatform, const FString& PlatformName,
 		TArray<FString>& UnsolicitedFiles, const FString& Filename, bool bIsCookable);
 
 	//////////////////////////////////////////////////////////////////////////
 	// CookWorker specific functions
 	/** Start a CookWorker Session */
-	UNREALED_API void StartCookAsCookWorker();
+	void StartCookAsCookWorker();
 	/** Construct the on-stack data for StartCook functions, based on the CookWorker's configuration */
-	UNREALED_API FBeginCookContext CreateCookWorkerContext();
-	UNREALED_API void CookAsCookWorkerFinished();
+	FBeginCookContext CreateCookWorkerContext();
+	void CookAsCookWorkerFinished();
 	/** Return the best packages to retract and give to other workers -the ones with least effort spent so far. */
-	UNREALED_API void GetPackagesToRetract(int32 NumToRetract, TArray<FName>& OutRetractionPackages);
+	void GetPackagesToRetract(int32 NumToRetract, TArray<FName>& OutRetractionPackages);
 
 	//////////////////////////////////////////////////////////////////////////
 	// general functions
 
 	/** Perform any special processing for freshly loaded packages */
-	UNREALED_API void ProcessUnsolicitedPackages(TArray<FName>* OutDiscoveredPackageNames = nullptr,
+	void ProcessUnsolicitedPackages(TArray<FName>* OutDiscoveredPackageNames = nullptr,
 		TMap<FName, UE::Cook::FInstigator>* OutInstigators = nullptr);
 
 	/**
@@ -1071,45 +1071,45 @@ private:
 	 * @param ReportingPackageData PackageData which caused the load, for e.g. generated packages
 	 * @return Whether LoadPackage was completely successful and the package can be cooked
 	 */
-	UNREALED_API bool LoadPackageForCooking(UE::Cook::FPackageData& PackageData, UPackage*& OutPackage,
+	bool LoadPackageForCooking(UE::Cook::FPackageData& PackageData, UPackage*& OutPackage,
 		UE::Cook::FPackageData* ReportingPackageData = nullptr);
 
 	/**
 	 * Read information about the previous cook from disk and set whether the current cook is iterative based on
 	 * previous cook, config, and commandline.
 	 */
-	UNREALED_API void LoadBeginCookIterativeFlags(FBeginCookContext& BeginContext);
+	void LoadBeginCookIterativeFlags(FBeginCookContext& BeginContext);
 	/** const because it is not always calledand should avoid sideeffects */
-	UNREALED_API void LoadBeginCookIterativeFlagsLocal(FBeginCookContext& BeginContext) const;
+	void LoadBeginCookIterativeFlagsLocal(FBeginCookContext& BeginContext) const;
 	/** Initialize the sandbox for a new cook session */
-	UNREALED_API void BeginCookSandbox(FBeginCookContext& BeginContext);
+	void BeginCookSandbox(FBeginCookContext& BeginContext);
 
 	/** Set parameters that rely on config settings from startup and don't depend on StartCook* functions. */
-	UNREALED_API void LoadInitializeConfigSettings(const FString& InOutputDirectoryOverride);
-	UNREALED_API void SetInitializeConfigSettings(UE::Cook::FInitializeConfigSettings&& Settings);
-	UNREALED_API void ParseCookFilters();
-	UNREALED_API void ParseCookFilters(const TCHAR* Parameter, const TCHAR* Message, TSet<FName>& OutFilterClasses);
+	void LoadInitializeConfigSettings(const FString& InOutputDirectoryOverride);
+	void SetInitializeConfigSettings(UE::Cook::FInitializeConfigSettings&& Settings);
+	void ParseCookFilters();
+	void ParseCookFilters(const TCHAR* Parameter, const TCHAR* Message, TSet<FName>& OutFilterClasses);
 
 	/** Set parameters that rely on config and CookByTheBook settings. */
-	UNREALED_API void LoadBeginCookConfigSettings(FBeginCookContext& BeginContext);
-	UNREALED_API void SetBeginCookConfigSettings(FBeginCookContext& BeginContext, UE::Cook::FBeginCookConfigSettings&& Settings);
-	UNREALED_API void SetNeverCookPackageConfigSettings(FBeginCookContext& BeginContext, UE::Cook::FBeginCookConfigSettings& Settings);
+	void LoadBeginCookConfigSettings(FBeginCookContext& BeginContext);
+	void SetBeginCookConfigSettings(FBeginCookContext& BeginContext, UE::Cook::FBeginCookConfigSettings&& Settings);
+	void SetNeverCookPackageConfigSettings(FBeginCookContext& BeginContext, UE::Cook::FBeginCookConfigSettings& Settings);
 
 	/** Finalize the package store */
-	UNREALED_API void FinalizePackageStore();
+	void FinalizePackageStore();
 
 	/** Empties SavePackageContexts and deletes the contents */
-	UNREALED_API void ClearPackageStoreContexts();
+	void ClearPackageStoreContexts();
 
 	/** Initialize shaders for the specified platforms when running cook on the fly. */
-	UNREALED_API void InitializeShadersForCookOnTheFly(const TArrayView<ITargetPlatform* const>& NewTargetPlatforms);
+	void InitializeShadersForCookOnTheFly(const TArrayView<ITargetPlatform* const>& NewTargetPlatforms);
 
 	/**
 	 * Some content plugins does not support all target platforms.
 	 * Build up a map of unsupported packages per platform that can be checked before saving.
 	 * const because it is not always called and should avoid side effects
 	 */
-	UNREALED_API void DiscoverPlatformSpecificNeverCookPackages(
+	void DiscoverPlatformSpecificNeverCookPackages(
 		const TArrayView<const ITargetPlatform* const>& TargetPlatforms, const TArray<FString>& UBTPlatformStrings,
 		UE::Cook::FBeginCookConfigSettings& Settings) const;
 
@@ -1120,7 +1120,7 @@ private:
 	 * @param Packages List of packages to use as the root set for dependency checking
 	 * @param Found return value, all objects which package is dependent on
 	 */
-	UNREALED_API void GetDependentPackages( const TSet<UPackage*>& Packages, TSet<FName>& Found);
+	void GetDependentPackages( const TSet<UPackage*>& Packages, TSet<FName>& Found);
 
 	/**
 	 * GetDependentPackages
@@ -1129,7 +1129,7 @@ private:
 	 * @param Root set of packages to use when looking for dependencies
 	 * @param FoundPackages list of packages which were found
 	 */
-	UNREALED_API void GetDependentPackages(const TSet<FName>& RootPackages, TSet<FName>& FoundPackages);
+	void GetDependentPackages(const TSet<FName>& RootPackages, TSet<FName>& FoundPackages);
 
 	/**
 	 * ContainsWorld
@@ -1138,7 +1138,7 @@ private:
 	 * @param PackageName to return if it contains the a UWorld object or a ULevel
 	 * @return true if the Package contains a UWorld or ULevel false otherwise
 	 */
-	UNREALED_API bool ContainsMap(const FName& PackageName) const;
+	bool ContainsMap(const FName& PackageName) const;
 
 
 	/** 
@@ -1148,7 +1148,7 @@ private:
 	 * @param RedirectedPaths map of original to redirected object paths
 	 * @return true if the Package contains a redirector false otherwise
 	 */
-	UNREALED_API bool ContainsRedirector(const FName& PackageName, TMap<FSoftObjectPath, FSoftObjectPath>& RedirectedPaths) const;
+	bool ContainsRedirector(const FName& PackageName, TMap<FSoftObjectPath, FSoftObjectPath>& RedirectedPaths) const;
 	
 	/**
 	 * Prepares save by calling BeginCacheForCookedPlatformData on all UObjects in the package. 
@@ -1160,12 +1160,12 @@ private:
 	 * @param bIsPrecaching true if called for precaching
 	 * @return false if time slice was reached, true if all objects have had BeginCacheForCookedPlatformData called
 	 */
-	UNREALED_API UE::Cook::EPollStatus PrepareSave(UE::Cook::FPackageData& PackageData, UE::Cook::FCookerTimer& Timer,
+	UE::Cook::EPollStatus PrepareSave(UE::Cook::FPackageData& PackageData, UE::Cook::FCookerTimer& Timer,
 		bool bPrecaching);
-	UNREALED_API UE::Cook::EPollStatus PrepareSaveInternal(UE::Cook::FPackageData& PackageData, UE::Cook::FCookerTimer& Timer,
+	UE::Cook::EPollStatus PrepareSaveInternal(UE::Cook::FPackageData& PackageData, UE::Cook::FCookerTimer& Timer,
 		bool bPrecaching);
 	/** Call BeginCacheForCookedPlatformData on all objects in a PackageData or a GeneratorPackage's current round. */
-	UNREALED_API UE::Cook::EPollStatus CallBeginCacheOnObjects(UE::Cook::FPackageData& PackageData, UPackage* Package,
+	UE::Cook::EPollStatus CallBeginCacheOnObjects(UE::Cook::FPackageData& PackageData, UPackage* Package,
 		TArray<UE::Cook::FCachedObjectInOuter>& Objects, int32& NextIndex, UE::Cook::FCookerTimer& Timer);
 
 	/**
@@ -1177,14 +1177,14 @@ private:
 	 *        the save again. If true, all data will be wiped.
 	 * @param ReleaseSaveReason Why the save data is being released, allows specifying how much to tear down
 	 */
-	UNREALED_API void ReleaseCookedPlatformData(UE::Cook::FPackageData& PackageData, UE::Cook::EStateChangeReason ReleaseSaveReason);
+	void ReleaseCookedPlatformData(UE::Cook::FPackageData& PackageData, UE::Cook::EStateChangeReason ReleaseSaveReason);
 
 	/**
 	 * Poll the PendingCookedPlatformDatas and release their resources when they are complete.
 	 * This is done inside of PumpSaveQueue, but is also required after a cancelled cook so that references to
 	 * the pending objects will be eventually dropped.
 	 */
-	UNREALED_API void TickCancels();
+	void TickCancels();
 
 	/**
 	* GetCurrentIniVersionStrings gets the current ini version strings for compare against previous cook
@@ -1192,7 +1192,7 @@ private:
 	* @param IniVersionStrings return list of the important current ini version strings
 	* @return false if function fails (should assume all platforms are out of date)
 	*/
-	UNREALED_API bool GetCurrentIniVersionStrings( const ITargetPlatform* TargetPlatform,
+	bool GetCurrentIniVersionStrings( const ITargetPlatform* TargetPlatform,
 		UE::Cook::FIniSettingContainer& IniVersionStrings ) const;
 
 	/**
@@ -1201,73 +1201,73 @@ private:
 	* @param IniVersionStrings return list of the previous cooks ini version strings
 	* @return false if function fails to find the ini version strings
 	*/
-	UNREALED_API bool GetCookedIniVersionStrings( const ITargetPlatform* TargetPlatform,
+	bool GetCookedIniVersionStrings( const ITargetPlatform* TargetPlatform,
 		UE::Cook::FIniSettingContainer& IniVersionStrings, TMap<FString, FString>& AdditionalStrings ) const;
 
 	/**
 	 * Test the CurrentCookSettings against the previous cooksettings to decide if we have to wipe the
 	 * previous cook even when running iteratively.
 	 */
-	UNREALED_API bool ArePreviousCookSettingsCompatible(const TMap<FName, FString>& CurrentCookSettings,
+	bool ArePreviousCookSettingsCompatible(const TMap<FName, FString>& CurrentCookSettings,
 		const ITargetPlatform* TargetPlatform) const;
 	/** Save the CurrentCookSettings into the output directory. */
-	UNREALED_API void SaveCookSettings(const TMap<FName, FString>& CurrentCookSettings, const ITargetPlatform* TargetPlatform);
+	void SaveCookSettings(const TMap<FName, FString>& CurrentCookSettings, const ITargetPlatform* TargetPlatform);
 	/**
 	 * Populate a map suitable for saving as an ini with the current value of the Cook settings that need to be
 	 * tested for compatibility.
 	 */
-	UNREALED_API TMap<FName, FString> CalculateCookSettingStrings() const;
+	TMap<FName, FString> CalculateCookSettingStrings() const;
 	/** Get the path to the CookSettings metadata file for the given platform and current output directory. */
-	UNREALED_API FString GetCookSettingsFileName(const ITargetPlatform* TargetPlatform) const;
+	FString GetCookSettingsFileName(const ITargetPlatform* TargetPlatform) const;
 
 	/**
 	 * Convert a path to a full sandbox path. Is effected by the cooking dlc settings.
 	 * This function should be used instead of calling the FSandbox Sandbox->ConvertToSandboxPath functions.
 	 */
-	UNREALED_API FString ConvertToFullSandboxPath( const FString &FileName, bool bForWrite = false ) const;
-	UNREALED_API FString ConvertToFullSandboxPath( const FString &FileName, bool bForWrite, const FString& PlatformName ) const;
+	FString ConvertToFullSandboxPath( const FString &FileName, bool bForWrite = false ) const;
+	FString ConvertToFullSandboxPath( const FString &FileName, bool bForWrite, const FString& PlatformName ) const;
 
 	/**
 	 * GetSandboxAssetRegistryFilename
 	 * 
 	 * @return full path of the asset registry in the sandbox
 	 */
-	UNREALED_API FString GetSandboxAssetRegistryFilename();
+	FString GetSandboxAssetRegistryFilename();
 
-	UNREALED_API FString GetCookedAssetRegistryFilename(const FString& PlatformName);
+	FString GetCookedAssetRegistryFilename(const FString& PlatformName);
 
 	// Return the filename to use for the cook metadata file, adjusted for DLC, sandbox, and platform.
 	FString GetCookedCookMetadataFilename(const FString& PlatformName);
 	void WriteCookMetadata(const ITargetPlatform* InTargetPlatform, uint64 InDevelopmentAssetRegistryHash);
 
 	/* @return Full path of the CachedEditorThumbnails.bin file in the sandbox */
-	UNREALED_API FString GetSandboxCachedEditorThumbnailsFilename();
+	FString GetSandboxCachedEditorThumbnailsFilename();
 
 	/**
 	 * Get the sandbox root directory for that platform. Is effected by the CookingDlc settings.
 	 * This should be used instead of calling the Sandbox function.
 	 */
-	UNREALED_API FString GetSandboxDirectory( const FString& PlatformName ) const;
+	FString GetSandboxDirectory( const FString& PlatformName ) const;
 
 	/* Create the delete-old-cooked-directory helper.*/
-	UNREALED_API FAsyncIODelete& GetAsyncIODelete();
+	FAsyncIODelete& GetAsyncIODelete();
 
 	/** Is the local CookOnTheFlyServer cooking a Project+Engine+EmbeddedPlugin Patch based on a previous Release? */
-	UNREALED_API bool IsCookingAgainstFixedBase() const;
+	bool IsCookingAgainstFixedBase() const;
 
 	/** Returns whether or not we should populate the Asset Registry using the main game content */
-	UNREALED_API bool ShouldPopulateFullAssetRegistry() const;
+	bool ShouldPopulateFullAssetRegistry() const;
 
 	/**
 	 * GetBaseDirectoryForDLC
 	 * 
 	 * @return return the path to the DLC
 	 */
-	UNREALED_API FString GetBaseDirectoryForDLC() const;
+	FString GetBaseDirectoryForDLC() const;
 
-	UNREALED_API FString GetContentDirectoryForDLC() const;
+	FString GetContentDirectoryForDLC() const;
 
-	UNREALED_API FString GetMountedAssetPathForDLC() const;
+	FString GetMountedAssetPathForDLC() const;
 
 	static FString GetMountedAssetPathForPlugin(const FString& InPluginName);
 
@@ -1277,7 +1277,7 @@ private:
 	 * Is the local CookOnTheFlyServer cooking a Project+Engine+EmbeddedPlugin Release that can be used as a
 	 * base for future DLC or Patches?
 	 */
-	UNREALED_API bool IsCreatingReleaseVersion();
+	bool IsCreatingReleaseVersion();
 
 	/**
 	 * Checks if important ini settings have changed since last cook for each target platform 
@@ -1285,14 +1285,14 @@ private:
 	 * @param TargetPlatforms to check if out of date
 	 * @param OutOfDateTargetPlatforms return list of out of date target platforms which should be cleaned
 	 */
-	UNREALED_API bool IniSettingsOutOfDate( const ITargetPlatform* TargetPlatform ) const;
+	bool IniSettingsOutOfDate( const ITargetPlatform* TargetPlatform ) const;
 
 	/**
 	 * Saves ini settings which are in the memory cache to the hard drive in ini files
 	 *
 	 * @param TargetPlatforms to save
 	 */
-	UNREALED_API bool SaveCurrentIniSettings( const ITargetPlatform* TargetPlatform ) const;
+	bool SaveCurrentIniSettings( const ITargetPlatform* TargetPlatform ) const;
 
 
 
@@ -1302,17 +1302,17 @@ private:
 		return (CookFlags & InCookFlags) != ECookInitializationFlags::None;
 	}
 
-	UNREALED_API void RouteBeginCacheForCookedPlatformData(UE::Cook::FPackageData& PackageData, UObject* Obj,
+	void RouteBeginCacheForCookedPlatformData(UE::Cook::FPackageData& PackageData, UObject* Obj,
 		const ITargetPlatform* TargetPlatform, UE::Cook::ECachedCookedPlatformDataEvent* ExistingEvent);
-	UNREALED_API bool RouteIsCachedCookedPlatformDataLoaded(UE::Cook::FPackageData& PackageData, UObject* Obj,
+	bool RouteIsCachedCookedPlatformDataLoaded(UE::Cook::FPackageData& PackageData, UObject* Obj,
 		const ITargetPlatform* TargetPlatform, UE::Cook::ECachedCookedPlatformDataEvent* ExistingEvent);
-	UNREALED_API EPackageWriterResult SavePackageBeginCacheForCookedPlatformData(FName PackageName,
+	EPackageWriterResult SavePackageBeginCacheForCookedPlatformData(FName PackageName,
 		const ITargetPlatform* TargetPlatform, TConstArrayView<UObject*> SaveableObjects, uint32 SaveFlags);
 
 	/** Cook (save) a package and process the results */
-	UNREALED_API void SaveCookedPackage(UE::Cook::FSaveCookedPackageContext& Context);
+	void SaveCookedPackage(UE::Cook::FSaveCookedPackageContext& Context);
 	/** Helper for package saves using ExternalActors: record ExternalActors for iterative builds. */
-	UNREALED_API void RecordExternalActorDependencies(TConstArrayView<FName> ExternalActorDependencies);
+	void RecordExternalActorDependencies(TConstArrayView<FName> ExternalActorDependencies);
 	friend class UE::Cook::FSaveCookedPackageContext;
 
 	/**
@@ -1320,80 +1320,80 @@ private:
 	 *
 	 * @param Platforms List of platforms to make global shader maps for
 	 */
-	UNREALED_API void SaveGlobalShaderMapFiles(const TArrayView<const ITargetPlatform* const>& Platforms,
+	void SaveGlobalShaderMapFiles(const TArrayView<const ITargetPlatform* const>& Platforms,
 		ODSCRecompileCommand RecompileCommand);
 
 
 	/** Create sandbox file in directory using current settings supplied */
-	UNREALED_API void CreateSandboxFile(FBeginCookContext& BeginContext);
+	void CreateSandboxFile(FBeginCookContext& BeginContext);
 	/** Gets the output directory respecting any command line overrides */
-	UNREALED_API FString GetOutputDirectoryOverride(FBeginCookContext& BeginContext) const;
+	FString GetOutputDirectoryOverride(FBeginCookContext& BeginContext) const;
 
 	/**
 	 * Populate cooked packages from the PackageWriter's previous manifest and assetregistry of cooked output.
 	 * Delete any out of date packages from the PackageWriter's manifest.
 	 */
-	UNREALED_API void PopulateCookedPackages(const TConstArrayView<const ITargetPlatform*> TargetPlatforms);
+	void PopulateCookedPackages(const TConstArrayView<const ITargetPlatform*> TargetPlatforms);
 
 	/** Waits for the AssetRegistry to complete so that we know any missing assets are missing on disk */
-	UNREALED_API void BlockOnAssetRegistry(TConstArrayView<FString> CommandlinePackages);
+	void BlockOnAssetRegistry(TConstArrayView<FString> CommandlinePackages);
 
 	/** Construct or refresh-for-filechanges the platform-specific asset registry for the given platforms */
-	UNREALED_API void RefreshPlatformAssetRegistries(const TArrayView<const ITargetPlatform* const>& TargetPlatforms);
+	void RefreshPlatformAssetRegistries(const TArrayView<const ITargetPlatform* const>& TargetPlatforms);
 
 	/** Generates long package names for all files to be cooked */
-	UNREALED_API void GenerateLongPackageNames(TArray<FName>& FilesInPath, TMap<FName, UE::Cook::FInstigator>& Instigators);
+	void GenerateLongPackageNames(TArray<FName>& FilesInPath, TMap<FName, UE::Cook::FInstigator>& Instigators);
 
-	UNREALED_API UE::Cook::EPollStatus ConditionalCreateGeneratorPackage(UE::Cook::FPackageData& PackageData, bool bPrecaching);
+	UE::Cook::EPollStatus ConditionalCreateGeneratorPackage(UE::Cook::FPackageData& PackageData, bool bPrecaching);
 
 	/** Generate the list of cook-time-created packages created by the Generator. */
-	UNREALED_API UE::Cook::EPollStatus QueueGeneratedPackages(UE::Cook::FGeneratorPackage& Generator,
+	UE::Cook::EPollStatus QueueGeneratedPackages(UE::Cook::FGeneratorPackage& Generator,
 		UE::Cook::FPackageData& PackageData);
 	/** Run additional steps in PrepareSave that are required when the package has a Generator. */
-	UNREALED_API UE::Cook::EPollStatus PrepareSaveGeneratedPackage(UE::Cook::FGeneratorPackage& Generator,
+	UE::Cook::EPollStatus PrepareSaveGeneratedPackage(UE::Cook::FGeneratorPackage& Generator,
 		UE::Cook::FPackageData& PackageData, UE::Cook::FCookerTimer& Timer, bool bPrecaching);
 	/** Call BeginCacheForCookedPlatformData on the objects the Generator plans to move into its main UPackage. */
-	UNREALED_API UE::Cook::EPollStatus BeginCacheObjectsToMove(UE::Cook::FGeneratorPackage& Generator,
+	UE::Cook::EPollStatus BeginCacheObjectsToMove(UE::Cook::FGeneratorPackage& Generator,
 		UE::Cook::FCookGenerationInfo& Info, UE::Cook::FCookerTimer& Timer,
 		TArray<ICookPackageSplitter::FGeneratedPackageForPreSave>& GeneratedPackagesForPresave);
 	/** Call the Generator's PreSaveGeneratorPackage to create/move objects into its main UPackage. */
-	UNREALED_API UE::Cook::EPollStatus PreSaveGeneratorPackage(UE::Cook::FPackageData& PackageData,
+	UE::Cook::EPollStatus PreSaveGeneratorPackage(UE::Cook::FPackageData& PackageData,
 		UE::Cook::FGeneratorPackage& Generator, UE::Cook::FCookGenerationInfo& Info,
 		TArray<ICookPackageSplitter::FGeneratedPackageForPreSave>& GeneratedPackagesForPresave);
 	/** Construct the list of generated packages that is required for some of the CookPackageSplitter interface calls. */
-	UNREALED_API void ConstructGeneratedPackagesForPresave(UE::Cook::FPackageData& PackageData, UE::Cook::FGeneratorPackage& Generator,
+	void ConstructGeneratedPackagesForPresave(UE::Cook::FPackageData& PackageData, UE::Cook::FGeneratorPackage& Generator,
 		TArray<ICookPackageSplitter::FGeneratedPackageForPreSave>& GeneratedPackagesForPresave);
 	/** Call BeginCacheForCookedPlatformData on any undeclared objects in the Generator's main UPackage after the move. */
-	UNREALED_API UE::Cook::EPollStatus BeginCachePostMove(UE::Cook::FGeneratorPackage& Generator,
+	UE::Cook::EPollStatus BeginCachePostMove(UE::Cook::FGeneratorPackage& Generator,
 		UE::Cook::FCookGenerationInfo& Info, UE::Cook::FCookerTimer& Timer);
 
 	/** Try creating (or finding from earlier creation) the generated package for later population */
-	UNREALED_API UPackage* TryCreateGeneratedPackage(UE::Cook::FGeneratorPackage& Generator, UE::Cook::FCookGenerationInfo& GeneratedInfo);
+	UPackage* TryCreateGeneratedPackage(UE::Cook::FGeneratorPackage& Generator, UE::Cook::FCookGenerationInfo& GeneratedInfo);
 	/** Try calling the splitter's populate to create the package */
-	UNREALED_API UE::Cook::EPollStatus TryPopulateGeneratedPackage(UE::Cook::FGeneratorPackage& Generator,
+	UE::Cook::EPollStatus TryPopulateGeneratedPackage(UE::Cook::FGeneratorPackage& Generator,
 		UE::Cook::FCookGenerationInfo& GeneratedInfo);
 
-	UNREALED_API ICookedPackageWriter& FindOrCreatePackageWriter(const ITargetPlatform* TargetPlatform);
-	UNREALED_API const ICookedPackageWriter* FindPackageWriter(const ITargetPlatform* TargetPlatform) const;
-	UNREALED_API void FindOrCreateSaveContexts(TConstArrayView<const ITargetPlatform*> TargetPlatforms);
-	UNREALED_API UE::Cook::FCookSavePackageContext& FindOrCreateSaveContext(const ITargetPlatform* TargetPlatform);
-	UNREALED_API const UE::Cook::FCookSavePackageContext* FindSaveContext(const ITargetPlatform* TargetPlatform) const;
+	ICookedPackageWriter& FindOrCreatePackageWriter(const ITargetPlatform* TargetPlatform);
+	const ICookedPackageWriter* FindPackageWriter(const ITargetPlatform* TargetPlatform) const;
+	void FindOrCreateSaveContexts(TConstArrayView<const ITargetPlatform*> TargetPlatforms);
+	UE::Cook::FCookSavePackageContext& FindOrCreateSaveContext(const ITargetPlatform* TargetPlatform);
+	const UE::Cook::FCookSavePackageContext* FindSaveContext(const ITargetPlatform* TargetPlatform) const;
 	/** Allocate a new FCookSavePackageContext and ICookedPackageWriter for the given platform. */
-	UNREALED_API UE::Cook::FCookSavePackageContext* CreateSaveContext(const ITargetPlatform* TargetPlatform);
+	UE::Cook::FCookSavePackageContext* CreateSaveContext(const ITargetPlatform* TargetPlatform);
 	/**
 	 * Delete files that exist for the package in the PackageWriter from a previous cook. Used when packages are
 	 * detected as invalidated for iterative cook after cook startup (e.g. generated packages).
 	 */
-	UNREALED_API void DeleteOutputForPackage(FName PackageName, const ITargetPlatform* TargetPlatform);
+	void DeleteOutputForPackage(FName PackageName, const ITargetPlatform* TargetPlatform);
 
 	/**
 	 * Set the given package as the active package for diagnostics (e.g. hidden dependencies).
  	 * Must be balanced with call to ClearActivePackage. Consider using FScopedActivePackage instead.
 	 * @param PackageTrackingOpsName Optional op for TObjectPtr tracking. Set to NAME_None to not set TObjectPtr context.
 	 */
-	UNREALED_API void SetActivePackage(FName PackageName, FName PackageTrackingOpsName);
+	void SetActivePackage(FName PackageName, FName PackageTrackingOpsName);
 	/** Clear the package for diagnostics. Must be balanced with call to SetActivePackage */
-	UNREALED_API void ClearActivePackage();
+	void ClearActivePackage();
 	struct FActivePackageData
 	{
 		FName PackageName;
@@ -1418,16 +1418,16 @@ private:
 		TGuardValue<bool> SoftGCGuard;
 	};
 	/** Callback for FGenericCrashContext; provides the current ActivePackage as context. */
-	UNREALED_API void DumpCrashContext(FCrashContextExtendedWriter& Writer);
+	void DumpCrashContext(FCrashContextExtendedWriter& Writer);
 	/** Callback for analytics when a new UPackage is loaded. */
-	UNREALED_API void OnDiscoveredPackageDebug(FName PackageName, const UE::Cook::FInstigator& Instigator);
+	void OnDiscoveredPackageDebug(FName PackageName, const UE::Cook::FInstigator& Instigator);
 	/** Callback for analytics when a TObjectPtr is read. */
-	UNREALED_API void OnObjectHandleReadDebug(const TArrayView<const UObject*const>& ReadObjects);
+	void OnObjectHandleReadDebug(const TArrayView<const UObject*const>& ReadObjects);
 	/** Send warnings/telemetry when a discovered or read package is found to be a hidden dependency. */
-	UNREALED_API void ReportHiddenDependency(FName Referencer, FName Dependency);
+	void ReportHiddenDependency(FName Referencer, FName Dependency);
 	void BroadcastCookByTheBookStarted();
 
-	static UNREALED_API UCookOnTheFlyServer* ActiveCOTFS;
+	static UCookOnTheFlyServer* ActiveCOTFS;
 	uint32		StatLoadedPackageCount = 0;
 	uint32		StatSavedPackageCount = 0;
 	int32		PackageDataFromBaseGameNum = 0;
@@ -1531,8 +1531,8 @@ private:
 	double LastCookableObjectTickTime = 0.;
 
 	// Cook events that can be listenned to
-	static UNREALED_API FOnCookByTheBookStarted CookByTheBookStartedEvent;
-	static UNREALED_API FOnCookByTheBookFinished CookByTheBookFinishedEvent;
+	static FOnCookByTheBookStarted CookByTheBookStartedEvent;
+	static FOnCookByTheBookFinished CookByTheBookFinishedEvent;
 
 
 	// These structs are TUniquePtr rather than inline members so we can keep their headers private.
