World Partition - Modified landscape spline builder to generate actors which share the same HLOD layer as the source splines

diff --git a/Engine/Source/Editor/UnrealEd/Private/WorldPartition/WorldPartitionLandscapeSplineMeshesBuilder.cpp b/Engine/Source/Editor/UnrealEd/Private/WorldPartition/WorldPartitionLandscapeSplineMeshesBuilder.cpp
--- a/Engine/Source/Editor/UnrealEd/Private/WorldPartition/WorldPartitionLandscapeSplineMeshesBuilder.cpp
+++ b/Engine/Source/Editor/UnrealEd/Private/WorldPartition/WorldPartitionLandscapeSplineMeshesBuilder.cpp
@@ -31,6 +31,7 @@
 #include "WorldPartition/WorldPartitionHelpers.h"
 #include "WorldPartition/WorldPartitionActorDescInstance.h"
 #include "ActorPartition/ActorPartitionSubsystem.h"
+#include "WorldPartition/HLOD/HLODLayer.h"
 #include "Rendering/ColorVertexBuffer.h"
 #include "EngineUtils.h"
 
@@ -264,12 +265,23 @@ ALandscapeSplineMeshesActor* UWorldPartitionLandscapeSplineMeshesBuilder::GetOrC
 	UActorPartitionSubsystem* ActorPartitionSubsystem = InWorld->GetSubsystem<UActorPartitionSubsystem>();
 	check(ActorPartitionSubsystem);
 
+	FGuid GridGuid = InLandscapeGuid;
+
+	AActor* MeshOwner = InMeshComponent->GetOwner();
+	UHLODLayer* HLODLayer = MeshOwner->GetHLODLayer();
+	if (InMeshComponent->IsHLODRelevant() && MeshOwner->IsHLODRelevant() && HLODLayer != nullptr)
+	{
+		// To get a new unique GUID, combine the original GUID with the HLOD Layer package guid
+		GridGuid = FGuid::Combine(GridGuid, HLODLayer->GetPackage()->GetPersistentGuid());
+	}
+
 	// Create or find the placement partition actor
-	auto OnActorCreated = [&InLandscapeGuid](APartitionActor* CreatedPartitionActor)
+	auto OnActorCreated = [&GridGuid, HLODLayer](APartitionActor* CreatedPartitionActor)
 	{
 		if (ALandscapeSplineMeshesActor* LandscapeSplineMeshesActor = CastChecked<ALandscapeSplineMeshesActor>(CreatedPartitionActor))
 		{
-			LandscapeSplineMeshesActor->SetGridGuid(InLandscapeGuid);
+			LandscapeSplineMeshesActor->SetGridGuid(GridGuid);
+			LandscapeSplineMeshesActor->SetHLODLayer(HLODLayer);
 		}
 	};
 
@@ -279,7 +291,7 @@ ALandscapeSplineMeshesActor* UWorldPartitionLandscapeSplineMeshesBuilder::GetOrC
 		InWorld->PersistentLevel,
 		FVector(FVector2D(InMeshComponent->GetComponentLocation()), 0),
 		0,
-		InLandscapeGuid,
+		GridGuid,
 		true,
 		OnActorCreated);
 
