Add a viewport action to focus the selected actors in the outliner (to be used with AlwaysFrameSelection off)

diff --git a/Engine/Source/Editor/LevelEditor/Private/LevelEditorActions.cpp b/Engine/Source/Editor/LevelEditor/Private/LevelEditorActions.cpp
--- a/Engine/Source/Editor/LevelEditor/Private/LevelEditorActions.cpp
+++ b/Engine/Source/Editor/LevelEditor/Private/LevelEditorActions.cpp
@@ -2449,6 +2449,20 @@ void FLevelEditorActionCallbacks::OnShowWorldProperties( TWeakPtr< SLevelEditor
 	LevelEditorModule.GetLevelEditorTabManager()->TryInvokeTab(FName("WorldSettingsTab"));
 }
 
+void FLevelEditorActionCallbacks::OnFocusOutlinerToSelection(TWeakPtr<SLevelEditor> LevelEditor)
+{
+	if (const TSharedPtr<SLevelEditor> Editor = LevelEditor.Pin())
+	{
+		for (TWeakPtr<ISceneOutliner> SceneOutliner : Editor->GetAllSceneOutliners())
+		{
+			if (const TSharedPtr<ISceneOutliner> Outliner = SceneOutliner.Pin())
+			{
+				Outliner->FrameSelectedItems();
+			}
+		}
+	}
+}
+
 void FLevelEditorActionCallbacks::OpenPlaceActors()
 {
 	FLevelEditorModule& LevelEditorModule = FModuleManager::GetModuleChecked<FLevelEditorModule>( TEXT("LevelEditor") );
diff --git a/Engine/Source/Editor/LevelEditor/Private/SLevelEditor.cpp b/Engine/Source/Editor/LevelEditor/Private/SLevelEditor.cpp
--- a/Engine/Source/Editor/LevelEditor/Private/SLevelEditor.cpp
+++ b/Engine/Source/Editor/LevelEditor/Private/SLevelEditor.cpp
@@ -178,6 +178,11 @@ void SLevelEditor::BindCommands()
 		FExecuteAction::CreateStatic( &FLevelEditorActionCallbacks::ExecuteExecCommand, FString( TEXT("CAMERA ALIGN ACTIVEVIEWPORTONLY") ) )
 		);
 
+	LevelEditorCommands->MapAction(
+		FEditorViewportCommands::Get().FocusOutlinerToSelection,
+		FExecuteAction::CreateStatic( &FLevelEditorActionCallbacks::OnFocusOutlinerToSelection, TWeakPtr< SLevelEditor >( SharedThis( this ) ) )
+		);
+
 	if (FPlayWorldCommands::GlobalPlayWorldActions.IsValid())
 	{
 		FUICommandList& PlayWorldActionList = *FPlayWorldCommands::GlobalPlayWorldActions;
diff --git a/Engine/Source/Editor/LevelEditor/Public/LevelEditorActions.h b/Engine/Source/Editor/LevelEditor/Public/LevelEditorActions.h
--- a/Engine/Source/Editor/LevelEditor/Public/LevelEditorActions.h
+++ b/Engine/Source/Editor/LevelEditor/Public/LevelEditorActions.h
@@ -1176,6 +1176,9 @@ public:
 	/** Select the world info actor and show the properties */
 	static void OnShowWorldProperties( TWeakPtr< SLevelEditor > LevelEditor );
 
+	/** Focuses the outliner on the selected actors */
+	static void OnFocusOutlinerToSelection(TWeakPtr<SLevelEditor> LevelEditor);
+
 	/** Open the Place Actors Panel */
 	static void OpenPlaceActors();
 
diff --git a/Engine/Source/Editor/SceneOutliner/Private/SSceneOutliner.cpp b/Engine/Source/Editor/SceneOutliner/Private/SSceneOutliner.cpp
--- a/Engine/Source/Editor/SceneOutliner/Private/SSceneOutliner.cpp
+++ b/Engine/Source/Editor/SceneOutliner/Private/SSceneOutliner.cpp
@@ -2471,6 +2471,15 @@ bool SSceneOutliner::CanUnpinSelectedItems() const
 	return CanUnpinItems(SelectedItems);
 }
 
+void SSceneOutliner::FrameSelectedItems()
+{
+	TArray<TSharedPtr<ISceneOutlinerTreeItem>> SelectedItems = GetSelectedItems();
+	if (!SelectedItems.IsEmpty())
+	{
+		ScrollItemIntoView(SelectedItems.Last());
+	}
+}
+
 FSceneOutlinerTreeItemPtr SSceneOutliner::FindParent(const ISceneOutlinerTreeItem& InItem) const
 {
 	FSceneOutlinerTreeItemPtr Parent = Mode->GetHierarchy()->FindOrCreateParentItem(InItem, TreeItemMap, /*bCreate=*/false);
diff --git a/Engine/Source/Editor/SceneOutliner/Public/ISceneOutliner.h b/Engine/Source/Editor/SceneOutliner/Public/ISceneOutliner.h
--- a/Engine/Source/Editor/SceneOutliner/Public/ISceneOutliner.h
+++ b/Engine/Source/Editor/SceneOutliner/Public/ISceneOutliner.h
@@ -157,6 +157,9 @@ public:
 	/** Returns true if any of the selected items can be unpinned */
 	virtual bool CanUnpinSelectedItems() const = 0;
 
+	/** Scrolls the outliner to the selected item(s). If more are selected, the chosen item is undeterministic. */
+	virtual void FrameSelectedItems() = 0;
+
 	/** Get the active SceneOutlinerMode */
 	const ISceneOutlinerMode* GetMode() const { return Mode; }
 
diff --git a/Engine/Source/Editor/SceneOutliner/Public/SSceneOutliner.h b/Engine/Source/Editor/SceneOutliner/Public/SSceneOutliner.h
--- a/Engine/Source/Editor/SceneOutliner/Public/SSceneOutliner.h
+++ b/Engine/Source/Editor/SceneOutliner/Public/SSceneOutliner.h
@@ -466,6 +466,12 @@ public:
 	 */
 	virtual bool CanUnpinSelectedItems() const override;
 	
+	/**
+	 * Scrolls the outliner to the selected item(s).
+	 * If more are selected, the chosen item is undeterministic.
+	 */
+	virtual void FrameSelectedItems() override;
+
 	/**
 	 * Returns the parent tree item for a given item if it exists, nullptr otherwise.
 	 */
diff --git a/Engine/Source/Editor/UnrealEd/Private/EditorViewportCommands.cpp b/Engine/Source/Editor/UnrealEd/Private/EditorViewportCommands.cpp
--- a/Engine/Source/Editor/UnrealEd/Private/EditorViewportCommands.cpp
+++ b/Engine/Source/Editor/UnrealEd/Private/EditorViewportCommands.cpp
@@ -163,6 +163,7 @@ void FEditorViewportCommands::RegisterCommands()
 	
 	UI_COMMAND( FocusAllViewportsToSelection, "Frame Selected Actors in All Viewports", "Centers all open viewports on the selected actors", EUserInterfaceActionType::Button, FInputChord(EKeys::F, EModifierKey::Shift) );
 	UI_COMMAND( FocusViewportToSelection, "Frame Selected", "Centers the viewport on the selected actors", EUserInterfaceActionType::Button, FInputChord( EKeys::F ) );
+	UI_COMMAND( FocusOutlinerToSelection, "Frame Selected in Outliner", "Focuses the selected actors in all open outliners", EUserInterfaceActionType::Button, FInputChord() );
 
 	UI_COMMAND( LocationGridSnap, "Grid Snap", "Enables or disables snapping to the grid when dragging objects around", EUserInterfaceActionType::ToggleButton, FInputChord() );
 	UI_COMMAND( RotationGridSnap, "Rotation Snap", "Enables or disables snapping objects to a rotation grid", EUserInterfaceActionType::ToggleButton, FInputChord() );
diff --git a/Engine/Source/Editor/UnrealEd/Public/EditorViewportCommands.h b/Engine/Source/Editor/UnrealEd/Public/EditorViewportCommands.h
--- a/Engine/Source/Editor/UnrealEd/Public/EditorViewportCommands.h
+++ b/Engine/Source/Editor/UnrealEd/Public/EditorViewportCommands.h
@@ -214,6 +214,7 @@ public:
 
 	TSharedPtr< FUICommandInfo > FocusAllViewportsToSelection;
 	TSharedPtr< FUICommandInfo > FocusViewportToSelection;
+	TSharedPtr< FUICommandInfo > FocusOutlinerToSelection;
 
 	/** Toggle automatic exposure */
 	TSharedPtr< FUICommandInfo > ToggleAutoExposure;
