Pasting / Duplicating an actor now moves the new one to the "Current Folder" (if it's set)

diff --git a/Engine/Source/Editor/UnrealEd/Private/Factories/EditorFactories.cpp b/Engine/Source/Editor/UnrealEd/Private/Factories/EditorFactories.cpp
--- a/Engine/Source/Editor/UnrealEd/Private/Factories/EditorFactories.cpp
+++ b/Engine/Source/Editor/UnrealEd/Private/Factories/EditorFactories.cpp
@@ -1267,10 +1267,20 @@ UObject* ULevelFactory::FactoryCreateText
 			const FString&	PropText = ActorMapElement.Value;
 			if ( Actor->ShouldImport(FStringView(PropText), bIsMoveToStreamingLevel) )
 			{
+				const FName OldPath = Actor->GetFolderPath();
+
 				Actor->PreEditChange(nullptr);
 				ImportObjectProperties( (uint8*)Actor, *PropText, Actor->GetClass(), Actor, Actor, Warn, 0, INDEX_NONE, NULL, &ExistingToNewMap );
 				bActorChanged = true;
 
+				// Path might have been set through OnLevelActorAdded (i.e. by UActorEditorContextSubsystem).
+				// If that's the case, then let's restore it here, as ImportObjectProperties might have
+				// overridden that with an imported value that we'd like to discard.
+				if (!OldPath.IsNone())
+				{
+					Actor->SetFolderPath(OldPath);
+				}
+
 				if (GWorld == World)
 				{
 					GEditor->SelectActor( Actor, true, false, true );
