World Partition - HLOD Layer: Made HLODMaterial properties non soft ptr
* This used to be necessary as those properties were marked as "config"
* Removed option that allowed the creation of HLOD layer assets from an ini file by the WorldPartitionConvertCommandlet - User must now point to existing assets

diff --git a/Engine/Plugins/Editor/WorldPartitionHLODUtilities/Source/Private/WorldPartition/HLOD/Builders/HLODBuilderMeshApproximate.cpp b/Engine/Plugins/Editor/WorldPartitionHLODUtilities/Source/Private/WorldPartition/HLOD/Builders/HLODBuilderMeshApproximate.cpp
--- a/Engine/Plugins/Editor/WorldPartitionHLODUtilities/Source/Private/WorldPartition/HLOD/Builders/HLODBuilderMeshApproximate.cpp
+++ b/Engine/Plugins/Editor/WorldPartitionHLODUtilities/Source/Private/WorldPartition/HLOD/Builders/HLODBuilderMeshApproximate.cpp
@@ -51,15 +51,11 @@ uint32 UHLODBuilderMeshApproximateSettings::GetCRC() const
 
 	uint32 Hash = Ar.GetCrc();
 
-	if (!HLODMaterial.IsNull())
+	if (HLODMaterial)
 	{
-		UMaterialInterface* Material = HLODMaterial.LoadSynchronous();
-		if (Material)
-		{
-			uint32 MaterialCRC = UHLODProxy::GetCRC(Material);
-			UE_LOG(LogHLODBuilder, VeryVerbose, TEXT(" - Material = %d"), MaterialCRC);
-			Hash = HashCombine(Hash, MaterialCRC);
-		}
+		uint32 MaterialCRC = UHLODProxy::GetCRC(HLODMaterial);
+		UE_LOG(LogHLODBuilder, VeryVerbose, TEXT(" - Material = %d"), MaterialCRC);
+		Hash = HashCombine(Hash, MaterialCRC);
 	}
 
 	return Hash;
@@ -86,7 +82,7 @@ TArray<UActorComponent*> UHLODBuilderMeshApproximate::Build(const FHLODBuildCont
 
 	const UHLODBuilderMeshApproximateSettings* MeshApproximateSettings = CastChecked<UHLODBuilderMeshApproximateSettings>(HLODBuilderSettings);
 	FMeshApproximationSettings UseSettings = MeshApproximateSettings->MeshApproximationSettings; // Make a copy as we may tweak some values
-	UMaterialInterface* HLODMaterial = MeshApproximateSettings->HLODMaterial.LoadSynchronous();
+	UMaterialInterface* HLODMaterial = MeshApproximateSettings->HLODMaterial;
 
 	// When using automatic texture sizing based on draw distance, use the MinVisibleDistance for this HLOD.
 	if (UseSettings.MaterialSettings.TextureSizingType == TextureSizingType_AutomaticFromMeshDrawDistance)
diff --git a/Engine/Plugins/Editor/WorldPartitionHLODUtilities/Source/Private/WorldPartition/HLOD/Builders/HLODBuilderMeshMerge.cpp b/Engine/Plugins/Editor/WorldPartitionHLODUtilities/Source/Private/WorldPartition/HLOD/Builders/HLODBuilderMeshMerge.cpp
--- a/Engine/Plugins/Editor/WorldPartitionHLODUtilities/Source/Private/WorldPartition/HLOD/Builders/HLODBuilderMeshMerge.cpp
+++ b/Engine/Plugins/Editor/WorldPartitionHLODUtilities/Source/Private/WorldPartition/HLOD/Builders/HLODBuilderMeshMerge.cpp
@@ -62,15 +62,11 @@ uint32 UHLODBuilderMeshMergeSettings::GetCRC() const
 
 	uint32 Hash = Ar.GetCrc();
 
-	if (!HLODMaterial.IsNull())
+	if (HLODMaterial)
 	{
-		UMaterialInterface* Material = HLODMaterial.LoadSynchronous();
-		if (Material)
-		{
-			uint32 MaterialCRC = UHLODProxy::GetCRC(Material);
-			UE_LOG(LogHLODBuilder, VeryVerbose, TEXT(" - Material = %d"), MaterialCRC);
-			Hash = HashCombine(Hash, MaterialCRC);
-		}
+		uint32 MaterialCRC = UHLODProxy::GetCRC(HLODMaterial);
+		UE_LOG(LogHLODBuilder, VeryVerbose, TEXT(" - Material = %d"), MaterialCRC);
+		Hash = HashCombine(Hash, MaterialCRC);
 	}
 
 	return Hash;
@@ -92,7 +88,7 @@ TArray<UActorComponent*> UHLODBuilderMeshMerge::Build(const FHLODBuildContext& I
 
 	const UHLODBuilderMeshMergeSettings* MeshMergeSettings = CastChecked<UHLODBuilderMeshMergeSettings>(HLODBuilderSettings);
 	FMeshMergingSettings UseSettings = MeshMergeSettings->MeshMergeSettings; // Make a copy as we may tweak some values
-	UMaterialInterface* HLODMaterial = MeshMergeSettings->HLODMaterial.LoadSynchronous();
+	UMaterialInterface* HLODMaterial = MeshMergeSettings->HLODMaterial;
 
 	// When using automatic texture sizing based on draw distance, use the MinVisibleDistance for this HLOD.
 	if (UseSettings.MaterialSettings.TextureSizingType == TextureSizingType_AutomaticFromMeshDrawDistance)
diff --git a/Engine/Plugins/Editor/WorldPartitionHLODUtilities/Source/Private/WorldPartition/HLOD/Builders/HLODBuilderMeshSimplify.cpp b/Engine/Plugins/Editor/WorldPartitionHLODUtilities/Source/Private/WorldPartition/HLOD/Builders/HLODBuilderMeshSimplify.cpp
--- a/Engine/Plugins/Editor/WorldPartitionHLODUtilities/Source/Private/WorldPartition/HLOD/Builders/HLODBuilderMeshSimplify.cpp
+++ b/Engine/Plugins/Editor/WorldPartitionHLODUtilities/Source/Private/WorldPartition/HLOD/Builders/HLODBuilderMeshSimplify.cpp
@@ -60,15 +60,11 @@ uint32 UHLODBuilderMeshSimplifySettings::GetCRC() const
 
 	uint32 Hash = Ar.GetCrc();
 
-	if (!HLODMaterial.IsNull())
+	if (HLODMaterial)
 	{
-		UMaterialInterface* Material = HLODMaterial.LoadSynchronous();
-		if (Material)
-		{
-			uint32 MaterialCRC = UHLODProxy::GetCRC(Material);
-			UE_LOG(LogHLODBuilder, VeryVerbose, TEXT(" - Material = %d"), MaterialCRC);
-			Hash = HashCombine(Hash, MaterialCRC);
-		}
+		uint32 MaterialCRC = UHLODProxy::GetCRC(HLODMaterial);
+		UE_LOG(LogHLODBuilder, VeryVerbose, TEXT(" - Material = %d"), MaterialCRC);
+		Hash = HashCombine(Hash, MaterialCRC);
 	}
 
 	return Hash;
@@ -87,7 +83,7 @@ TArray<UActorComponent*> UHLODBuilderMeshSimplify::Build(const FHLODBuildContext
 
 	const UHLODBuilderMeshSimplifySettings* MeshSimplifySettings = CastChecked<UHLODBuilderMeshSimplifySettings>(HLODBuilderSettings);
 	FMeshProxySettings UseSettings = MeshSimplifySettings->MeshSimplifySettings; // Make a copy as we may tweak some values
-	UMaterialInterface* HLODMaterial = MeshSimplifySettings->HLODMaterial.LoadSynchronous();
+	UMaterialInterface* HLODMaterial = MeshSimplifySettings->HLODMaterial;
 
 	// When using automatic texture sizing based on draw distance, use the MinVisibleDistance for this HLOD.
 	if (UseSettings.MaterialSettings.TextureSizingType == TextureSizingType_AutomaticFromMeshDrawDistance)
diff --git a/Engine/Plugins/Editor/WorldPartitionHLODUtilities/Source/Private/WorldPartition/HLOD/Utilities/WorldPartitionHLODUtilities.cpp b/Engine/Plugins/Editor/WorldPartitionHLODUtilities/Source/Private/WorldPartition/HLOD/Utilities/WorldPartitionHLODUtilities.cpp
--- a/Engine/Plugins/Editor/WorldPartitionHLODUtilities/Source/Private/WorldPartition/HLOD/Utilities/WorldPartitionHLODUtilities.cpp
+++ b/Engine/Plugins/Editor/WorldPartitionHLODUtilities/Source/Private/WorldPartition/HLOD/Utilities/WorldPartitionHLODUtilities.cpp
@@ -395,17 +395,17 @@ UHLODBuilderSettings* FWorldPartitionHLODUtilities::CreateHLODBuilderSettings(UH
 		{
 		case EHLODLayerType::MeshMerge:
 			CastChecked<UHLODBuilderMeshMergeSettings>(HLODBuilderSettings)->MeshMergeSettings = InHLODLayer->MeshMergeSettings_DEPRECATED;
-			CastChecked<UHLODBuilderMeshMergeSettings>(HLODBuilderSettings)->HLODMaterial = InHLODLayer->HLODMaterial_DEPRECATED;
+			CastChecked<UHLODBuilderMeshMergeSettings>(HLODBuilderSettings)->HLODMaterial = InHLODLayer->HLODMaterial_DEPRECATED.LoadSynchronous();
 			break;
 
 		case EHLODLayerType::MeshSimplify:
 			CastChecked<UHLODBuilderMeshSimplifySettings>(HLODBuilderSettings)->MeshSimplifySettings = InHLODLayer->MeshSimplifySettings_DEPRECATED;
-			CastChecked<UHLODBuilderMeshSimplifySettings>(HLODBuilderSettings)->HLODMaterial = InHLODLayer->HLODMaterial_DEPRECATED;
+			CastChecked<UHLODBuilderMeshSimplifySettings>(HLODBuilderSettings)->HLODMaterial = InHLODLayer->HLODMaterial_DEPRECATED.LoadSynchronous();
 			break;
 
 		case EHLODLayerType::MeshApproximate:
 			CastChecked<UHLODBuilderMeshApproximateSettings>(HLODBuilderSettings)->MeshApproximationSettings = InHLODLayer->MeshApproximationSettings_DEPRECATED;
-			CastChecked<UHLODBuilderMeshApproximateSettings>(HLODBuilderSettings)->HLODMaterial = InHLODLayer->HLODMaterial_DEPRECATED;
+			CastChecked<UHLODBuilderMeshApproximateSettings>(HLODBuilderSettings)->HLODMaterial = InHLODLayer->HLODMaterial_DEPRECATED.LoadSynchronous();
 			break;
 		};
 	}
diff --git a/Engine/Plugins/Editor/WorldPartitionHLODUtilities/Source/Public/WorldPartition/HLOD/Builders/HLODBuilderMeshApproximate.h b/Engine/Plugins/Editor/WorldPartitionHLODUtilities/Source/Public/WorldPartition/HLOD/Builders/HLODBuilderMeshApproximate.h
--- a/Engine/Plugins/Editor/WorldPartitionHLODUtilities/Source/Public/WorldPartition/HLOD/Builders/HLODBuilderMeshApproximate.h
+++ b/Engine/Plugins/Editor/WorldPartitionHLODUtilities/Source/Public/WorldPartition/HLOD/Builders/HLODBuilderMeshApproximate.h
@@ -12,7 +12,7 @@ class UMaterialInterface;
 class UMaterial;
 
 
-UCLASS(Blueprintable, Config = Engine, PerObjectConfig)
+UCLASS(Blueprintable)
 class WORLDPARTITIONHLODUTILITIES_API UHLODBuilderMeshApproximateSettings : public UHLODBuilderSettings
 {
 	GENERATED_UCLASS_BODY()
@@ -20,12 +20,12 @@ class WORLDPARTITIONHLODUTILITIES_API UHLODBuilderMeshApproximateSettings : publ
 	virtual uint32 GetCRC() const override;
 
 	/** Mesh approximation settings */
-	UPROPERTY(EditAnywhere, Config, Category = HLOD, meta=(EditInline))
+	UPROPERTY(EditAnywhere, Category = HLOD, meta=(EditInline))
 	FMeshApproximationSettings MeshApproximationSettings;
 
 	/** Material that will be used by the generated HLOD static mesh */
-	UPROPERTY(EditAnywhere, Config, AdvancedDisplay, Category = HLOD, meta = (DisplayName = "HLOD Material"))
-	TSoftObjectPtr<UMaterialInterface> HLODMaterial;
+	UPROPERTY(EditAnywhere, AdvancedDisplay, Category = HLOD, meta = (DisplayName = "HLOD Material"))
+	TObjectPtr<UMaterialInterface> HLODMaterial;
 };
 
 
diff --git a/Engine/Plugins/Editor/WorldPartitionHLODUtilities/Source/Public/WorldPartition/HLOD/Builders/HLODBuilderMeshMerge.h b/Engine/Plugins/Editor/WorldPartitionHLODUtilities/Source/Public/WorldPartition/HLOD/Builders/HLODBuilderMeshMerge.h
--- a/Engine/Plugins/Editor/WorldPartitionHLODUtilities/Source/Public/WorldPartition/HLOD/Builders/HLODBuilderMeshMerge.h
+++ b/Engine/Plugins/Editor/WorldPartitionHLODUtilities/Source/Public/WorldPartition/HLOD/Builders/HLODBuilderMeshMerge.h
@@ -12,7 +12,7 @@ class UMaterialInterface;
 class UMaterial;
 
 
-UCLASS(Blueprintable, Config = Engine, PerObjectConfig)
+UCLASS(Blueprintable)
 class WORLDPARTITIONHLODUTILITIES_API UHLODBuilderMeshMergeSettings : public UHLODBuilderSettings
 {
 	GENERATED_UCLASS_BODY()
@@ -20,12 +20,12 @@ class WORLDPARTITIONHLODUTILITIES_API UHLODBuilderMeshMergeSettings : public UHL
 	virtual uint32 GetCRC() const override;
 
 	/** Merged mesh generation settings */
-	UPROPERTY(EditAnywhere, Config, Category = HLOD)
+	UPROPERTY(EditAnywhere, Category = HLOD)
 	FMeshMergingSettings MeshMergeSettings;
 
 	/** Material that will be used by the generated HLOD static mesh */
-	UPROPERTY(EditAnywhere, Config, AdvancedDisplay, Category = HLOD, meta = (DisplayName = "HLOD Material"))
-	TSoftObjectPtr<UMaterialInterface> HLODMaterial;
+	UPROPERTY(EditAnywhere, AdvancedDisplay, Category = HLOD, meta = (DisplayName = "HLOD Material"))
+	TObjectPtr<UMaterialInterface> HLODMaterial;
 };
 
 
diff --git a/Engine/Plugins/Editor/WorldPartitionHLODUtilities/Source/Public/WorldPartition/HLOD/Builders/HLODBuilderMeshSimplify.h b/Engine/Plugins/Editor/WorldPartitionHLODUtilities/Source/Public/WorldPartition/HLOD/Builders/HLODBuilderMeshSimplify.h
--- a/Engine/Plugins/Editor/WorldPartitionHLODUtilities/Source/Public/WorldPartition/HLOD/Builders/HLODBuilderMeshSimplify.h
+++ b/Engine/Plugins/Editor/WorldPartitionHLODUtilities/Source/Public/WorldPartition/HLOD/Builders/HLODBuilderMeshSimplify.h
@@ -12,7 +12,7 @@ class UMaterialInterface;
 class UMaterial;
 
 
-UCLASS(Blueprintable, Config = Engine, PerObjectConfig)
+UCLASS(Blueprintable)
 class WORLDPARTITIONHLODUTILITIES_API UHLODBuilderMeshSimplifySettings : public UHLODBuilderSettings
 {
 	GENERATED_UCLASS_BODY()
@@ -20,12 +20,12 @@ class WORLDPARTITIONHLODUTILITIES_API UHLODBuilderMeshSimplifySettings : public
 	virtual uint32 GetCRC() const override;
 
 	/** Simplified mesh generation settings */
-	UPROPERTY(EditAnywhere, Config, Category = HLOD)
+	UPROPERTY(EditAnywhere, Category = HLOD)
 	FMeshProxySettings MeshSimplifySettings;
 
 	/** Material that will be used by the generated HLOD static mesh */
-	UPROPERTY(EditAnywhere, Config, AdvancedDisplay, Category = HLOD, meta = (DisplayName = "HLOD Material"))
-	TSoftObjectPtr<UMaterialInterface> HLODMaterial;
+	UPROPERTY(EditAnywhere, AdvancedDisplay, Category = HLOD, meta = (DisplayName = "HLOD Material"))
+	TObjectPtr<UMaterialInterface> HLODMaterial;
 };
 
 
diff --git a/Engine/Source/Editor/UnrealEd/Classes/Commandlets/WorldPartitionConvertCommandlet.h b/Engine/Source/Editor/UnrealEd/Classes/Commandlets/WorldPartitionConvertCommandlet.h
--- a/Engine/Source/Editor/UnrealEd/Classes/Commandlets/WorldPartitionConvertCommandlet.h
+++ b/Engine/Source/Editor/UnrealEd/Classes/Commandlets/WorldPartitionConvertCommandlet.h
@@ -28,7 +28,7 @@ struct FHLODLayerActorMapping
 	TSoftClassPtr<AActor> ActorClass;
 
 	UPROPERTY()
-	FString HLODLayer;
+	FTopLevelAssetPath HLODLayer;
 };
 
 UCLASS(Config = Engine, MinimalAPI)
@@ -69,7 +69,6 @@ protected:
 
 	UNREALED_API void SetupHLOD();
 	UNREALED_API void SetupHLODLayerAssets();
-	UNREALED_API UHLODLayer* CreateHLODLayerFromINI(const FString& InHLODLayerName);
 
 	UNREALED_API void SetActorGuid(AActor* Actor, const FGuid& NewGuid);
 	UNREALED_API void CreateWorldMiniMapTexture(UWorld* World);
@@ -117,13 +116,7 @@ protected:
 	FVector WorldExtent;
 
 	UPROPERTY(Config)
-	FString HLODLayerAssetsPath;
-
-	UPROPERTY(Config)
-	FString DefaultHLODLayerName;
-
-	UPROPERTY(Config)
-	FString DefaultHLODLayerAsset;
+	FTopLevelAssetPath DefaultHLODLayerAsset;
 
 	UPROPERTY(Config)
 	FString FoliageTypePath;
@@ -131,9 +124,6 @@ protected:
 	UPROPERTY(Config)
 	TArray<FHLODLayerActorMapping> HLODLayersForActorClasses;
 
-	UPROPERTY(Transient)
-	TMap<FString, TObjectPtr<UHLODLayer>> HLODLayers;
-
 	UPROPERTY(Config)
 	uint32 LandscapeGridSize;
 
diff --git a/Engine/Source/Editor/UnrealEd/Private/Commandlets/WorldPartitionConvertCommandlet.cpp b/Engine/Source/Editor/UnrealEd/Private/Commandlets/WorldPartitionConvertCommandlet.cpp
--- a/Engine/Source/Editor/UnrealEd/Private/Commandlets/WorldPartitionConvertCommandlet.cpp
+++ b/Engine/Source/Editor/UnrealEd/Private/Commandlets/WorldPartitionConvertCommandlet.cpp
@@ -242,14 +242,10 @@ UWorldPartition* UWorldPartitionConvertCommandlet::CreateWorldPartition(AWorldSe
 		WorldPartition->RuntimeHash->LoadConfig(*RuntimeHashClass, *LevelConfigFilename);
 
 		// Use specified existing default HLOD layer if valid 
-		if (UHLODLayer* ExistingHLODLayer = LoadObject<UHLODLayer>(NULL, *DefaultHLODLayerAsset, nullptr, LOAD_NoWarn))
+		if (UHLODLayer* ExistingHLODLayer = LoadObject<UHLODLayer>(NULL, *DefaultHLODLayerAsset.ToString(), nullptr, LOAD_NoWarn))
 		{
 			WorldPartition->DefaultHLODLayer = ExistingHLODLayer;
 		}
-		else if(UHLODLayer* FoundLayer = HLODLayers.FindRef(DefaultHLODLayerName))
-		{
-			WorldPartition->DefaultHLODLayer = FoundLayer;
-		}
 	}
 
 	if ((WorldPartition->DefaultHLODLayer == UHLODLayer::GetEngineDefaultHLODLayersSetup()) && !bDisableStreaming)
@@ -583,40 +579,6 @@ bool UWorldPartitionConvertCommandlet::RenameWorldPackageWithSuffix(UWorld* Worl
 	return true;
 }
 
-
-UHLODLayer* UWorldPartitionConvertCommandlet::CreateHLODLayerFromINI(const FString& InHLODLayerName)
-{
-	const FString PackagePath = HLODLayerAssetsPath / InHLODLayerName;
-	UPackage* AssetPackage = CreatePackage(*PackagePath);
-	if (!AssetPackage)
-	{
-		UE_LOG(LogWorldPartitionConvertCommandlet, Error, TEXT("Package \"%s\" creation failed"), *PackagePath);
-		return nullptr;
-	}
-
-	// Make sure we overwrite any existing HLODLayer asset package
-	AssetPackage->MarkAsFullyLoaded();
-
-	UHLODLayer* HLODLayer = NewObject<UHLODLayer>(AssetPackage, *InHLODLayerName, RF_Public | RF_Standalone);
-	if (!HLODLayer)
-	{
-		UE_LOG(LogWorldPartitionConvertCommandlet, Error, TEXT("HLODLayer \"%s\" creation failed"), *InHLODLayerName);
-		return nullptr;
-	}
-
-	HLODLayer->LoadConfig(nullptr, *LevelConfigFilename);
-
-	// Notify the asset registry
-	FAssetRegistryModule::AssetCreated(HLODLayer);
-
-	// Mark the package dirty...
-	HLODLayer->Modify();
-
-	PackagesToSave.Add(HLODLayer->GetOutermost());
-
-	return HLODLayer;
-}
-
 void UWorldPartitionConvertCommandlet::SetupHLOD()
 {
 	// No need to spawn HLOD actors during the conversion
@@ -627,23 +589,13 @@ void UWorldPartitionConvertCommandlet::SetupHLOD()
 
 void UWorldPartitionConvertCommandlet::SetupHLODLayerAssets()
 {
-	TArray<FString> HLODLayerSectionsNames;
-	if (GConfig->GetPerObjectConfigSections(LevelConfigFilename, TEXT("HLODLayer"), HLODLayerSectionsNames))
-	{
-		for(const FString& HLODLayerSectionName : HLODLayerSectionsNames)
-		{
-			FString HLODLayerName(*HLODLayerSectionName.Left(HLODLayerSectionName.Find(TEXT(" "))));
-			UHLODLayer* HLODLayer = CreateHLODLayerFromINI(HLODLayerName);
-			HLODLayers.Add(HLODLayerName, HLODLayer);
-		}
-	}
-
 	// Assign HLOD layers to the classes listed in the level config
 	for (const FHLODLayerActorMapping& Entry : HLODLayersForActorClasses)
 	{
-		UHLODLayer* HLODLayer = HLODLayers.FindRef(Entry.HLODLayer);
-		if (!ensure(HLODLayer))
+		UHLODLayer* HLODLayer = LoadObject<UHLODLayer>(NULL, *Entry.HLODLayer.ToString(), nullptr, LOAD_NoWarn);
+		if (!HLODLayer)
 		{
+			UE_LOG(LogWorldPartitionConvertCommandlet, Warning, TEXT("Unable to load HLOD Layer %s, skipping assignment to class %s"), *Entry.HLODLayer.ToString(), *Entry.ActorClass.ToString());
 			continue;
 		}
 
@@ -1714,11 +1666,6 @@ int32 UWorldPartitionConvertCommandlet::Main(const FString& Params)
 				WorldPartition->EditorHash->SaveConfig(CPF_Config, *LevelConfigFilename);
 				WorldPartition->RuntimeHash->SaveConfig(CPF_Config, *LevelConfigFilename);
 			}
-			
-			for(const auto& Pair : HLODLayers)
-			{
-				Pair.Value->SaveConfig(CPF_Config, *LevelConfigFilename);
-			}
 
 			UE_LOG(LogWorldPartitionConvertCommandlet, Display, TEXT("Generated ini file: %s"), *LevelConfigFilename);
 		}
