[Gizmos] Fix an issue where a hidden mode was not hidden in the mode menu

diff --git a/Engine/Source/Editor/UnrealEd/Private/Tools/DefaultEdMode.cpp b/Engine/Source/Editor/UnrealEd/Private/Tools/DefaultEdMode.cpp
--- a/Engine/Source/Editor/UnrealEd/Private/Tools/DefaultEdMode.cpp
+++ b/Engine/Source/Editor/UnrealEd/Private/Tools/DefaultEdMode.cpp
@@ -65,7 +65,7 @@ UAssetEdModeDefault::UAssetEdModeDefault()
 		FAssetEdModes::EM_AssetDefault,
 		NSLOCTEXT("AssetDefaultMode", "DisplayName", "AssetSelection"),
 		FSlateIcon("EditorStyle", "LevelEditor.SelectMode", "LevelEditor.SelectMode.Small"),
-		true, 0);
+		false, 0);
 }
 
 TSharedRef<FLegacyEdModeWidgetHelper> UAssetEdModeDefault::CreateWidgetHelper()
