TextureAssetActions resize to power of two : also update MipGen and NeverStream

#rnx

diff --git a/Engine/Source/Editor/UnrealEd/Private/TextureAssetActions.cpp b/Engine/Source/Editor/UnrealEd/Private/TextureAssetActions.cpp
--- a/Engine/Source/Editor/UnrealEd/Private/TextureAssetActions.cpp
+++ b/Engine/Source/Editor/UnrealEd/Private/TextureAssetActions.cpp
@@ -680,8 +680,15 @@ static void DoResizeTextureSourceToPowerOfTwo(UTexture * Texture)
 		BeforeSourceSize.X,BeforeSourceSize.Y,AfterSourceSize.X,AfterSourceSize.Y);
 
 	// ?? if Texture was set to stretch to Pow2, we could remove that now, but leaving it is harmless
+	//Texture->PowerOfTwoMode = ETexturePowerOfTwoSetting::None;
 
-	// @@ ?? if Texture was set to MipGen = NoMipMaps, change to FromTextureGroup ?
+	// if Texture was set to MipGen = NoMipMaps, change to FromTextureGroup ?
+	// assume that since we changed to Pow2, we want mips and streaming
+	if ( Texture->MipGenSettings == TMGS_NoMipmaps )
+	{
+		Texture->MipGenSettings = TMGS_FromTextureGroup;
+	}
+	Texture->NeverStream = false;
 
 	// this counts as a reimport :
 	UE::TextureUtilitiesCommon::ApplyDefaultsForNewlyImportedTextures(Texture,true);
