Interchange: Fix a crash if re-importing in editor build with GEditor null.

#rnx

diff --git a/Engine/Source/Editor/UnrealEd/Private/Editor.cpp b/Engine/Source/Editor/UnrealEd/Private/Editor.cpp
--- a/Engine/Source/Editor/UnrealEd/Private/Editor.cpp
+++ b/Engine/Source/Editor/UnrealEd/Private/Editor.cpp
@@ -730,19 +730,23 @@ void FReimportManager::SortHandlersIfNeeded()
 
 void FReimportManager::OnInterchangePostReimported(UObject* ReimportAsset) const
 {
-	if (ReimportAsset)
+	if (!ReimportAsset)
+	{
+		return;
+	}
+
+	if (FEngineAnalytics::IsAvailable())
+	{
+		TArray<FAnalyticsEventAttribute> Attributes;
+		Attributes.Add(FAnalyticsEventAttribute(TEXT("ObjectType"), ReimportAsset->GetClass()->GetName()));
+		FEngineAnalytics::GetProvider().RecordEvent(TEXT("Editor.Usage.AssetReimported"), Attributes);
+	}
+
+	PostReimport.Broadcast(ReimportAsset, true);
+
+	if (GEditor)
 	{
-		TArray<UObject*> ObjectArray;
-		ObjectArray.Add(ReimportAsset);
-		//UAssetToolsImpl::Get().SyncBrowserToAssets(ObjectArray);
 		GEditor->BroadcastObjectReimported(ReimportAsset);
-		if (FEngineAnalytics::IsAvailable())
-		{
-			TArray<FAnalyticsEventAttribute> Attributes;
-			Attributes.Add(FAnalyticsEventAttribute(TEXT("ObjectType"), ReimportAsset->GetClass()->GetName()));
-			FEngineAnalytics::GetProvider().RecordEvent(TEXT("Editor.Usage.AssetReimported"), Attributes);
-		}
-		PostReimport.Broadcast(ReimportAsset, true);
 		GEditor->RedrawAllViewports();
 	}
 }
