Sequencer: Check for valid bone index before dereferencing.

Animating the skeletal mesh isn't supported on export

diff --git a/Engine/Source/Editor/UnrealEd/Private/Fbx/FbxAnimationExport.cpp b/Engine/Source/Editor/UnrealEd/Private/Fbx/FbxAnimationExport.cpp
--- a/Engine/Source/Editor/UnrealEd/Private/Fbx/FbxAnimationExport.cpp
+++ b/Engine/Source/Editor/UnrealEd/Private/Fbx/FbxAnimationExport.cpp
@@ -658,6 +658,12 @@ void FFbxExporter::ExportAnimTrack(IAnimTrackAdapter& AnimTrackAdapter, AActor*
 		// Add the animation data to the bone nodes
 		for(int32 BoneIndex = 0; BoneIndex < BoneNodes.Num(); ++BoneIndex)
 		{
+			if (!InSkeletalMeshComponent->GetSkeletalMeshAsset()->GetRefSkeleton().IsValidIndex(BoneIndex))
+			{
+				UE_LOG(LogFbxAnimationExport, Warning, TEXT("Invalid BoneIndex %d, did the skeleton change? (animating the skeleton is not currently supported)"), BoneIndex);
+				continue;
+			}
+
 			FName BoneName = InSkeletalMeshComponent->GetSkeletalMeshAsset()->GetRefSkeleton().GetBoneName(BoneIndex);
 			FbxNode* CurrentBoneNode = BoneNodes[BoneIndex];
 
