Implemented material graph drag&drop support for SparseVolumeTexture assets.

#rnx

diff --git a/Engine/Source/Editor/UnrealEd/Private/MaterialGraphSchema.cpp b/Engine/Source/Editor/UnrealEd/Private/MaterialGraphSchema.cpp
--- a/Engine/Source/Editor/UnrealEd/Private/MaterialGraphSchema.cpp
+++ b/Engine/Source/Editor/UnrealEd/Private/MaterialGraphSchema.cpp
@@ -20,6 +20,7 @@
 #include "Materials/MaterialFunction.h"
 #include "MaterialGraph/MaterialGraph.h"
 #include "Engine/Texture.h"
+#include "SparseVolumeTexture/SparseVolumeTexture.h"
 #include "MaterialGraph/MaterialGraphNode_Base.h"
 #include "MaterialGraph/MaterialGraphNode_Comment.h"
 #include "MaterialGraph/MaterialGraphNode.h"
@@ -32,6 +33,7 @@
 #include "Materials/MaterialExpressionPinBase.h"
 #include "Materials/MaterialExpressionFunctionInput.h"
 #include "Materials/MaterialExpressionTextureSample.h"
+#include "Materials/MaterialExpressionSparseVolumeTextureSample.h"
 #include "Materials/MaterialExpressionFunctionOutput.h"
 #include "Materials/MaterialExpressionReroute.h"
 #include "Materials/MaterialExpressionNamedReroute.h"
@@ -836,6 +838,7 @@ void UMaterialGraphSchema::DroppedAssetsOnGraph(const TArray<struct FAssetData>&
 		UClass* MaterialExpressionClass = Cast<UClass>(Asset);
 		UMaterialFunctionInterface* Func = Cast<UMaterialFunctionInterface>(Asset);
 		UTexture* Tex = Cast<UTexture>(Asset);
+		USparseVolumeTexture* SparseVolumeTexture = Cast<USparseVolumeTexture>(Asset);
 		UMaterialParameterCollection* ParameterCollection = Cast<UMaterialParameterCollection>(Asset);
 		
 		if (MaterialExpressionClass && MaterialExpressionClass->IsChildOf(UMaterialExpression::StaticClass()))
@@ -877,6 +880,16 @@ void UMaterialGraphSchema::DroppedAssetsOnGraph(const TArray<struct FAssetData>&
 
 			bAddedNode = true;
 		}
+		else if ( SparseVolumeTexture )
+		{
+			UMaterialExpressionSparseVolumeTextureSample* SparseVolumeTextureSampleNode = CastChecked<UMaterialExpressionSparseVolumeTextureSample>(
+				FMaterialEditorUtilities::CreateNewMaterialExpression(Graph, UMaterialExpressionSparseVolumeTextureSample::StaticClass(), ExpressionPosition, true, true));
+			SparseVolumeTextureSampleNode->SparseVolumeTexture = SparseVolumeTexture;
+
+			FMaterialEditorUtilities::ForceRefreshExpressionPreviews(Graph);
+
+			bAddedNode = true;
+		}
 		else if ( ParameterCollection )
 		{
 			UMaterialExpressionCollectionParameter* CollectionParameterNode = CastChecked<UMaterialExpressionCollectionParameter>(
@@ -1198,6 +1211,7 @@ void UMaterialGraphSchema::GetAssetsGraphHoverMessage(const TArray<FAssetData>&
 		UClass* MaterialExpressionClass = Cast<UClass>(Asset);
 		UMaterialFunctionInterface* Func = Cast<UMaterialFunctionInterface>(Asset);
 		UTexture* Tex = Cast<UTexture>(Asset);
+		USparseVolumeTexture* SparseVolumeTexture = Cast<USparseVolumeTexture>(Asset);
 		UMaterialParameterCollection* ParameterCollection = Cast<UMaterialParameterCollection>(Asset);
 
 		if (MaterialExpressionClass && MaterialExpressionClass->IsChildOf(UMaterialExpression::StaticClass()))
@@ -1212,6 +1226,10 @@ void UMaterialGraphSchema::GetAssetsGraphHoverMessage(const TArray<FAssetData>&
 		{
 			OutOkIcon = true;
 		}
+		else if (SparseVolumeTexture)
+		{
+			OutOkIcon = true;
+		}
 		else if (ParameterCollection)
 		{
 			OutOkIcon = true;
