Creating a montage directly from a sequence now respects the sequence framerate

diff --git a/Engine/Source/Editor/UnrealEd/Private/Factories/AnimMontageFactory.cpp b/Engine/Source/Editor/UnrealEd/Private/Factories/AnimMontageFactory.cpp
--- a/Engine/Source/Editor/UnrealEd/Private/Factories/AnimMontageFactory.cpp
+++ b/Engine/Source/Editor/UnrealEd/Private/Factories/AnimMontageFactory.cpp
@@ -87,6 +87,8 @@ UObject* UAnimMontageFactory::FactoryCreateNew(UClass* Class, UObject* InParent,
 			NewTrack.AnimTrack.AnimSegments.Add(NewSegment);
 
 			AnimMontage->SetCompositeLength(SourceAnimation->GetPlayLength());
+
+			AnimMontage->UpdateCommonTargetFrameRate();
 		}
 
 		AnimMontage->SetSkeleton(TargetSkeleton);
diff --git a/Engine/Source/Runtime/Engine/Classes/Animation/AnimMontage.h b/Engine/Source/Runtime/Engine/Classes/Animation/AnimMontage.h
--- a/Engine/Source/Runtime/Engine/Classes/Animation/AnimMontage.h
+++ b/Engine/Source/Runtime/Engine/Classes/Animation/AnimMontage.h
@@ -615,6 +615,7 @@ class UAnimMontage : public UAnimCompositeBase
 	GENERATED_UCLASS_BODY()
 
 	friend struct FAnimMontageInstance;
+	friend class UAnimMontageFactory;
 	
 	UPROPERTY(EditAnywhere, BlueprintReadOnly, Category = BlendOption)
 	EMontageBlendMode BlendModeIn;
