Fix comments in AssetEditorSubsystem

diff --git a/Engine/Source/Editor/UnrealEd/Public/Subsystems/AssetEditorSubsystem.h b/Engine/Source/Editor/UnrealEd/Public/Subsystems/AssetEditorSubsystem.h
--- a/Engine/Source/Editor/UnrealEd/Public/Subsystems/AssetEditorSubsystem.h
+++ b/Engine/Source/Editor/UnrealEd/Public/Subsystems/AssetEditorSubsystem.h
@@ -478,13 +478,13 @@ private:
 	/** Map of FEditorModeId to EditorModeInfo for all known UEdModes when the subsystem initialized */
 	RegisteredModeInfoMap EditorModes;
 
-	/** Event that is triggered whenever a mode is unregistered */
+	/** Event that is triggered whenever a mode is registered or unregistered. */
 	FRegisteredModesChangedEvent OnEditorModesChangedEvent;
 
-	/** Event that is triggered whenever a mode is unregistered */
+	/** Event that is triggered whenever a mode is registered. Includes the mode's ID. */
 	FOnModeRegistered OnEditorModeRegisteredEvent;
 
-	/** Event that is triggered whenever a mode is unregistered */
+	/** Event that is triggered whenever a mode is unregistered. Includes the mode's ID. */
 	FOnModeUnregistered OnEditorModeUnregisteredEvent;
 	
 	/**
