PR #11293: Fix color casting on editor HDR screenshots

#tests Ran the editor with -hdr and set the cvar r.HDR.EnableHDROutput 1. Took a screenshot with F9. Opened the screenshot in renderdoc: the screeshot was identical to the display in Unreal viewport (while it was different with previous code). Ran auto test without issue.

diff --git a/Engine/Source/Editor/UnrealEd/Private/EditorViewportClient.cpp b/Engine/Source/Editor/UnrealEd/Private/EditorViewportClient.cpp
--- a/Engine/Source/Editor/UnrealEd/Private/EditorViewportClient.cpp
+++ b/Engine/Source/Editor/UnrealEd/Private/EditorViewportClient.cpp
@@ -5892,7 +5892,7 @@ void FEditorViewportClient::TakeScreenshot(FViewport* InViewport, bool bInValida
 	{
 		TArray<FLinearColor> RawPixels;
 		RawPixels.SetNum(CaptureRect.Area());
-		if (!InViewport->ReadLinearColorPixels(RawPixels, FReadSurfaceDataFlags(RCM_UNorm, CubeFace_MAX), CaptureRect))
+		if (!InViewport->ReadLinearColorPixels(RawPixels, FReadSurfaceDataFlags(RCM_MinMax, CubeFace_MAX), CaptureRect))
 		{
 			// Failed to read the image from the viewport
 			SaveMessagePtr->SetText(NSLOCTEXT("UnrealEd", "ScreenshotFailedViewport", "Screenshot failed, unable to read image from viewport"));
diff --git a/Engine/Source/Runtime/SlateRHIRenderer/Private/SlateRHIRenderer.cpp b/Engine/Source/Runtime/SlateRHIRenderer/Private/SlateRHIRenderer.cpp
--- a/Engine/Source/Runtime/SlateRHIRenderer/Private/SlateRHIRenderer.cpp
+++ b/Engine/Source/Runtime/SlateRHIRenderer/Private/SlateRHIRenderer.cpp
@@ -1523,7 +1523,7 @@ void FSlateRHIRenderer::DrawWindow_RenderThread(FRHICommandListImmediate& RHICmd
 		{
 			if (OutHDRScreenshotData != nullptr)
 			{
-				RHICmdList.ReadSurfaceData(BackBuffer, ClampedScreenshotRect, *OutHDRScreenshotData, FReadSurfaceDataFlags());
+				RHICmdList.ReadSurfaceData(BackBuffer, ClampedScreenshotRect, *OutHDRScreenshotData, FReadSurfaceDataFlags(RCM_MinMax));
 			}
 			else
 			{
