World Partition - Minimap builder: Make sure to delete extra minimap actors in case the world contains multiple of them

diff --git a/Engine/Source/Editor/UnrealEd/Private/WorldPartition/WorldPartitionMiniMapBuilder.cpp b/Engine/Source/Editor/UnrealEd/Private/WorldPartition/WorldPartitionMiniMapBuilder.cpp
--- a/Engine/Source/Editor/UnrealEd/Private/WorldPartition/WorldPartitionMiniMapBuilder.cpp
+++ b/Engine/Source/Editor/UnrealEd/Private/WorldPartition/WorldPartitionMiniMapBuilder.cpp
@@ -244,6 +244,21 @@ bool UWorldPartitionMiniMapBuilder::PostRun(UWorld* World, FPackageSourceControl
 		return false;
 	}
 
+	TArray<FString> ModifiedFiles = { PackageFileName };
+
+	// Make sure to delete unneeded minimap actors in case there is more than one in the map
+	for (TActorIterator<AWorldPartitionMiniMap> It(World); It; ++It)
+	{
+		if (*It != WorldMiniMap)
+		{
+			if (UPackage* ExternalPackage = It->GetExternalPackage())
+			{
+				PackageHelper.Delete(ExternalPackage);
+				ModifiedFiles.Add(SourceControlHelpers::PackageFilename(ExternalPackage));
+			}
+		}
+	}
+
 	const FString ChangeDescription = FString::Printf(TEXT("Rebuilt minimap for %s"), *World->GetName());
-	return OnFilesModified({ PackageFileName }, ChangeDescription);
+	return OnFilesModified(ModifiedFiles, ChangeDescription);
 }
