Fixed crash where the world DefaultBrush was inadvertently destroyed during drop preview teardown.

diff --git a/Engine/Source/Editor/UnrealEd/Private/LevelEditorViewport.cpp b/Engine/Source/Editor/UnrealEd/Private/LevelEditorViewport.cpp
--- a/Engine/Source/Editor/UnrealEd/Private/LevelEditorViewport.cpp
+++ b/Engine/Source/Editor/UnrealEd/Private/LevelEditorViewport.cpp
@@ -1684,11 +1684,27 @@ void FLevelEditorViewportClient::DestroyDropPreviewElements()
 		return true; // true means continue
 	});
 
-	DropPreviewElements->ForEachElement<ITypedElementWorldInterface>([this](const TTypedElement<ITypedElementWorldInterface>& InElement)
+	// Used for special casing BSP backwards compatibility: the builder brush is used as a preview object, and
+	// we don't want to delete it.
+	const UObject* DefaultBrush = GetWorld() ? GetWorld()->GetDefaultBrush() : nullptr;
+
+	DropPreviewElements->ForEachElement<ITypedElementWorldInterface>([this, DefaultBrush](const TTypedElement<ITypedElementWorldInterface>& InElement)
 	{
 		// Make sure to remove any items keyed by this handle, otherwise we'll get complaints about
 		// remaining references to the handle after destroying the item
 		PreDragElementTransforms.Remove(InElement);
+
+		// Don't delete the BSP builder brush, which legacy code uses for previews and expects to always exist.
+		if (DefaultBrush)
+		{
+			ITypedElementObjectInterface* ObjectInterface = UTypedElementRegistry::GetInstance()->GetElementInterface<ITypedElementObjectInterface>(InElement);
+			const UObject* Object = ObjectInterface ? ObjectInterface->GetObject(InElement) : nullptr;
+			if (Object && Object == DefaultBrush)
+			{
+				// true means continue - don't fall through to the deletion.
+				return true;
+			}
+		}
 		
 		FTypedElementDeletionOptions Options;
 		Options
