[Component Visualizer] Fix an issue where spline component visualizer was not correctly re-evaluating and clearing in some cases for viewport selections

diff --git a/Engine/Source/Editor/ComponentVisualizers/Private/SplineComponentVisualizer.cpp b/Engine/Source/Editor/ComponentVisualizers/Private/SplineComponentVisualizer.cpp
--- a/Engine/Source/Editor/ComponentVisualizers/Private/SplineComponentVisualizer.cpp
+++ b/Engine/Source/Editor/ComponentVisualizers/Private/SplineComponentVisualizer.cpp
@@ -1672,9 +1672,11 @@ bool FSplineComponentVisualizer::HandleFrustumSelect(const FConvexVolume& InFrus
 			SelectionState->ClearSelectedSegmentIndex();
 			SelectionState->ClearSelectedTangentHandle();
 		}
+
+		return true;
 	}
 
-	return true;
+	return false;
 }
 
 bool FSplineComponentVisualizer::HasFocusOnSelectionBoundingBox(FBox& OutBoundingBox)
diff --git a/Engine/Source/Editor/UnrealEd/Private/EditorSelectUtils.cpp b/Engine/Source/Editor/UnrealEd/Private/EditorSelectUtils.cpp
--- a/Engine/Source/Editor/UnrealEd/Private/EditorSelectUtils.cpp
+++ b/Engine/Source/Editor/UnrealEd/Private/EditorSelectUtils.cpp
@@ -308,22 +308,32 @@ void UUnrealEdEngine::OnEditorElementSelectionChanged(const UTypedElementSelecti
 
 		for (int32 CompIdx = 0; CompIdx < Components.Num(); CompIdx++)
 		{
-			UActorComponent* Comp = Components[CompIdx];
+			TWeakObjectPtr<UActorComponent> Comp(Components[CompIdx]);
 			if (Comp->IsRegistered())
 			{
 				// Try and find a visualizer
 				TSharedPtr<FComponentVisualizer> Visualizer = FindComponentVisualizer(Comp->GetClass());
-				if (Visualizer.IsValid() && (Comp == SelectedComponent || Visualizer->ShouldShowForSelectedSubcomponents(Comp)))
+				if (Visualizer.IsValid())
 				{
-					FCachedComponentVisualizer CachedComponentVisualizer(Comp, Visualizer);
+					FCachedComponentVisualizer CachedComponentVisualizer(Comp.Get(), Visualizer);
 					FComponentVisualizerForSelection Temp{ CachedComponentVisualizer };
 
 					FComponentVisualizerForSelection& ComponentVisualizerForSelection = VisualizersForSelection.Add_GetRef(MoveTemp(Temp));
 
-					if (Comp != SelectedComponent)
+					ComponentVisualizerForSelection.IsEnabledDelegate.Emplace([bIsSelectedComponent = (Comp == SelectedComponent), WeakViz = TWeakPtr<FComponentVisualizer>(Visualizer), Comp]()
 					{
-						ComponentVisualizerForSelection.IsEnabledDelegate.Emplace([]() { return GetDefault<UEditorPerProjectUserSettings>()->bShowSelectionSubcomponents == true; });
-					}
+						if (bIsSelectedComponent || (GetDefault<UEditorPerProjectUserSettings>()->bShowSelectionSubcomponents == true))
+						{
+							return true;
+						}
+
+						bool bShouldShowWithSelection = true;
+						if (TSharedPtr<FComponentVisualizer> PinnedViz = WeakViz.Pin())
+						{
+							bShouldShowWithSelection = PinnedViz->ShouldShowForSelectedSubcomponents(Comp.Get());
+						}
+						return bShouldShowWithSelection;
+					});
 				}
 			}
 		}
@@ -354,21 +364,23 @@ void UUnrealEdEngine::OnEditorElementSelectionChanged(const UTypedElementSelecti
 	}
 
 	// Restore the active component visualizer, since an undo/redo may have changed the selection
-	if (!ComponentVisManager.IsActive() && VisualizersForSelection.Num() > 0)
+	bool bDidSetVizThisTick = false;
+	if (VisualizersForSelection.Num() > 0)
 	{
 		for (FComponentVisualizerForSelection& VisualizerForSelection : VisualizersForSelection)
 		{
-			if (VisualizerForSelection.ComponentVisualizer.Visualizer->GetEditedComponent() != nullptr)
+			if (!ComponentVisManager.IsActive() && VisualizerForSelection.ComponentVisualizer.Visualizer->GetEditedComponent() != nullptr)
 			{
 				ComponentVisManager.SetActiveComponentVis(GCurrentLevelEditingViewportClient, VisualizerForSelection.ComponentVisualizer.Visualizer);
+				bDidSetVizThisTick = true;
 				break;
 			}
 		}
+	}
 
-		if (!ComponentVisManager.IsActive())
-		{
-			ComponentVisManager.ClearActiveComponentVis();
-		}
+	if (ComponentVisManager.IsActive() && (VisualizersForSelection.Num() == 0 || !bDidSetVizThisTick))
+	{
+		ComponentVisManager.ClearActiveComponentVis();
 	}
 
 #if PLATFORM_MAC
