Fix crashes around creating and deleting pose assets

Found an additional crash when investigating this ticket - when previewing a newly-created (i.e. unsaved) pose asset, moving the pose preview sliders would crash. Adding a refresh of the curve index table to fix this.
Original crash was caused by calling into the anim graph when the skeletal mesh is null. No nodes are built to be robust to this, so adding an guard similar to that in USkeletalMeshComponent::InitAnim()

diff --git a/Engine/Source/Editor/UnrealEd/Private/Animation/DebugSkelMeshComponent.cpp b/Engine/Source/Editor/UnrealEd/Private/Animation/DebugSkelMeshComponent.cpp
--- a/Engine/Source/Editor/UnrealEd/Private/Animation/DebugSkelMeshComponent.cpp
+++ b/Engine/Source/Editor/UnrealEd/Private/Animation/DebugSkelMeshComponent.cpp
@@ -532,35 +532,38 @@ void UDebugSkelMeshComponent::InitAnim(bool bForceReinit)
 
 	Super::InitAnim(bForceReinit);
 
-	// if PreviewInstance is nullptr, create here once
-	if (PreviewInstance == nullptr)
+	if(GetSkeletalMeshAsset() != nullptr)
 	{
-		PreviewInstance = CreatePreviewInstance();
-		check(PreviewInstance);
+		// if PreviewInstance is nullptr, create here once
+		if (PreviewInstance == nullptr)
+		{
+			PreviewInstance = CreatePreviewInstance();
+			check(PreviewInstance);
 
-		//Set transactional flag in order to restore slider position when undo operation is performed
-		PreviewInstance->SetFlags(RF_Transactional);
-	}
+			//Set transactional flag in order to restore slider position when undo operation is performed
+			PreviewInstance->SetFlags(RF_Transactional);
+		}
 
-	// if anim script instance is null because it's not playing a blueprint, set to PreviewInstnace by default
-	// that way if user would like to modify bones or do extra stuff, it will work
-	if (AnimScriptInstance == nullptr)
-	{
-		AnimScriptInstance = PreviewInstance;
-		AnimScriptInstance->InitializeAnimation();
-	}
-	else
-	{
-		// Make sure we initialize the preview instance here, as we want the required bones to be up to date
-		// even if we arent using the instance right now.
-		PreviewInstance->InitializeAnimation();
-	}
+		// if anim script instance is null because it's not playing a blueprint, set to PreviewInstnace by default
+		// that way if user would like to modify bones or do extra stuff, it will work
+		if (AnimScriptInstance == nullptr)
+		{
+			AnimScriptInstance = PreviewInstance;
+			AnimScriptInstance->InitializeAnimation();
+		}
+		else
+		{
+			// Make sure we initialize the preview instance here, as we want the required bones to be up to date
+			// even if we arent using the instance right now.
+			PreviewInstance->InitializeAnimation();
+		}
 
-	if(PostProcessAnimInstance)
-	{
-		// Add the same settings as the preview instance in this case.
-		PostProcessAnimInstance->RootMotionMode = ERootMotionMode::RootMotionFromEverything;
-		PostProcessAnimInstance->bUseMultiThreadedAnimationUpdate = false;
+		if(PostProcessAnimInstance)
+		{
+			// Add the same settings as the preview instance in this case.
+			PostProcessAnimInstance->RootMotionMode = ERootMotionMode::RootMotionFromEverything;
+			PostProcessAnimInstance->bUseMultiThreadedAnimationUpdate = false;
+		}
 	}
 }
 
diff --git a/Engine/Source/Runtime/Engine/Private/Animation/PoseAsset.cpp b/Engine/Source/Runtime/Engine/Private/Animation/PoseAsset.cpp
--- a/Engine/Source/Runtime/Engine/Private/Animation/PoseAsset.cpp
+++ b/Engine/Source/Runtime/Engine/Private/Animation/PoseAsset.cpp
@@ -1223,6 +1223,8 @@ void UPoseAsset::PostProcessData()
 	}
 
 	UpdateTrackBoneIndices();
+
+	PoseContainer.RebuildCurveIndexTable();
 }
 
 void UPoseAsset::BreakAnimationSequenceGUIDComparison()
