Split hit result filtering for Actor positioning in Editor in two passes to fix threading issue. This will be required to fix an ensure that could soon be reported by added threading validation in FInstancedActorHandle.
- gamethread part will filter based UObject based informations and FInstancedActorHandles and will provided list of remaining WeakObject pointers on the remaining primitives
- rendering thread can then apply the ViewRelevance filter

diff --git a/Engine/Source/Editor/UnrealEd/Private/Editor/ActorPositioning.cpp b/Engine/Source/Editor/UnrealEd/Private/Editor/ActorPositioning.cpp
--- a/Engine/Source/Editor/UnrealEd/Private/Editor/ActorPositioning.cpp
+++ b/Engine/Source/Editor/UnrealEd/Private/Editor/ActorPositioning.cpp
@@ -46,9 +46,6 @@ FActorPositionTraceResult FActorPositioning::TraceWorldForPositionWithDefault(co
 FActorPositionTraceResult FActorPositioning::TraceWorldForPosition(const FViewportCursorLocation& Cursor, const FSceneView& View, const TArray<AActor*>* IgnoreActors)
 {
 	const auto* ViewportClient = Cursor.GetViewportClient();
-
-	FActorPositionTraceResult Results;
-
 	const auto ViewportType = ViewportClient->GetViewportType();
 
 	// Start with a ray that encapsulates the entire world
@@ -63,63 +60,83 @@ FActorPositionTraceResult FActorPositioning::TraceWorldForPosition(const FViewpo
 	return TraceWorldForPosition(*ViewportClient->GetWorld(), View, RayStart, RayEnd, IgnoreActors);
 }
 
-/** Check to see if the specified hit result should be ignored from actor positioning calculations for the specified scene view */
-bool IsHitIgnored(const FHitResult& InHit, const FSceneView& InSceneView)
+/**
+ * Prunes list of hit results for actor positioning calculations based on conditions that could be tested
+ * on the game thread and returns a list of primitives for the remaining this.
+ */
+TArray<TWeakObjectPtr<const UPrimitiveComponent>> FilterHitsGameThread(TArray<FHitResult>& InOutHits)
 {
-	// We're using the SceneProxy and ViewRelevance here, we should execute from the render thread but since
-	// we're also accessing UObjects, the game thread should be blocked during this call.
-	check(IsInParallelRenderingThread());
-	const FActorInstanceHandle& HitObjHandle = InHit.HitObjectHandle;
-	
-	// Try and find a primitive component for the hit
-	const UPrimitiveComponent* PrimitiveComponent = Cast<UPrimitiveComponent>(HitObjHandle.GetRootComponent());
+	check(IsInGameThread() || IsInParallelGameThread());
+	TArray<TWeakObjectPtr<const UPrimitiveComponent>> WeakPrimitives;
+	WeakPrimitives.Reserve(InOutHits.Num());
 
-	if (!PrimitiveComponent)
-	{
-		PrimitiveComponent = InHit.Component.Get();
-	}
-	if (PrimitiveComponent && PrimitiveComponent->IsA(ULandscapeHeightfieldCollisionComponent::StaticClass()))
+	InOutHits.RemoveAll([&WeakPrimitives](const FHitResult& Hit)
 	{
-		PrimitiveComponent = CastChecked<ULandscapeHeightfieldCollisionComponent>(PrimitiveComponent)->GetRenderComponent();
-	}
+		if (Hit.bStartPenetrating)
+		{
+			return true;
+		}
 
-	if (InHit.bStartPenetrating)
-	{
-		return true;
-	}
+		const FActorInstanceHandle& HitObjHandle = Hit.HitObjectHandle;
 
-	if (!PrimitiveComponent)
-	{
-		// If we don't have a primitive component, either ignore the hit, or pass it through if the CVar is set appropriately.
-		// i.e. ignoring is the inverse of the "allow non primitive hits" CVar.
-		return !ActorPositioningLocals::bAllowNonPrimitiveComponentHits;
-	}
+		// Try and find a primitive component for the hit
+		const UPrimitiveComponent* PrimitiveComponent = Cast<UPrimitiveComponent>(HitObjHandle.GetRootComponent());
 
-	// Ignore volumes and shapes
-	if (HitObjHandle.DoesRepresentClass(AVolume::StaticClass()))
-	{
-		return true;
-	}
-	else if (PrimitiveComponent->IsA(UShapeComponent::StaticClass()))
-	{
-		return true;
-	}
+		if (!PrimitiveComponent)
+		{
+			PrimitiveComponent = Hit.Component.Get();
+		}
+		if (PrimitiveComponent && PrimitiveComponent->IsA(ULandscapeHeightfieldCollisionComponent::StaticClass()))
+		{
+			PrimitiveComponent = CastChecked<ULandscapeHeightfieldCollisionComponent>(PrimitiveComponent)->GetRenderComponent();
+		}
+
+		if (!PrimitiveComponent)
+		{
+			// If we don't have a primitive component, either ignore the hit, or pass it through if the CVar is set appropriately.
+			// i.e. ignoring is the inverse of the "allow non primitive hits" CVar.
+			return !ActorPositioningLocals::bAllowNonPrimitiveComponentHits;
+		}
 
-	// Only use this component if it is visible in the specified scene views
-	bool bIsRenderedOnScreen = false;
-	bool bIgnoreTranslucentPrimitive = false;
-	const bool bConsiderInvisibleComponentForPlacement = PrimitiveComponent->bConsiderForActorPlacementWhenHidden;
-	{				
-		if (PrimitiveComponent && PrimitiveComponent->SceneProxy)
+		// Ignore volumes and shapes
+		if (HitObjHandle.DoesRepresentClass(AVolume::StaticClass()))
 		{
-			const FPrimitiveViewRelevance ViewRelevance = PrimitiveComponent->SceneProxy->GetViewRelevance(&InSceneView);
-			// BSP is a bit special in that its bDrawRelevance is false even when drawn as wireframe because InSceneView.Family->EngineShowFlags.BSPTriangles is off
-			bIsRenderedOnScreen = ViewRelevance.bDrawRelevance || (PrimitiveComponent->IsA(UModelComponent::StaticClass()) && InSceneView.Family->EngineShowFlags.BSP);
-			bIgnoreTranslucentPrimitive = ViewRelevance.HasTranslucency() && !GetDefault<UEditorPerProjectUserSettings>()->bAllowSelectTranslucent;
+			return true;
 		}
+
+		if (PrimitiveComponent->IsA(UShapeComponent::StaticClass()))
+		{
+			return true;
+		}
+
+		WeakPrimitives.Add(PrimitiveComponent);
+		return false;
+	});
+
+	return MoveTemp(WeakPrimitives);
+}
+
+/** Check to see if the specified hit result should be ignored from actor positioning calculations for the specified scene view */
+bool IsHitIgnoredRenderingThread(const TWeakObjectPtr<const UPrimitiveComponent>& InWeakPrimitiveComponent, const FSceneView& InSceneView)
+{
+	// We're using the SceneProxy and ViewRelevance here, we should execute from the render thread
+	check(IsInParallelRenderingThread());
+
+	const UPrimitiveComponent* PrimitiveComponent = InWeakPrimitiveComponent.Get();
+	if (PrimitiveComponent && PrimitiveComponent->SceneProxy)
+	{
+		const bool bConsiderInvisibleComponentForPlacement = PrimitiveComponent->bConsiderForActorPlacementWhenHidden;
+
+		// Only use this component if it is visible in the specified scene views
+		const FPrimitiveViewRelevance ViewRelevance = PrimitiveComponent->SceneProxy->GetViewRelevance(&InSceneView);
+		// BSP is a bit special in that its bDrawRelevance is false even when drawn as wireframe because InSceneView.Family->EngineShowFlags.BSPTriangles is off
+		const bool bIsRenderedOnScreen = ViewRelevance.bDrawRelevance || (PrimitiveComponent->IsA(UModelComponent::StaticClass()) && InSceneView.Family->EngineShowFlags.BSP);
+		const bool bIgnoreTranslucentPrimitive = ViewRelevance.HasTranslucency() && !GetDefault<UEditorPerProjectUserSettings>()->bAllowSelectTranslucent;
+		
+		return (!bIsRenderedOnScreen && !bConsiderInvisibleComponentForPlacement) || bIgnoreTranslucentPrimitive;
 	}
 
-	return (!bIsRenderedOnScreen && !bConsiderInvisibleComponentForPlacement) || bIgnoreTranslucentPrimitive;
+	return false;
 }
 
 FActorPositionTraceResult FActorPositioning::TraceWorldForPosition(const UWorld& InWorld, const FSceneView& InSceneView, const FVector& RayStart, const FVector& RayEnd, const TArray<AActor*>* IgnoreActors)
@@ -137,13 +154,20 @@ FActorPositionTraceResult FActorPositioning::TraceWorldForPosition(const UWorld&
 	if ( InWorld.LineTraceMultiByObjectType(Hits, RayStart, RayEnd, FCollisionObjectQueryParams(FCollisionObjectQueryParams::InitType::AllObjects), Param) )
 	{
 		{
-			// Send IsHitIgnored on the render thread since we're accessing view relevance
+			// Filter out anything that should be ignored based on information accessible on the game thread
+			// and build list of remaining weak primitive components that need to be filtered on the rendering thread 
+			TArray<TWeakObjectPtr<const UPrimitiveComponent>> WeakPrimitives = FilterHitsGameThread(Hits);
+			ensure(Hits.Num() == WeakPrimitives.Num());
+
+			// Send IsHitIgnoredRenderingThread on the render thread since we're accessing view relevance
 			ENQUEUE_RENDER_COMMAND(TraceWorldForPosition_FilterHitsByViewRelevance)(
-				[&Hits, &InSceneView](FRHICommandListImmediate& RHICmdList)
+				[&Hits, &WeakPrimitives, &InSceneView](FRHICommandListImmediate& RHICmdList)
 				{
 					// Filter out anything that should be ignored
-					Hits.RemoveAll([&InSceneView](const FHitResult& Hit){
-						return IsHitIgnored(Hit, InSceneView);
+					int32 Index = 0;
+					Hits.RemoveAll([&Index, &InSceneView, &WeakPrimitives](const FHitResult&)
+					{
+						return IsHitIgnoredRenderingThread(WeakPrimitives[Index++], InSceneView);
 					});
 				}
 			);
