World Partition Runtime HashSet: when saving a copy of a world, replace HLOd layer setups with duplicated ones.

diff --git a/Engine/Source/Editor/UnrealEd/Private/FileHelpers.cpp b/Engine/Source/Editor/UnrealEd/Private/FileHelpers.cpp
--- a/Engine/Source/Editor/UnrealEd/Private/FileHelpers.cpp
+++ b/Engine/Source/Editor/UnrealEd/Private/FileHelpers.cpp
@@ -30,6 +30,7 @@
 #include "UncontrolledChangelistsModule.h"
 #include "SourceControlOperations.h"
 #include "Editor/UnrealEdEngine.h"
+#include "Serialization/ArchiveReplaceObjectRef.h"
 #include "Settings/EditorLoadingSavingSettings.h"
 #include "Factories/Factory.h"
 #include "Factories/FbxSceneImportFactory.h"
@@ -75,6 +76,7 @@
 #include "HierarchicalLOD.h"
 #include "WorldPartition/IWorldPartitionEditorModule.h"
 #include "WorldPartition/ActorDescContainer.h"
+#include "WorldPartition/WorldPartitionRuntimeHash.h"
 #include "WorldPartition/WorldPartition.h"
 #include "WorldPartition/WorldPartitionEditorPerProjectUserSettings.h"
 #include "WorldPartition/HLOD/HLODLayer.h"
@@ -841,17 +843,22 @@ static bool SaveWorld(UWorld* World,
 
 					if (bIsTempPackage)
 					{
-						if (UHLODLayer* CurrentHLODLayer = RenamedWorldPartition->GetDefaultHLODLayer())
+						if (UHLODLayer* CurHLODLayer = RenamedWorldPartition->GetDefaultHLODLayer())
 						{
-							CurrentHLODLayer = UHLODLayer::DuplicateHLODLayersSetup(CurrentHLODLayer, NewPackageName, NewWorldAssetName);
-							
-							RenamedWorldPartition->SetDefaultHLODLayer(CurrentHLODLayer);
+							UHLODLayer* NewHLODLayer = UHLODLayer::DuplicateHLODLayersSetup(CurHLODLayer, NewPackageName, NewWorldAssetName);
+
+							RenamedWorldPartition->SetDefaultHLODLayer(NewHLODLayer);
 
-							while (CurrentHLODLayer)
+							TMap<UHLODLayer*, UHLODLayer*> ReplacementMap;
+							while (NewHLODLayer)
 							{
-								PackagesToSave.Add(CurrentHLODLayer->GetPackage());
-								CurrentHLODLayer = CurrentHLODLayer->GetParentLayer();
+								PackagesToSave.Add(NewHLODLayer->GetPackage());
+								ReplacementMap.Add(CurHLODLayer, NewHLODLayer);
+								CurHLODLayer = CurHLODLayer->GetParentLayer();
+								NewHLODLayer = NewHLODLayer->GetParentLayer();
 							}
+							
+							FArchiveReplaceObjectRef<UHLODLayer> ReplaceObjectRefAr(RenamedWorldPartition->RuntimeHash, ReplacementMap, EArchiveReplaceObjectFlags::IgnoreOuterRef | EArchiveReplaceObjectFlags::IgnoreArchetypeRef);							
 						}
 					}
 				}
