Add editor delegate to allow ForceDeleteObject behaviour to be extended by plugins.

diff --git a/Engine/Source/Editor/UnrealEd/Private/Editor.cpp b/Engine/Source/Editor/UnrealEd/Private/Editor.cpp
--- a/Engine/Source/Editor/UnrealEd/Private/Editor.cpp
+++ b/Engine/Source/Editor/UnrealEd/Private/Editor.cpp
@@ -152,6 +152,7 @@ FEditorDelegates::FOnAssetsAddExtraObjectsToDelete		FEditorDelegates::OnAssetsAd
 FEditorDelegates::FOnAssetsPreDelete					FEditorDelegates::OnAssetsPreDelete;
 FEditorDelegates::FOnAssetsDeleted						FEditorDelegates::OnAssetsDeleted;
 FEditorDelegates::FOnAssetDragStarted					FEditorDelegates::OnAssetDragStarted;
+FEditorDelegates::FOnPreForceDeleteObjects				FEditorDelegates::OnPreForceDeleteObjects;
 FSimpleMulticastDelegate								FEditorDelegates::OnEnableGestureRecognizerChanged;
 FSimpleMulticastDelegate								FEditorDelegates::OnActionAxisMappingsChanged;
 FEditorDelegates::FOnAddLevelToWorld					FEditorDelegates::OnAddLevelToWorld;
diff --git a/Engine/Source/Editor/UnrealEd/Private/ObjectTools.cpp b/Engine/Source/Editor/UnrealEd/Private/ObjectTools.cpp
--- a/Engine/Source/Editor/UnrealEd/Private/ObjectTools.cpp
+++ b/Engine/Source/Editor/UnrealEd/Private/ObjectTools.cpp
@@ -3350,6 +3350,8 @@ namespace ObjectTools
 
 		GWarn->BeginSlowTask( NSLOCTEXT("UnrealEd", "Deleting", "Deleting"), true );
 
+		FEditorDelegates::OnPreForceDeleteObjects.Broadcast(ShownObjectsToDelete);
+		
 		struct FSCSNodeToDelete
 		{
 			USimpleConstructionScript* SimpleConstructionScript;
diff --git a/Engine/Source/Editor/UnrealEd/Public/Editor.h b/Engine/Source/Editor/UnrealEd/Public/Editor.h
--- a/Engine/Source/Editor/UnrealEd/Public/Editor.h
+++ b/Engine/Source/Editor/UnrealEd/Public/Editor.h
@@ -136,6 +136,8 @@ struct FEditorDelegates
 	DECLARE_MULTICAST_DELEGATE_OneParam(FOnAssetsAddExtraObjectsToDelete, TArray<UObject*>&);
 	/** delegate type for when a user requests to delete certain assets... DOES NOT mean the asset(s) will be deleted (the user could cancel) */
 	DECLARE_MULTICAST_DELEGATE_OneParam(FOnAssetsPreDelete, const TArray<UObject*>&);
+	/** delegate type for when a user requested force deleting objects. The objects(s) will be deleted (no possibility to cancel), so implementations should delete references */
+	DECLARE_MULTICAST_DELEGATE_OneParam(FOnPreForceDeleteObjects, const TArray<UObject*>&);
 	/** delegate type for when one or more assets have been deleted */
 	DECLARE_MULTICAST_DELEGATE_OneParam(FOnAssetsDeleted, const TArray<UClass*>& /*DeletedAssetClasses*/);
 	/** delegate type for when a user starts dragging something out of content browser (can be multiple assets) */
@@ -320,6 +322,8 @@ struct FEditorDelegates
 	static UNREALED_API FOnAssetsDeleted OnAssetsDeleted;
 	/** Called when a user starts dragging something out of content browser (can be multiple assets) */
 	static UNREALED_API FOnAssetDragStarted OnAssetDragStarted;
+	/** Called when the user requests objects to be force deleted.  There is no possibility to cancel once this callback is made */
+    static UNREALED_API FOnPreForceDeleteObjects OnPreForceDeleteObjects;
 	/** Called when a user changes the UInputSettings::bEnableGestureRecognizer setting to refresh the available actions. */
 	static UNREALED_API FSimpleMulticastDelegate OnEnableGestureRecognizerChanged;
 	/** Called when Action or Axis mappings have been changed */
