Hide PathTracing menu entry on platforms (such as Vulkan) which support raytracing but cannot yet run the path tracer

diff --git a/Engine/Source/Editor/UnrealEd/Private/SEditorViewportViewMenu.cpp b/Engine/Source/Editor/UnrealEd/Private/SEditorViewportViewMenu.cpp
--- a/Engine/Source/Editor/UnrealEd/Private/SEditorViewportViewMenu.cpp
+++ b/Engine/Source/Editor/UnrealEd/Private/SEditorViewportViewMenu.cpp
@@ -147,7 +147,9 @@ void SEditorViewportViewMenu::FillViewMenu(UToolMenu* Menu) const
 			if (IsRayTracingEnabled())
 			{
 				static auto PathTracingCvar = IConsoleManager::Get().FindTConsoleVariableDataInt(TEXT("r.PathTracing"));
-				if (PathTracingCvar && PathTracingCvar->GetValueOnAnyThread() != 0)
+				const bool bPathTracingSupported = FDataDrivenShaderPlatformInfo::GetSupportsPathTracing(GMaxRHIShaderPlatform);
+				const bool bPathTracingEnabled = PathTracingCvar && PathTracingCvar->GetValueOnAnyThread() != 0;
+				if (bPathTracingSupported && bPathTracingEnabled)
 				{
 					Section.AddMenuEntry(BaseViewportActions.PathTracingMode, UViewModeUtils::GetViewModeDisplayName(VMI_PathTracing));
 				}
