Animation Editor Preview improvements
- Added commonly used root motion functions to AnimSequenceHelper
- Take initial root motion transform into account when previewing root motion
- Separated UDebugSkelMeshComponent ProcessRootMotionMode to requested and actual so that the yser setting does not get reset if the animation does not support it
- Store reference transform for the whole root motion in UDebugSkelMeshComponent (this will allow later to preview more root motion data)
- Separated root motion and turntable transforms to allow to combine root motion and turntable modes
- Changed RootMotionProcessMode to LoopAndReset
- React preview immediately to mirro table change

diff --git a/Engine/Source/Editor/Persona/Private/SAnimViewportToolBar.cpp b/Engine/Source/Editor/Persona/Private/SAnimViewportToolBar.cpp
--- a/Engine/Source/Editor/Persona/Private/SAnimViewportToolBar.cpp
+++ b/Engine/Source/Editor/Persona/Private/SAnimViewportToolBar.cpp
@@ -952,6 +952,7 @@ void SAnimViewportToolBar::OnMirrorDataTableSelected(const FAssetData& SelectedM
 		if (Mesh && PreviewInstance)
 		{
 			PreviewInstance->SetMirrorDataTable(MirrorDataTable);
+			PreviewComp->OnMirrorDataTableChanged();
 		}
 	}
 }
diff --git a/Engine/Source/Editor/Persona/Private/SAnimationEditorViewport.cpp b/Engine/Source/Editor/Persona/Private/SAnimationEditorViewport.cpp
--- a/Engine/Source/Editor/Persona/Private/SAnimationEditorViewport.cpp
+++ b/Engine/Source/Editor/Persona/Private/SAnimationEditorViewport.cpp
@@ -2095,7 +2095,7 @@ void SAnimationEditorViewportTabBody::SetProcessRootMotionMode(EProcessRootMotio
 bool SAnimationEditorViewportTabBody::IsProcessRootMotionModeSet(EProcessRootMotionMode Mode) const
 {
 	const UDebugSkelMeshComponent* PreviewComponent = GetPreviewScene()->GetPreviewMeshComponent();
-	return PreviewComponent ? (PreviewComponent->GetProcessRootMotionMode() == Mode) : false;
+	return PreviewComponent ? (PreviewComponent->GetRequestedProcessRootMotionMode() == Mode) : false;
 }
 
 bool SAnimationEditorViewportTabBody::CanUseProcessRootMotionMode(EProcessRootMotionMode Mode) const
diff --git a/Engine/Source/Editor/UnrealEd/Classes/Animation/DebugSkelMeshComponent.h b/Engine/Source/Editor/UnrealEd/Classes/Animation/DebugSkelMeshComponent.h
--- a/Engine/Source/Editor/UnrealEd/Classes/Animation/DebugSkelMeshComponent.h
+++ b/Engine/Source/Editor/UnrealEd/Classes/Animation/DebugSkelMeshComponent.h
@@ -248,6 +248,10 @@ class UDebugSkelMeshComponent : public USkeletalMeshComponent
 	UPROPERTY()
 	uint32 bPreviewRootMotion_DEPRECATED : 1;
 
+	/** Requested Process root motion mode, ProcessRootMotionMode gets set based on requested mode and what is supported. */
+	UPROPERTY(transient)
+	EProcessRootMotionMode RequestedProcessRootMotionMode;
+
 	/** Process root motion mode */
 	UPROPERTY(transient)
 	EProcessRootMotionMode ProcessRootMotionMode;
@@ -298,7 +302,13 @@ class UDebugSkelMeshComponent : public USkeletalMeshComponent
 
 	/** Storage of Source Animation Pose for when bDisplaySourceAnimation == true, as they have to be calculated */
 	TArray<FTransform> SourceAnimationPoses;
-	
+
+	/** Accumulated root motion. */
+	FTransform RootMotionTransform;
+
+	/** Transform representing the actor transform at the beginning of the animation. */
+	FTransform RootMotionReferenceTransform;
+
 	/** Array of bones to render bone weights for */
 	UPROPERTY(transient)
 	TArray<int32> BonesOfInterest;
@@ -392,6 +402,21 @@ class UDebugSkelMeshComponent : public USkeletalMeshComponent
 	UNREALED_API virtual void ShowReferencePose(bool bRefPose);
 	UNREALED_API virtual bool IsReferencePoseShown() const;
 
+	/** @return Accumulated root motion with turntable animation. */
+	FTransform GetRootMotionTransform() const
+	{
+		return RootMotionTransform * TurntableTransform;
+	}
+
+	/** @return Transform representing the actor transform at the beginning of the animation with turntable animation. */
+	FTransform GetRootMotionReferenceTransform() const
+	{
+		return RootMotionReferenceTransform * TurntableTransform;
+	}
+
+	/** Called when mirror data table changes on anim instance. */
+	UNREALED_API void OnMirrorDataTableChanged();
+
 	/**
 	 * Update material information depending on color render mode 
 	 * Refresh/replace materials 
@@ -451,10 +476,13 @@ class UDebugSkelMeshComponent : public USkeletalMeshComponent
 	/** Whether we are processing root motion or not */
 	UNREALED_API bool IsProcessingRootMotion() const;
 
+	/** Gets requested process root motion mode, can differ from GetProcessRootMotionMode() if the current asset does not support root motion. */
+	UNREALED_API EProcessRootMotionMode GetRequestedProcessRootMotionMode() const;
+
 	/** Gets process root motion mode */
 	UNREALED_API EProcessRootMotionMode GetProcessRootMotionMode() const;
 
-	/** Sets process root motion mode. Note: disabling root motion preview resets transform. */
+	/** Sets process root motion mode, the request may be ignored if current asset does not support the mode. Note: disabling root motion preview resets transform. */
 	UNREALED_API void SetProcessRootMotionMode(EProcessRootMotionMode Mode);
 
 	/** Whether the supplied root motion mode can be used for the current asset */
@@ -577,6 +605,8 @@ private:
 
 	// Rebuilds the cloth bounds for the asset.
 	UNREALED_API void RebuildCachedClothBounds();
+
+	UNREALED_API void SetProcessRootMotionModeInternal(EProcessRootMotionMode Mode);
 protected:
 
 	// Overridden to support single clothing ticks
@@ -589,6 +619,8 @@ public:
 	EPersonaTurnTableMode::Type TurnTableMode;
 	/** Current turn table speed scaling */
 	float TurnTableSpeedScaling;
+	/** Current turn table transform */
+	FTransform TurntableTransform;
 
 	UNREALED_API virtual void TickComponent(float DeltaTime, enum ELevelTick TickType, FActorComponentTickFunction *ThisTickFunction) override;
 
diff --git a/Engine/Source/Editor/UnrealEd/Private/Animation/DebugSkelMeshComponent.cpp b/Engine/Source/Editor/UnrealEd/Private/Animation/DebugSkelMeshComponent.cpp
--- a/Engine/Source/Editor/UnrealEd/Private/Animation/DebugSkelMeshComponent.cpp
+++ b/Engine/Source/Editor/UnrealEd/Private/Animation/DebugSkelMeshComponent.cpp
@@ -16,6 +16,7 @@
 #include "AnimPreviewInstance.h"
 #include "Animation/AnimComposite.h"
 #include "Animation/BlendSpace.h"
+#include "Animation/AnimSequenceHelpers.h"
 #include "Rendering/SkeletalMeshRenderData.h"
 #include "Rendering/SkeletalMeshModel.h"
 
@@ -41,6 +42,10 @@ UDebugSkelMeshComponent::UDebugSkelMeshComponent(const FObjectInitializer& Objec
 
 	TurnTableSpeedScaling = 1.f;
 	TurnTableMode = EPersonaTurnTableMode::Stopped;
+	TurntableTransform = FTransform::Identity;
+
+	RootMotionTransform = FTransform::Identity;
+	RootMotionReferenceTransform = FTransform::Identity;
 
 	bPauseClothingSimulationWithAnim = false;
 	bPerformSingleClothingTick = false;
@@ -51,6 +56,7 @@ UDebugSkelMeshComponent::UDebugSkelMeshComponent(const FObjectInitializer& Objec
 	
 	CachedClothBounds = FBoxSphereBounds(ForceInit);
 
+	RequestedProcessRootMotionMode = EProcessRootMotionMode::LoopAndReset;
 	ProcessRootMotionMode = EProcessRootMotionMode::Ignore;
 	ConsumeRootMotionPreviousPlaybackTime = 0.f;
 }
@@ -183,166 +189,128 @@ void UDebugSkelMeshComponent::ConsumeRootMotion(const FVector& FloorMin, const F
 		return;
 	}
 
-	// Helper to extract root motion manually from the current asset
-	auto ExtractRootMotionFromAnimationAsset = [](const UAnimationAsset* Animation, const UMirrorDataTable* MirrorDataTable, float StartPosition, float EndPosition) -> FTransform
+	// Force ProcessRootMotionMode to Ignore if the current asset/animation blueprint is not using root motion. 
+	if (ProcessRootMotionMode != EProcessRootMotionMode::Ignore && DoesCurrentAssetHaveRootMotion() == false)
 	{
-		if (const UAnimMontage* Anim = Cast<UAnimMontage>(Animation))
-		{
-			return Anim->ExtractRootMotionFromTrackRange(StartPosition, EndPosition);
-		}
-		
-		if (const UAnimComposite* Anim = Cast<UAnimComposite>(Animation))
-		{
-			FRootMotionMovementParams RootMotion;
-			Anim->ExtractRootMotionFromTrack(Anim->AnimationTrack, StartPosition, EndPosition, RootMotion);
-			return RootMotion.GetRootMotionTransform();
-		}
-		
-		if (const UAnimSequence* Anim = Cast<UAnimSequence>(Animation))
-		{
-			FTransform RootMotion = Anim->ExtractRootMotionFromRange(StartPosition, EndPosition);
-			if (MirrorDataTable)
-			{
-				FVector T = RootMotion.GetTranslation();
-				T = FAnimationRuntime::MirrorVector(T, MirrorDataTable->MirrorAxis);
-
-				FQuat Q = RootMotion.GetRotation();
-				Q = FAnimationRuntime::MirrorQuat(Q, MirrorDataTable->MirrorAxis);
-
-				FVector S = RootMotion.GetScale3D();
-				return FTransform(Q, T, S);
-			}
-			return RootMotion;
-		}
-
-		return FTransform::Identity;
-	};
+		SetProcessRootMotionModeInternal(EProcessRootMotionMode::Ignore);
+	}
 
-	// Force ProcessRootMotionMode to Ignore if the current asset/animation blueprint is not using root motion. 
-	if(ProcessRootMotionMode != EProcessRootMotionMode::Ignore && DoesCurrentAssetHaveRootMotion() == false)
+	// If our requested mode became available, use it.
+	if (ProcessRootMotionMode != RequestedProcessRootMotionMode && CanUseProcessRootMotionMode(RequestedProcessRootMotionMode))
 	{
-		SetProcessRootMotionMode(EProcessRootMotionMode::Ignore);
+		SetProcessRootMotionModeInternal(RequestedProcessRootMotionMode);
 	}
 
 	//Extract root motion regardless of where we use it so that we don't hit problems with it building up in the instance
 	FRootMotionMovementParams ExtractedRootMotion = ConsumeRootMotion_Internal(1.0f);
 	if (PreviewInstance->GetMirrorDataTable())
 	{
-		FTransform RootMotion = ExtractedRootMotion.GetRootMotionTransform();
-		FVector T = RootMotion.GetTranslation();
-		T = FAnimationRuntime::MirrorVector(T, PreviewInstance->GetMirrorDataTable()->MirrorAxis);
-
-		FQuat Q = RootMotion.GetRotation();
-		Q = FAnimationRuntime::MirrorQuat(Q, PreviewInstance->GetMirrorDataTable()->MirrorAxis);
-
-		FVector S = RootMotion.GetScale3D();
-		ExtractedRootMotion.Set(FTransform(Q, T, S));
+		const FTransform MirroredTransform = UE::Anim::MirrorTransform(ExtractedRootMotion.GetRootMotionTransform(), *PreviewInstance->GetMirrorDataTable());
+		ExtractedRootMotion.Set(MirroredTransform);
 	}
 
 	const float CurrentTime = PreviewInstance->GetCurrentTime();
 	const float PreviousTime = ConsumeRootMotionPreviousPlaybackTime;
 	ConsumeRootMotionPreviousPlaybackTime = CurrentTime;
 
+	bool bLooped = false;
+
 	if (ProcessRootMotionMode != EProcessRootMotionMode::Ignore && !ShouldBlendPhysicsBones())
 	{
-		if(PreviewInstance->IsPlaying())
+		if (PreviewInstance->IsPlaying())
 		{
+			// Figure out if the animation has looped, and the start end position of the root motion extraction.
+			float SectionStartPosition = 0.0f;
+
+			// We have to deal with montage explicitly because we can have multiple sections and we want to reset the position when the section loops
+			// and depending on the composition, CurrentTime < PreviousTime (or CurrentTime > PreviousTime when playing in reverse) is not enough
+			if (const UAnimMontage* Montage = Cast<UAnimMontage>(PreviewInstance->CurrentAsset))
+			{
+				const int32 PreviewStartSectionIdx = Montage->CompositeSections.IsValidIndex(PreviewInstance->MontagePreviewStartSectionIdx) ? PreviewInstance->MontagePreviewStartSectionIdx : Montage->GetSectionIndexFromPosition(CurrentTime);
+				const int32 FirstSectionIdx = PreviewInstance->MontagePreview_FindFirstSectionAsInMontage(PreviewStartSectionIdx);
+				const int32 LastSectionIdx = PreviewInstance->MontagePreview_FindLastSection(FirstSectionIdx);
+
+				// If FirstSection == LastSection we are previewing a single section
+				// In this case to know if we have looped we just need to check if CurrentTime < PreviousTime (or the oposite if we are playing the montage in reverse)
+				if (FirstSectionIdx == LastSectionIdx)
+				{
+					bLooped = PreviewInstance->IsReverse() ? (CurrentTime > PreviousTime) : (CurrentTime < PreviousTime);
+				}
+				// Otherwise, we are previewing a montage with multiple section. In this case we check if section at CurrentTime is the FirstSection and the section at PreviewTime is the LastSection (or the opposite if we are playing the montage in reverse)
+				else
+				{
+					const int32 SectionIndexPrevTime = Montage->GetSectionIndexFromPosition(PreviousTime);
+					const int32 SectionIndexCurrentTime = Montage->GetSectionIndexFromPosition(CurrentTime);
+					bLooped = PreviewInstance->IsReverse() ? (SectionIndexPrevTime == FirstSectionIdx && SectionIndexCurrentTime == LastSectionIdx) : (SectionIndexPrevTime == LastSectionIdx && SectionIndexCurrentTime == FirstSectionIdx);
+				}
+
+				// If we have looped...
+				if (bLooped)
+				{
+					float StartTime = 0.0f, EndTime = 0.0f;
+					Montage->GetSectionStartAndEndTime(LastSectionIdx, StartTime, EndTime);
+					SectionStartPosition = StartTime;
+				}
+			}
+			else // CurrentAsset is not a Montage
+			{
+				bLooped = PreviewInstance->IsReverse() ? (CurrentTime > PreviousTime) : (CurrentTime < PreviousTime);
+			}
+			
 			// Loop Mode: Preview mesh will consume root motion continually
 			if (ProcessRootMotionMode == EProcessRootMotionMode::Loop)
 			{
-				AddLocalTransform(ExtractedRootMotion.GetRootMotionTransform());
+				const FTransform RootMotion = ExtractedRootMotion.GetRootMotionTransform();
+				RootMotionTransform = RootMotion * RootMotionTransform;
 
 				//Handle moving component so that it stays within the editor floor
-				FTransform CurrentTransform = GetRelativeTransform();
-				FVector Trans = CurrentTransform.GetTranslation();
+				FVector Trans = RootMotionTransform.GetLocation();
 				Trans.X = WrapInRange(Trans.X, FloorMin.X, FloorMax.X);
 				Trans.Y = WrapInRange(Trans.Y, FloorMin.Y, FloorMax.Y);
-				CurrentTransform.SetTranslation(Trans);
-				SetRelativeTransform(CurrentTransform);
+				const FVector WrapDelta = Trans - RootMotionTransform.GetTranslation();
+				RootMotionTransform.SetTranslation(Trans);
+				
+				// If the location wraps, move the root motion reference too.
+				if (!WrapDelta.IsNearlyZero())
+				{
+					RootMotionReferenceTransform.AddToTranslation(WrapDelta);
+				}
+
+				if (bLooped)
+				{
+					RootMotionReferenceTransform = RootMotionTransform;
+				}
 			}
 			// Loop and Reset Mode: Preview mesh will consume root motion resetting the position back to the origin every time the animation loops
 			else if (ProcessRootMotionMode == EProcessRootMotionMode::LoopAndReset)
 			{
-				// In this mode we have to deal with montage explicitly because we can have multiple sections and we want to reset the position when the section loops
-				// and depending on the composition, CurrentTime < PreviousTime (or CurrentTime > PreviousTime when playing in reverse) is not enough
-				if (const UAnimMontage* Montage = Cast<UAnimMontage>(PreviewInstance->CurrentAsset))
+				RootMotionReferenceTransform = FTransform::Identity;
+				
+				if (bLooped)
 				{
-					const int32 PreviewStartSectionIdx = Montage->CompositeSections.IsValidIndex(PreviewInstance->MontagePreviewStartSectionIdx) ? PreviewInstance->MontagePreviewStartSectionIdx : Montage->GetSectionIndexFromPosition(CurrentTime);
-					const int32 FirstSectionIdx = PreviewInstance->MontagePreview_FindFirstSectionAsInMontage(PreviewStartSectionIdx);
-					const int32 LastSectionIdx = PreviewInstance->MontagePreview_FindLastSection(FirstSectionIdx);
-
-					bool bLooped = false;
-
-					// If FirstSection == LastSection we are previewing a single section
-					// In this case to know if we have looped we just need to check if CurrentTime < PreviousTime (or the oposite if we are playing the montage in reverse)
-					if (FirstSectionIdx == LastSectionIdx)
-					{
-						bLooped = PreviewInstance->IsReverse() ? (CurrentTime > PreviousTime) : (CurrentTime < PreviousTime);
-					}
-					// Otherwise, we are previewing a montage with multiple section. In this case we check if section at CurrentTime is the FirstSection and the section at PreviewTime is the LastSection (or the opposite if we are playing the montage in reverse)
-					else
-					{
-						const int32 SectionIndexPrevTime = Montage->GetSectionIndexFromPosition(PreviousTime);
-						const int32 SectionIndexCurrentTime = Montage->GetSectionIndexFromPosition(CurrentTime);
-						bLooped = PreviewInstance->IsReverse() ? (SectionIndexPrevTime == FirstSectionIdx && SectionIndexCurrentTime == LastSectionIdx) : (SectionIndexPrevTime == LastSectionIdx && SectionIndexCurrentTime == FirstSectionIdx);
-					}
-
-					// If we have looped...
-					if (bLooped)
-					{
-						//... and we are playing in reverse, reset the position of the mesh back to the origin and apply root motion from the end of the LastSection to CurrentTime
-						if (PreviewInstance->IsReverse())
-						{
-							float StartTime, EndTime;
-							Montage->GetSectionStartAndEndTime(LastSectionIdx, StartTime, EndTime);
-							const FTransform RootMotionDelta = ExtractRootMotionFromAnimationAsset(PreviewInstance->CurrentAsset, PreviewInstance->GetMirrorDataTable(), EndTime, CurrentTime);
-							SetRelativeTransform(RootMotionDelta);
-						}
-						// Otherwise, reset the position of the mesh back to the origin and apply root motion from the beginning of the FirstSection to CurrentTime
-						else
-						{
-							float StartTime, EndTime;
-							Montage->GetSectionStartAndEndTime(FirstSectionIdx, StartTime, EndTime);
-							const FTransform RootMotionDelta = ExtractRootMotionFromAnimationAsset(PreviewInstance->CurrentAsset, PreviewInstance->GetMirrorDataTable(), StartTime, CurrentTime);
-							SetRelativeTransform(RootMotionDelta);
-						}
-
-					}
-					// we haven't looped, just keep applying the extracted root motion
-					else
-					{
-						AddLocalTransform(ExtractedRootMotion.GetRootMotionTransform());
-					}
+					const FTransform InitialTransform = UE::Anim::ExtractRootTransformFromAnimationAsset(PreviewInstance->CurrentAsset, SectionStartPosition);
+					const FTransform RootMotionDelta = UE::Anim::ExtractRootMotionFromAnimationAsset(PreviewInstance->CurrentAsset, PreviewInstance->GetMirrorDataTable(), SectionStartPosition, CurrentTime);
+					RootMotionTransform = InitialTransform * RootMotionDelta;
 				}
-				else // CurrentAsset is not a Montage
+				else
 				{
-					bool bLooped = PreviewInstance->IsReverse() ? (CurrentTime > PreviousTime) : (CurrentTime < PreviousTime);
-					if (bLooped)
-					{
-						if (PreviewInstance->IsReverse())
-						{
-							const FTransform RootMotionDelta = ExtractRootMotionFromAnimationAsset(PreviewInstance->CurrentAsset, PreviewInstance->GetMirrorDataTable(),  PreviewInstance->CurrentAsset->GetPlayLength(), CurrentTime);
-							SetRelativeTransform(RootMotionDelta);
-						}
-						else
-						{
-							const FTransform RootMotionDelta = ExtractRootMotionFromAnimationAsset(PreviewInstance->CurrentAsset, PreviewInstance->GetMirrorDataTable(), 0.f, CurrentTime);
-							SetRelativeTransform(RootMotionDelta);
-						}
-					}
-					else
-					{
-						AddLocalTransform(ExtractedRootMotion.GetRootMotionTransform());
-					}
+					const FTransform RootMotion = ExtractedRootMotion.GetRootMotionTransform();
+					RootMotionTransform = RootMotion * RootMotionTransform;
 				}
 			}
 		}
 		else // Not Playing. When not playing user can still scrub the time line but animation is not ticking so we have to extract and apply root motion manually
 		{
-			const FTransform RootMotion = ExtractRootMotionFromAnimationAsset(PreviewInstance->CurrentAsset, PreviewInstance->GetMirrorDataTable(), PreviousTime, CurrentTime);
-			AddLocalTransform(RootMotion);
+			const FTransform RootMotion = UE::Anim::ExtractRootMotionFromAnimationAsset(PreviewInstance->CurrentAsset, PreviewInstance->GetMirrorDataTable(), PreviousTime, CurrentTime);
+			RootMotionTransform = RootMotion * RootMotionTransform;
 		}
 	}
+	else
+	{
+		RootMotionTransform = FTransform::Identity;
+	}
+
+	SetRelativeTransform(RootMotionTransform * TurntableTransform, /*bSweep*/false, /*OutSweepResult*/nullptr, bLooped ? ETeleportType::ResetPhysics : ETeleportType::None);
 }
 
 bool UDebugSkelMeshComponent::IsProcessingRootMotion() const 
@@ -355,6 +323,11 @@ EProcessRootMotionMode UDebugSkelMeshComponent::GetProcessRootMotionMode() const
 	return ProcessRootMotionMode;
 }
 
+EProcessRootMotionMode UDebugSkelMeshComponent::GetRequestedProcessRootMotionMode() const
+{
+	return RequestedProcessRootMotionMode;
+}
+
 bool UDebugSkelMeshComponent::DoesCurrentAssetHaveRootMotion() const
 {
 	if (PreviewInstance)
@@ -407,9 +380,9 @@ bool UDebugSkelMeshComponent::CanUseProcessRootMotionMode(EProcessRootMotionMode
 	}
 	
 	// Disable Loop modes if the current asset or animation blueprint doesn't have root motion
-	if(Mode != EProcessRootMotionMode::Ignore)
+	if (Mode != EProcessRootMotionMode::Ignore)
 	{
-		if(!DoesCurrentAssetHaveRootMotion())
+		if (!DoesCurrentAssetHaveRootMotion())
 		{
 			return false;
 		}
@@ -429,22 +402,43 @@ bool UDebugSkelMeshComponent::CanUseProcessRootMotionMode(EProcessRootMotionMode
 
 void UDebugSkelMeshComponent::SetProcessRootMotionMode(EProcessRootMotionMode Mode)
 {
-	if(CanUseProcessRootMotionMode(Mode))
+	RequestedProcessRootMotionMode = Mode;
+	
+	if (CanUseProcessRootMotionMode(Mode))
 	{
-		ProcessRootMotionMode = Mode;
+		SetProcessRootMotionModeInternal(Mode);
+	}
+}
+
+void UDebugSkelMeshComponent::SetProcessRootMotionModeInternal(EProcessRootMotionMode Mode)
+{
+	ProcessRootMotionMode = Mode;
 
-		if (ProcessRootMotionMode == EProcessRootMotionMode::Ignore)
+	if (ProcessRootMotionMode == EProcessRootMotionMode::LoopAndReset || ProcessRootMotionMode == EProcessRootMotionMode::Loop)
+	{
+		// Reset transform
+		const float CurrentTime = PreviewInstance->GetCurrentTime();
+		float SectionStartPosition = 0.0f;
+		if (const UAnimMontage* Montage = Cast<UAnimMontage>(PreviewInstance->CurrentAsset))
 		{
-			if (TurnTableMode == EPersonaTurnTableMode::Stopped)
-			{
-				SetWorldTransform(FTransform());
-			}
-			else
-			{
-				SetRelativeLocation(FVector::ZeroVector);
-			}
+			const int32 PreviewStartSectionIdx = Montage->CompositeSections.IsValidIndex(PreviewInstance->MontagePreviewStartSectionIdx) ? PreviewInstance->MontagePreviewStartSectionIdx : Montage->GetSectionIndexFromPosition(CurrentTime);
+			const int32 FirstSectionIdx = PreviewInstance->MontagePreview_FindFirstSectionAsInMontage(PreviewStartSectionIdx);
+			const int32 LastSectionIdx = PreviewInstance->MontagePreview_FindLastSection(FirstSectionIdx);
+			float StartTime = 0.0f, EndTime = 0.0f;
+			Montage->GetSectionStartAndEndTime(LastSectionIdx, StartTime, EndTime);
+			SectionStartPosition = StartTime;
 		}
+	
+		const FTransform InitialTransform = UE::Anim::ExtractRootTransformFromAnimationAsset(PreviewInstance->CurrentAsset, SectionStartPosition);
+		const FTransform RootMotionDelta = UE::Anim::ExtractRootMotionFromAnimationAsset(PreviewInstance->CurrentAsset, PreviewInstance->GetMirrorDataTable(), SectionStartPosition, CurrentTime);
+		RootMotionTransform = InitialTransform * RootMotionDelta;
+	}
+	else if (ProcessRootMotionMode == EProcessRootMotionMode::Ignore)
+	{
+		RootMotionTransform = FTransform::Identity;
 	}
+	
+	SetRelativeTransform(RootMotionTransform * TurntableTransform, /*bSweep*/false, /*OutSweepResult*/nullptr, ETeleportType::ResetPhysics);
 }
 
 bool UDebugSkelMeshComponent::IsTrackingAttachedLOD() const
@@ -641,12 +635,9 @@ void UDebugSkelMeshComponent::EnablePreview(bool bEnable, UAnimationAsset* Previ
 			// Reset to previous animation asset's root motion playback time to prevent this from influencing the new animation asset previewing during root motion consumption.
 			ConsumeRootMotionPreviousPlaybackTime = 0.0f;
 			
-			// Reset mesh relative transform. It could be offset from the origin if we are processing root motion
-			if (!GetRelativeTransform().Equals(FTransform::Identity))
-			{
-				SetRelativeTransform(FTransform::Identity);
-			}
-
+			// Update requested process root motion mode, the new asset might support requested processing mode.
+			// Note: This might further reset the transform.
+			SetProcessRootMotionModeInternal(CanUseProcessRootMotionMode(RequestedProcessRootMotionMode) ? RequestedProcessRootMotionMode : EProcessRootMotionMode::Ignore);
 		}
 		else if (IsPreviewOn())
 		{
@@ -686,6 +677,31 @@ void UDebugSkelMeshComponent::PostInitMeshObject(FSkeletalMeshObject* InMeshObje
 	}
 }
 
+void UDebugSkelMeshComponent::OnMirrorDataTableChanged()
+{
+	if (ProcessRootMotionMode == EProcessRootMotionMode::LoopAndReset)
+	{
+		// Reset transform
+		const float CurrentTime = PreviewInstance->GetCurrentTime();
+		float SectionStartPosition = 0.0f;
+		if (const UAnimMontage* Montage = Cast<UAnimMontage>(PreviewInstance->CurrentAsset))
+		{
+			const int32 PreviewStartSectionIdx = Montage->CompositeSections.IsValidIndex(PreviewInstance->MontagePreviewStartSectionIdx) ? PreviewInstance->MontagePreviewStartSectionIdx : Montage->GetSectionIndexFromPosition(CurrentTime);
+			const int32 FirstSectionIdx = PreviewInstance->MontagePreview_FindFirstSectionAsInMontage(PreviewStartSectionIdx);
+			const int32 LastSectionIdx = PreviewInstance->MontagePreview_FindLastSection(FirstSectionIdx);
+			float StartTime = 0.0f, EndTime = 0.0f;
+			Montage->GetSectionStartAndEndTime(LastSectionIdx, StartTime, EndTime);
+			SectionStartPosition = StartTime;
+		}
+	
+		const FTransform InitialTransform = UE::Anim::ExtractRootTransformFromAnimationAsset(PreviewInstance->CurrentAsset, SectionStartPosition);
+		const FTransform RootMotionDelta = UE::Anim::ExtractRootMotionFromAnimationAsset(PreviewInstance->CurrentAsset, PreviewInstance->GetMirrorDataTable(), SectionStartPosition, CurrentTime);
+		RootMotionTransform = InitialTransform * RootMotionDelta;
+	
+		SetRelativeTransform(RootMotionTransform * TurntableTransform, /*bSweep*/false, /*OutSweepResult*/nullptr, ETeleportType::ResetPhysics);
+	}
+}
+
 void UDebugSkelMeshComponent::SetShowBoneWeight(bool bNewShowBoneWeight)
 {
 	// Check we are actually changing it!
@@ -1140,9 +1156,10 @@ void UDebugSkelMeshComponent::TickComponent(float DeltaTime, enum ELevelTick Tic
 		return;
 	}
 
+	// Update turn table transform. The transform is applied in ConsumeRootMotion().
 	if (TurnTableMode == EPersonaTurnTableMode::Playing)
 	{
-		FRotator Rotation = GetRelativeTransform().Rotator();
+		FRotator Rotation = TurntableTransform.Rotator();
 		// Take into account time dilation, so it doesn't affect turn table turn rate.
 		float CurrentTimeDilation = 1.0f;
 		if (UWorld* MyWorld = GetWorld())
@@ -1150,7 +1167,11 @@ void UDebugSkelMeshComponent::TickComponent(float DeltaTime, enum ELevelTick Tic
 			CurrentTimeDilation = MyWorld->GetWorldSettings()->GetEffectiveTimeDilation();
 		}
 		Rotation.Yaw += 36.f * TurnTableSpeedScaling * DeltaTime / FMath::Max(CurrentTimeDilation, KINDA_SMALL_NUMBER);
-		SetRelativeRotation(Rotation);
+		TurntableTransform.SetRotation(Rotation.Quaternion());
+	}
+	else if (TurnTableMode == EPersonaTurnTableMode::Stopped)
+	{
+		TurntableTransform = FTransform::Identity;
 	}
 
     // Brute force approach to ensure that when materials are changed the names are cached parameter names are updated 
diff --git a/Engine/Source/Runtime/Engine/Classes/Animation/AnimSequenceHelpers.h b/Engine/Source/Runtime/Engine/Classes/Animation/AnimSequenceHelpers.h
--- a/Engine/Source/Runtime/Engine/Classes/Animation/AnimSequenceHelpers.h
+++ b/Engine/Source/Runtime/Engine/Classes/Animation/AnimSequenceHelpers.h
@@ -182,7 +182,35 @@ namespace Anim {
 	* @param	KeyIndex		Key index to retrieve from the track.
 	*/
 	ENGINE_API void ExtractBoneTransform(const struct FRawAnimSequenceTrack& RawTrack, FTransform& OutTransform, int32 KeyIndex);
-		
+
+
+	/**
+	 * Mirrors transform based on the specified mirror table.
+	 *
+	 * @param	Transform		Transform to mirror
+	 * @param	MirrorDataTable	Animation data mirroring data table.
+	 * @returns Transform mirrored based on the mirror data table.
+	 */
+	ENGINE_API FTransform MirrorTransform(const FTransform& Transform, const UMirrorDataTable& MirrorDataTable);
+
+	/**
+	 * Extract relative root motion between start and end position of the animation asset.
+	 *
+	 * @param	Animation		Animation asset to extract the root motion from.
+	 * @param	MirrorDataTable	[optional] Animation data mirroring data table.
+	 * @param	StartPosition	Start Position of the animation.
+	 * @param	EndPosition		End Position of the animation.
+	 */
+	ENGINE_API FTransform ExtractRootMotionFromAnimationAsset(const UAnimationAsset* Animation, const UMirrorDataTable* MirrorDataTable, float StartPosition, float EndPosition);
+
+	/**
+	 * Extract absolute root transform at specific position in time from the animation asset.
+	 *
+	 * @param	Animation		Animation asset to extract the root motion from.
+	 * @param	Position		Start Position of the animation.
+	 */
+	ENGINE_API FTransform ExtractRootTransformFromAnimationAsset(const UAnimationAsset* Animation, float Position);
+
 	namespace Compression
 	{
 		/**
diff --git a/Engine/Source/Runtime/Engine/Private/Animation/AnimSequenceHelpers.cpp b/Engine/Source/Runtime/Engine/Private/Animation/AnimSequenceHelpers.cpp
--- a/Engine/Source/Runtime/Engine/Private/Animation/AnimSequenceHelpers.cpp
+++ b/Engine/Source/Runtime/Engine/Private/Animation/AnimSequenceHelpers.cpp
@@ -6,6 +6,9 @@
 
 
 #include "Animation/AnimSequence.h"
+#include "Animation/AnimMontage.h"
+#include "Animation/AnimComposite.h"
+#include "Animation/MirrorDataTable.h"
 #include "BonePose.h"
 
 #include "AnimationRuntime.h"
@@ -191,7 +194,68 @@ void CopyCurveDataToModel(const FRawCurveTracks& CurveData, const USkeleton* Ske
 		OutTransform.SetScale3D(DefaultScale3D);
 	}	
 }
+
+FTransform MirrorTransform(const FTransform& Transform, const UMirrorDataTable& MirrorDataTable)
+{
+	FVector T = Transform.GetTranslation();
+	T = FAnimationRuntime::MirrorVector(T, MirrorDataTable.MirrorAxis);
+
+	FQuat Q = Transform.GetRotation();
+	Q = FAnimationRuntime::MirrorQuat(Q, MirrorDataTable.MirrorAxis);
+
+	const FVector S = Transform.GetScale3D();
+	return FTransform(Q, T, S);
+}
+
+FTransform ExtractRootMotionFromAnimationAsset(const UAnimationAsset* Animation, const UMirrorDataTable* MirrorDataTable, float StartPosition, float EndPosition)
+{
+	FTransform Result = FTransform::Identity;
 	
+	if (const UAnimMontage* Montage = Cast<UAnimMontage>(Animation))
+	{
+		Result = Montage->ExtractRootMotionFromTrackRange(StartPosition, EndPosition);
+	}
+	else if (const UAnimComposite* AnimComposite = Cast<UAnimComposite>(Animation))
+	{
+		FRootMotionMovementParams RootMotion;
+		AnimComposite->ExtractRootMotionFromTrack(AnimComposite->AnimationTrack, StartPosition, EndPosition, RootMotion);
+		Result = RootMotion.GetRootMotionTransform();
+	}
+	else if (const UAnimSequence* AnimSequence = Cast<UAnimSequence>(Animation))
+	{
+		Result = AnimSequence->ExtractRootMotionFromRange(StartPosition, EndPosition);
+	}
+
+	if (MirrorDataTable)
+	{
+		Result =  MirrorTransform(Result, *MirrorDataTable);
+	}
+
+	return Result;
+}
+
+FTransform ExtractRootTransformFromAnimationAsset(const UAnimationAsset* Animation, float Time)
+{
+	FTransform Result = FTransform::Identity;
+	if (const UAnimMontage* AnimMontage = Cast<UAnimMontage>(Animation))
+	{
+		if (const FAnimSegment* Segment = AnimMontage->SlotAnimTracks[0].AnimTrack.GetSegmentAtTime(Time))
+		{
+			if (const UAnimSequence* AnimSequence = Cast<UAnimSequence>(Segment->GetAnimReference()))
+			{
+				const float AnimSequenceTime = Segment->ConvertTrackPosToAnimPos(Time);
+				Result = AnimSequence->ExtractRootTrackTransform(AnimSequenceTime, nullptr);
+			}	
+		}
+	}
+	else if (const UAnimSequence* AnimSequence = Cast<UAnimSequence>(Animation))
+	{
+		Result = AnimSequence->ExtractRootTrackTransform(Time, nullptr);
+	}
+
+	return Result;
+}
+
 Retargeting::FRetargetingScope::FRetargetingScope(const USkeleton* InSourceSkeleton, FCompactPose& ToRetargetPose, const DataModel::FEvaluationContext& InEvaluationContext)
 	: SourceSkeleton(InSourceSkeleton),
 	RetargetPose(ToRetargetPose),
