Add way to display the axis for transforms using the MakeEditWidget meta tag.
This is enabled by default, but the original behavior can be brought back using ULevelEditorViewportSettings::bAllowEditWidgetAxisDisplay

diff --git a/Engine/Source/Editor/UnrealEd/Classes/Settings/LevelEditorViewportSettings.h b/Engine/Source/Editor/UnrealEd/Classes/Settings/LevelEditorViewportSettings.h
--- a/Engine/Source/Editor/UnrealEd/Classes/Settings/LevelEditorViewportSettings.h
+++ b/Engine/Source/Editor/UnrealEd/Classes/Settings/LevelEditorViewportSettings.h
@@ -321,6 +321,10 @@ class ULevelEditorViewportSettings
 	UPROPERTY(EditAnywhere, config, Category = LookAndFeel, meta=(AdvancedDisplay))
 	uint32 bShowActorEditorContext : 1;
 
+	/** If true, the Edit widget of a transform will display the axis */
+	UPROPERTY(EditAnywhere, config, Category=LookAndFeel, meta=( DisplayName = "Enable Axis drawing for transform Edit Widget" ))
+	uint32 bAllowEditWidgetAxisDisplay:1;
+
 	/** How fast the perspective camera moves when flying through the world. */
 	UPROPERTY(config, meta=(UIMin = "1", UIMax = "8", ClampMin="1", ClampMax="8"))
 	int32 CameraSpeed;
diff --git a/Engine/Source/Editor/UnrealEd/Private/Settings/SettingsClasses.cpp b/Engine/Source/Editor/UnrealEd/Private/Settings/SettingsClasses.cpp
--- a/Engine/Source/Editor/UnrealEd/Private/Settings/SettingsClasses.cpp
+++ b/Engine/Source/Editor/UnrealEd/Private/Settings/SettingsClasses.cpp
@@ -944,6 +944,7 @@ ULevelEditorViewportSettings::ULevelEditorViewportSettings( const FObjectInitial
 	bAllowArcballRotate = false;
 	bAllowScreenRotate = false;
 	bShowActorEditorContext = true;
+	bAllowEditWidgetAxisDisplay = true;
 	MouseSensitivty = .2f;
 	bUseLegacyCameraMovementNotifications = false;
 	// Set a default preview mesh
diff --git a/Engine/Source/Editor/UnrealEd/Private/Tools/LegacyEdModeWidgetHelpers.cpp b/Engine/Source/Editor/UnrealEd/Private/Tools/LegacyEdModeWidgetHelpers.cpp
--- a/Engine/Source/Editor/UnrealEd/Private/Tools/LegacyEdModeWidgetHelpers.cpp
+++ b/Engine/Source/Editor/UnrealEd/Private/Tools/LegacyEdModeWidgetHelpers.cpp
@@ -11,6 +11,7 @@
 #include "EngineUtils.h"
 #include "CanvasTypes.h"
 #include "CanvasItem.h"
+#include "Settings/LevelEditorViewportSettings.h"
 
 const FName FLegacyEdModeWidgetHelper::MD_MakeEditWidget(TEXT("MakeEditWidget"));
 const FName FLegacyEdModeWidgetHelper::MD_ValidateWidgetUsing(TEXT("ValidateWidgetUsing"));
@@ -635,6 +636,9 @@ void FLegacyEdModeWidgetHelper::Render(const FSceneView* View, FViewport* Viewpo
 			UClass* Class = BestSelectedItem->GetClass();
 			TArray<FPropertyWidgetInfo> WidgetInfos;
 			GetPropertyWidgetInfos(Class, BestSelectedItem, WidgetInfos);
+
+			const bool bAllowEditWidgetAxisDisplay = GetDefault<ULevelEditorViewportSettings>()->bAllowEditWidgetAxisDisplay;
+
 			FEditorScriptExecutionGuard ScriptGuard;
 			for (const FPropertyWidgetInfo& WidgetInfo : WidgetInfos)
 			{
@@ -654,6 +658,11 @@ void FLegacyEdModeWidgetHelper::Render(const FSceneView* View, FViewport* Viewpo
 
 				if (bHitTesting) PDI->SetHitProxy(new HPropertyWidgetProxy(WidgetInfo.PropertyName, WidgetInfo.PropertyIndex, WidgetInfo.bIsTransform));
 				DrawWireDiamond(PDI, WidgetTM, WidgetRadius, WidgetColor, SDPG_Foreground);
+				if (bAllowEditWidgetAxisDisplay && WidgetInfo.bIsTransform)
+				{
+					constexpr float AxisScale = 1.5f;
+					DrawCoordinateSystem(PDI, WorldWidgetTransform.GetTranslation(), WorldWidgetTransform.GetRotation().Rotator(), AxisScale * WidgetRadius, SDPG_Foreground);
+				}
 				if (bHitTesting) PDI->SetHitProxy(NULL);
 			}
 		}
