Fix regression with CL 30443392 (change of default param value)

This was causing MakeCurrentLevel/SetLevelVisibilityNoGlobalUpdateInternal to dirty the persistent level

#rnx

diff --git a/Engine/Source/Editor/UnrealEd/Private/EditorLevelUtils.cpp b/Engine/Source/Editor/UnrealEd/Private/EditorLevelUtils.cpp
--- a/Engine/Source/Editor/UnrealEd/Private/EditorLevelUtils.cpp
+++ b/Engine/Source/Editor/UnrealEd/Private/EditorLevelUtils.cpp
@@ -1194,7 +1194,7 @@ void SetLevelVisibilityNoGlobalUpdateInternal(ULevel* Level, const bool bShouldB
 		{
 			if (ModifyMode == ELevelVisibilityDirtyMode::ModifyOnChange)
 			{
-				CurLevelModel->Modify();
+				CurLevelModel->Modify(false);
 			}
 
 			for (TArray<FBspSurf>::TIterator SurfaceIterator(CurLevelModel->Surfs); SurfaceIterator; ++SurfaceIterator)
