-Fix divergence on actor factory usage introduced in the previous refactoring
-Factory chosen in FLevelEditorViewportClient::TryPlacingAssetObject must be passed as a param in the call to UE::AssetPlacementUtil::PlaceAssetInCurrentLevel since if we leave it to nullptr it'll choose a different one (chosing from GEditor vs choosing from the edtior  PlacementSubsystem)

diff --git a/Engine/Source/Editor/UnrealEd/Private/LevelEditorViewport.cpp b/Engine/Source/Editor/UnrealEd/Private/LevelEditorViewport.cpp
--- a/Engine/Source/Editor/UnrealEd/Private/LevelEditorViewport.cpp
+++ b/Engine/Source/Editor/UnrealEd/Private/LevelEditorViewport.cpp
@@ -315,7 +315,9 @@ TArray<FTypedElementHandle> FLevelEditorViewportClient::TryPlacingAssetObject(UL
 
 		if (!bPlaced && ActorFactory != NULL )
 		{
-			PlacedItems = UE::AssetPlacementUtil::PlaceAssetInCurrentLevel(ObjToUse, PlacementOptions);
+			UE::AssetPlacementUtil::FExtraPlaceAssetOptions OverridenPlacementOptions = PlacementOptions;
+			OverridenPlacementOptions.FactoryToUse = TScriptInterface<IAssetFactoryInterface>(ActorFactory);
+			PlacedItems = UE::AssetPlacementUtil::PlaceAssetInCurrentLevel(ObjToUse, OverridenPlacementOptions);
 			ForEachActorInElementArray(PlacedItems, [&PlacedActors](AActor& Actor)
 			{
 				PlacedActors.Add(&Actor);
