RayTracing: Limit the menu entries for the RT debug menu to what is actually supported

Most of the debug modes in the menu require support for raytracing shaders, except for Barycentrics which has an inline implementation.

Also, re-order some of the entries so that barycentrics (the most basic mode) comes first, followed by PrimaryRays, followed by the payload debugging modes and the instance/triangle visualizations.

diff --git a/Engine/Source/Editor/UnrealEd/Private/RayTracingDebugVisualizationMenuCommands.cpp b/Engine/Source/Editor/UnrealEd/Private/RayTracingDebugVisualizationMenuCommands.cpp
--- a/Engine/Source/Editor/UnrealEd/Private/RayTracingDebugVisualizationMenuCommands.cpp
+++ b/Engine/Source/Editor/UnrealEd/Private/RayTracingDebugVisualizationMenuCommands.cpp
@@ -1,6 +1,7 @@
 // Copyright Epic Games, Inc. All Rights Reserved.
 
 #include "RayTracingDebugVisualizationMenuCommands.h"
+#include "DataDrivenShaderPlatformInfo.h"
 #include "Containers/UnrealString.h"
 #include "Framework/Commands/InputChord.h"
 #include "Internationalization/Text.h"
@@ -31,45 +32,55 @@ void FRayTracingDebugVisualizationMenuCommands::BuildCommandMap()
 
 void FRayTracingDebugVisualizationMenuCommands::CreateRayTracingDebugVisualizationCommands()
 {
-	RayTracingDebugModeNames.Add(LOCTEXT("Radiance", "Radiance"));
-	RayTracingDebugModeNames.Add(LOCTEXT("World Normal", "World Normal"));
-	RayTracingDebugModeNames.Add(LOCTEXT("BaseColor", "BaseColor"));
-	RayTracingDebugModeNames.Add(LOCTEXT("DiffuseColor", "DiffuseColor"));
-	RayTracingDebugModeNames.Add(LOCTEXT("SpecularColor", "SpecularColor"));
-	RayTracingDebugModeNames.Add(LOCTEXT("Opacity", "Opacity"));
-	RayTracingDebugModeNames.Add(LOCTEXT("Metallic", "Metallic"));
-	RayTracingDebugModeNames.Add(LOCTEXT("Specular", "Specular"));
-	RayTracingDebugModeNames.Add(LOCTEXT("Roughness", "Roughness"));
-	RayTracingDebugModeNames.Add(LOCTEXT("Ior", "Ior"));
-	RayTracingDebugModeNames.Add(LOCTEXT("ShadingModelID", "ShadingModelID"));
-	RayTracingDebugModeNames.Add(LOCTEXT("BlendingMode", "BlendingMode"));
-	RayTracingDebugModeNames.Add(LOCTEXT("PrimitiveLightingChannelMask", "PrimitiveLightingChannelMask"));
-	RayTracingDebugModeNames.Add(LOCTEXT("CustomData", "CustomData"));
-	RayTracingDebugModeNames.Add(LOCTEXT("GBufferAO", "GBufferAO"));
-	RayTracingDebugModeNames.Add(LOCTEXT("IndirectIrradiance", "IndirectIrradiance"));
-	RayTracingDebugModeNames.Add(LOCTEXT("World Position", "World Position"));
-	RayTracingDebugModeNames.Add(LOCTEXT("HitKind", "HitKind"));
+	bool bRayTracingShadersSupported = FDataDrivenShaderPlatformInfo::GetSupportsRayTracingShaders(GMaxRHIShaderPlatform);
+	
+	// always supported (as long as either inline RT or RT shaders work)
 	RayTracingDebugModeNames.Add(LOCTEXT("Barycentrics", "Barycentrics"));
-	RayTracingDebugModeNames.Add(LOCTEXT("PrimaryRays", "PrimaryRays"));
-	RayTracingDebugModeNames.Add(LOCTEXT("World Tangent", "World Tangent"));
-	RayTracingDebugModeNames.Add(LOCTEXT("Anisotropy", "Anisotropy"));
-	RayTracingDebugModeNames.Add(LOCTEXT("Instances", "Instances"));
-	RayTracingDebugModeNames.Add(LOCTEXT("Instance Overlap", "Instance Overlap"));
-	RayTracingDebugModeNames.Add(LOCTEXT("Triangle Hit Count", "Triangle Hit Count"));
-	RayTracingDebugModeNames.Add(LOCTEXT("Hit Count Per Instance", "Hit Count Per Instance"));
-
-	if (GRHISupportsShaderTimestamp)
+	if (bRayTracingShadersSupported)
 	{
-		RayTracingDebugModeNames.Add(LOCTEXT("Performance", "Performance"));
+		// runs basic lighting calculations on hits
+		RayTracingDebugModeNames.Add(LOCTEXT("PrimaryRays", "PrimaryRays"));
+
+		// all of these fields reflect entries in the payload which require running a CHS
+		RayTracingDebugModeNames.Add(LOCTEXT("Radiance", "Radiance"));
+		RayTracingDebugModeNames.Add(LOCTEXT("World Normal", "World Normal"));
+		RayTracingDebugModeNames.Add(LOCTEXT("BaseColor", "BaseColor"));
+		RayTracingDebugModeNames.Add(LOCTEXT("DiffuseColor", "DiffuseColor"));
+		RayTracingDebugModeNames.Add(LOCTEXT("SpecularColor", "SpecularColor"));
+		RayTracingDebugModeNames.Add(LOCTEXT("Opacity", "Opacity"));
+		RayTracingDebugModeNames.Add(LOCTEXT("Metallic", "Metallic"));
+		RayTracingDebugModeNames.Add(LOCTEXT("Specular", "Specular"));
+		RayTracingDebugModeNames.Add(LOCTEXT("Roughness", "Roughness"));
+		RayTracingDebugModeNames.Add(LOCTEXT("Ior", "Ior"));
+		RayTracingDebugModeNames.Add(LOCTEXT("ShadingModelID", "ShadingModelID"));
+		RayTracingDebugModeNames.Add(LOCTEXT("BlendingMode", "BlendingMode"));
+		RayTracingDebugModeNames.Add(LOCTEXT("PrimitiveLightingChannelMask", "PrimitiveLightingChannelMask"));
+		RayTracingDebugModeNames.Add(LOCTEXT("CustomData", "CustomData"));
+		RayTracingDebugModeNames.Add(LOCTEXT("GBufferAO", "GBufferAO"));
+		RayTracingDebugModeNames.Add(LOCTEXT("IndirectIrradiance", "IndirectIrradiance"));
+		RayTracingDebugModeNames.Add(LOCTEXT("World Position", "World Position"));
+		RayTracingDebugModeNames.Add(LOCTEXT("HitKind", "HitKind"));
+		RayTracingDebugModeNames.Add(LOCTEXT("World Tangent", "World Tangent"));
+		RayTracingDebugModeNames.Add(LOCTEXT("Anisotropy", "Anisotropy"));
+
+		// debugging the geometry itself
+		RayTracingDebugModeNames.Add(LOCTEXT("Instances", "Instances"));
+		RayTracingDebugModeNames.Add(LOCTEXT("Instance Overlap", "Instance Overlap"));
+		RayTracingDebugModeNames.Add(LOCTEXT("Triangle Hit Count", "Triangle Hit Count"));
+		RayTracingDebugModeNames.Add(LOCTEXT("Hit Count Per Instance", "Hit Count Per Instance"));
+	
+		if (GRHISupportsShaderTimestamp)
+		{
+			RayTracingDebugModeNames.Add(LOCTEXT("Performance", "Performance"));
+		}
+	
+		RayTracingDebugModeNames.Add(LOCTEXT("Triangles", "Triangles"));
+		RayTracingDebugModeNames.Add(LOCTEXT("FarField", "FarField"));
+		RayTracingDebugModeNames.Add(LOCTEXT("Dynamic Instances", "Dynamic Instances"));
+		RayTracingDebugModeNames.Add(LOCTEXT("Proxy Type", "Proxy Type"));
+		RayTracingDebugModeNames.Add(LOCTEXT("Picker", "Picker"));
+		RayTracingDebugModeNames.Add(LOCTEXT("Light Grid Occupancy", "Light Grid Occupancy"));
 	}
-
-	RayTracingDebugModeNames.Add(LOCTEXT("Triangles", "Triangles"));
-	RayTracingDebugModeNames.Add(LOCTEXT("FarField", "FarField"));
-	RayTracingDebugModeNames.Add(LOCTEXT("Dynamic Instances", "Dynamic Instances"));
-	RayTracingDebugModeNames.Add(LOCTEXT("Proxy Type", "Proxy Type"));
-	RayTracingDebugModeNames.Add(LOCTEXT("Picker", "Picker"));
-	RayTracingDebugModeNames.Add(LOCTEXT("Light Grid Occupancy", "Light Grid Occupancy"));
-
 	for ( int32 RayTracingDebugIndex = 0; RayTracingDebugIndex < RayTracingDebugModeNames.Num(); ++RayTracingDebugIndex)
 	{
 		const FText CommandNameText = RayTracingDebugModeNames[RayTracingDebugIndex];
