Only log shader compile stats if we've compiled more shaders since we last tried to log stats.

- Printing still happens on an interval (controlled by CVarShaderCompilerStatsPrintoutInterval)
- If we have hit our interval, and the number of shaders compiled is different since last time, we print.

diff --git a/Engine/Source/Editor/UnrealEd/Private/ShaderStatsCollector.cpp b/Engine/Source/Editor/UnrealEd/Private/ShaderStatsCollector.cpp
--- a/Engine/Source/Editor/UnrealEd/Private/ShaderStatsCollector.cpp
+++ b/Engine/Source/Editor/UnrealEd/Private/ShaderStatsCollector.cpp
@@ -123,7 +123,14 @@ void FShaderStatsReporter::LogStats()
 	}
 	
 	AggregateStats(LocalStats);
-	LocalStats.WriteStatSummary();
+
+	static uint32 PreviousTotalShadersCompiled = 0;
+	const uint32 CurrentTotalShadersCompiled = LocalStats.GetTotalShadersCompiled();
+	if (CurrentTotalShadersCompiled > PreviousTotalShadersCompiled)
+	{
+		LocalStats.WriteStatSummary();
+		PreviousTotalShadersCompiled = CurrentTotalShadersCompiled;
+	}
 }
 
 void FShaderStatsReporter::AggregateStats(FShaderCompilerStats& OutStats) const
