Fix Select All Actor Children was not part of a transaction

diff --git a/Engine/Source/Editor/UnrealEd/Private/Subsystems/EditorActorSubsystem.cpp b/Engine/Source/Editor/UnrealEd/Private/Subsystems/EditorActorSubsystem.cpp
--- a/Engine/Source/Editor/UnrealEd/Private/Subsystems/EditorActorSubsystem.cpp
+++ b/Engine/Source/Editor/UnrealEd/Private/Subsystems/EditorActorSubsystem.cpp
@@ -320,15 +320,17 @@ void UEditorActorSubsystem::SelectAllChildren(bool bRecurseChildren)
 		return;
 	}
 
+	FText TransactionLabel;
 	if (bRecurseChildren)
 	{
-		const FScopedTransaction Transaction(NSLOCTEXT("UnrealEd", "SelectAllDescendants", "Select All Descendants"));
+		TransactionLabel = NSLOCTEXT("UnrealEd", "SelectAllDescendants", "Select All Descendants");
 	}
 	else
 	{
-		const FScopedTransaction Transaction(NSLOCTEXT("UnrealEd", "SelectAllChildren", "Select All Children"));
+		TransactionLabel = NSLOCTEXT("UnrealEd", "SelectAllChildren", "Select All Children");
 	}
 
+	const FScopedTransaction Transaction(TransactionLabel);
 	GUnrealEd->edactSelectAllChildren(bRecurseChildren);
 }
 
