[UX Improvement] Make the measure tool color in viewport configurable via setting

diff --git a/Engine/Source/Editor/UnrealEd/Classes/Settings/EditorStyleSettings.h b/Engine/Source/Editor/UnrealEd/Classes/Settings/EditorStyleSettings.h
--- a/Engine/Source/Editor/UnrealEd/Classes/Settings/EditorStyleSettings.h
+++ b/Engine/Source/Editor/UnrealEd/Classes/Settings/EditorStyleSettings.h
@@ -84,6 +84,10 @@ public:
 	UPROPERTY(EditAnywhere, config, Category=UserInterface, meta=(DisplayName="Additional Viewport Selection Colors"))
 	FLinearColor AdditionalSelectionColors[6];
 
+	/** The color used for overlay tools inside of the viewport, like the measure tool */
+	UPROPERTY(EditAnywhere, config, Category = UserInterface)
+	FLinearColor ViewportToolOverlayColor;
+
 	UPROPERTY(config)
 	bool bEnableEditorWindowBackgroundColor;
 
diff --git a/Engine/Source/Editor/UnrealEd/Private/DragTool_Measure.cpp b/Engine/Source/Editor/UnrealEd/Private/DragTool_Measure.cpp
--- a/Engine/Source/Editor/UnrealEd/Private/DragTool_Measure.cpp
+++ b/Engine/Source/Editor/UnrealEd/Private/DragTool_Measure.cpp
@@ -78,9 +78,9 @@ void FDragTool_Measure::Render(const FSceneView* View, FCanvas* Canvas)
 
 	if (View != nullptr && Canvas != nullptr && Length >= 1.f)
 	{
-		const FLinearColor SelectionColor = GetDefault<UEditorStyleSettings>()->SelectionColor;
+		const FLinearColor ToolColor = GetDefault<UEditorStyleSettings>()->ViewportToolOverlayColor;
 		FCanvasLineItem LineItem( PixelStart, PixelEnd );
-		LineItem.SetColor( SelectionColor );
+		LineItem.SetColor( ToolColor );
 		Canvas->DrawItem( LineItem );
 
 		const FVector2D PixelMid = FVector2D(PixelStart + ((PixelEnd - PixelStart) / 2));
@@ -113,7 +113,7 @@ void FDragTool_Measure::Render(const FSceneView* View, FCanvas* Canvas)
 		const FText LengthStr = FText::AsNumber( Length / Divisor, &Options );
 
 
-		FCanvasTextItem TextItem( FVector2D( FMath::FloorToFloat(PixelMid.X), FMath::FloorToFloat(PixelMid.Y) ), LengthStr, GEngine->GetSmallFont(), SelectionColor );
+		FCanvasTextItem TextItem( FVector2D( FMath::FloorToFloat(PixelMid.X), FMath::FloorToFloat(PixelMid.Y) ), LengthStr, GEngine->GetSmallFont(), ToolColor );
 		TextItem.bCentreX = true;
 		Canvas->DrawItem( TextItem );
 	}
diff --git a/Engine/Source/Editor/UnrealEd/Private/Settings/EditorStyleClasses.cpp b/Engine/Source/Editor/UnrealEd/Private/Settings/EditorStyleClasses.cpp
--- a/Engine/Source/Editor/UnrealEd/Private/Settings/EditorStyleClasses.cpp
+++ b/Engine/Source/Editor/UnrealEd/Private/Settings/EditorStyleClasses.cpp
@@ -32,6 +32,7 @@ UEditorStyleSettings::UEditorStyleSettings( const FObjectInitializer& ObjectInit
 		FStyleColors::AccentYellow.GetSpecifiedColor(),
 		FStyleColors::AccentGreen.GetSpecifiedColor(),
 	}
+	, ViewportToolOverlayColor(FLinearColor::White)
 {
 	bEnableUserEditorLayoutManagement = true;
 
