Scene Outliner: Add a key binding to scroll to the Current Folder

diff --git a/Engine/Source/Editor/LevelEditor/Private/LevelEditorActions.cpp b/Engine/Source/Editor/LevelEditor/Private/LevelEditorActions.cpp
--- a/Engine/Source/Editor/LevelEditor/Private/LevelEditorActions.cpp
+++ b/Engine/Source/Editor/LevelEditor/Private/LevelEditorActions.cpp
@@ -104,6 +104,7 @@
 #include "ActorGroupingUtils.h"
 #include "LevelUtils.h"
 #include "ISceneOutliner.h"
+#include "SceneOutlinerStandaloneTypes.h"
 #include "ISettingsModule.h"
 #include "PlatformInfo.h"
 #include "Misc/CoreMisc.h"
@@ -2463,6 +2464,27 @@ void FLevelEditorActionCallbacks::OnFocusOutlinerToSelection(TWeakPtr<SLevelEdit
 	}
 }
 
+void FLevelEditorActionCallbacks::OnFocusOutlinerToContextFolder(TWeakPtr<SLevelEditor> LevelEditor)
+{
+	if (const TSharedPtr<SLevelEditor> Editor = LevelEditor.Pin())
+	{
+		if (UWorld* World = Editor->GetWorld())
+		{
+			const FFolder ContextFolder = FActorFolders::Get().GetActorEditorContextFolder(*World);
+			if (ContextFolder.IsValid())
+			{
+				for (TWeakPtr<ISceneOutliner> SceneOutliner : Editor->GetAllSceneOutliners())
+				{
+					if (const TSharedPtr<ISceneOutliner> Outliner = SceneOutliner.Pin())
+					{
+						Outliner->FrameItem(ContextFolder);
+					}
+				}
+			}
+		}
+	}
+}
+
 void FLevelEditorActionCallbacks::OpenPlaceActors()
 {
 	FLevelEditorModule& LevelEditorModule = FModuleManager::GetModuleChecked<FLevelEditorModule>( TEXT("LevelEditor") );
diff --git a/Engine/Source/Editor/LevelEditor/Private/SLevelEditor.cpp b/Engine/Source/Editor/LevelEditor/Private/SLevelEditor.cpp
--- a/Engine/Source/Editor/LevelEditor/Private/SLevelEditor.cpp
+++ b/Engine/Source/Editor/LevelEditor/Private/SLevelEditor.cpp
@@ -183,6 +183,11 @@ void SLevelEditor::BindCommands()
 		FExecuteAction::CreateStatic( &FLevelEditorActionCallbacks::OnFocusOutlinerToSelection, TWeakPtr< SLevelEditor >( SharedThis( this ) ) )
 		);
 
+	LevelEditorCommands->MapAction(
+		FEditorViewportCommands::Get().FocusOutlinerToContextFolder,
+		FExecuteAction::CreateStatic( &FLevelEditorActionCallbacks::OnFocusOutlinerToContextFolder, TWeakPtr< SLevelEditor >( SharedThis( this ) ) )
+		);
+
 	if (FPlayWorldCommands::GlobalPlayWorldActions.IsValid())
 	{
 		FUICommandList& PlayWorldActionList = *FPlayWorldCommands::GlobalPlayWorldActions;
diff --git a/Engine/Source/Editor/LevelEditor/Public/LevelEditorActions.h b/Engine/Source/Editor/LevelEditor/Public/LevelEditorActions.h
--- a/Engine/Source/Editor/LevelEditor/Public/LevelEditorActions.h
+++ b/Engine/Source/Editor/LevelEditor/Public/LevelEditorActions.h
@@ -1179,6 +1179,9 @@ public:
 	/** Focuses the outliner on the selected actors */
 	static void OnFocusOutlinerToSelection(TWeakPtr<SLevelEditor> LevelEditor);
 
+	/** Focuses the outliner on the context folder if it's currently set */
+	static void OnFocusOutlinerToContextFolder(TWeakPtr<SLevelEditor> LevelEditor);
+
 	/** Open the Place Actors Panel */
 	static void OpenPlaceActors();
 
diff --git a/Engine/Source/Editor/SceneOutliner/Private/SSceneOutliner.cpp b/Engine/Source/Editor/SceneOutliner/Private/SSceneOutliner.cpp
--- a/Engine/Source/Editor/SceneOutliner/Private/SSceneOutliner.cpp
+++ b/Engine/Source/Editor/SceneOutliner/Private/SSceneOutliner.cpp
@@ -2483,6 +2483,14 @@ void SSceneOutliner::FrameSelectedItems()
 	}
 }
 
+void SSceneOutliner::FrameItem(const FSceneOutlinerTreeItemID& Item)
+{
+	if (const TSharedPtr<ISceneOutlinerTreeItem>* TreeItem = TreeItemMap.Find(Item))
+	{
+		ScrollItemIntoView(*TreeItem);
+	}
+}
+
 FSceneOutlinerTreeItemPtr SSceneOutliner::FindParent(const ISceneOutlinerTreeItem& InItem) const
 {
 	FSceneOutlinerTreeItemPtr Parent = Mode->GetHierarchy()->FindOrCreateParentItem(InItem, TreeItemMap, /*bCreate=*/false);
diff --git a/Engine/Source/Editor/SceneOutliner/Public/ISceneOutliner.h b/Engine/Source/Editor/SceneOutliner/Public/ISceneOutliner.h
--- a/Engine/Source/Editor/SceneOutliner/Public/ISceneOutliner.h
+++ b/Engine/Source/Editor/SceneOutliner/Public/ISceneOutliner.h
@@ -160,6 +160,9 @@ public:
 	/** Scrolls the outliner to the selected item(s). If more are selected, the chosen item is undeterministic. */
 	virtual void FrameSelectedItems() = 0;
 
+	/** Scrolls the outliner to the specified item if it's found. No action is done otherwise. */
+	virtual void FrameItem(const FSceneOutlinerTreeItemID& Item) = 0;
+
 	/** Get the active SceneOutlinerMode */
 	const ISceneOutlinerMode* GetMode() const { return Mode; }
 
diff --git a/Engine/Source/Editor/SceneOutliner/Public/SSceneOutliner.h b/Engine/Source/Editor/SceneOutliner/Public/SSceneOutliner.h
--- a/Engine/Source/Editor/SceneOutliner/Public/SSceneOutliner.h
+++ b/Engine/Source/Editor/SceneOutliner/Public/SSceneOutliner.h
@@ -472,6 +472,12 @@ public:
 	 */
 	virtual void FrameSelectedItems() override;
 
+	/**
+	 * Scrolls the outliner to the specified item if it's found.
+	 * No action is done otherwise.
+	 */
+	virtual void FrameItem(const FSceneOutlinerTreeItemID& Item) override;
+
 	/**
 	 * Returns the parent tree item for a given item if it exists, nullptr otherwise.
 	 */
diff --git a/Engine/Source/Editor/UnrealEd/Private/EditorViewportCommands.cpp b/Engine/Source/Editor/UnrealEd/Private/EditorViewportCommands.cpp
--- a/Engine/Source/Editor/UnrealEd/Private/EditorViewportCommands.cpp
+++ b/Engine/Source/Editor/UnrealEd/Private/EditorViewportCommands.cpp
@@ -164,6 +164,7 @@ void FEditorViewportCommands::RegisterCommands()
 	UI_COMMAND( FocusAllViewportsToSelection, "Frame Selected Actors in All Viewports", "Centers all open viewports on the selected actors", EUserInterfaceActionType::Button, FInputChord(EKeys::F, EModifierKey::Shift) );
 	UI_COMMAND( FocusViewportToSelection, "Frame Selected", "Centers the viewport on the selected actors", EUserInterfaceActionType::Button, FInputChord( EKeys::F ) );
 	UI_COMMAND( FocusOutlinerToSelection, "Frame Selected in Outliner", "Focuses the selected actors in all open outliners", EUserInterfaceActionType::Button, FInputChord() );
+	UI_COMMAND( FocusOutlinerToContextFolder, "Frame Current Folder in Outliner", "Focuses the current folder (if it's set) in the outliner", EUserInterfaceActionType::Button, FInputChord() );
 
 	UI_COMMAND( LocationGridSnap, "Grid Snap", "Enables or disables snapping to the grid when dragging objects around", EUserInterfaceActionType::ToggleButton, FInputChord() );
 	UI_COMMAND( RotationGridSnap, "Rotation Snap", "Enables or disables snapping objects to a rotation grid", EUserInterfaceActionType::ToggleButton, FInputChord() );
diff --git a/Engine/Source/Editor/UnrealEd/Public/EditorViewportCommands.h b/Engine/Source/Editor/UnrealEd/Public/EditorViewportCommands.h
--- a/Engine/Source/Editor/UnrealEd/Public/EditorViewportCommands.h
+++ b/Engine/Source/Editor/UnrealEd/Public/EditorViewportCommands.h
@@ -215,6 +215,7 @@ public:
 	TSharedPtr< FUICommandInfo > FocusAllViewportsToSelection;
 	TSharedPtr< FUICommandInfo > FocusViewportToSelection;
 	TSharedPtr< FUICommandInfo > FocusOutlinerToSelection;
+	TSharedPtr< FUICommandInfo > FocusOutlinerToContextFolder;
 
 	/** Toggle automatic exposure */
 	TSharedPtr< FUICommandInfo > ToggleAutoExposure;
