Fixed ComponentTypeRegistry blocking unload of ActorComponent blueprint classes

diff --git a/Engine/Source/Editor/UnrealEd/Private/ComponentTypeRegistry.cpp b/Engine/Source/Editor/UnrealEd/Private/ComponentTypeRegistry.cpp
--- a/Engine/Source/Editor/UnrealEd/Private/ComponentTypeRegistry.cpp
+++ b/Engine/Source/Editor/UnrealEd/Private/ComponentTypeRegistry.cpp
@@ -774,11 +774,15 @@ FComponentTypeRegistry::FComponentTypeRegistry()
 	Data->Invalidate();
 
 	FCoreUObjectDelegates::ReloadCompleteDelegate.AddRaw(this, &FComponentTypeRegistry::OnReloadComplete);
+
+	// Clear references when unloaded otherwise unload will fail and leak the object
+	FKismetEditorUtilities::OnBlueprintGeneratedClassUnloaded.AddRaw(this, &FComponentTypeRegistry::OnBlueprintGeneratedClassUnloaded);
 }
 
 FComponentTypeRegistry::~FComponentTypeRegistry()
 {
 	FCoreUObjectDelegates::ReloadCompleteDelegate.RemoveAll(this);
+	FKismetEditorUtilities::OnBlueprintGeneratedClassUnloaded.RemoveAll(this);
 }
 
 void FComponentTypeRegistry::OnReloadComplete(EReloadCompleteReason Reason)
@@ -786,6 +790,42 @@ void FComponentTypeRegistry::OnReloadComplete(EReloadCompleteReason Reason)
 	Data->ForceRefreshComponentList();
 }
 
+void FComponentTypeRegistry::OnBlueprintGeneratedClassUnloaded(UBlueprintGeneratedClass* BlueprintGeneratedClass)
+{
+	if (!BlueprintGeneratedClass)
+	{
+		return;
+	}
+
+	if (!BlueprintGeneratedClass->IsChildOf(UActorComponent::StaticClass()))
+	{
+		return;
+	}
+
+	bool bModified = false;
+	for (FComponentClassComboEntryPtr& ComboEntry : Data->ComponentClassList)
+	{
+		if (ComboEntry->OnBlueprintGeneratedClassUnloaded(BlueprintGeneratedClass))
+		{
+			bModified = true;
+		}
+	}
+
+	for (FComponentTypeEntry& ComponentTypeEntry : Data->ComponentTypeList)
+	{
+		if (ComponentTypeEntry.ComponentClass == BlueprintGeneratedClass)
+		{
+			ComponentTypeEntry.ComponentClass = nullptr;
+			bModified = true;
+		}
+	}
+
+	if (bModified)
+	{
+		Data->bNeedsRefreshNextTick = true;
+	}
+}
+
 void FComponentTypeRegistry::InvalidateClass(TSubclassOf<UActorComponent> /*ClassToUpdate*/)
 {
 	Data->Invalidate();
diff --git a/Engine/Source/Editor/UnrealEd/Private/SComponentClassCombo.cpp b/Engine/Source/Editor/UnrealEd/Private/SComponentClassCombo.cpp
--- a/Engine/Source/Editor/UnrealEd/Private/SComponentClassCombo.cpp
+++ b/Engine/Source/Editor/UnrealEd/Private/SComponentClassCombo.cpp
@@ -53,6 +53,25 @@ void FComponentClassComboEntry::AddReferencedObjects(FReferenceCollector& Collec
 	Collector.AddReferencedObject(IconClass);
 }
 
+bool FComponentClassComboEntry::OnBlueprintGeneratedClassUnloaded(UBlueprintGeneratedClass* BlueprintGeneratedClass)
+{
+	bool bModified = false;
+
+	if (BlueprintGeneratedClass == ComponentClass)
+	{
+		ComponentClass = nullptr;
+		bModified = true;
+	}
+
+	if (BlueprintGeneratedClass == IconClass)
+	{
+		IconClass = nullptr;
+		bModified = true;
+	}
+
+	return bModified;
+}
+
 void SComponentClassCombo::Construct(const FArguments& InArgs)
 {
 	PrevSelectedIndex = INDEX_NONE;
diff --git a/Engine/Source/Editor/UnrealEd/Public/ComponentTypeRegistry.h b/Engine/Source/Editor/UnrealEd/Public/ComponentTypeRegistry.h
--- a/Engine/Source/Editor/UnrealEd/Public/ComponentTypeRegistry.h
+++ b/Engine/Source/Editor/UnrealEd/Public/ComponentTypeRegistry.h
@@ -7,6 +7,7 @@
 #include "Components/ActorComponent.h"
 
 class FComponentClassComboEntry;
+class UBlueprintGeneratedClass;
 
 DECLARE_MULTICAST_DELEGATE(FOnComponentTypeListChanged);
 
@@ -53,6 +54,7 @@ struct FComponentTypeRegistry
 
 private:
 	void OnReloadComplete(EReloadCompleteReason Reason);
+	void OnBlueprintGeneratedClassUnloaded(UBlueprintGeneratedClass* BlueprintGeneratedClass);
 
 private:
 	FComponentTypeRegistry();
diff --git a/Engine/Source/Editor/UnrealEd/Public/SComponentClassCombo.h b/Engine/Source/Editor/UnrealEd/Public/SComponentClassCombo.h
--- a/Engine/Source/Editor/UnrealEd/Public/SComponentClassCombo.h
+++ b/Engine/Source/Editor/UnrealEd/Public/SComponentClassCombo.h
@@ -160,6 +160,7 @@ public:
 	int32 GetSortPriority() const { return CustomizationArgs.SortPriority; }
 
 	void AddReferencedObjects(FReferenceCollector& Collector);
+	bool OnBlueprintGeneratedClassUnloaded(UBlueprintGeneratedClass* BlueprintGeneratedClass);
 
 	// If the component type is not loaded, this stores data that can be used for class filtering.
 	TSharedPtr<IUnloadedBlueprintData> GetUnloadedBlueprintData() const { return UnloadedBlueprintData; }
