Fixed crash setting a preview animation blueprint

Delegate callbacks fired when setting a debugged object could concurrently modify the LinkedInstances array while it was being iterated

diff --git a/Engine/Source/Editor/UnrealEd/Private/AnimationEditorUtils.cpp b/Engine/Source/Editor/UnrealEd/Private/AnimationEditorUtils.cpp
--- a/Engine/Source/Editor/UnrealEd/Private/AnimationEditorUtils.cpp
+++ b/Engine/Source/Editor/UnrealEd/Private/AnimationEditorUtils.cpp
@@ -1095,7 +1095,8 @@ namespace AnimationEditorUtils
 				if(const USkeletalMeshComponent* Component = Cast<USkeletalMeshComponent>(InRootObjectBeingDebugged->GetOuter()))
 				{
 					// See if we have any linked instances
-					for(UAnimInstance* LinkedInstance : Component->GetLinkedAnimInstances())
+					const TArray<UAnimInstance*> LinkedInstances = Component->GetLinkedAnimInstances();
+					for(UAnimInstance* LinkedInstance : LinkedInstances)
 					{
 						if(UAnimBlueprint* LinkedAnimBlueprint = Cast<UAnimBlueprint>(LinkedInstance->GetClass()->ClassGeneratedBy))
 						{
@@ -1109,7 +1110,8 @@ namespace AnimationEditorUtils
 				if(const USkeletalMeshComponent* Component = Cast<USkeletalMeshComponent>(OldDebuggedObject->GetOuter()))
 				{
 					// See if we have any linked instances
-					for(UAnimInstance* LinkedInstance : Component->GetLinkedAnimInstances())
+					const TArray<UAnimInstance*> LinkedInstances = Component->GetLinkedAnimInstances();
+					for(UAnimInstance* LinkedInstance : LinkedInstances)
 					{
 						if(UAnimBlueprint* LinkedAnimBlueprint = Cast<UAnimBlueprint>(LinkedInstance->GetClass()->ClassGeneratedBy))
 						{
