When checking for if a BP added variable is in use, we should also check if there are any Component Bound event nodes that reference it.

diff --git a/Engine/Source/Editor/BlueprintGraph/Classes/K2Node_ComponentBoundEvent.h b/Engine/Source/Editor/BlueprintGraph/Classes/K2Node_ComponentBoundEvent.h
--- a/Engine/Source/Editor/BlueprintGraph/Classes/K2Node_ComponentBoundEvent.h
+++ b/Engine/Source/Editor/BlueprintGraph/Classes/K2Node_ComponentBoundEvent.h
@@ -77,6 +77,10 @@ class UK2Node_ComponentBoundEvent : public UK2Node_Event
 
 	BLUEPRINTGRAPH_API void InitializeComponentBoundEventParams(FObjectProperty const* InComponentProperty, const FMulticastDelegateProperty* InDelegateProperty);
 
+public:
+
+	BLUEPRINTGRAPH_API FName GetComponentPropertyName() const { return ComponentPropertyName; }
+
 private:
 
 	/** Returns true if there is an FObjectProperty on this blueprint with a name that matches ComponentPropertyName */
diff --git a/Engine/Source/Editor/UnrealEd/Private/Kismet2/BlueprintEditorUtils.cpp b/Engine/Source/Editor/UnrealEd/Private/Kismet2/BlueprintEditorUtils.cpp
--- a/Engine/Source/Editor/UnrealEd/Private/Kismet2/BlueprintEditorUtils.cpp
+++ b/Engine/Source/Editor/UnrealEd/Private/Kismet2/BlueprintEditorUtils.cpp
@@ -5845,6 +5845,19 @@ bool FBlueprintEditorUtils::IsVariableUsed(const UBlueprint* Blueprint, const FN
 					return true;
 				}
 
+				// Check for all component bound event nodes. This variable may be referenced by a bound event
+				// (i.e. "On Component Hit" or any of the delegates you can add from the details panel with a "+" button)
+				TArray<UK2Node_ComponentBoundEvent*> ComponentBoundEventNodes;
+				CurrentGraph->GetNodesOfClass(ComponentBoundEventNodes);
+
+				if (Algo::AnyOf(ComponentBoundEventNodes, [&VariableName](const UK2Node_ComponentBoundEvent* EventNode)
+					{
+						return EventNode->GetComponentPropertyName() == VariableName;
+					}))
+				{
+					return true;
+				}
+
 				// Check all K2Node's which specify private/internal function referencing behavior
 				TArray<const UK2Node*> GraphNodes;
 				CurrentGraph->GetNodesOfClass(GraphNodes);
