MPCook: Support for ConfigAccessTracking
*) Handle nullptr when serializing ITargetPlatform; some systems need to track editor-agnostic data and serializing that is a nullptr ITargetPlatform* is convenient.
*) Remove includes when the code using them is ifdef'd out.
*) Remove ReportAccessedIniSettings hooks. The information will instead be part of the information reported from CookWorker -> Zen for package-specific accesses, and will be replicated using an IMPCollector for the total list of all accesses throughout the cook.
*) Add BroadcastCookByTheBookFinished subfunction to record the declaration of IMPCollectors constructed privately by UCookOnTheFlyServer in the same place as external systems declare their IMPCollectors.
#rnx

diff --git a/Engine/Source/Editor/UnrealEd/Classes/CookOnTheSide/CookOnTheFlyServer.h b/Engine/Source/Editor/UnrealEd/Classes/CookOnTheSide/CookOnTheFlyServer.h
--- a/Engine/Source/Editor/UnrealEd/Classes/CookOnTheSide/CookOnTheFlyServer.h
+++ b/Engine/Source/Editor/UnrealEd/Classes/CookOnTheSide/CookOnTheFlyServer.h
@@ -1426,6 +1426,7 @@ private:
 	/** Send warnings/telemetry when a discovered or read package is found to be a hidden dependency. */
 	void ReportHiddenDependency(FName Referencer, FName Dependency);
 	void BroadcastCookByTheBookStarted();
+	void BroadcastCookByTheBookFinished();
 
 	static UCookOnTheFlyServer* ActiveCOTFS;
 	uint32		StatLoadedPackageCount = 0;
diff --git a/Engine/Source/Editor/UnrealEd/Private/CookOnTheFlyServer.cpp b/Engine/Source/Editor/UnrealEd/Private/CookOnTheFlyServer.cpp
--- a/Engine/Source/Editor/UnrealEd/Private/CookOnTheFlyServer.cpp
+++ b/Engine/Source/Editor/UnrealEd/Private/CookOnTheFlyServer.cpp
@@ -7521,8 +7521,6 @@ void GetAdditionalCurrentIniVersionStrings( const UCookOnTheFlyServer* CookOnThe
 	IniVersionMap.Add(TEXT("IsUsingShaderCodeLibrary"), FString::Printf(TEXT("%d"), PackagingSettings->bShareMaterialShaderCode && CookOnTheFlyServer->IsUsingShaderCodeLibrary()));
 }
 
-
-
 bool UCookOnTheFlyServer::GetCurrentIniVersionStrings( const ITargetPlatform* TargetPlatform, UE::Cook::FIniSettingContainer& IniVersionStrings ) const
 {
 	{
@@ -7556,7 +7554,6 @@ bool UCookOnTheFlyServer::GetCurrentIniVersionStrings( const ITargetPlatform* Ta
 		}
 	}
 
-
 	// remove any which are filtered out
 	FString EditorPrefix(TEXT("Editor."));
 	for ( const FString& Filter : ConfigSettingDenyList )
@@ -7720,7 +7717,7 @@ void UCookOnTheFlyServer::OnFConfigDeleted(const FConfigFile* Config)
 	}
 
 	FScopeLock Lock(&ConfigFileCS);
-	WorkerRequests->ReportAccessedIniSettings(*this, *Config);
+	ProcessAccessedIniSettings(Config, AccessedIniStrings);
 	OpenConfigFiles.Remove(Config);
 }
 
@@ -8181,15 +8178,15 @@ bool UCookOnTheFlyServer::SaveCurrentIniSettings(const ITargetPlatform* TargetPl
 		for (const auto& CurrentIniFilename : CurrentIniSettings)
 		{
 			const FName& Filename = CurrentIniFilename.Key;
-			for ( const auto& CurrentSection : CurrentIniFilename.Value )
+			for (const auto& CurrentSection : CurrentIniFilename.Value)
 			{
 				const FName& Section = CurrentSection.Key;
-				for ( const auto& CurrentValue : CurrentSection.Value )
+				for (const auto& CurrentValue : CurrentSection.Value)
 				{
 					const FName& ValueName = CurrentValue.Key;
 					const TArray<FString>& Values = CurrentValue.Value;
 
-					for ( int Index = 0; Index < Values.Num(); ++Index )
+					for (int Index = 0; Index < Values.Num(); ++Index)
 					{
 						FString NewKey = FString::Printf(TEXT("%s:%s:%s:%d"), *Filename.ToString(), *Section.ToString(), *ValueName.ToString(), Index);
 						ConfigFile.AddToSection(NAME_UsedSettings, *NewKey, Values[Index]);
@@ -9979,10 +9976,7 @@ void UCookOnTheFlyServer::CookByTheBookFinished()
 	}
 
 	ShutdownCookSession();
-	PRAGMA_DISABLE_DEPRECATION_WARNINGS;
-	CookByTheBookFinishedEvent.Broadcast();
-	PRAGMA_ENABLE_DEPRECATION_WARNINGS;
-	UE::Cook::FDelegates::CookByTheBookFinished.Broadcast(*this);
+	BroadcastCookByTheBookFinished();
 	UE_LOG(LogCook, Display, TEXT("Done!"));
 }
 
@@ -11405,10 +11399,7 @@ void UCookOnTheFlyServer::CookAsCookWorkerFinished()
 	LogCookWorkerStats();
 	if (IsDirectorCookByTheBook())
 	{
-		PRAGMA_DISABLE_DEPRECATION_WARNINGS;
-		CookByTheBookFinishedEvent.Broadcast();
-		PRAGMA_ENABLE_DEPRECATION_WARNINGS;
-		UE::Cook::FDelegates::CookByTheBookFinished.Broadcast(*this);
+		BroadcastCookByTheBookFinished();
 	}
 	CookWorkerClient->FlushLogs();
 }
@@ -13079,6 +13070,25 @@ void UCookOnTheFlyServer::BroadcastCookByTheBookStarted()
 #if ENABLE_LOW_LEVEL_MEM_TRACKER
 	FLowLevelMemTracker::Get().UpdateStatsPerFrame();
 #endif
+
+	// Register collectors used internally by CookOnTheFlyServer.
+	// External systems would do this during CookByTheBookStarted.Broadcast
+	if (GetProcessType() != UE::Cook::EProcessType::SingleProcess)
+	{
+	}
+}
+
+void UCookOnTheFlyServer::BroadcastCookByTheBookFinished()
+{
+	// Unregister collectors used internally by CookOnTheFlyServer.
+	if (GetProcessType() != UE::Cook::EProcessType::SingleProcess)
+	{
+	}
+
+	PRAGMA_DISABLE_DEPRECATION_WARNINGS;
+	CookByTheBookFinishedEvent.Broadcast();
+	PRAGMA_ENABLE_DEPRECATION_WARNINGS;
+	UE::Cook::FDelegates::CookByTheBookFinished.Broadcast(*this);
 }
 
 #undef LOCTEXT_NAMESPACE
diff --git a/Engine/Source/Editor/UnrealEd/Private/Cooker/IWorkerRequests.h b/Engine/Source/Editor/UnrealEd/Private/Cooker/IWorkerRequests.h
--- a/Engine/Source/Editor/UnrealEd/Private/Cooker/IWorkerRequests.h
+++ b/Engine/Source/Editor/UnrealEd/Private/Cooker/IWorkerRequests.h
@@ -72,7 +72,6 @@ public:
 	virtual void OnRemoveSessionPlatform(const ITargetPlatform* TargetPlatform) = 0;
 
 	// Functions that are written locally for Directors but are sent to Director for Workers
-	virtual void ReportAccessedIniSettings(UCookOnTheFlyServer& COTFS, const FConfigFile& Config) = 0;
 	virtual void ReportDemoteToIdle(UE::Cook::FPackageData& PackageData, ESuppressCookReason Reason) = 0;
 	virtual void ReportPromoteToSaveComplete(UE::Cook::FPackageData& PackageData) = 0;
 
diff --git a/Engine/Source/Editor/UnrealEd/Private/Cooker/PackageBuildDependencyTracker.cpp b/Engine/Source/Editor/UnrealEd/Private/Cooker/PackageBuildDependencyTracker.cpp
--- a/Engine/Source/Editor/UnrealEd/Private/Cooker/PackageBuildDependencyTracker.cpp
+++ b/Engine/Source/Editor/UnrealEd/Private/Cooker/PackageBuildDependencyTracker.cpp
@@ -2,12 +2,14 @@
 
 #include "PackageBuildDependencyTracker.h"
 
+#if UE_WITH_PACKAGE_ACCESS_TRACKING
 #include "Async/UniqueLock.h"
 #include "HAL/Platform.h"
 #include "Logging/LogMacros.h"
 #include "Misc/PackageAccessTrackingOps.h"
 #include "UObject/Class.h"
 #include "UObject/Package.h"
+#endif
 
 #if UE_WITH_PACKAGE_ACCESS_TRACKING
 
diff --git a/Engine/Source/Editor/UnrealEd/Private/Cooker/PackageBuildDependencyTracker.h b/Engine/Source/Editor/UnrealEd/Private/Cooker/PackageBuildDependencyTracker.h
--- a/Engine/Source/Editor/UnrealEd/Private/Cooker/PackageBuildDependencyTracker.h
+++ b/Engine/Source/Editor/UnrealEd/Private/Cooker/PackageBuildDependencyTracker.h
@@ -2,10 +2,13 @@
 
 #pragma once
 
-#include "Async/Mutex.h"
 #include "Misc/PackageAccessTracking.h"
+
+#if UE_WITH_PACKAGE_ACCESS_TRACKING
+#include "Async/Mutex.h"
 #include "Templates/UnrealTemplate.h"
 #include "UObject/ObjectHandle.h"
+#endif
 
 #if UE_WITH_PACKAGE_ACCESS_TRACKING
 
diff --git a/Engine/Source/Editor/UnrealEd/Private/Cooker/WorkerRequestsLocal.cpp b/Engine/Source/Editor/UnrealEd/Private/Cooker/WorkerRequestsLocal.cpp
--- a/Engine/Source/Editor/UnrealEd/Private/Cooker/WorkerRequestsLocal.cpp
+++ b/Engine/Source/Editor/UnrealEd/Private/Cooker/WorkerRequestsLocal.cpp
@@ -104,11 +104,6 @@ void FWorkerRequestsLocal::OnRemoveSessionPlatform(const ITargetPlatform* Target
 	ExternalRequests.OnRemoveSessionPlatform(TargetPlatform);
 }
 
-void FWorkerRequestsLocal::ReportAccessedIniSettings(UCookOnTheFlyServer& COTFS, const FConfigFile& Config)
-{
-	COTFS.ProcessAccessedIniSettings(&Config, COTFS.AccessedIniStrings);
-}
-
 void FWorkerRequestsLocal::ReportDemoteToIdle(UE::Cook::FPackageData& PackageData, ESuppressCookReason Reason)
 {
 }
diff --git a/Engine/Source/Editor/UnrealEd/Private/Cooker/WorkerRequestsLocal.h b/Engine/Source/Editor/UnrealEd/Private/Cooker/WorkerRequestsLocal.h
--- a/Engine/Source/Editor/UnrealEd/Private/Cooker/WorkerRequestsLocal.h
+++ b/Engine/Source/Editor/UnrealEd/Private/Cooker/WorkerRequestsLocal.h
@@ -42,7 +42,6 @@ public:
 	virtual void AddEditorActionCallback(FSchedulerCallback&& Callback) override;
 	virtual void AddPublicInterfaceRequest(FFilePlatformRequest&& Request, bool bForceFrontOfQueue) override;
 
-	virtual void ReportAccessedIniSettings(UCookOnTheFlyServer& COTFS, const FConfigFile& Config) override;
 	virtual void ReportDemoteToIdle(UE::Cook::FPackageData& PackageData, ESuppressCookReason Reason) override;
 	virtual void ReportPromoteToSaveComplete(UE::Cook::FPackageData& PackageData) override;
 	virtual void GetInitializeConfigSettings(UCookOnTheFlyServer& COTFS, const FString& OutputOverrideDirectory,
diff --git a/Engine/Source/Editor/UnrealEd/Private/Cooker/WorkerRequestsRemote.cpp b/Engine/Source/Editor/UnrealEd/Private/Cooker/WorkerRequestsRemote.cpp
--- a/Engine/Source/Editor/UnrealEd/Private/Cooker/WorkerRequestsRemote.cpp
+++ b/Engine/Source/Editor/UnrealEd/Private/Cooker/WorkerRequestsRemote.cpp
@@ -105,13 +105,6 @@ void FWorkerRequestsRemote::OnRemoveSessionPlatform(const ITargetPlatform* Targe
 	ExternalRequests.OnRemoveSessionPlatform(TargetPlatform);
 }
 
-void FWorkerRequestsRemote::ReportAccessedIniSettings(UCookOnTheFlyServer& COTFS, const FConfigFile& Config)
-{
-	FIniSettingContainer AccessedIniStrings;
-	COTFS.ProcessAccessedIniSettings(&Config, COTFS.AccessedIniStrings);
-	// MPCOOKTODO: Not yet implemented
-}
-
 void FWorkerRequestsRemote::ReportDemoteToIdle(UE::Cook::FPackageData& PackageData, ESuppressCookReason Reason)
 {
 	CookWorkerClient.ReportDemoteToIdle(PackageData, Reason);
diff --git a/Engine/Source/Editor/UnrealEd/Private/Cooker/WorkerRequestsRemote.h b/Engine/Source/Editor/UnrealEd/Private/Cooker/WorkerRequestsRemote.h
--- a/Engine/Source/Editor/UnrealEd/Private/Cooker/WorkerRequestsRemote.h
+++ b/Engine/Source/Editor/UnrealEd/Private/Cooker/WorkerRequestsRemote.h
@@ -47,7 +47,6 @@ public:
 	virtual void AddEditorActionCallback(FSchedulerCallback&& Callback) override;
 	virtual void AddPublicInterfaceRequest(FFilePlatformRequest&& Request, bool bForceFrontOfQueue) override;
 
-	virtual void ReportAccessedIniSettings(UCookOnTheFlyServer& COTFS, const FConfigFile& Config) override;
 	virtual void ReportDemoteToIdle(UE::Cook::FPackageData& PackageData, ESuppressCookReason Reason) override;
 	virtual void ReportPromoteToSaveComplete(UE::Cook::FPackageData& PackageData) override;
 	virtual void GetInitializeConfigSettings(UCookOnTheFlyServer& COTFS, const FString& OutputOverrideDirectory,
diff --git a/Engine/Source/Runtime/CoreUObject/Private/Cooker/MPCollector.cpp b/Engine/Source/Runtime/CoreUObject/Private/Cooker/MPCollector.cpp
--- a/Engine/Source/Runtime/CoreUObject/Private/Cooker/MPCollector.cpp
+++ b/Engine/Source/Runtime/CoreUObject/Private/Cooker/MPCollector.cpp
@@ -9,11 +9,16 @@ namespace UE::Cook
 
 static uint8 PlatformToIntImpl(const ITargetPlatform* Platform, TConstArrayView<const ITargetPlatform*> Platforms)
 {
-	int32 Result = Platforms.IndexOfByKey(Platform);
-	check(Result != INDEX_NONE);
-	uint8 NarrowResult = static_cast<uint8>(Result);
-	check(static_cast<int32>(NarrowResult) == Result);
-	return NarrowResult;
+	if (Platform)
+	{
+		int32 Result = Platforms.IndexOfByKey(Platform);
+		check(Result != INDEX_NONE && Result < static_cast<int32>(MAX_uint8));
+		return static_cast<uint8>(Result);
+	}
+	else
+	{
+		return MAX_uint8;
+	}
 }
 
 static const ITargetPlatform* IntToPlatformImpl(uint8 PlatformAsInt, TConstArrayView<const ITargetPlatform*> Platforms)
