Fix performance issue when exporting hi res mesh.

#rnx

diff --git a/Engine/Source/Editor/UnrealEd/Private/Fbx/FbxMainExport.cpp b/Engine/Source/Editor/UnrealEd/Private/Fbx/FbxMainExport.cpp
--- a/Engine/Source/Editor/UnrealEd/Private/Fbx/FbxMainExport.cpp
+++ b/Engine/Source/Editor/UnrealEd/Private/Fbx/FbxMainExport.cpp
@@ -4583,6 +4583,9 @@ bool FFbxExporter::ExportStaticMeshFromMeshDescription(FbxMesh* Mesh
 		}
 		//Build the edge so we can set the edge hardness
 		Mesh->BuildMeshEdgeArray();
+
+		Mesh->BeginGetMeshEdgeIndexForPolygon();
+		int32 PolygonIndex = 0;
 		for (const FPolygonGroupID& PolygonGroupID : MeshDescription->PolygonGroups().GetElementIDs())
 		{
 			for (const FTriangleID& TriangleID : MeshDescription->GetPolygonGroupTriangles(PolygonGroupID))
@@ -4608,8 +4611,7 @@ bool FFbxExporter::ExportStaticMeshFromMeshDescription(FbxMesh* Mesh
 					}
 					ProcessEdges.Add(MatchEdgeId);
 
-					bool ReverseEdge = false;
-					int32 FbxEdgeIndex = Mesh->GetMeshEdgeIndex(EdgeStart.GetValue(), EdgeEnd.GetValue(), ReverseEdge);
+					int32 FbxEdgeIndex = Mesh->GetMeshEdgeIndexForPolygon(PolygonIndex, CornerIndex);
 					if (FbxEdgeIndex == -1)
 					{
 						continue;
@@ -4622,8 +4624,10 @@ bool FFbxExporter::ExportStaticMeshFromMeshDescription(FbxMesh* Mesh
 					int32 LayerAddIndex = SmoothingInfoLayer->GetDirectArray().Add(EdgeHardnessValue);
 					ensure(LayerAddIndex == FbxEdgeIndex);
 				}
+				PolygonIndex++;
 			}
 		}
+		Mesh->EndGetMeshEdgeIndexForPolygon();
 	}
 
 	// Create and fill in the vertex color data source.
