UE-203644 VR PIE: bOneHeadsetEachProcess option allowing each separate process PIE instance to own an HMD or HMD simulator.
-New setting bOneHeadsetEachProcess which when true will allow PIE separate-process instances to run without -nohmd and with -HMDSimulator which an xr plugin may consume to do special vr device simulation.  Even if that does not happen if we are running with OpenXR we will attempt to create an additional OpenXR instance to connect to an hmd and an OpenXR runtime may allow us some kind of connection, these pie instances might also not run correctly.

diff --git a/Engine/Source/Editor/UnrealEd/Classes/Settings/LevelEditorPlaySettings.h b/Engine/Source/Editor/UnrealEd/Classes/Settings/LevelEditorPlaySettings.h
--- a/Engine/Source/Editor/UnrealEd/Classes/Settings/LevelEditorPlaySettings.h
+++ b/Engine/Source/Editor/UnrealEd/Classes/Settings/LevelEditorPlaySettings.h
@@ -420,6 +420,14 @@ private:
 	UPROPERTY(config, EditAnywhere, Category="Multiplayer Options")
 	FLinearColor ServerDebugDrawingColorTint;
 
+	/** 
+	* When True each PIE process is launched with "-HMDSimulator" argument.  
+	* The usefullness of this will vary by XR platform.  
+	* The PIE instances may get special -HMDSimulator behavior from an XR plugin, they may successfully make connections to the HMD hardware, their attempt to connect to hardware may be rejected by the runtime.
+	*/
+	UPROPERTY(config, EditAnywhere, Category = "Multiplayer Options", meta = (EditCondition = "!RunUnderOneProcess"))
+	bool bOneHeadsetEachProcess;
+
 private:
 	UNREALED_API void PushDebugDrawingSettings();
 
@@ -513,6 +521,8 @@ public:
 		}
 	}
 
+	bool IsOneHeadsetEachProcess() const { return bOneHeadsetEachProcess; }
+
 public:
 
 	/** The last known screen size for the first instance window (in pixels). */
diff --git a/Engine/Source/Editor/UnrealEd/Private/PlayLevelNewProcess.cpp b/Engine/Source/Editor/UnrealEd/Private/PlayLevelNewProcess.cpp
--- a/Engine/Source/Editor/UnrealEd/Private/PlayLevelNewProcess.cpp
+++ b/Engine/Source/Editor/UnrealEd/Private/PlayLevelNewProcess.cpp
@@ -178,12 +178,21 @@ void UEditorEngine::LaunchNewProcess(const FRequestPlaySessionParams& InParams,
 		CommandLine += TEXT(" -vulkan -faketouches -featureleveles31");
 	}
 
-	// If they're trying to launch a new process (from the editor) in VR, this will fail because the editor
-	// owns the HMD resource, so we warn, and then fall back. They will need to use single-process for VR preview.
+	// VRPreview handling
 	if (InParams.SessionPreviewTypeOverride.Get(EPlaySessionPreviewType::NoPreview) == EPlaySessionPreviewType::VRPreview)
 	{
-		CommandLine += TEXT(" -nohmd");
-		GLog->CategorizedLogf(FName("LogHMD"), ELogVerbosity::Warning, TEXT("Standalone Game VR not supported, please use VR Preview."));
+		if (!InParams.EditorPlaySettings->IsOneHeadsetEachProcess())
+		{
+			// If they're trying to launch a new process (from the editor) in VR, this will fail because the editor
+			// owns the HMD resource, so we warn, and then fall back. They will need to use single-process for VR preview.
+			CommandLine += TEXT(" -nohmd");
+			UE_LOG(LogPlayLevel, Warning, TEXT("Standalone Game VR not supported, please use VR Preview. Launching separate process PIE with -nohmd."));
+		}
+		else if (InInstanceNum != 0) // PIE instance 0 is normally run in the editor process, so we may not see it here. That instance get the real HMD, so no simulator argument is passed.
+		{
+			CommandLine += TEXT(" -HMDSimulator");
+			UE_LOG(LogPlayLevel, Log, TEXT("Launching separate process PIE with -HMDSimulator. See bOneHeadsetEachProcess editor preference tooltip for more information about this."));
+		}
 	}
 
 	// if we had -emulatestereo on the commandline, also pass it to the new process
diff --git a/Engine/Source/Editor/UnrealEd/Private/Settings/SettingsClasses.cpp b/Engine/Source/Editor/UnrealEd/Private/Settings/SettingsClasses.cpp
--- a/Engine/Source/Editor/UnrealEd/Private/Settings/SettingsClasses.cpp
+++ b/Engine/Source/Editor/UnrealEd/Private/Settings/SettingsClasses.cpp
@@ -514,6 +514,8 @@ ULevelEditorPlaySettings::ULevelEditorPlaySettings( const FObjectInitializer& Ob
 	bShowServerDebugDrawingByDefault = true;
 	ServerDebugDrawingColorTintStrength = 0.0f;
 	ServerDebugDrawingColorTint = FLinearColor(0.0f, 0.0f, 0.0f, 1.0f);
+
+	bOneHeadsetEachProcess = false;
 }
 
 void ULevelEditorPlaySettings::PushDebugDrawingSettings()
