Added the ability to optionally run content validation during cook

The following flags can be passed when running the cook commandlet:
  * -RunAssetValidation. Run asset validation (EditorValidatorSubsystem) on assets loaded during cook.
  * -RunMapValidation. Run map validation (MapCheck) on maps loaded during cook.
  * -ValidationErrorsAreFatal. Consider validation errors (from RunAssetValidation or RunMapValidation) as fatal (preventing the package from being cooked).

Any validation errors are downgraded to warnings when running without -ValidationErrorsAreFatal to ensure that things looking for log errors or a non-zero return code don't consider the errors as failing the cook, and the "Cook.AllowContentValidation" CVar controls whether validation can run even if the flags are passed.

These flags (-RunAssetValidation in particular) replace the older (and broken) bValidateAssetsWhileSavingForCook option in EditorValidatorSubsystem.

diff --git a/Engine/Plugins/Editor/DataValidation/Source/DataValidation/Private/DataValidationModule.cpp b/Engine/Plugins/Editor/DataValidation/Source/DataValidation/Private/DataValidationModule.cpp
--- a/Engine/Plugins/Editor/DataValidation/Source/DataValidation/Private/DataValidationModule.cpp
+++ b/Engine/Plugins/Editor/DataValidation/Source/DataValidation/Private/DataValidationModule.cpp
@@ -3,8 +3,10 @@
 #include "DataValidationModule.h"
 
 #include "AssetRegistry/ARFilter.h"
+#include "AssetRegistry/AssetDataToken.h"
 #include "UObject/ObjectSaveContext.h"
 
+#include "CookOnTheSide/CookOnTheFlyServer.h"
 #include "Editor.h"
 #include "Framework/Application/SlateApplication.h"
 #include "ToolMenus.h"
@@ -12,10 +14,12 @@
 #include "Misc/MessageDialog.h"
 #include "DataValidationCommandlet.h"
 #include "Elements/Framework/TypedElementSelectionSet.h"
+#include "Logging/MessageLog.h"
 
 #include "ContentBrowserMenuContexts.h"
 #include "LevelEditorMenuContext.h"
 
+#include "EditorValidatorHelpers.h"
 #include "EditorValidatorSubsystem.h"
 #include "ISettingsModule.h"
 #include "Algo/RemoveIf.h"
@@ -41,6 +45,8 @@ class FDataValidationModule : public IDataValidationModule
 private:
 	void OnPackageSaved(const FString& PackageFileName, UPackage* Package, FObjectPostSaveContext ObjectSaveContext);
 
+	EDataValidationResult OnValidateSourcePackageDuringCook(UPackage* Package, FDataValidationContext& ValidationContext);
+
 	// Adds Asset and any assets it depends on to the set DependentAssets
 	void FindAssetDependencies(const FAssetRegistryModule& AssetRegistryModule, const FAssetData& Asset, TSet<FAssetData>& DependentAssets);
 
@@ -53,6 +59,8 @@ IMPLEMENT_MODULE(FDataValidationModule, DataValidation)
 
 void FDataValidationModule::StartupModule()
 {	
+	UCookOnTheFlyServer::OnValidateSourcePackage().BindRaw(this, &FDataValidationModule::OnValidateSourcePackageDuringCook);
+
 	if (!IsRunningCommandlet() && !IsRunningGame() && FSlateApplication::IsInitialized())
 	{
 		// add the File->DataValidation menu subsection
@@ -72,6 +80,8 @@ void FDataValidationModule::StartupModule()
 
 void FDataValidationModule::ShutdownModule()
 {
+	UCookOnTheFlyServer::OnValidateSourcePackage().Unbind();
+
 	if (!IsRunningCommandlet() && !IsRunningGame() && !IsRunningDedicatedServer())
 	{
 		// remove menu extension
@@ -354,4 +364,59 @@ void FDataValidationModule::OnPackageSaved(const FString& PackageFileName, UPack
 	}
 }
 
+EDataValidationResult FDataValidationModule::OnValidateSourcePackageDuringCook(UPackage* Package, FDataValidationContext& ValidationContext)
+{
+	EDataValidationResult FinalValidationResult = EDataValidationResult::NotValidated;
+
+	if (UEditorValidatorSubsystem* EditorValidationSubsystem = GEditor->GetEditorSubsystem<UEditorValidatorSubsystem>())
+	{
+		TArray<FAssetData> AssetList;
+		IAssetRegistry::GetChecked().GetAssetsByPackageName(Package->GetFName(), AssetList);
+
+		{
+			FValidateAssetsSettings Settings;
+			Settings.ValidationUsecase = ValidationContext.GetValidationUsecase();
+			AssetList.RemoveAll([EditorValidationSubsystem, &Settings, &ValidationContext](const FAssetData& AssetData)
+			{
+				return !EditorValidationSubsystem->ShouldValidateAsset(AssetData, Settings, ValidationContext);
+			});
+		}
+
+		if (AssetList.Num() > 0)
+		{
+			// Broadcast the Editor event before we start validating. This lets other systems (such as Sequencer) restore the state
+			// of the level to what is actually saved on disk before performing validation.
+			if (FEditorDelegates::OnPreAssetValidation.IsBound())
+			{
+				FEditorDelegates::OnPreAssetValidation.Broadcast();
+			}
+
+			FMessageLog DataValidationLog(UE::DataValidation::MessageLogName);
+			for (const FAssetData& AssetData : AssetList)
+			{
+				if (UObject* Asset = AssetData.FastGetAsset(/*bLoad*/false))
+				{
+					DataValidationLog.Info()
+						->AddToken(FAssetDataToken::Create(AssetData))
+						->AddToken(FTextToken::Create(LOCTEXT("Data.ValidatingAsset", "Validating asset")));
+
+					const EDataValidationResult ValidationResult = EditorValidationSubsystem->IsObjectValidWithContext(Asset, ValidationContext);
+					FinalValidationResult = CombineDataValidationResults(FinalValidationResult, ValidationResult);
+
+					UE::DataValidation::AddAssetValidationMessages(AssetData, DataValidationLog, ValidationContext);
+					DataValidationLog.Flush();
+				}
+			}
+
+			// Broadcast now that we're complete so other systems can go back to their previous state.
+			if (FEditorDelegates::OnPostAssetValidation.IsBound())
+			{
+				FEditorDelegates::OnPostAssetValidation.Broadcast();
+			}
+		}
+	}
+
+	return FinalValidationResult;
+}
+
 #undef LOCTEXT_NAMESPACE
diff --git a/Engine/Plugins/Editor/DataValidation/Source/DataValidation/Private/EditorValidatorSubsystem.cpp b/Engine/Plugins/Editor/DataValidation/Source/DataValidation/Private/EditorValidatorSubsystem.cpp
--- a/Engine/Plugins/Editor/DataValidation/Source/DataValidation/Private/EditorValidatorSubsystem.cpp
+++ b/Engine/Plugins/Editor/DataValidation/Source/DataValidation/Private/EditorValidatorSubsystem.cpp
@@ -118,6 +118,11 @@ bool UEditorValidatorSubsystem::ShouldValidateAsset(
 	const FValidateAssetsSettings& 	Settings,
 	FDataValidationContext& 		InContext) const
 {
+	if (Asset.HasAnyPackageFlags(PKG_Cooked))
+	{
+		return false;
+	}
+
 	FNameBuilder AssetPackageNameBuilder(Asset.PackageName);
 	FStringView AssetPackageNameView = AssetPackageNameBuilder.ToView();
 
@@ -126,6 +131,11 @@ bool UEditorValidatorSubsystem::ShouldValidateAsset(
 		return false;
 	}
 
+	if (FPackageName::IsVersePackage(AssetPackageNameView))
+	{
+		return false;
+	}
+
 	if (IsPathExcludedFromValidation(AssetPackageNameView))
 	{
 		return false;
@@ -254,8 +264,7 @@ EDataValidationResult UEditorValidatorSubsystem::IsObjectValidWithContext(
 	
 	if (ensure(InObject))
 	{
-		FAssetRegistryModule& AssetRegistryModule = FModuleManager::LoadModuleChecked<FAssetRegistryModule>("AssetRegistry");
-		IAssetRegistry& AssetRegistry = AssetRegistryModule.Get();
+		IAssetRegistry& AssetRegistry = IAssetRegistry::GetChecked();
 
 		FAssetData AssetData = AssetRegistry.GetAssetByObjectPath(FSoftObjectPath(InObject), true);
 		if (!AssetData.IsValid())
@@ -344,8 +353,7 @@ EDataValidationResult UEditorValidatorSubsystem::ValidateAssetsInternal(
 	const FValidateAssetsSettings& 	InSettings,
 	FValidateAssetsResults& 		OutResults) const
 {
-	FAssetRegistryModule& AssetRegistryModule = FModuleManager::LoadModuleChecked<FAssetRegistryModule>("AssetRegistry");
-	IAssetRegistry& AssetRegistry = AssetRegistryModule.Get();
+	IAssetRegistry& AssetRegistry = IAssetRegistry::GetChecked();
 
 	FScopedSlowTask SlowTask(AssetDataList.Num(), LOCTEXT("DataValidation.ValidateAssetsTask", "Validating Assets"));
 	SlowTask.MakeDialog();
@@ -706,8 +714,7 @@ void UEditorValidatorSubsystem::ValidateAllSavedPackages()
 {
 	TRACE_CPUPROFILER_EVENT_SCOPE(UEditorValidatorSubsystem::ValidateAllSavedPackages);
 
-	FAssetRegistryModule& AssetRegistryModule = FModuleManager::LoadModuleChecked<FAssetRegistryModule>("AssetRegistry");
-	IAssetRegistry& AssetRegistry = AssetRegistryModule.Get();
+	IAssetRegistry& AssetRegistry = IAssetRegistry::GetChecked();
 
 	// Prior to validation, make sure Asset Registry is updated. This is done by ticking the DirectoryWatcher module, which 
 	// is responsible of scanning modified asset files.
@@ -796,8 +803,7 @@ EDataValidationResult UEditorValidatorSubsystem::ValidateChangelistsInternal(
 	SlowTask.Visibility = ESlowTaskVisibility::Invisible;
 	SlowTask.MakeDialog();
 
-	FAssetRegistryModule& AssetRegistryModule = FModuleManager::LoadModuleChecked<FAssetRegistryModule>("AssetRegistry");
-	IAssetRegistry& AssetRegistry = AssetRegistryModule.Get();
+	IAssetRegistry& AssetRegistry = IAssetRegistry::GetChecked();
 
 	if (AssetRegistry.IsLoadingAssets())
 	{
@@ -880,8 +886,7 @@ void UEditorValidatorSubsystem::GatherAssetsToValidateFromChangelist(
 	TSet<FAssetData>& 				OutAssets,
 	FDataValidationContext& 		InContext) const
 {
-	FAssetRegistryModule& AssetRegistryModule = FModuleManager::LoadModuleChecked<FAssetRegistryModule>("AssetRegistry");
-	IAssetRegistry& AssetRegistry = AssetRegistryModule.Get();
+	IAssetRegistry& AssetRegistry = IAssetRegistry::GetChecked();
 	
 	LoadValidators();
 	
@@ -973,8 +978,7 @@ void UEditorValidatorSubsystem::LoadValidators()
 
 TArray<FAssetData> UEditorValidatorSubsystem::GetAssetsResolvingRedirectors(FARFilter& InFilter)
 {
-	FAssetRegistryModule& AssetRegistryModule = FModuleManager::LoadModuleChecked<FAssetRegistryModule>("AssetRegistry");
-	IAssetRegistry& AssetRegistry = AssetRegistryModule.Get();
+	IAssetRegistry& AssetRegistry = IAssetRegistry::GetChecked();
 
 	TArray<FAssetData> Found;
 	AssetRegistry.GetAssets(InFilter, Found);
diff --git a/Engine/Source/Editor/UnrealEd/Classes/CookOnTheSide/CookOnTheFlyServer.h b/Engine/Source/Editor/UnrealEd/Classes/CookOnTheSide/CookOnTheFlyServer.h
--- a/Engine/Source/Editor/UnrealEd/Classes/CookOnTheSide/CookOnTheFlyServer.h
+++ b/Engine/Source/Editor/UnrealEd/Classes/CookOnTheSide/CookOnTheFlyServer.h
@@ -93,6 +93,9 @@ enum class ECookByTheBookOptions
 	DlcLoadMainAssetRegistry =			0x00020000, // If cooking DLC, populate the main game asset registry
 	ZenStore =							0x00040000, // Store cooked data in Zen Store
 	DlcReevaluateUncookedAssets =		0x00080000, // If cooking DLC, ignore assets in the base asset registry that were not cooked, so that this cook has an opportunity to cook the assets
+	RunAssetValidation =				0x00100000, // Run asset validation (EditorValidatorSubsystem) on assets loaded during cook
+	RunMapValidation =					0x00200000, // Run map validation (MapCheck) on maps loaded during cook
+	ValidationErrorsAreFatal =			0x00400000, // Consider validation errors (from RunAssetValidation or RunMapValidation) as fatal (preventing the package from being cooked)
 };
 ENUM_CLASS_FLAGS(ECookByTheBookOptions);
 
@@ -685,6 +688,9 @@ public:
 	/** Execute class-specific special case cook postloads and reference discovery on a given package. */
 	UNREALED_API void PostLoadPackageFixup(UE::Cook::FPackageData& PackageData, UPackage* Package);
 
+	/** Execute validation on a loaded source package. */
+	UNREALED_API EDataValidationResult ValidateSourcePackage(UE::Cook::FPackageData& PackageData, UPackage* Package);
+
 	/** Tick CBTB until it finishes or needs to yield. Should only be called when in CookByTheBook Mode. */
 	UNREALED_API uint32 TickCookByTheBook(const float TimeSlice, ECookTickFlags TickFlags = ECookTickFlags::None);
 	/** Tick COTF until it finishes or needs to yield. Should only be called when in CookOnTheFly Mode. */
@@ -841,6 +847,9 @@ public:
 	UNREALED_API void OnObjectUpdated( UObject *Object );
 	UNREALED_API void OnObjectSaved( UObject *ObjectSaved, FObjectPreSaveContext SaveContext );
 
+	DECLARE_DELEGATE_RetVal_TwoParams(EDataValidationResult, FOnValidateSourcePackage, UPackage* /*Package*/, FDataValidationContext& /*ValidationContext*/);
+	static FOnValidateSourcePackage& OnValidateSourcePackage() { return ValidateSourcePackageEvent; }
+
 	DECLARE_MULTICAST_DELEGATE(FOnCookByTheBookStarted);
 	UE_DEPRECATED(5.4, "Use UE::Cook::FDelegates::CookByTheBookStarted (CoreUObject/Public/UObject/ICookInfo.h.")
 	static FOnCookByTheBookStarted& OnCookByTheBookStarted() { return CookByTheBookStartedEvent; };
@@ -1531,6 +1540,9 @@ private:
 	/** Tracking for the ticking of tickable cook objects */
 	double LastCookableObjectTickTime = 0.;
 
+	// Hook into the DataValidation subsystem from the cooker
+	static UNREALED_API FOnValidateSourcePackage ValidateSourcePackageEvent;
+
 	// Cook events that can be listenned to
 	static FOnCookByTheBookStarted CookByTheBookStartedEvent;
 	static FOnCookByTheBookFinished CookByTheBookFinishedEvent;
diff --git a/Engine/Source/Editor/UnrealEd/Private/Commandlets/CookCommandlet.cpp b/Engine/Source/Editor/UnrealEd/Private/Commandlets/CookCommandlet.cpp
--- a/Engine/Source/Editor/UnrealEd/Private/Commandlets/CookCommandlet.cpp
+++ b/Engine/Source/Editor/UnrealEd/Private/Commandlets/CookCommandlet.cpp
@@ -66,6 +66,12 @@ struct FScopeRootObject
 	}
 };
 
+bool bAllowContentValidation = true;
+FAutoConsoleVariableRef AllowContentValidationCVar(
+	TEXT("Cook.AllowContentValidation"),
+	bAllowContentValidation,
+	TEXT("True to allow content validation to run during cook (if requested), or false to disable it."));
+
 }
 
 UCookCommandlet::UCookCommandlet( const FObjectInitializer& ObjectInitializer )
@@ -407,6 +413,13 @@ bool UCookCommandlet::CookByTheBook(const TArray<ITargetPlatform*>& Platforms)
 	CookOptions |= Switches.Contains(TEXT("DlcReevaluateUncookedAssets")) ? ECookByTheBookOptions::DlcReevaluateUncookedAssets : ECookByTheBookOptions::None;
 	bool bCookList = Switches.Contains(TEXT("CookList"));
 
+	if (UE::Cook::bAllowContentValidation)
+	{
+		CookOptions |= Switches.Contains(TEXT("RunAssetValidation")) ? ECookByTheBookOptions::RunAssetValidation : ECookByTheBookOptions::None;
+		CookOptions |= Switches.Contains(TEXT("RunMapValidation")) ? ECookByTheBookOptions::RunMapValidation : ECookByTheBookOptions::None;
+		CookOptions |= Switches.Contains(TEXT("ValidationErrorsAreFatal")) ? ECookByTheBookOptions::ValidationErrorsAreFatal : ECookByTheBookOptions::None;
+	}
+
 	if (bCookSinglePackage)
 	{
 		const ECookByTheBookOptions SinglePackageFlags = ECookByTheBookOptions::NoAlwaysCookMaps | ECookByTheBookOptions::NoDefaultMaps | ECookByTheBookOptions::NoGameAlwaysCookPackages |
diff --git a/Engine/Source/Editor/UnrealEd/Private/CookOnTheFlyServer.cpp b/Engine/Source/Editor/UnrealEd/Private/CookOnTheFlyServer.cpp
--- a/Engine/Source/Editor/UnrealEd/Private/CookOnTheFlyServer.cpp
+++ b/Engine/Source/Editor/UnrealEd/Private/CookOnTheFlyServer.cpp
@@ -98,6 +98,7 @@
 #include "Misc/CommandLine.h"
 #include "Misc/ConfigCacheIni.h"
 #include "Misc/CoreDelegates.h"
+#include "Misc/DataValidation.h"
 #include "Misc/FileHelper.h"
 #include "Misc/LocalTimestampDirectoryVisitor.h"
 #include "Misc/NetworkVersion.h"
@@ -218,6 +219,7 @@ static FAutoConsoleVariableRef CVarCookDisplayWarnBusyTime(
 /// Cook on the fly server
 ///////////////////////////////////////////////////////////////
 UCookOnTheFlyServer* UCookOnTheFlyServer::ActiveCOTFS = nullptr;
+UCookOnTheFlyServer::FOnValidateSourcePackage UCookOnTheFlyServer::ValidateSourcePackageEvent;
 UCookOnTheFlyServer::FOnCookByTheBookStarted UCookOnTheFlyServer::CookByTheBookStartedEvent;
 UCookOnTheFlyServer::FOnCookByTheBookFinished UCookOnTheFlyServer::CookByTheBookFinishedEvent;
 
@@ -2934,6 +2936,20 @@ void UCookOnTheFlyServer::LoadPackageInQueue(UE::Cook::FPackageData& PackageData
 		return;
 	}
 
+	if (ValidateSourcePackage(PackageData, LoadedPackage) == EDataValidationResult::Invalid)
+	{
+		if (EnumHasAnyFlags(CookByTheBookOptions->StartupOptions, ECookByTheBookOptions::ValidationErrorsAreFatal))
+		{
+			UE_LOG(LogCook, Error, TEXT("%s failed validation"), *LoadedPackage->GetName());
+
+			PackageData.SetPlatformsCooked(PlatformManager->GetSessionPlatforms(), ECookResult::Failed);
+			RejectPackageToLoad(PackageData, TEXT("failed validation"), ESuppressCookReason::ValidationError);
+			return;
+		}
+
+		UE_LOG(LogCook, Warning, TEXT("%s failed validation"), *LoadedPackage->GetName());
+	}
+
 	PostLoadPackageFixup(PackageData, LoadedPackage);
 	PackageData.SetPackage(LoadedPackage);
 	PackageData.SendToState(EPackageState::Save, ESendFlags::QueueAdd, EStateChangeReason::Loaded);
@@ -2964,6 +2980,168 @@ void UCookOnTheFlyServer::RejectPackageToLoad(UE::Cook::FPackageData& PackageDat
 	DemoteToIdle(PackageData, UE::Cook::ESendFlags::QueueAdd, Reason);
 }
 
+EDataValidationResult UCookOnTheFlyServer::ValidateSourcePackage(UE::Cook::FPackageData& PackageData, UPackage* Package)
+{
+	UE_SCOPED_HIERARCHICAL_COOKTIMER(ValidateSourcePackage);
+
+	// Don't validate packages if validation is disabled
+	if (!EnumHasAnyFlags(CookByTheBookOptions->StartupOptions, ECookByTheBookOptions::RunAssetValidation | ECookByTheBookOptions::RunMapValidation))
+	{
+		return EDataValidationResult::NotValidated;
+	}
+
+	// Don't validate packages generated during cook
+	if (PackageData.IsGenerated())
+	{
+		return EDataValidationResult::NotValidated;
+	}
+
+	// Don't validate packages that are already cooked
+	if (Package->HasAnyPackageFlags(PKG_Cooked))
+	{
+		return EDataValidationResult::NotValidated;
+	}
+
+	FNameBuilder PackageName(Package->GetFName());
+	const FStringView ContentRootName = FPackageName::SplitPackageNameRoot(PackageName.ToView(), nullptr);
+
+	// Don't validate Verse packages, as the Verse compiler handles that
+	if (FPackageName::IsVersePackage(PackageName.ToView()))
+	{
+		return EDataValidationResult::NotValidated;
+	}
+
+	// When cooking DLC, don't validate anything outside of the DLC plugin
+	if (IsCookingDLC() && ContentRootName != CookByTheBookOptions->DlcName)
+	{
+		return EDataValidationResult::NotValidated;
+	}
+
+	// When cooking a project, don't validate any engine content as it may not pass the project specific validators
+	if (FApp::HasProjectName())
+	{
+		if (ContentRootName == TEXTVIEW("Engine"))
+		{
+			return EDataValidationResult::NotValidated;
+		}
+		if (TSharedPtr<IPlugin> Plugin = IPluginManager::Get().FindPlugin(ContentRootName);
+			Plugin && Plugin->GetLoadedFrom() == EPluginLoadedFrom::Engine)
+		{
+			return EDataValidationResult::NotValidated;
+		}
+	}
+
+	// Don't validate packages that won't actually be cooked
+	{
+		UAssetManager& AssetManager = UAssetManager::Get();
+		
+		if (!AssetManager.VerifyCanCookPackage(this, Package->GetFName(), /*bLogError*/false))
+		{
+			return EDataValidationResult::NotValidated;
+		}
+
+		bool bShouldCookForAnyPlatform = false;
+		for (const ITargetPlatform* TargetPlatform : PlatformManager->GetSessionPlatforms())
+		{
+			if (AssetManager.ShouldCookForPlatform(Package, TargetPlatform))
+			{
+				if (const TSet<FName>* NeverCookPackages = PackageTracker->PlatformSpecificNeverCookPackages.Find(TargetPlatform);
+					!NeverCookPackages || !NeverCookPackages->Contains(Package->GetFName()))
+				{
+					bShouldCookForAnyPlatform = true;
+					break;
+				}
+			}
+		}
+		if (!bShouldCookForAnyPlatform)
+		{
+			return EDataValidationResult::NotValidated;
+		}
+	}
+
+#if DEBUG_COOKONTHEFLY 
+	UE_LOG(LogCook, Display, TEXT("Validating package %s"), *PackageName);
+#endif
+
+	EDataValidationResult FinalValidationResult = EDataValidationResult::NotValidated;
+
+	UWorld* World = nullptr;
+	if (Package->HasAnyPackageFlags(PKG_ContainsMap))
+	{
+		World = UWorld::FindWorldInPackage(Package);
+	}
+
+	bool bRunCleanupWorld = false;
+	if (World && !World->bIsWorldInitialized)
+	{
+		FWorldInitializationValues IVS;
+		IVS.AllowAudioPlayback(false);
+		IVS.RequiresHitProxies(false);
+		IVS.ShouldSimulatePhysics(false);
+		IVS.EnableTraceCollision(true);
+		IVS.SetTransactional(false);
+		IVS.CreateWorldPartition(true);
+
+		World->InitWorld(IVS);
+		bRunCleanupWorld = true;
+	}
+
+	// Run asset validation if requested
+	if (EnumHasAnyFlags(CookByTheBookOptions->StartupOptions, ECookByTheBookOptions::RunAssetValidation) && ValidateSourcePackageEvent.IsBound())
+	{
+		TArray<FAssetData> ExternalObjects;
+		if (World)
+		{
+			const FString ExternalActorsPathForWorld = ULevel::GetExternalActorsPath(Package);
+			AssetRegistry->GetAssetsByPath(*ExternalActorsPathForWorld, ExternalObjects, /*bRecursive*/true, /*bIncludeOnlyOnDiskAssets*/true);
+		}
+
+		static const FName NAME_AssetCheck = "AssetCheck";
+
+		FMessageLogScopedOverride AssetCheckLogOverride(NAME_AssetCheck);
+		if (!EnumHasAnyFlags(CookByTheBookOptions->StartupOptions, ECookByTheBookOptions::ValidationErrorsAreFatal))
+		{
+			AssetCheckLogOverride.RemapMessageSeverity(EMessageSeverity::Error, EMessageSeverity::Warning);
+		}
+
+		FDataValidationContext ValidationContext(IsRunningCookCommandlet(), EDataValidationUsecase::Save, ExternalObjects);
+		const EDataValidationResult ValidationResult = ValidateSourcePackageEvent.Execute(Package, ValidationContext);
+		FinalValidationResult = CombineDataValidationResults(FinalValidationResult, ValidationResult);
+	}
+
+	// Run map validation if requested
+	if (EnumHasAnyFlags(CookByTheBookOptions->StartupOptions, ECookByTheBookOptions::RunMapValidation) && World)
+	{
+		static const FName NAME_MapCheck = "MapCheck";
+
+		FMessageLogScopedOverride MapCheckLogOverride(NAME_MapCheck);
+		if (!EnumHasAnyFlags(CookByTheBookOptions->StartupOptions, ECookByTheBookOptions::ValidationErrorsAreFatal))
+		{
+			MapCheckLogOverride.RemapMessageSeverity(EMessageSeverity::Error, EMessageSeverity::Warning);
+		}
+
+		GEditor->Exec(World, TEXT("MAP CHECK"));
+		
+		FMessageLog MapCheckLog(NAME_MapCheck);
+		if (MapCheckLog.NumMessages(EMessageSeverity::Error) > 0)
+		{
+			FinalValidationResult = EDataValidationResult::Invalid;
+		}
+	}
+
+	if (bRunCleanupWorld)
+	{
+		checkf(World, TEXT("bRunCleanupWorld was true but World was null!"));
+		checkf(World->bIsWorldInitialized, TEXT("bRunCleanupWorld was true but World->bIsWorldInitialized was false!"));
+
+		World->ClearWorldComponents();
+		World->CleanupWorld();
+		World->SetPhysicsScene(nullptr);
+	}
+
+	return FinalValidationResult;
+}
+
 void UCookOnTheFlyServer::QueueDiscoveredPackage(UE::Cook::FPackageData& PackageData,
 	UE::Cook::FInstigator&& Instigator, UE::Cook::FDiscoveredPlatformSet&& ReachablePlatforms, bool bUrgent)
 {
diff --git a/Engine/Source/Editor/UnrealEd/Private/Cooker/CookTypes.cpp b/Engine/Source/Editor/UnrealEd/Private/Cooker/CookTypes.cpp
--- a/Engine/Source/Editor/UnrealEd/Private/Cooker/CookTypes.cpp
+++ b/Engine/Source/Editor/UnrealEd/Private/Cooker/CookTypes.cpp
@@ -63,6 +63,7 @@ const TCHAR* LexToString(ESuppressCookReason Reason)
 	case ESuppressCookReason::Redirected: return TEXT("Redirected");
 	case ESuppressCookReason::OrphanedGenerated: return TEXT("OrphanedGenerated");
 	case ESuppressCookReason::LoadError: return TEXT("LoadError");
+	case ESuppressCookReason::ValidationError: return TEXT("ValidationError");
 	case ESuppressCookReason::SaveError: return TEXT("SaveError");
 	case ESuppressCookReason::OnlyEditorOnly: return TEXT("OnlyEditorOnly");
 	case ESuppressCookReason::CookCanceled: return TEXT("CookCanceled");
diff --git a/Engine/Source/Editor/UnrealEd/Private/Cooker/CookTypes.h b/Engine/Source/Editor/UnrealEd/Private/Cooker/CookTypes.h
--- a/Engine/Source/Editor/UnrealEd/Private/Cooker/CookTypes.h
+++ b/Engine/Source/Editor/UnrealEd/Private/Cooker/CookTypes.h
@@ -152,6 +152,7 @@ namespace UE::Cook
 		Redirected,
 		OrphanedGenerated,
 		LoadError,
+		ValidationError,
 		SaveError,
 		OnlyEditorOnly,
 		CookCanceled,
