Moved the loaded package validation callback for the cooker to UE::Cook::FDelegates

#rnx

diff --git a/Engine/Plugins/Editor/DataValidation/Source/DataValidation/Private/DataValidationModule.cpp b/Engine/Plugins/Editor/DataValidation/Source/DataValidation/Private/DataValidationModule.cpp
--- a/Engine/Plugins/Editor/DataValidation/Source/DataValidation/Private/DataValidationModule.cpp
+++ b/Engine/Plugins/Editor/DataValidation/Source/DataValidation/Private/DataValidationModule.cpp
@@ -6,7 +6,6 @@
 #include "AssetRegistry/AssetDataToken.h"
 #include "UObject/ObjectSaveContext.h"
 
-#include "CookOnTheSide/CookOnTheFlyServer.h"
 #include "Editor.h"
 #include "Framework/Application/SlateApplication.h"
 #include "ToolMenus.h"
@@ -15,6 +14,7 @@
 #include "DataValidationCommandlet.h"
 #include "Elements/Framework/TypedElementSelectionSet.h"
 #include "Logging/MessageLog.h"
+#include "UObject/ICookInfo.h"
 
 #include "ContentBrowserMenuContexts.h"
 #include "LevelEditorMenuContext.h"
@@ -59,7 +59,7 @@ IMPLEMENT_MODULE(FDataValidationModule, DataValidation)
 
 void FDataValidationModule::StartupModule()
 {	
-	UCookOnTheFlyServer::OnValidateSourcePackage().BindRaw(this, &FDataValidationModule::OnValidateSourcePackageDuringCook);
+	UE::Cook::FDelegates::ValidateSourcePackage.BindRaw(this, &FDataValidationModule::OnValidateSourcePackageDuringCook);
 
 	if (!IsRunningCommandlet() && !IsRunningGame() && FSlateApplication::IsInitialized())
 	{
@@ -80,7 +80,7 @@ void FDataValidationModule::StartupModule()
 
 void FDataValidationModule::ShutdownModule()
 {
-	UCookOnTheFlyServer::OnValidateSourcePackage().Unbind();
+	UE::Cook::FDelegates::ValidateSourcePackage.Unbind();
 
 	if (!IsRunningCommandlet() && !IsRunningGame() && !IsRunningDedicatedServer())
 	{
diff --git a/Engine/Source/Editor/UnrealEd/Classes/CookOnTheSide/CookOnTheFlyServer.h b/Engine/Source/Editor/UnrealEd/Classes/CookOnTheSide/CookOnTheFlyServer.h
--- a/Engine/Source/Editor/UnrealEd/Classes/CookOnTheSide/CookOnTheFlyServer.h
+++ b/Engine/Source/Editor/UnrealEd/Classes/CookOnTheSide/CookOnTheFlyServer.h
@@ -847,9 +847,6 @@ public:
 	UNREALED_API void OnObjectUpdated( UObject *Object );
 	UNREALED_API void OnObjectSaved( UObject *ObjectSaved, FObjectPreSaveContext SaveContext );
 
-	DECLARE_DELEGATE_RetVal_TwoParams(EDataValidationResult, FOnValidateSourcePackage, UPackage* /*Package*/, FDataValidationContext& /*ValidationContext*/);
-	static FOnValidateSourcePackage& OnValidateSourcePackage() { return ValidateSourcePackageEvent; }
-
 	DECLARE_MULTICAST_DELEGATE(FOnCookByTheBookStarted);
 	UE_DEPRECATED(5.4, "Use UE::Cook::FDelegates::CookByTheBookStarted (CoreUObject/Public/UObject/ICookInfo.h.")
 	static FOnCookByTheBookStarted& OnCookByTheBookStarted() { return CookByTheBookStartedEvent; };
@@ -1540,9 +1537,6 @@ private:
 	/** Tracking for the ticking of tickable cook objects */
 	double LastCookableObjectTickTime = 0.;
 
-	// Hook into the DataValidation subsystem from the cooker
-	static UNREALED_API FOnValidateSourcePackage ValidateSourcePackageEvent;
-
 	// Cook events that can be listenned to
 	static FOnCookByTheBookStarted CookByTheBookStartedEvent;
 	static FOnCookByTheBookFinished CookByTheBookFinishedEvent;
diff --git a/Engine/Source/Editor/UnrealEd/Private/CookOnTheFlyServer.cpp b/Engine/Source/Editor/UnrealEd/Private/CookOnTheFlyServer.cpp
--- a/Engine/Source/Editor/UnrealEd/Private/CookOnTheFlyServer.cpp
+++ b/Engine/Source/Editor/UnrealEd/Private/CookOnTheFlyServer.cpp
@@ -219,7 +219,6 @@ static FAutoConsoleVariableRef CVarCookDisplayWarnBusyTime(
 /// Cook on the fly server
 ///////////////////////////////////////////////////////////////
 UCookOnTheFlyServer* UCookOnTheFlyServer::ActiveCOTFS = nullptr;
-UCookOnTheFlyServer::FOnValidateSourcePackage UCookOnTheFlyServer::ValidateSourcePackageEvent;
 UCookOnTheFlyServer::FOnCookByTheBookStarted UCookOnTheFlyServer::CookByTheBookStartedEvent;
 UCookOnTheFlyServer::FOnCookByTheBookFinished UCookOnTheFlyServer::CookByTheBookFinishedEvent;
 
@@ -3087,7 +3086,7 @@ EDataValidationResult UCookOnTheFlyServer::ValidateSourcePackage(UE::Cook::FPack
 	}
 
 	// Run asset validation if requested
-	if (EnumHasAnyFlags(CookByTheBookOptions->StartupOptions, ECookByTheBookOptions::RunAssetValidation) && ValidateSourcePackageEvent.IsBound())
+	if (EnumHasAnyFlags(CookByTheBookOptions->StartupOptions, ECookByTheBookOptions::RunAssetValidation) && UE::Cook::FDelegates::ValidateSourcePackage.IsBound())
 	{
 		TArray<FAssetData> ExternalObjects;
 		if (World)
@@ -3105,7 +3104,7 @@ EDataValidationResult UCookOnTheFlyServer::ValidateSourcePackage(UE::Cook::FPack
 		}
 
 		FDataValidationContext ValidationContext(IsRunningCookCommandlet(), EDataValidationUsecase::Save, ExternalObjects);
-		const EDataValidationResult ValidationResult = ValidateSourcePackageEvent.Execute(Package, ValidationContext);
+		const EDataValidationResult ValidationResult = UE::Cook::FDelegates::ValidateSourcePackage.Execute(Package, ValidationContext);
 		FinalValidationResult = CombineDataValidationResults(FinalValidationResult, ValidationResult);
 	}
 
diff --git a/Engine/Source/Runtime/CoreUObject/Private/UObject/ICookInfo.cpp b/Engine/Source/Runtime/CoreUObject/Private/UObject/ICookInfo.cpp
--- a/Engine/Source/Runtime/CoreUObject/Private/UObject/ICookInfo.cpp
+++ b/Engine/Source/Runtime/CoreUObject/Private/UObject/ICookInfo.cpp
@@ -47,6 +47,7 @@ FString FInstigator::ToString() const
 
 FCookInfoEvent FDelegates::CookByTheBookStarted;
 FCookInfoEvent FDelegates::CookByTheBookFinished;
+FValidateSourcePackage FDelegates::ValidateSourcePackage;
 
 }
 
diff --git a/Engine/Source/Runtime/CoreUObject/Public/UObject/ICookInfo.h b/Engine/Source/Runtime/CoreUObject/Public/UObject/ICookInfo.h
--- a/Engine/Source/Runtime/CoreUObject/Public/UObject/ICookInfo.h
+++ b/Engine/Source/Runtime/CoreUObject/Public/UObject/ICookInfo.h
@@ -13,6 +13,10 @@
 
 #if WITH_EDITOR
 
+enum class EDataValidationResult : uint8;
+class FDataValidationContext;
+class UPackage;
+
 namespace UE::Cook { class IMPCollector; }
 
 namespace UE::Cook
@@ -132,6 +136,7 @@ public:
 };
 
 DECLARE_MULTICAST_DELEGATE_OneParam(FCookInfoEvent, ICookInfo&);
+DECLARE_DELEGATE_RetVal_TwoParams(EDataValidationResult, FValidateSourcePackage, UPackage* /*Package*/, FDataValidationContext& /*ValidationContext*/);
 
 /** UE::Cook::FDelegates: callbacks for cook events. */
 struct FDelegates
@@ -139,6 +144,7 @@ struct FDelegates
 public:
 	static COREUOBJECT_API FCookInfoEvent CookByTheBookStarted;
 	static COREUOBJECT_API FCookInfoEvent CookByTheBookFinished;
+	static COREUOBJECT_API FValidateSourcePackage ValidateSourcePackage;
 };
 
 } // namespace UE::Cook
