Skeletal Mesh: Temporarily revive the code for storing the deprecated retarget base pose when re-importing

#rnx

diff --git a/Engine/Source/Editor/UnrealEd/Private/ImportUtils/SkeletalMeshImportUtils.cpp b/Engine/Source/Editor/UnrealEd/Private/ImportUtils/SkeletalMeshImportUtils.cpp
--- a/Engine/Source/Editor/UnrealEd/Private/ImportUtils/SkeletalMeshImportUtils.cpp
+++ b/Engine/Source/Editor/UnrealEd/Private/ImportUtils/SkeletalMeshImportUtils.cpp
@@ -350,6 +350,9 @@ TSharedPtr<FExistingSkelMeshData> SkeletalMeshImportUtils::SaveExistingSkelMeshD
 	{
 		ExistingMeshDataPtr->ExistingMaterials = SourceMaterials;
 	}
+	PRAGMA_DISABLE_DEPRECATION_WARNINGS
+	ExistingMeshDataPtr->ExistingRetargetBasePose = SourceSkeletalMesh->GetRetargetBasePose();
+	PRAGMA_ENABLE_DEPRECATION_WARNINGS
 
 	ExistingMeshDataPtr->ExistingInlineReductionCacheDatas = SourceMeshModel->InlineReductionCacheDatas;
 
@@ -708,6 +711,15 @@ void SkeletalMeshImportUtils::RestoreExistingSkelMeshData(const TSharedPtr<const
 	//Do everything we need for base LOD re-import
 	if (SafeReimportLODIndex == 0)
 	{
+		PRAGMA_DISABLE_DEPRECATION_WARNINGS
+		// this is not ideal. Ideally we'll have to save only diff with indicating which joints, 
+		// but for now, we allow them to keep the previous pose IF the element count is same
+		if (MeshData->ExistingRetargetBasePose.Num() == SkeletalMesh->GetRefSkeleton().GetRawBoneNum())
+		{
+			SkeletalMesh->SetRetargetBasePose(MeshData->ExistingRetargetBasePose);
+		}
+		PRAGMA_ENABLE_DEPRECATION_WARNINGS
+
 		// Assign sockets from old version of this SkeletalMesh.
 		// Only copy ones for bones that exist in the new mesh.
 		for (int32 i = 0; i < MeshData->ExistingSockets.Num(); i++)
diff --git a/Engine/Source/Editor/UnrealEd/Public/ImportUtils/SkelImport.h b/Engine/Source/Editor/UnrealEd/Public/ImportUtils/SkelImport.h
--- a/Engine/Source/Editor/UnrealEd/Public/ImportUtils/SkelImport.h
+++ b/Engine/Source/Editor/UnrealEd/Public/ImportUtils/SkelImport.h
@@ -81,6 +81,7 @@ struct FExistingSkelMeshData
 	TArray<UPhysicsAsset*>					ExistingPhysicsAssets;
 	UPhysicsAsset*							ExistingShadowPhysicsAsset;
 	USkeleton*								ExistingSkeleton;
+	TArray<FTransform>						ExistingRetargetBasePose;
 	USkeletalMeshLODSettings*				ExistingLODSettings;
 	TSubclassOf<UAnimInstance>				ExistingPostProcessAnimBlueprint;
 	TSoftObjectPtr<UObject>					ExistingDefaultAnimatingRig;
