Downgrade UE_LOG errors and ensures to warnings when running cook validation without -ValidationErrorsAreFatal

#rnx

diff --git a/Engine/Source/Editor/UnrealEd/Private/CookOnTheFlyServer.cpp b/Engine/Source/Editor/UnrealEd/Private/CookOnTheFlyServer.cpp
--- a/Engine/Source/Editor/UnrealEd/Private/CookOnTheFlyServer.cpp
+++ b/Engine/Source/Editor/UnrealEd/Private/CookOnTheFlyServer.cpp
@@ -3062,6 +3062,7 @@ EDataValidationResult UCookOnTheFlyServer::ValidateSourcePackage(UE::Cook::FPack
 	UE_LOG(LogCook, Display, TEXT("Validating package %s"), *PackageName);
 #endif
 
+	const bool bWarningsAsErrors = !EnumHasAnyFlags(CookByTheBookOptions->StartupOptions, ECookByTheBookOptions::ValidationErrorsAreFatal);
 	EDataValidationResult FinalValidationResult = EDataValidationResult::NotValidated;
 
 	UWorld* World = nullptr;
@@ -3098,11 +3099,13 @@ EDataValidationResult UCookOnTheFlyServer::ValidateSourcePackage(UE::Cook::FPack
 		static const FName NAME_AssetCheck = "AssetCheck";
 
 		FMessageLogScopedOverride AssetCheckLogOverride(NAME_AssetCheck);
-		if (!EnumHasAnyFlags(CookByTheBookOptions->StartupOptions, ECookByTheBookOptions::ValidationErrorsAreFatal))
+		if (bWarningsAsErrors)
 		{
 			AssetCheckLogOverride.RemapMessageSeverity(EMessageSeverity::Error, EMessageSeverity::Warning);
 		}
 
+		TGuardValue<bool> LogErrorsAsWarnings(GWarn->TreatErrorsAsWarnings, GWarn->TreatErrorsAsWarnings || bWarningsAsErrors);
+
 		FDataValidationContext ValidationContext(IsRunningCookCommandlet(), EDataValidationUsecase::Save, ExternalObjects);
 		const EDataValidationResult ValidationResult = UE::Cook::FDelegates::ValidateSourcePackage.Execute(Package, ValidationContext);
 		FinalValidationResult = CombineDataValidationResults(FinalValidationResult, ValidationResult);
@@ -3114,11 +3117,13 @@ EDataValidationResult UCookOnTheFlyServer::ValidateSourcePackage(UE::Cook::FPack
 		static const FName NAME_MapCheck = "MapCheck";
 
 		FMessageLogScopedOverride MapCheckLogOverride(NAME_MapCheck);
-		if (!EnumHasAnyFlags(CookByTheBookOptions->StartupOptions, ECookByTheBookOptions::ValidationErrorsAreFatal))
+		if (bWarningsAsErrors)
 		{
 			MapCheckLogOverride.RemapMessageSeverity(EMessageSeverity::Error, EMessageSeverity::Warning);
 		}
 
+		TGuardValue<bool> LogErrorsAsWarnings(GWarn->TreatErrorsAsWarnings, GWarn->TreatErrorsAsWarnings || bWarningsAsErrors);
+
 		GEditor->Exec(World, TEXT("MAP CHECK"));
 		
 		FMessageLog MapCheckLog(NAME_MapCheck);
