Editor: Check for valid component

diff --git a/Engine/Source/Editor/UnrealEd/Private/EditorSelectUtils.cpp b/Engine/Source/Editor/UnrealEd/Private/EditorSelectUtils.cpp
--- a/Engine/Source/Editor/UnrealEd/Private/EditorSelectUtils.cpp
+++ b/Engine/Source/Editor/UnrealEd/Private/EditorSelectUtils.cpp
@@ -309,7 +309,7 @@ void UUnrealEdEngine::OnEditorElementSelectionChanged(const UTypedElementSelecti
 		for (int32 CompIdx = 0; CompIdx < Components.Num(); CompIdx++)
 		{
 			TWeakObjectPtr<UActorComponent> Comp(Components[CompIdx]);
-			if (Comp->IsRegistered())
+			if (Comp.IsValid() && Comp->IsRegistered())
 			{
 				// Try and find a visualizer
 				TSharedPtr<FComponentVisualizer> Visualizer = FindComponentVisualizer(Comp->GetClass());
