Fixed typo in variable name

#rnx

diff --git a/Engine/Source/Editor/UnrealEd/Private/CookOnTheFlyServer.cpp b/Engine/Source/Editor/UnrealEd/Private/CookOnTheFlyServer.cpp
--- a/Engine/Source/Editor/UnrealEd/Private/CookOnTheFlyServer.cpp
+++ b/Engine/Source/Editor/UnrealEd/Private/CookOnTheFlyServer.cpp
@@ -3062,7 +3062,7 @@ EDataValidationResult UCookOnTheFlyServer::ValidateSourcePackage(UE::Cook::FPack
 	UE_LOG(LogCook, Display, TEXT("Validating package %s"), *PackageName);
 #endif
 
-	const bool bWarningsAsErrors = !EnumHasAnyFlags(CookByTheBookOptions->StartupOptions, ECookByTheBookOptions::ValidationErrorsAreFatal);
+	const bool bLogErrorsAsWarnings = !EnumHasAnyFlags(CookByTheBookOptions->StartupOptions, ECookByTheBookOptions::ValidationErrorsAreFatal);
 	EDataValidationResult FinalValidationResult = EDataValidationResult::NotValidated;
 
 	UWorld* World = nullptr;
@@ -3099,12 +3099,12 @@ EDataValidationResult UCookOnTheFlyServer::ValidateSourcePackage(UE::Cook::FPack
 		static const FName NAME_AssetCheck = "AssetCheck";
 
 		FMessageLogScopedOverride AssetCheckLogOverride(NAME_AssetCheck);
-		if (bWarningsAsErrors)
+		if (bLogErrorsAsWarnings)
 		{
 			AssetCheckLogOverride.RemapMessageSeverity(EMessageSeverity::Error, EMessageSeverity::Warning);
 		}
 
-		TGuardValue<bool> LogErrorsAsWarnings(GWarn->TreatErrorsAsWarnings, GWarn->TreatErrorsAsWarnings || bWarningsAsErrors);
+		TGuardValue<bool> LogErrorsAsWarnings(GWarn->TreatErrorsAsWarnings, GWarn->TreatErrorsAsWarnings || bLogErrorsAsWarnings);
 
 		FDataValidationContext ValidationContext(IsRunningCookCommandlet(), EDataValidationUsecase::Save, ExternalObjects);
 		const EDataValidationResult ValidationResult = UE::Cook::FDelegates::ValidateSourcePackage.Execute(Package, ValidationContext);
@@ -3117,12 +3117,12 @@ EDataValidationResult UCookOnTheFlyServer::ValidateSourcePackage(UE::Cook::FPack
 		static const FName NAME_MapCheck = "MapCheck";
 
 		FMessageLogScopedOverride MapCheckLogOverride(NAME_MapCheck);
-		if (bWarningsAsErrors)
+		if (bLogErrorsAsWarnings)
 		{
 			MapCheckLogOverride.RemapMessageSeverity(EMessageSeverity::Error, EMessageSeverity::Warning);
 		}
 
-		TGuardValue<bool> LogErrorsAsWarnings(GWarn->TreatErrorsAsWarnings, GWarn->TreatErrorsAsWarnings || bWarningsAsErrors);
+		TGuardValue<bool> LogErrorsAsWarnings(GWarn->TreatErrorsAsWarnings, GWarn->TreatErrorsAsWarnings || bLogErrorsAsWarnings);
 
 		GEditor->Exec(World, TEXT("MAP CHECK"));
 		
