Fix Pilot Actor not creating a transaction.

diff --git a/Engine/Source/Editor/UnrealEd/Private/LevelEditorViewport.cpp b/Engine/Source/Editor/UnrealEd/Private/LevelEditorViewport.cpp
--- a/Engine/Source/Editor/UnrealEd/Private/LevelEditorViewport.cpp
+++ b/Engine/Source/Editor/UnrealEd/Private/LevelEditorViewport.cpp
@@ -3583,6 +3583,22 @@ void FLevelEditorViewportClient::TrackingStarted( const FInputEventState& InInpu
 	// Begin transacting.  Give the current editor mode an opportunity to do the transacting.
 	const bool bTrackingHandledExternally = ModeTools->StartTracking(this, Viewport);
 
+	// If we are have a lock due to piloting, don't modify the tracking transaction until the piloting is stopped in 'SetActorLock'
+	AActor* ActiveActorLock = GetActiveActorLock().Get();
+	if (ActiveActorLock && PilotingTransaction.IsActive() && PilotingTransaction.TransCount > 0)
+	{
+		PreDragActorTransforms.Empty();
+		PreDragElementTransforms.Empty();
+
+		if (PilotingTransaction.IsActive() || PilotingTransaction.IsPending())
+		{
+			// Suspend actor/component modification during each delta step to avoid recording unnecessary overhead into the transaction buffer
+			GEditor->DisableDeltaModification(true);
+		}
+
+		return;
+	}
+
 	TrackingTransaction.End();
 
 	// Re-initialize new tracking only if a new button was pressed, otherwise we continue the previous one.
@@ -3692,12 +3708,11 @@ void FLevelEditorViewportClient::TrackingStarted( const FInputEventState& InInpu
 		}
 		else
 		{
-			AActor* ActiveActorLock = GetActiveActorLock().Get();
-			if (ActiveActorLock && !ActiveActorLock->IsLockLocation() && TrackingTransaction.TransCount == 0)
+			if (ActiveActorLock && !ActiveActorLock->IsLockLocation() && TrackingTransaction.TransCount == 0 && PilotingTransaction.TransCount == 0)
 			{
 				// Open a tracking transaction to contain the locked actor changes
-				TrackingTransaction.TransCount++;
-				TrackingTransaction.Begin(LOCTEXT("PilotTransaction", "Pilot Actor"));
+				PilotingTransaction.TransCount++;
+				PilotingTransaction.Begin(LOCTEXT("PilotTransaction", "Pilot Actor"));
 			}
 		}
 
@@ -3713,10 +3728,19 @@ void FLevelEditorViewportClient::TrackingStarted( const FInputEventState& InInpu
 
 void FLevelEditorViewportClient::TrackingStopped()
 {
-
 	// Only disable the duplicate on next drag flag if we actually dragged the mouse.
 	bDuplicateOnNextDrag = false;
 
+	// If we are have a lock due to piloting, don't stop the transaction until the piloting is stopped in 'SetActorLock'
+	AActor* ActiveActorLock = GetActiveActorLock().Get();
+	if (ActiveActorLock && PilotingTransaction.IsActive() && PilotingTransaction.TransCount > 0)
+	{
+		PreDragActorTransforms.Empty();
+		PreDragElementTransforms.Empty();
+
+		return;
+	}
+
 	// here we check to see if anything of worth actually changed when ending our MouseMovement
 	// If the TransCount > 0 (we changed something of value) so we need to call PostEditMove() on stuff
 	// if we didn't change anything then don't call PostEditMove()
@@ -4372,6 +4396,18 @@ UActorComponent* FLevelEditorViewportClient::FindViewComponentForActor(AActor co
 
 void FLevelEditorViewportClient::SetActorLock(AActor* Actor)
 {
+	// If we had an active lock and are clearing it, also end the transaction for that lock
+	if (!Actor)
+	{
+		AActor* ActiveActorLock = GetActiveActorLock().Get();
+		if (ActiveActorLock && !ActiveActorLock->IsLockLocation() && PilotingTransaction.TransCount > 0)
+		{
+			PilotingTransaction.TransCount--;
+			PilotingTransaction.End();
+
+			ModeTools->ActorMoveNotify();
+		}
+	}
 	SetActorLock(FLevelViewportActorLock(Actor));
 }
 
diff --git a/Engine/Source/Editor/UnrealEd/Public/LevelEditorViewport.h b/Engine/Source/Editor/UnrealEd/Public/LevelEditorViewport.h
--- a/Engine/Source/Editor/UnrealEd/Public/LevelEditorViewport.h
+++ b/Engine/Source/Editor/UnrealEd/Public/LevelEditorViewport.h
@@ -958,8 +958,12 @@ private:
 
 	UWorld* World;
 
+	/** Global shared transaction for all mouse interactions. */
 	FTrackingTransaction TrackingTransaction;
 
+	/** Transaction excluse to pilot operations. Seperate as it externally handles & prevents tracking transactions */
+	FTrackingTransaction PilotingTransaction;
+
 	/** Represents the last known drop preview mouse position. */
 	int32 DropPreviewMouseX;
 	int32 DropPreviewMouseY;
