UPackage::SavedHash: Add the API that replaces UPackage->Guid. FIoHash SavedHash (or PackageSavedHash ) is now used instead on UPackage, FPackageSummary, and FAssetPackageData. Remark Guid as deprecated; we will remove it in 5.5. This replacement only applies to the editor; in the cooked game we are removing UPackage->Guid with no replacment because it is unnecessary and costs memory.
[RN] Minor, Foundation

diff --git a/Engine/Source/Editor/UnrealEd/Private/Commandlets/AssetRegistryGenerator.cpp b/Engine/Source/Editor/UnrealEd/Private/Commandlets/AssetRegistryGenerator.cpp
--- a/Engine/Source/Editor/UnrealEd/Private/Commandlets/AssetRegistryGenerator.cpp
+++ b/Engine/Source/Editor/UnrealEd/Private/Commandlets/AssetRegistryGenerator.cpp
@@ -1287,9 +1287,7 @@ void FAssetRegistryGenerator::ComputePackageDifferences(const FComputeDifference
 				else
 				{
 					FGeneratorPackageInfo& Info = OutDifference.GeneratorPackages.FindOrAdd(GeneratorName);
-					PRAGMA_DISABLE_DEPRECATION_WARNINGS;
-					Info.Generated.Add(PackageName, PreviousPackageData->PackageGuid);
-					PRAGMA_ENABLE_DEPRECATION_WARNINGS;
+					Info.Generated.Add(PackageName, PreviousPackageData->GetPackageSavedHash());
 				}
 			}
 			else
@@ -1325,7 +1323,7 @@ void FAssetRegistryGenerator::ComputePackageDifferences(const FComputeDifference
 		EDifference* GeneratorDifference = OutDifference.Packages.Find(GeneratorName);
 		if (GeneratorDifference && *GeneratorDifference == EDifference::RemovedCooked)
 		{
-			for (const TPair<FName, FGuid>& Generated : Iter->Value.Generated)
+			for (const TPair<FName, FIoHash>& Generated : Iter->Value.Generated)
 			{
 				OutDifference.Packages.Add(Generated.Key, EDifference::RemovedCooked);
 			}
@@ -1427,9 +1425,7 @@ bool FAssetRegistryGenerator::ComputePackageDifferences_IsPackageFileUnchanged(
 	const FComputeDifferenceOptions& Options, FName PackageName, const FAssetPackageData& CurrentPackageData,
 	const FAssetPackageData& PreviousPackageData)
 {
-	PRAGMA_DISABLE_DEPRECATION_WARNINGS
-	if (CurrentPackageData.PackageGuid != PreviousPackageData.PackageGuid)
-	PRAGMA_ENABLE_DEPRECATION_WARNINGS
+	if (CurrentPackageData.GetPackageSavedHash() != PreviousPackageData.GetPackageSavedHash())
 	{
 		return false;
 	}
@@ -1491,9 +1487,7 @@ void FAssetRegistryGenerator::ComputePackageRemovals(const FAssetRegistryState&
 				}
 				else
 				{
-					PRAGMA_DISABLE_DEPRECATION_WARNINGS;
-					OutGeneratorPackages.FindOrAdd(GeneratorName).Generated.Add(PackageName, PreviousPackageData->PackageGuid);
-					PRAGMA_ENABLE_DEPRECATION_WARNINGS;
+					OutGeneratorPackages.FindOrAdd(GeneratorName).Generated.Add(PackageName, PreviousPackageData->GetPackageSavedHash());
 				}
 			}
 			else
@@ -1517,7 +1511,7 @@ void FAssetRegistryGenerator::ComputePackageRemovals(const FAssetRegistryState&
 		FName GeneratorName = Iter->Key;
 		if (RemovedPackageSet.Contains(GeneratorName))
 		{
-			for (const TPair<FName, FGuid>& Generated : Iter->Value.Generated)
+			for (const TPair<FName, FIoHash>& Generated : Iter->Value.Generated)
 			{
 				RemovedPackageSet.Add(Generated.Key);
 			}
@@ -2790,9 +2784,7 @@ void FAssetRegistryPackageMessage::Write(FCbWriter& Writer) const
 		Writer.BeginObject("P");
 		{
 			// Currently we only replicate Guid and ImportedClasses, since these are the only fields set by generated pacakges
-			PRAGMA_DISABLE_DEPRECATION_WARNINGS;
-			Writer << "G" << OverrideAssetPackageData->PackageGuid;
-			PRAGMA_ENABLE_DEPRECATION_WARNINGS;
+			Writer << "H" << OverrideAssetPackageData->GetPackageSavedHash();
 			Writer << "C" << OverrideAssetPackageData->ImportedClasses;
 		}
 		Writer.EndObject();
@@ -2830,13 +2822,13 @@ bool FAssetRegistryPackageMessage::TryRead(FCbObjectView Object)
 	if (OverrideAssetPackageDataField.HasValue())
 	{
 		OverrideAssetPackageData.Emplace();
-		// Currently we only replicate Guid and ImportedClasses, since these are the only fields set by generated pacakges
-		PRAGMA_DISABLE_DEPRECATION_WARNINGS;
-		if (!LoadFromCompactBinary(OverrideAssetPackageDataField["G"], OverrideAssetPackageData->PackageGuid))
+		// Currently we only replicate Guid and ImportedClasses, since these are the only fields set by generated packages
+		FIoHash PackageSavedHash;
+		if (!LoadFromCompactBinary(OverrideAssetPackageDataField["H"], PackageSavedHash))
 		{
 			return false;
 		}
-		PRAGMA_ENABLE_DEPRECATION_WARNINGS;
+		OverrideAssetPackageData->SetPackageSavedHash(PackageSavedHash);
 		if (!LoadFromCompactBinary(OverrideAssetPackageDataField["C"], OverrideAssetPackageData->ImportedClasses))
 		{
 			return false;
diff --git a/Engine/Source/Editor/UnrealEd/Private/Commandlets/DiffAssetRegistriesCommandlet.cpp b/Engine/Source/Editor/UnrealEd/Private/Commandlets/DiffAssetRegistriesCommandlet.cpp
--- a/Engine/Source/Editor/UnrealEd/Private/Commandlets/DiffAssetRegistriesCommandlet.cpp
+++ b/Engine/Source/Editor/UnrealEd/Private/Commandlets/DiffAssetRegistriesCommandlet.cpp
@@ -404,7 +404,7 @@ void UDiffAssetRegistriesCommandlet::ConsistencyCheck(const FString& OldPath, co
 			return;
 		}
 
-		if (!OldState.Load(SerializedAssetData))
+		if (!NewState.Load(SerializedAssetData))
 		{
 			UE_LOG(LogDiffAssets, Error, TEXT("Failed to parse file '%s' as asset registry."), *NewPath);
 			return;
@@ -418,7 +418,7 @@ void UDiffAssetRegistriesCommandlet::ConsistencyCheck(const FString& OldPath, co
 	}
 	// We're looking for packages that the Cooked check says are modified, but that the Guid check says are not
 	// We're ignoring new packages for this, as those are obviously going to change
-	TSet<FName> GuidModified, CookModified;
+	TSet<FName> HashModified, CookModified;
 	TSet<FName> New;
 
 	for (const TPair<FName, const FAssetPackageData*>& Pair : NewState.GetAssetPackageDataMap())
@@ -433,11 +433,9 @@ void UDiffAssetRegistriesCommandlet::ConsistencyCheck(const FString& OldPath, co
 		}
 		else
 		{
-			PRAGMA_DISABLE_DEPRECATION_WARNINGS
-			if (Data->PackageGuid != PrevData->PackageGuid)
-			PRAGMA_ENABLE_DEPRECATION_WARNINGS
+			if (Data->GetPackageSavedHash() != PrevData->GetPackageSavedHash())
 			{
-				GuidModified.Add(Name);
+				HashModified.Add(Name);
 			}
 			if (Data->CookedHash != PrevData->CookedHash)
 			{
@@ -446,8 +444,8 @@ void UDiffAssetRegistriesCommandlet::ConsistencyCheck(const FString& OldPath, co
 		}
 	}
 
-	// recurse through the referencer lists to fill out GuidModified
-	TArray<FName> Recurse = GuidModified.Array();
+	// recurse through the referencer lists to fill out HashModified
+	TArray<FName> Recurse = HashModified.Array();
 	
 	for (int32 RecurseIndex = 0; RecurseIndex < Recurse.Num(); RecurseIndex++)
 	{
@@ -458,9 +456,9 @@ void UDiffAssetRegistriesCommandlet::ConsistencyCheck(const FString& OldPath, co
 		for (const FAssetIdentifier& Referencer : Referencers)
 		{
 			FName ReferencerPackage = Referencer.PackageName;
-			if (!New.Contains(ReferencerPackage) && !GuidModified.Contains(ReferencerPackage))
+			if (!New.Contains(ReferencerPackage) && !HashModified.Contains(ReferencerPackage))
 			{
-				GuidModified.Add(ReferencerPackage);
+				HashModified.Add(ReferencerPackage);
 				Recurse.Add(ReferencerPackage);
 			}
 		}
@@ -469,13 +467,13 @@ void UDiffAssetRegistriesCommandlet::ConsistencyCheck(const FString& OldPath, co
 	int64 Changes = 0;
 	int64 ChangeBytes = 0;
 
-	// find all entries of CookModified that do not exist in GuidModified
+	// find all entries of CookModified that do not exist in HashModified
 	const TMap<FName, const FAssetPackageData*>& PackageMap = NewState.GetAssetPackageDataMap();
 	for (FName const &Package : CookModified)
 	{
 		const FAssetPackageData* Data = PackageMap[Package];
 
-		if (!GuidModified.Contains(Package) && Data->DiskSize >= 0)
+		if (!HashModified.Contains(Package) && Data->DiskSize >= 0)
 		{
 			++Changes;
 			ChangeBytes += Data->DiskSize;
@@ -1156,9 +1154,7 @@ void UDiffAssetRegistriesCommandlet::DiffAssetRegistries(const FString& OldPath,
 				New.Add(Name);
 				RecordAdd(Name, *Data);
 			}
-			PRAGMA_DISABLE_DEPRECATION_WARNINGS
-			else if (Data->PackageGuid != PrevData->PackageGuid)
-			PRAGMA_ENABLE_DEPRECATION_WARNINGS
+			else if (Data->GetPackageSavedHash() != PrevData->GetPackageSavedHash())
 			{
 				Modified.Add(Name);
 			}
@@ -1232,9 +1228,7 @@ void UDiffAssetRegistriesCommandlet::DiffAssetRegistries(const FString& OldPath,
 			{
 				RecordEdit(Name, *Data, *PrevData);
 				AssetPathFlags.FindOrAdd(Name) |= EAssetFlags::HashChange;
-				PRAGMA_DISABLE_DEPRECATION_WARNINGS
-				if (Data->PackageGuid != PrevData->PackageGuid)
-				PRAGMA_ENABLE_DEPRECATION_WARNINGS
+				if (Data->GetPackageSavedHash() != PrevData->GetPackageSavedHash())
 				{
 					AssetPathFlags.FindOrAdd(Name) |= EAssetFlags::GuidChange;
 				}
diff --git a/Engine/Source/Editor/UnrealEd/Private/Commandlets/PackageUtilities.cpp b/Engine/Source/Editor/UnrealEd/Private/Commandlets/PackageUtilities.cpp
--- a/Engine/Source/Editor/UnrealEd/Private/Commandlets/PackageUtilities.cpp
+++ b/Engine/Source/Editor/UnrealEd/Private/Commandlets/PackageUtilities.cpp
@@ -999,9 +999,7 @@ void FPkgInfoReporter_Log::GeneratePackageReport( FLinkerLoad* InLinker /*=nullp
 
 	if (!IsHideSaveUnstable())
 	{
-		PRAGMA_DISABLE_DEPRECATION_WARNINGS
-		Out.Logf(ELogVerbosity::Display, TEXT("\t             Guid: %s"), *Linker->Summary.Guid.ToString());
-		PRAGMA_ENABLE_DEPRECATION_WARNINGS
+		Out.Logf(ELogVerbosity::Display, TEXT("\t        SavedHash: %s"), *WriteToString<40>(Linker->Summary.GetSavedHash()));
 	}
 	Out.Logf(ELogVerbosity::Display, TEXT("\t   PersistentGuid: %s"), *Linker->Summary.PersistentGuid.ToString());
 	Out.Logf(ELogVerbosity::Display, TEXT("\t      Generations:"));
diff --git a/Engine/Source/Editor/UnrealEd/Private/CookOnTheFlyServer.cpp b/Engine/Source/Editor/UnrealEd/Private/CookOnTheFlyServer.cpp
--- a/Engine/Source/Editor/UnrealEd/Private/CookOnTheFlyServer.cpp
+++ b/Engine/Source/Editor/UnrealEd/Private/CookOnTheFlyServer.cpp
@@ -6407,9 +6407,7 @@ void FSaveCookedPackageContext::FinishPlatform()
 			Info.Status = IPackageWriter::ECommitStatus::Error;
 		}
 		Info.PackageName = Package->GetFName();
-		PRAGMA_DISABLE_DEPRECATION_WARNINGS;
-		Info.PackageGuid = AssetPackageData ? AssetPackageData->PackageGuid : FGuid();
-		PRAGMA_ENABLE_DEPRECATION_WARNINGS;
+		Info.PackageHash = AssetPackageData ? AssetPackageData->GetPackageSavedHash() : FIoHash();
 		if (TargetDomainDependencies)
 		{
 			Info.Attachments.Add({ "Dependencies", TargetDomainDependencies });
@@ -6499,9 +6497,7 @@ void FSaveCookedPackageContext::FinishPlatform()
 				ImportedClasses.Sort(FNameLexicalLess());
 
 				OverrideAssetPackageData.Emplace();
-				PRAGMA_DISABLE_DEPRECATION_WARNINGS;
-				OverrideAssetPackageData->PackageGuid = GeneratedInfo->Guid;
-				PRAGMA_ENABLE_DEPRECATION_WARNINGS;
+				OverrideAssetPackageData->SetPackageSavedHash(GeneratedInfo->PackageHash);
 				OverrideAssetPackageData->ImportedClasses = ImportedClasses;
 			}
 		}
@@ -8597,7 +8593,7 @@ void UCookOnTheFlyServer::PopulateCookedPackages(TArrayView<const ITargetPlatfor
 				FPackageData* PackageData = PackageDatas->TryAddPackageDataByPackageName(Generator, false /* bRequireExists */);
 				if (PackageData && PackageData->FindOrAddPlatformData(TargetPlatform).IsCookAttempted())
 				{
-					for (const TPair<FName, FGuid>& GeneratedPair : Iter->Value.Generated)
+					for (const TPair<FName, FIoHash>& GeneratedPair : Iter->Value.Generated)
 					{
 						UpdateCookedPackage(GeneratedPair.Key, false /* bRequireExists */, true /* bIterativelyUnmodified */);
 					}
diff --git a/Engine/Source/Editor/UnrealEd/Private/Cooker/CookPackageData.cpp b/Engine/Source/Editor/UnrealEd/Private/Cooker/CookPackageData.cpp
--- a/Engine/Source/Editor/UnrealEd/Private/Cooker/CookPackageData.cpp
+++ b/Engine/Source/Editor/UnrealEd/Private/Cooker/CookPackageData.cpp
@@ -1688,21 +1688,21 @@ bool FGeneratorPackage::TryGenerateList(UObject* OwnerObject, FPackageDatas& Pac
 		PackageData->SetGeneratedOwner(this);
 		PackageData->SetWorkerAssignmentConstraint(FWorkerId::Local());
 
-		// Create the Guid from the GenerationHash and Dependencies
-		GeneratedInfo.CreateGuid();
+		// Create the Hash from the GenerationHash and Dependencies
+		GeneratedInfo.CreatePackageHash();
 
-		FGuid PreviousGuid;
-		if (PreviousGeneratedPackages.RemoveAndCopyValue(PackageFName, PreviousGuid) && !bHybridIterativeEnabled)
+		FIoHash PreviousHash;
+		if (PreviousGeneratedPackages.RemoveAndCopyValue(PackageFName, PreviousHash) && !bHybridIterativeEnabled)
 		{
 			bool bIterativelyUnmodified;
-			GeneratedInfo.IterativeCookValidateOrClear(*this, PlatformsToCook, PreviousGuid, bIterativelyUnmodified);
+			GeneratedInfo.IterativeCookValidateOrClear(*this, PlatformsToCook, PreviousHash, bIterativelyUnmodified);
 			++(bIterativelyUnmodified ? NumIterativeUnmodified : NumIterativeModified);
 		}
 	}
 	if (!PreviousGeneratedPackages.IsEmpty())
 	{
 		NumIterativeRemoved = PreviousGeneratedPackages.Num();
-		for (TPair<FName, FGuid>& Pair : PreviousGeneratedPackages)
+		for (TPair<FName, FIoHash>& Pair : PreviousGeneratedPackages)
 		{
 			for (const ITargetPlatform* TargetPlatform : PlatformsToCook)
 			{
@@ -1934,9 +1934,7 @@ UPackage* FGeneratorPackage::CreateGeneratedUPackage(FCookGenerationInfo& Genera
 	++DetailedCookStats::NumRequestedLoads;
 #endif
 	UPackage* GeneratedPackage = CreatePackage(GeneratedPackageName);
-	PRAGMA_DISABLE_DEPRECATION_WARNINGS;
-	GeneratedPackage->SetGuid(GeneratedInfo.Guid);
-	PRAGMA_ENABLE_DEPRECATION_WARNINGS;
+	GeneratedPackage->SetSavedHash(GeneratedInfo.PackageHash);
 	GeneratedPackage->SetPersistentGuid(InOwnerPackage->GetPersistentGuid());
 	GeneratedPackage->SetPackageFlags(PKG_CookGenerated);
 	GeneratedInfo.SetHasCreatedPackage(true);
@@ -2274,7 +2272,7 @@ EPollStatus FCookGenerationInfo::RefreshPackageObjects(FGeneratorPackage& Genera
 	return EPollStatus::Success;
 }
 
-void FCookGenerationInfo::CreateGuid()
+void FCookGenerationInfo::CreatePackageHash()
 {
 	FBlake3 Blake3;
 	Blake3.Update(&GenerationHash, sizeof(GenerationHash));
@@ -2284,21 +2282,29 @@ void FCookGenerationInfo::CreateGuid()
 		TOptional<FAssetPackageData> DependencyData = AssetRegistry.GetAssetPackageDataCopy(Dependency.AssetId.PackageName);
 		if (DependencyData)
 		{
-			PRAGMA_DISABLE_DEPRECATION_WARNINGS;
-			Blake3.Update(&DependencyData->PackageGuid, sizeof(DependencyData->PackageGuid));
-			PRAGMA_ENABLE_DEPRECATION_WARNINGS;
+			Blake3.Update(&DependencyData->GetPackageSavedHash().GetBytes(), sizeof(DependencyData->GetPackageSavedHash().GetBytes()));
 		}
 	}
-	FBlake3Hash GeneratedHash = Blake3.Finalize();
-	const uint32* HashInts = reinterpret_cast<const uint32*>(GeneratedHash.GetBytes());
-	Guid = FGuid(HashInts[0], HashInts[1], HashInts[2], HashInts[3]);
+	PackageHash = FIoHash(Blake3.Finalize());
+	// We store the PackageHash as a FIoHash, but UPackage and FAssetPackageData store it as a FGuid, which is smaller,
+	// so we have to remove any data which doesn't fit into FGuid. This can be removed when we remove the deprecated
+	// Guid storage on UPackage.
+	PRAGMA_DISABLE_DEPRECATION_WARNINGS;
+	constexpr int SizeDifference = sizeof(PackageHash) - sizeof(DeclVal<UPackage>().GetGuid());
+	if (SizeDifference > 0)
+	{
+		FMemory::Memset(((uint8*)&PackageHash.GetBytes()) + (sizeof(PackageHash.GetBytes()) - SizeDifference),
+			0, SizeDifference);
+	}
+	PRAGMA_ENABLE_DEPRECATION_WARNINGS;
 }
 
 void FCookGenerationInfo::IterativeCookValidateOrClear(FGeneratorPackage& Generator,
-	TConstArrayView<const ITargetPlatform*> RequestedPlatforms, const FGuid& PreviousGuid, bool& bOutIterativelyUnmodified)
+	TConstArrayView<const ITargetPlatform*> RequestedPlatforms, const FIoHash& PreviousPackageHash,
+	bool& bOutIterativelyUnmodified)
 {
 	UCookOnTheFlyServer& COTFS = Generator.GetOwner().GetPackageDatas().GetCookOnTheFlyServer();
-	bOutIterativelyUnmodified = PreviousGuid == this->Guid;
+	bOutIterativelyUnmodified = PreviousPackageHash == this->PackageHash;
 	if (bOutIterativelyUnmodified)
 	{
 		// If not directly modified, mark it as indirectly modified if any of its dependencies
diff --git a/Engine/Source/Editor/UnrealEd/Private/Cooker/CookPackageData.h b/Engine/Source/Editor/UnrealEd/Private/Cooker/CookPackageData.h
--- a/Engine/Source/Editor/UnrealEd/Private/Cooker/CookPackageData.h
+++ b/Engine/Source/Editor/UnrealEd/Private/Cooker/CookPackageData.h
@@ -13,6 +13,7 @@
 #include "Containers/UnrealString.h"
 #include "HAL/CriticalSection.h"
 #include "HAL/Platform.h"
+#include "IO/IoHash.h"
 #include "Misc/EnumClassFlags.h"
 #include "Misc/Optional.h"
 #include "Misc/ScopeRWLock.h"
@@ -878,8 +879,8 @@ public:
 		KeepReferencedPackages.Append(InKeepReferencedPackages);
 	}
 
-	/** Create the Guid for this generated package, based on dependencies and GenerationHash. */
-	void CreateGuid();
+	/** Create the hash for this generated package, based on dependencies and GenerationHash. */
+	void CreatePackageHash();
 
 	/**
 	 * If the package has iterative results from a previous cook that were not invalidated by dependency changes,
@@ -887,12 +888,13 @@ public:
 	 * incremental cooks handle invalidation by querying the TargetDomainDigest during the RequestCluster.
 	 */
 	void IterativeCookValidateOrClear(FGeneratorPackage& Generator,
-		TConstArrayView<const ITargetPlatform*> RequestedPlatforms, const FGuid& PreviousGuid, bool& bOutIterativelyUnmodified);
+		TConstArrayView<const ITargetPlatform*> RequestedPlatforms, const FIoHash& PreviousHash,
+		bool& bOutIterativelyUnmodified);
 
 	TConstArrayView<FAssetDependency> GetDependencies() const { return PackageDependencies; }
 
 public:
-	FGuid Guid;
+	FIoHash PackageHash;
 	FString RelativePath;
 	FString GeneratedRootPath;
 	FBlake3Hash GenerationHash;
@@ -975,7 +977,7 @@ public:
 
 	UPackage* GetOwnerPackage() const { return OwnerPackage.Get(); };
 	void SetOwnerPackage(UPackage* InPackage) { OwnerPackage = InPackage; }
-	void SetPreviousGeneratedPackages(TMap<FName, FGuid>&& Packages) { PreviousGeneratedPackages = MoveTemp(Packages); }
+	void SetPreviousGeneratedPackages(TMap<FName, FIoHash>&& Packages) { PreviousGeneratedPackages = MoveTemp(Packages); }
 
 	TConstArrayView<FName> GetExternalActorDependencies() { check(IsInitialized()); return ExternalActorDependencies; }
 	TArray<FName> ReleaseExternalActorDependencies() { TArray<FName> Result = MoveTemp(ExternalActorDependencies); return Result; }
@@ -992,7 +994,7 @@ private:
 	/** Recorded list of packages to generate from the splitter, and data we need about them */
 	TArray<FCookGenerationInfo> PackagesToGenerate;
 	TWeakObjectPtr<UPackage> OwnerPackage;
-	TMap<FName, FGuid> PreviousGeneratedPackages;
+	TMap<FName, FIoHash> PreviousGeneratedPackages;
 	TArray<FName> ExternalActorDependencies;
 
 	int32 NextPopulateIndex = 0;
diff --git a/Engine/Source/Editor/UnrealEd/Private/EditorDomain/EditorDomainUtils.cpp b/Engine/Source/Editor/UnrealEd/Private/EditorDomain/EditorDomainUtils.cpp
--- a/Engine/Source/Editor/UnrealEd/Private/EditorDomain/EditorDomainUtils.cpp
+++ b/Engine/Source/Editor/UnrealEd/Private/EditorDomain/EditorDomainUtils.cpp
@@ -24,6 +24,7 @@
 #include "HAL/IConsoleManager.h"
 #include "Hash/Blake3.h"
 #include "Interfaces/IPluginManager.h"
+#include "IO/IoHash.h"
 #include "Memory/SharedBuffer.h"
 #include "Misc/CommandLine.h"
 #include "Misc/CoreDelegates.h"
@@ -342,9 +343,7 @@ FPackageDigest CalculatePackageDigest(IAssetRegistry& AssetRegistry, FName Packa
 	Writer.Update(EditorDomainVersion, FCString::Strlen(EditorDomainVersion)*sizeof(EditorDomainVersion[0]));
 	uint8 EditorDomainSaveUnversioned = GetEditorDomainSaveUnversioned() ? 1 : 0;
 	Writer.Update(&EditorDomainSaveUnversioned, sizeof(EditorDomainSaveUnversioned));
-	PRAGMA_DISABLE_DEPRECATION_WARNINGS;
-	Writer.Update(&PackageData.PackageGuid, sizeof(PackageData.PackageGuid));
-	PRAGMA_ENABLE_DEPRECATION_WARNINGS;
+	Writer.Update(&PackageData.GetPackageSavedHash().GetBytes(), sizeof(PackageData.GetPackageSavedHash().GetBytes()));
 	Writer.Update(&GPackageFileUEVersion, sizeof(GPackageFileUEVersion));
 	Writer.Update(&GPackageFileLicenseeUEVersion, sizeof(GPackageFileLicenseeUEVersion));
 	TArray<int32> CustomVersionHandles;
@@ -2051,10 +2050,10 @@ void PutBulkDataList(FName PackageName, FSharedBuffer Buffer)
 	Owner.KeepAlive();
 }
 
-FIoHash GetPackageAndGuidHash(const FGuid& PackageGuid, const FGuid& BulkDataId)
+FIoHash GetPackageAndGuidHash(const FIoHash& PackageHash, const FGuid& BulkDataId)
 {
 	FBlake3 Builder;
-	Builder.Update(&PackageGuid, sizeof(PackageGuid));
+	Builder.Update(&PackageHash.GetBytes(), sizeof(PackageHash.GetBytes()));
 	Builder.Update(&BulkDataId, sizeof(BulkDataId));
 	return Builder.Finalize();
 }
@@ -2069,10 +2068,7 @@ void GetBulkDataPayloadId(FName PackageName, const FGuid& BulkDataId, UE::Derive
 	{
 		return;
 	}
-	PRAGMA_DISABLE_DEPRECATION_WARNINGS;
-	const FGuid& PackageGuid = PackageData->PackageGuid;
-	PRAGMA_ENABLE_DEPRECATION_WARNINGS;
-	FIoHash PackageAndGuidHash = GetPackageAndGuidHash(PackageGuid, BulkDataId);
+	FIoHash PackageAndGuidHash = GetPackageAndGuidHash(PackageData->GetPackageSavedHash(), BulkDataId);
 
 	using namespace UE::DerivedData;
 	ICache& Cache = GetCache();
@@ -2094,10 +2090,7 @@ void PutBulkDataPayloadId(FName PackageName, const FGuid& BulkDataId, FSharedBuf
 	{
 		return;
 	}
-	PRAGMA_DISABLE_DEPRECATION_WARNINGS;
-	const FGuid& PackageGuid = PackageData->PackageGuid;
-	PRAGMA_ENABLE_DEPRECATION_WARNINGS;
-	FIoHash PackageAndGuidHash = GetPackageAndGuidHash(PackageGuid, BulkDataId);
+	FIoHash PackageAndGuidHash = GetPackageAndGuidHash(PackageData->GetPackageSavedHash(), BulkDataId);
 
 	using namespace UE::DerivedData;
 	ICache& Cache = GetCache();
diff --git a/Engine/Source/Editor/UnrealEd/Private/PlayLevel.cpp b/Engine/Source/Editor/UnrealEd/Private/PlayLevel.cpp
--- a/Engine/Source/Editor/UnrealEd/Private/PlayLevel.cpp
+++ b/Engine/Source/Editor/UnrealEd/Private/PlayLevel.cpp
@@ -2307,9 +2307,7 @@ UWorld* UEditorEngine::CreatePIEWorldByDuplication(FWorldContext &WorldContext,
 	PlayWorldPackage->SetPackageFlags(PKG_PlayInEditor | PKG_NewlyCreated);
 	PlayWorldPackage->SetPIEInstanceID(WorldContext.PIEInstance);
 	PlayWorldPackage->SetLoadedPath(InPackage->GetLoadedPath());
-	PRAGMA_DISABLE_DEPRECATION_WARNINGS
-	PlayWorldPackage->SetGuid( InPackage->GetGuid() );
-	PRAGMA_ENABLE_DEPRECATION_WARNINGS
+	PlayWorldPackage->SetSavedHash( InPackage->GetSavedHash() );
 	PlayWorldPackage->MarkAsFullyLoaded();
 
 	// check(GPlayInEditorID == -1 || GPlayInEditorID == WorldContext.PIEInstance);
diff --git a/Engine/Source/Editor/UnrealEd/Public/Commandlets/AssetRegistryGenerator.h b/Engine/Source/Editor/UnrealEd/Public/Commandlets/AssetRegistryGenerator.h
--- a/Engine/Source/Editor/UnrealEd/Public/Commandlets/AssetRegistryGenerator.h
+++ b/Engine/Source/Editor/UnrealEd/Public/Commandlets/AssetRegistryGenerator.h
@@ -81,7 +81,7 @@ public:
 	/** Info about a GeneratorPackage (see ICookPackageSplitter) loaded from previous iterative cooks. */
 	struct FGeneratorPackageInfo
 	{
-		TMap<FName, FGuid> Generated;
+		TMap<FName, FIoHash> Generated;
 	};
 
 	enum EDifference
diff --git a/Engine/Source/Runtime/AssetRegistry/Private/AssetRegistry.cpp b/Engine/Source/Runtime/AssetRegistry/Private/AssetRegistry.cpp
--- a/Engine/Source/Runtime/AssetRegistry/Private/AssetRegistry.cpp
+++ b/Engine/Source/Runtime/AssetRegistry/Private/AssetRegistry.cpp
@@ -1246,10 +1246,10 @@ void FAssetRegistryImpl::ReadScriptPackages()
 				if (Package && FPackageName::IsScriptPackage(Package->GetName()))
 				{
 					FAssetPackageData* ScriptPackageData = State.CreateOrGetAssetPackageData(Package->GetFName());
-					// Get the guid off the script package, it is updated when script is changed so we need to refresh it every run
-					PRAGMA_DISABLE_DEPRECATION_WARNINGS
-					ScriptPackageData->PackageGuid = Package->GetGuid();
-					PRAGMA_ENABLE_DEPRECATION_WARNINGS
+#if WITH_EDITORONLY_DATA
+					// Get the hash off the script package, it is updated when script is changed so we need to refresh it every run
+					ScriptPackageData->SetPackageSavedHash(Package->GetSavedHash());
+#endif
 				}
 			}
 		}
@@ -5756,9 +5756,9 @@ void FAssetRegistryImpl::DependencyDataGathered(const double TickStartTime, TMul
 							if (Package)
 							{
 								FAssetPackageData* ScriptPackageData = State.CreateOrGetAssetPackageData(DependencyPackageName);
-								PRAGMA_DISABLE_DEPRECATION_WARNINGS
-								ScriptPackageData->PackageGuid = Package->GetGuid();
-								PRAGMA_ENABLE_DEPRECATION_WARNINGS
+#if WITH_EDITORONLY_DATA
+								ScriptPackageData->SetPackageSavedHash(Package->GetSavedHash());
+#endif
 							}
 						}
 					}
diff --git a/Engine/Source/Runtime/AssetRegistry/Private/AssetRegistryState.cpp b/Engine/Source/Runtime/AssetRegistry/Private/AssetRegistryState.cpp
--- a/Engine/Source/Runtime/AssetRegistry/Private/AssetRegistryState.cpp
+++ b/Engine/Source/Runtime/AssetRegistry/Private/AssetRegistryState.cpp
@@ -3063,9 +3063,7 @@ void FAssetRegistryState::Dump(const TArray<FString>& Arguments, TArray<FString>
 			PageBuffer.Append(TEXT("	"));
 			Key.AppendString(PageBuffer);
 			PageBuffer.Append(TEXT(" : "));
-			PRAGMA_DISABLE_DEPRECATION_WARNINGS
-			PageBuffer.Append(PackageData->PackageGuid.ToString());
-			PRAGMA_ENABLE_DEPRECATION_WARNINGS
+			PageBuffer << PackageData->GetPackageSavedHash();
 			PageBuffer.Appendf(TEXT(" : %d bytes"), PackageData->DiskSize);
 			AddLine();
 		}
diff --git a/Engine/Source/Runtime/AssetRegistry/Private/PackageReader.cpp b/Engine/Source/Runtime/AssetRegistry/Private/PackageReader.cpp
--- a/Engine/Source/Runtime/AssetRegistry/Private/PackageReader.cpp
+++ b/Engine/Source/Runtime/AssetRegistry/Private/PackageReader.cpp
@@ -673,9 +673,9 @@ bool FPackageReader::ReadDependencyData(FPackageDependencyData& OutDependencyDat
 		OutDependencyData.bHasPackageData = true;
 		FAssetPackageData& PackageData = OutDependencyData.PackageData;
 		PackageData.DiskSize = PackageFileSize;
-		PRAGMA_DISABLE_DEPRECATION_WARNINGS;
-		PackageData.PackageGuid = PackageFileSummary.Guid;
-		PRAGMA_ENABLE_DEPRECATION_WARNINGS;
+#if WITH_EDITORONLY_DATA
+		PackageData.SetPackageSavedHash(PackageFileSummary.GetSavedHash());
+#endif
 		PackageData.SetCustomVersions(PackageFileSummary.GetCustomVersionContainer().GetAllVersions());
 		PackageData.FileVersionUE = PackageFileSummary.GetFileVersionUE();
 		PackageData.FileVersionLicenseeUE = PackageFileSummary.GetFileVersionLicenseeUE();
diff --git a/Engine/Source/Runtime/Core/Public/Serialization/PackageWriter.h b/Engine/Source/Runtime/Core/Public/Serialization/PackageWriter.h
--- a/Engine/Source/Runtime/Core/Public/Serialization/PackageWriter.h
+++ b/Engine/Source/Runtime/Core/Public/Serialization/PackageWriter.h
@@ -7,6 +7,7 @@
 #include "Async/Future.h"
 #include "Containers/StringView.h"
 #include "IO/IoDispatcher.h"
+#include "IO/IoHash.h"
 #include "Misc/AssertionMacros.h"
 #include "Misc/DateTime.h"
 #include "Misc/EnumClassFlags.h"
@@ -94,10 +95,20 @@ public:
 	struct FCommitPackageInfo
 	{
 		FName PackageName;
+		FIoHash PackageHash;
+		UE_DEPRECATED(5.4, "Use PackageHash instead")
 		FGuid PackageGuid;
 		TArray<FCommitAttachmentInfo> Attachments;
 		ECommitStatus Status;
 		EWriteOptions WriteOptions;
+
+		PRAGMA_DISABLE_DEPRECATION_WARNINGS;
+		inline FCommitPackageInfo() = default;
+		inline FCommitPackageInfo(const FCommitPackageInfo&) = default;
+		inline FCommitPackageInfo(FCommitPackageInfo&&) = default;
+		inline FCommitPackageInfo& operator=(FCommitPackageInfo&) = default;
+		inline FCommitPackageInfo& operator=(FCommitPackageInfo&&) = default;
+		PRAGMA_ENABLE_DEPRECATION_WARNINGS;
 	};
 
 	/** Finalize a package started with BeginPackage()
diff --git a/Engine/Source/Runtime/CoreUObject/Private/AssetRegistry/AssetData.cpp b/Engine/Source/Runtime/CoreUObject/Private/AssetRegistry/AssetData.cpp
--- a/Engine/Source/Runtime/CoreUObject/Private/AssetRegistry/AssetData.cpp
+++ b/Engine/Source/Runtime/CoreUObject/Private/AssetRegistry/AssetData.cpp
@@ -8,6 +8,7 @@
 #include "AssetRegistry/ARFilter.h"
 #include "Containers/Set.h"
 #include "HAL/CriticalSection.h"
+#include "HAL/PlatformMath.h"
 #include "Misc/AsciiSet.h"
 #include "Misc/PathViews.h"
 #include "Misc/ScopeRWLock.h"
@@ -945,17 +946,35 @@ void FAssetPackageData::SerializeForCacheInternal(FArchive& Ar, FAssetPackageDat
 	}
 }
 
-COREUOBJECT_API void FAssetPackageData::SerializeForCache(FArchive& Ar)
+void FAssetPackageData::SerializeForCache(FArchive& Ar)
 {
 	// Calling with hard-coded version and using force-inline on SerializeForCacheInternal eliminates the cost of its if-statements
 	SerializeForCacheInternal(Ar, *this, FAssetRegistryVersion::LatestVersion);
 }
 
-COREUOBJECT_API void FAssetPackageData::SerializeForCacheOldVersion(FArchive& Ar, FAssetRegistryVersion::Type Version)
+void FAssetPackageData::SerializeForCacheOldVersion(FArchive& Ar, FAssetRegistryVersion::Type Version)
 {
 	SerializeForCacheInternal(Ar, *this, Version);
 }
 
+FIoHash FAssetPackageData::GetPackageSavedHash() const
+{
+	FIoHash Result;
+	PRAGMA_DISABLE_DEPRECATION_WARNINGS;
+	FMemory::Memcpy(&Result.GetBytes(), &this->PackageGuid,
+		FMath::Min(sizeof(Result.GetBytes()), sizeof(PackageGuid)));
+	PRAGMA_ENABLE_DEPRECATION_WARNINGS;
+	return Result;
+}
+
+void FAssetPackageData::SetPackageSavedHash(const FIoHash& InHash)
+{
+	PRAGMA_DISABLE_DEPRECATION_WARNINGS;
+	PackageGuid = FGuid();
+	FMemory::Memcpy(&PackageGuid, &InHash.GetBytes(), FMath::Min(sizeof(PackageGuid), sizeof(InHash.GetBytes())));
+	PRAGMA_ENABLE_DEPRECATION_WARNINGS;
+}
+
 void FARFilter::PostSerialize(const FArchive& Ar)
 {
 	PRAGMA_DISABLE_DEPRECATION_WARNINGS
diff --git a/Engine/Source/Runtime/CoreUObject/Private/UObject/LinkerDiff.cpp b/Engine/Source/Runtime/CoreUObject/Private/UObject/LinkerDiff.cpp
--- a/Engine/Source/Runtime/CoreUObject/Private/UObject/LinkerDiff.cpp
+++ b/Engine/Source/Runtime/CoreUObject/Private/UObject/LinkerDiff.cpp
@@ -67,6 +67,21 @@ void FLinkerDiff::PrintDiff(FOutputDevice& Ar)
 		Builder.Append(LexToString(RHSRef.Member));	\
 		Diffs.Add(Builder.ToString());				\
 	}
+#define COMPARE_CALCULATED_VALUE(LHSValue, RHSValue, ValueNameText) \
+	{												\
+		auto LocalLHS = LHSValue;					\
+		auto LocalRHS = RHSValue;					\
+		if (LocalLHS != LocalRHS)					\
+		{											\
+			TStringBuilder<256> Builder;			\
+			Builder.Append(ValueNameText);			\
+			Builder.Append(TEXT(": "));				\
+			Builder.Append(LexToString(LocalLHS));	\
+			Builder.Append(TEXT(" vs "));			\
+			Builder.Append(LexToString(LocalRHS));	\
+			Diffs.Add(Builder.ToString());			\
+		}											\
+	}
 
 
 void FLinkerDiff::Generate(FLinker* LHSLinker, FLinker* RHSLinker)
@@ -119,10 +134,8 @@ void FLinkerDiff::GenerateSummaryDiff(FLinker* LHSLinker, FLinker* RHSLinker)
 	//COMPARE_MEMBER(LHSSummary, RHSSummary, SearchableNamesOffset);
 	//COMPARE_MEMBER(LHSSummary, RHSSummary, ThumbnailTableOffset);
 
-	PRAGMA_DISABLE_DEPRECATION_WARNINGS
-	COMPARE_MEMBER(LHSSummary, RHSSummary, Guid);
-	PRAGMA_ENABLE_DEPRECATION_WARNINGS
 #if WITH_EDITORONLY_DATA
+	COMPARE_CALCULATED_VALUE(LHSSummary.GetSavedHash(), RHSSummary.GetSavedHash(), TEXT("SavedHash"));
 	COMPARE_MEMBER(LHSSummary, RHSSummary, PersistentGuid);
 #endif
 	// Generations
diff --git a/Engine/Source/Runtime/CoreUObject/Private/UObject/LinkerLoad.cpp b/Engine/Source/Runtime/CoreUObject/Private/UObject/LinkerLoad.cpp
--- a/Engine/Source/Runtime/CoreUObject/Private/UObject/LinkerLoad.cpp
+++ b/Engine/Source/Runtime/CoreUObject/Private/UObject/LinkerLoad.cpp
@@ -1686,12 +1686,9 @@ FLinkerLoad::ELinkerStatus FLinkerLoad::UpdateFromPackageFileSummary()
 		// Propagate package file size
 		LinkerRootPackage->SetFileSize(Loader ? Loader->TotalSize() : 0);
 
-		// Propagate package Guids
-		PRAGMA_DISABLE_DEPRECATION_WARNINGS
-		LinkerRootPackage->SetGuid( Summary.Guid );
-		PRAGMA_ENABLE_DEPRECATION_WARNINGS
-
+		// Propagate package hashes
 #if WITH_EDITORONLY_DATA
+		LinkerRootPackage->SetSavedHash(Summary.GetSavedHash());
 		LinkerRootPackage->SetPersistentGuid( Summary.PersistentGuid );
 #endif
 
diff --git a/Engine/Source/Runtime/CoreUObject/Private/UObject/Package.cpp b/Engine/Source/Runtime/CoreUObject/Private/UObject/Package.cpp
--- a/Engine/Source/Runtime/CoreUObject/Private/UObject/Package.cpp
+++ b/Engine/Source/Runtime/CoreUObject/Private/UObject/Package.cpp
@@ -4,6 +4,7 @@
 
 #include "AssetRegistry/AssetData.h"
 #include "HAL/FileManager.h"
+#include "HAL/PlatformMath.h"
 #include "Misc/AssetRegistryInterface.h"
 #include "Misc/ITransaction.h"
 #include "Misc/PackageName.h"
@@ -392,6 +393,30 @@ void UPackage::SetLoadedByEditorPropertiesOnly(bool bIsEditorOnly, bool bRecursi
 		FixupPackageEditorOnlyFlag(GetFName(), bRecursive);
 	}
 }
+
+FIoHash UPackage::GetSavedHash()
+{
+#if UE_STRIP_DEPRECATED_PROPERTIES
+	return FIoHash();
+#else
+	FIoHash Result;
+	PRAGMA_DISABLE_DEPRECATION_WARNINGS;
+	FMemory::Memcpy(&Result.GetBytes(), &Guid, FMath::Min(sizeof(Result.GetBytes()), sizeof(Guid)));
+	PRAGMA_ENABLE_DEPRECATION_WARNINGS;
+	return Result;
+#endif
+}
+
+void UPackage::SetSavedHash(const FIoHash& InSavedHash)
+{
+#if !UE_STRIP_DEPRECATED_PROPERTIES
+	PRAGMA_DISABLE_DEPRECATION_WARNINGS;
+	Guid = FGuid();
+	FMemory::Memcpy(&Guid, &InSavedHash.GetBytes(), FMath::Min(sizeof(Guid), sizeof(InSavedHash.GetBytes())));
+	PRAGMA_ENABLE_DEPRECATION_WARNINGS;
+#endif
+}
+
 #endif
 
 #if WITH_EDITORONLY_DATA
diff --git a/Engine/Source/Runtime/CoreUObject/Private/UObject/PackageFileSummary.cpp b/Engine/Source/Runtime/CoreUObject/Private/UObject/PackageFileSummary.cpp
--- a/Engine/Source/Runtime/CoreUObject/Private/UObject/PackageFileSummary.cpp
+++ b/Engine/Source/Runtime/CoreUObject/Private/UObject/PackageFileSummary.cpp
@@ -1,11 +1,13 @@
 // Copyright Epic Games, Inc. All Rights Reserved.
 
 #include "UObject/PackageFileSummary.h"
+
+#include "HAL/PlatformMath.h"
 #include "Misc/Compression.h"
-#include "UObject/Linker.h"
+#include "Runtime/Launch/Resources/Version.h"
 #include "Serialization/StructuredArchive.h"
+#include "UObject/Linker.h"
 #include "UObject/UObjectGlobals.h"
-#include "Runtime/Launch/Resources/Version.h"
 
 FPackageFileSummary::FPackageFileSummary()
 {
@@ -270,9 +272,10 @@ void operator<<(FStructuredArchive::FSlot Slot, FPackageFileSummary& Sum)
 			else
 			{
 				// By assigning the current package guid, we maintain a stable persistent guid, so we can reference this package even if it wasn't resaved.
-				PRAGMA_DISABLE_DEPRECATION_WARNINGS
-				Sum.PersistentGuid = Sum.Guid;
-				PRAGMA_ENABLE_DEPRECATION_WARNINGS
+				FIoHash SavedHash = Sum.GetSavedHash();
+				Sum.PersistentGuid = FGuid();
+				FMemory::Memcpy(&Sum.PersistentGuid, &SavedHash.GetBytes(),
+					FMath::Min(sizeof(Sum.PersistentGuid), sizeof(SavedHash.GetBytes())));
 			}
 
 			// The owner persistent guid was added in VER_UE4_ADDED_PACKAGE_OWNER but removed in the next version VER_UE4_NON_OUTER_PACKAGE_IMPORT
@@ -447,6 +450,25 @@ void operator<<(FStructuredArchive::FSlot Slot, FPackageFileSummary& Sum)
 	}
 }
 
+#if WITH_EDITORONLY_DATA
+FIoHash FPackageFileSummary::GetSavedHash() const
+{
+	FIoHash Result;
+	PRAGMA_DISABLE_DEPRECATION_WARNINGS;
+	FMemory::Memcpy(&Result.GetBytes(), &Guid, FMath::Min(sizeof(Result.GetBytes()), sizeof(Guid)));
+	PRAGMA_ENABLE_DEPRECATION_WARNINGS;
+	return Result;
+}
+
+void FPackageFileSummary::SetSavedHash(const FIoHash& InSavedHash)
+{
+	PRAGMA_DISABLE_DEPRECATION_WARNINGS;
+	Guid = FGuid();
+	FMemory::Memcpy(&Guid, &InSavedHash.GetBytes(), FMath::Min(sizeof(Guid), sizeof(InSavedHash.GetBytes())));
+	PRAGMA_ENABLE_DEPRECATION_WARNINGS;
+}
+#endif
+
 FArchive& operator<<( FArchive& Ar, FPackageFileSummary& Sum )
 {
 	FStructuredArchiveFromArchive(Ar).GetSlot() << Sum;
diff --git a/Engine/Source/Runtime/CoreUObject/Private/UObject/SavePackage2.cpp b/Engine/Source/Runtime/CoreUObject/Private/UObject/SavePackage2.cpp
--- a/Engine/Source/Runtime/CoreUObject/Private/UObject/SavePackage2.cpp
+++ b/Engine/Source/Runtime/CoreUObject/Private/UObject/SavePackage2.cpp
@@ -1109,19 +1109,30 @@ ESavePackageResult BuildLinker(FSaveContext& SaveContext)
 		SaveContext.UpdatePackageLinkerVersions();
 	}
 
+#if WITH_EDITORONLY_DATA
 	PRAGMA_DISABLE_DEPRECATION_WARNINGS;
-	if (SaveContext.GetSaveArgs().OutputPackageGuid)
+	const FGuid* OutputPackageGuid = SaveContext.GetSaveArgs().OutputPackageGuid.GetPtrOrNull();
+	PRAGMA_ENABLE_DEPRECATION_WARNINGS;
+	if (OutputPackageGuid)
 	{
-		SaveContext.GetLinker()->Summary.Guid = *SaveContext.GetSaveArgs().OutputPackageGuid;
-		
+		FIoHash SavedHash;
+		FMemory::Memcpy(&SavedHash.GetBytes(), OutputPackageGuid,
+			FMath::Min(sizeof(SavedHash.GetBytes()), sizeof(*OutputPackageGuid)));
+		SaveContext.GetLinker()->Summary.SetSavedHash(SavedHash);
 	}
 	else
 	{
-		SaveContext.GetLinker()->Summary.Guid = SaveContext.IsKeepGuid() ?
-			SaveContext.GetPackage()->GetGuid() : SaveContext.GetPackage()->MakeNewGuid();
+		if (!SaveContext.IsKeepGuid())
+		{
+			// TODO: Change this to instead be calculated later after the save from a hash of the saved data.
+			// Remove IsKeepGuid and OutputPackageGuid when that occurs.
+			FGuid Guid = FGuid::NewGuid();
+			FIoHash SavedHash;
+			FMemory::Memcpy(&SavedHash.GetBytes(), &Guid, FMath::Min(sizeof(SavedHash.GetBytes()), sizeof(Guid)));
+			SaveContext.GetPackage()->SetSavedHash(SavedHash);
+		}
+		SaveContext.GetLinker()->Summary.SetSavedHash(SaveContext.GetPackage()->GetSavedHash());
 	}
-	PRAGMA_ENABLE_DEPRECATION_WARNINGS;
-#if WITH_EDITORONLY_DATA
 	SaveContext.GetLinker()->Summary.PersistentGuid = SaveContext.GetPackage()->GetPersistentGuid();
 #endif
 	SaveContext.GetLinker()->Summary.Generations = TArray<FGenerationInfo>{ FGenerationInfo(0, 0) };
diff --git a/Engine/Source/Runtime/CoreUObject/Private/UObject/UObjectGlobals.cpp b/Engine/Source/Runtime/CoreUObject/Private/UObject/UObjectGlobals.cpp
--- a/Engine/Source/Runtime/CoreUObject/Private/UObject/UObjectGlobals.cpp
+++ b/Engine/Source/Runtime/CoreUObject/Private/UObject/UObjectGlobals.cpp
@@ -6206,13 +6206,15 @@ namespace UECodeGen_Private
 #endif
 
 		NewPackage->SetPackageFlags(Params.PackageFlags);
-		PRAGMA_DISABLE_DEPRECATION_WARNINGS
-		NewPackage->SetGuid(FGuid(Params.BodyCRC, Params.DeclarationsCRC, 0u, 0u));
-		PRAGMA_ENABLE_DEPRECATION_WARNINGS
 #if WITH_EDITORONLY_DATA
-		// Store the CRC in UPackage::PersistentGuid since UPackage::Guid will be removed.
 		// Replace the PersistentGuid generated from UPackage::PostInitProperties() that changes every time.
-		NewPackage->SetPersistentGuid(FGuid(Params.BodyCRC, Params.DeclarationsCRC, 0u, 0u));
+		FGuid DeterministicGuid(Params.BodyCRC, Params.DeclarationsCRC, 0u, 0u);
+		NewPackage->SetPersistentGuid(DeterministicGuid);
+		// Set the initial saved hash to a value based on the CRCs; this is needed for script packages.
+		FIoHash SavedHash;
+		FMemory::Memcpy(&SavedHash.GetBytes(), &DeterministicGuid,
+			FMath::Min(sizeof(&SavedHash.GetBytes()), sizeof(DeterministicGuid)));
+		NewPackage->SetSavedHash(SavedHash);
 #endif
 
 #if WITH_RELOAD
diff --git a/Engine/Source/Runtime/CoreUObject/Public/AssetRegistry/AssetData.h b/Engine/Source/Runtime/CoreUObject/Public/AssetRegistry/AssetData.h
--- a/Engine/Source/Runtime/CoreUObject/Public/AssetRegistry/AssetData.h
+++ b/Engine/Source/Runtime/CoreUObject/Public/AssetRegistry/AssetData.h
@@ -903,7 +903,7 @@ class FAssetPackageData
 {
 public:
 	/** Guid of the source package, uniquely identifies an asset package */
-	UE_DEPRECATED(4.27, "UPackage::Guid has not been used by the engine for a long time and FAssetPackageData::PackageGuid will be removed.")
+	UE_DEPRECATED(5.4, "Use GetPackageSavedHash instead.")
 	FGuid PackageGuid;
 
 	/** MD5 of the cooked package on disk, for tracking nondeterministic changes */
@@ -976,7 +976,7 @@ public:
 	void SetHasVirtualizedPayloads(bool bValue) { Flags = (Flags & ~FLAG_HAS_VIRTUALIZED_PAYLOADS) | (bValue ? FLAG_HAS_VIRTUALIZED_PAYLOADS : 0); }
 
 	/**
-	 * Serialize as part of the registry cache. This is not meant to be serialized as part of a package so  it does not handle versions normally
+	 * Serialize as part of the registry cache. This is not meant to be serialized as part of a package so it does not handle versions normally
 	 * To version this data change FAssetRegistryVersion
 	 */
 	COREUOBJECT_API void SerializeForCache(FArchive& Ar);
@@ -988,6 +988,10 @@ public:
 		return ImportedClasses.GetAllocatedSize();
 	}
 
+	/** Hash of the package's .uasset/.umap file when it was last saved by the editor. */
+	COREUOBJECT_API FIoHash GetPackageSavedHash() const;
+	COREUOBJECT_API void SetPackageSavedHash(const FIoHash& InHash);
+
 private:
 	FORCEINLINE void SerializeForCacheInternal(FArchive& Ar, FAssetPackageData& PackageData, FAssetRegistryVersion::Type Version);
 
diff --git a/Engine/Source/Runtime/CoreUObject/Public/UObject/Package.h b/Engine/Source/Runtime/CoreUObject/Public/UObject/Package.h
--- a/Engine/Source/Runtime/CoreUObject/Public/UObject/Package.h
+++ b/Engine/Source/Runtime/CoreUObject/Public/UObject/Package.h
@@ -28,6 +28,10 @@
 #include "UObject/ObjectVersion.h"
 #include "UObject/UObjectGlobals.h"
 
+#if WITH_EDITORONLY_DATA
+#include "IO/IoHash.h"
+#endif
+
 class Error;
 class FArchive;
 class FLinkerLoad;
@@ -259,13 +263,13 @@ private:
 
 	/** Chunk IDs for the streaming install chunks this package will be placed in.  Empty for no chunk. Used during cooking. */
 	TArray<int32> ChunkIDs;
-#endif
 
 #if !UE_STRIP_DEPRECATED_PROPERTIES
 	/** GUID of package if it was loaded from disk. Changes at every save. */
-	UE_DEPRECATED(4.27, "UPackage::Guid has not been used by the engine for a long time and it will be removed.")
+	UE_DEPRECATED(5.4, "Use GetSavedHash/SetSavedHash instead.")
 	FGuid Guid;
 #endif
+#endif // WITH_EDITORONLY_DATA
 
 	/** Package Flags */
 	uint32 PackageFlagsPrivate;
@@ -776,6 +780,10 @@ public:
 	{
 		PersistentGuid = NewPersistentGuid;
 	}
+
+	/** Hash of the package's .uasset/.umap file when it was last saved by the editor. */
+	COREUOBJECT_API FIoHash GetSavedHash();
+	COREUOBJECT_API void SetSavedHash(const FIoHash& InSavedHash);
 #endif
 
 #if WITH_RELOAD
@@ -803,8 +811,8 @@ public:
 		WorldTileInfo = MoveTemp(InWorldTileInfo);
 	}
 
-	/** returns our Guid */
-	UE_DEPRECATED(4.27, "UPackage::Guid has not been used by the engine for a long time and GetGuid will be removed.")
+#if WITH_EDITORONLY_DATA
+	UE_DEPRECATED(5.4, "Use GetSavedHash/SetSavedHash instead.")
 	FORCEINLINE FGuid GetGuid() const
 	{
 #if !UE_STRIP_DEPRECATED_PROPERTIES
@@ -812,9 +820,10 @@ public:
 #else
 		return FGuid();
 #endif
+	
 	}
 	/** makes our a new fresh Guid */
-	UE_DEPRECATED(4.27, "UPackage::Guid has not been used by the engine for a long time and MakeNewGuid will be removed.")
+	UE_DEPRECATED(5.4, "Use GetSavedHash/SetSavedHash instead.")
 	FORCEINLINE FGuid MakeNewGuid()
 	{
 #if !UE_STRIP_DEPRECATED_PROPERTIES
@@ -825,13 +834,14 @@ public:
 #endif
 	}
 	/** sets a specific Guid */
-	UE_DEPRECATED(4.27, "UPackage::Guid has not been used by the engine for a long time and SetGuid will be removed.")
+	UE_DEPRECATED(5.4, "Use GetSavedHash/SetSavedHash instead.")
 	FORCEINLINE void SetGuid(FGuid NewGuid)
 	{
 #if !UE_STRIP_DEPRECATED_PROPERTIES
 		Guid = NewGuid;
 #endif
 	}
+#endif // WITH_EDITORONLY_DATA
 
 	/** returns our FileSize */
 	FORCEINLINE int64 GetFileSize() const
diff --git a/Engine/Source/Runtime/CoreUObject/Public/UObject/PackageFileSummary.h b/Engine/Source/Runtime/CoreUObject/Public/UObject/PackageFileSummary.h
--- a/Engine/Source/Runtime/CoreUObject/Public/UObject/PackageFileSummary.h
+++ b/Engine/Source/Runtime/CoreUObject/Public/UObject/PackageFileSummary.h
@@ -12,6 +12,10 @@
 #include "Serialization/StructuredArchive.h"
 #include "UObject/ObjectVersion.h"
 
+#if WITH_EDITORONLY_DATA
+#include "IO/IoHash.h"
+#endif
+
 class FArchive;
 struct FCompressedChunk;
 
@@ -162,10 +166,7 @@ public:
 	*/
 	int32		ThumbnailTableOffset;
 
-	/**
-	* Current id for this package
-	*/
-	UE_DEPRECATED(4.27, "UPackage::Guid has not been used by the engine for a long time and FPackageFileSummary::Guid will be removed.")
+	UE_DEPRECATED(5.4, "Use GetSavedHash/SetSavedHash instead.")
 	FGuid	Guid;
 
 #if WITH_EDITORONLY_DATA
@@ -331,6 +332,12 @@ public:
 	/** Set the summary package flags while stripping out temporary flags (i.e. NewlyCreated, IsSaving) */
 	void SetPackageFlags(uint32 InPackageFlags);
 
+#if WITH_EDITORONLY_DATA
+	/** Hash of the package's .uasset/.umap file when it was last saved by the editor. */
+	COREUOBJECT_API FIoHash GetSavedHash() const;
+	COREUOBJECT_API void SetSavedHash(const FIoHash& InSavedHash);
+#endif
+
 	/** I/O functions */
 	friend COREUOBJECT_API FArchive& operator<<(FArchive& Ar, FPackageFileSummary& Sum);
 	friend COREUOBJECT_API void operator<<(FStructuredArchive::FSlot Slot, FPackageFileSummary& Sum);
diff --git a/Engine/Source/Runtime/Engine/Classes/Engine/Engine.h b/Engine/Source/Runtime/Engine/Classes/Engine/Engine.h
--- a/Engine/Source/Runtime/Engine/Classes/Engine/Engine.h
+++ b/Engine/Source/Runtime/Engine/Classes/Engine/Engine.h
@@ -2600,12 +2600,6 @@ public:
 		return INDEX_NONE;
 	}
 
-	/** Looks up the GUID of a package on disk. The package must NOT be in the auto-download cache.
-	 * This may require loading the header of the package in question and is therefore slow.
-	 */
-	UE_DEPRECATED(4.27, "UPackage::Guid has not been used by the engine for a long time and UEngine::GetPackageGuid will be removed.")
-	static ENGINE_API FGuid GetPackageGuid(FName PackageName, bool bForPIE);
-
 	static ENGINE_API void PreGarbageCollect();
 
 	/**
diff --git a/Engine/Source/Runtime/Engine/Private/UnrealEngine.cpp b/Engine/Source/Runtime/Engine/Private/UnrealEngine.cpp
--- a/Engine/Source/Runtime/Engine/Private/UnrealEngine.cpp
+++ b/Engine/Source/Runtime/Engine/Private/UnrealEngine.cpp
@@ -11376,49 +11376,6 @@ void UEngine::AddTextureStreamingLoc(FVector InLoc, float BoostFactor, bool bOve
 	IStreamingManager::Get().AddViewLocation(InLoc, BoostFactor, bOverrideLocation, OverrideDuration);
 }
 
-/** Looks up the GUID of a package on disk. The package must NOT be in the autodownload cache.
-* This may require loading the header of the package in question and is therefore slow.
-*/
-FGuid UEngine::GetPackageGuid(FName PackageName, bool bForPIE)
-{
-	FGuid Result(0,0,0,0);
-	// There is no package guid support when using the I/O dispatcher
-	if (FIoDispatcher::IsInitialized() && FIoDispatcher::Get().DoesChunkExist(CreatePackageDataChunkId(FPackageId::FromName(PackageName))))
-	{
-		return Result;
-	}
-	FPackagePath PackagePath;
-	if (!FPackagePath::TryFromMountedName(PackageName.ToString(), PackagePath))
-	{
-		return Result;
-	}
-
-	uint32 LoadFlags = LOAD_NoWarn | LOAD_NoVerify;
-	if (bForPIE)
-	{
-		LoadFlags |= LOAD_PackageForPIE;
-	}
-	UPackage* PackageToReset = nullptr;
-
-	FLinkerLoad* Linker = LoadPackageLinker(nullptr, PackagePath, LoadFlags, nullptr, nullptr, [&PackageToReset, &Result](FLinkerLoad* InLinker)
-	{
-		if (InLinker != nullptr && InLinker->LinkerRoot != nullptr)
-		{
-			PRAGMA_DISABLE_DEPRECATION_WARNINGS
-			Result = InLinker->LinkerRoot->GetGuid();
-			PRAGMA_ENABLE_DEPRECATION_WARNINGS
-			PackageToReset = InLinker->LinkerRoot;
-		}
-	});
-
-	if (PackageToReset)
-	{
-		ResetLoaders(PackageToReset);
-	}
-
-	return Result;
-}
-
 #if !(UE_BUILD_SHIPPING || UE_BUILD_TEST)
 // Critical section used in AddOnScreenDebugMessage to handle concurrent insert.
 static FCriticalSection GOnScreenMessageCS;
diff --git a/Engine/Source/Runtime/Engine/Private/World.cpp b/Engine/Source/Runtime/Engine/Private/World.cpp
--- a/Engine/Source/Runtime/Engine/Private/World.cpp
+++ b/Engine/Source/Runtime/Engine/Private/World.cpp
@@ -3919,9 +3919,7 @@ UWorld* UWorld::DuplicateWorldForPIE(const FString& PackageName, UWorld* OwningW
 	PIELevelPackage->SetPackageFlags(PKG_PlayInEditor | PKG_NewlyCreated);
 	PIELevelPackage->SetPIEInstanceID(PIEInstanceID);
 	PIELevelPackage->SetLoadedPath(EditorLevelPackage->GetLoadedPath());
-	PRAGMA_DISABLE_DEPRECATION_WARNINGS
-	PIELevelPackage->SetGuid( EditorLevelPackage->GetGuid() );
-	PRAGMA_ENABLE_DEPRECATION_WARNINGS
+	PIELevelPackage->SetSavedHash( EditorLevelPackage->GetSavedHash() );
 	PIELevelPackage->MarkAsFullyLoaded();
 
 	ULevel::StreamedLevelsOwningWorld.Add(PIELevelPackage->GetFName(), OwningWorld);
@@ -8835,9 +8833,9 @@ static ULevel* DuplicateLevelWithPrefix(ULevel* InLevel, int32 InstanceID )
 	NewPackage->SetPackageFlags( PKG_PlayInEditor );
 	NewPackage->SetPIEInstanceID(InstanceID);
 	NewPackage->SetLoadedPath(OriginalPackage->GetLoadedPath());
-	PRAGMA_DISABLE_DEPRECATION_WARNINGS
-	NewPackage->SetGuid( OriginalPackage->GetGuid() );
-	PRAGMA_ENABLE_DEPRECATION_WARNINGS
+#if WITH_EDITORONLY_DATA
+	NewPackage->SetSavedHash( OriginalPackage->GetSavedHash() );
+#endif
 	NewPackage->MarkAsFullyLoaded();
 
 	FSoftObjectPath::AddPIEPackageName(NewPackage->GetFName());
