Editor: Support filtering out non-actor hits from editor object placement raycasts by adding an appropriate typed element interface function and using it in level editor.

diff --git a/Engine/Source/Editor/UnrealEd/Private/Editor/ActorPositioning.cpp b/Engine/Source/Editor/UnrealEd/Private/Editor/ActorPositioning.cpp
--- a/Engine/Source/Editor/UnrealEd/Private/Editor/ActorPositioning.cpp
+++ b/Engine/Source/Editor/UnrealEd/Private/Editor/ActorPositioning.cpp
@@ -19,183 +19,70 @@
 #include "LandscapeHeightfieldCollisionComponent.h"
 #include "LandscapeComponent.h"
 #include "Editor/EditorPerProjectUserSettings.h"
+#include "Editor/ObjectPositioning.h"
 
 namespace ActorPositioningLocals
 {
-	bool bAllowNonPrimitiveComponentHits = true;
-	static FAutoConsoleVariableRef CVarAllowNonPrimitiveComponentHits(
-		TEXT("PlacementMode.AllowNonPrimitiveComponentHits"),
-		bAllowNonPrimitiveComponentHits,
-		TEXT("When raycasting the world in placement mode, allow hits of physics objects that are not tied to a UPrimitiveComponent (to work with non-actor workflows)."));
-}
-
-FActorPositionTraceResult FActorPositioning::TraceWorldForPositionWithDefault(const FViewportCursorLocation& Cursor, const FSceneView& View, const TArray<AActor*>* IgnoreActors)
-{
-	FActorPositionTraceResult Results = TraceWorldForPosition(Cursor, View, IgnoreActors);
-	if (Results.State == FActorPositionTraceResult::Failed)
+	FActorPositionTraceResult ConvertToActorPositionTraceResult(const UE::Positioning::FObjectPositioningTraceResult& ResultIn)
 	{
-		Results.State = FActorPositionTraceResult::Default;
+		using namespace UE::Positioning;
 
-		// And put it in front of the camera
-		const float DistanceMultiplier = ( Cursor.GetViewportType() == LVT_Perspective ) ? GetDefault<ULevelEditorViewportSettings>()->BackgroundDropDistance : 0.0f;
-		Results.Location = Cursor.GetOrigin() + Cursor.GetDirection() * DistanceMultiplier;
-	}
-	return Results;
-}
+		FActorPositionTraceResult ResultOut;
+		switch (ResultIn.State)
+		{
+		case FObjectPositioningTraceResult::HitSuccess:
+			ResultOut.State = FActorPositionTraceResult::HitSuccess;
+			break;
+		case FObjectPositioningTraceResult::Default:
+			ResultOut.State = FActorPositionTraceResult::Default;
+			break;
+		case FObjectPositioningTraceResult::Failed:
+			ResultOut.State = FActorPositionTraceResult::Failed;
+			break;
+		}
 
-FActorPositionTraceResult FActorPositioning::TraceWorldForPosition(const FViewportCursorLocation& Cursor, const FSceneView& View, const TArray<AActor*>* IgnoreActors)
-{
-	const auto* ViewportClient = Cursor.GetViewportClient();
-	const auto ViewportType = ViewportClient->GetViewportType();
+		ResultOut.Location = ResultIn.Location;
+		ResultOut.SurfaceNormal = ResultIn.SurfaceNormal;
+		ResultOut.HitActor = Cast<AActor>(ResultIn.HitObject.Get());
 
-	// Start with a ray that encapsulates the entire world
-	FVector RayStart = Cursor.GetOrigin();
-	if (ViewportType == LVT_OrthoXY || ViewportType == LVT_OrthoXZ || ViewportType == LVT_OrthoYZ ||
-		ViewportType == LVT_OrthoNegativeXY || ViewportType == LVT_OrthoNegativeXZ || ViewportType == LVT_OrthoNegativeYZ)
-	{
-		RayStart -= Cursor.GetDirection() * HALF_WORLD_MAX/2;
+		return ResultOut;
 	}
-
-	const FVector RayEnd = RayStart + Cursor.GetDirection() * HALF_WORLD_MAX;
-	return TraceWorldForPosition(*ViewportClient->GetWorld(), View, RayStart, RayEnd, IgnoreActors);
 }
 
-/**
- * Prunes list of hit results for actor positioning calculations based on conditions that could be tested
- * on the game thread and returns a list of primitives for the remaining this.
- */
-TArray<TWeakObjectPtr<const UPrimitiveComponent>> FilterHitsGameThread(TArray<FHitResult>& InOutHits)
+FActorPositionTraceResult FActorPositioning::TraceWorldForPositionWithDefault(const FViewportCursorLocation& Cursor, const FSceneView& View, const TArray<AActor*>* IgnoreActors)
 {
-	check(IsInGameThread() || IsInParallelGameThread());
-	TArray<TWeakObjectPtr<const UPrimitiveComponent>> WeakPrimitives;
-	WeakPrimitives.Reserve(InOutHits.Num());
-
-	InOutHits.RemoveAll([&WeakPrimitives](const FHitResult& Hit)
+	FCollisionQueryParams Param(SCENE_QUERY_STAT(DragDropTrace), true);
+	if (IgnoreActors)
 	{
-		if (Hit.bStartPenetrating)
-		{
-			return true;
-		}
-
-		const FActorInstanceHandle& HitObjHandle = Hit.HitObjectHandle;
-
-		// Try and find a primitive component for the hit
-		const UPrimitiveComponent* PrimitiveComponent = Cast<UPrimitiveComponent>(HitObjHandle.GetRootComponent());
-
-		if (!PrimitiveComponent)
-		{
-			PrimitiveComponent = Hit.Component.Get();
-		}
-		if (PrimitiveComponent && PrimitiveComponent->IsA(ULandscapeHeightfieldCollisionComponent::StaticClass()))
-		{
-			PrimitiveComponent = CastChecked<ULandscapeHeightfieldCollisionComponent>(PrimitiveComponent)->GetRenderComponent();
-		}
-
-		if (!PrimitiveComponent)
-		{
-			// If we don't have a primitive component, either ignore the hit, or pass it through if the CVar is set appropriately.
-			// i.e. ignoring is the inverse of the "allow non primitive hits" CVar.
-			return !ActorPositioningLocals::bAllowNonPrimitiveComponentHits;
-		}
-
-		// Ignore volumes and shapes
-		if (HitObjHandle.DoesRepresentClass(AVolume::StaticClass()))
-		{
-			return true;
-		}
-
-		if (PrimitiveComponent->IsA(UShapeComponent::StaticClass()))
-		{
-			return true;
-		}
-
-		WeakPrimitives.Add(PrimitiveComponent);
-		return false;
-	});
+		Param.AddIgnoredActors(*IgnoreActors);
+	}
 
-	return MoveTemp(WeakPrimitives);
+	return ActorPositioningLocals::ConvertToActorPositionTraceResult(
+		UE::Positioning::TraceWorldForPositionWithDefault(Cursor, View, IgnoreActors ? &Param : nullptr));
 }
 
-/** Check to see if the specified hit result should be ignored from actor positioning calculations for the specified scene view */
-bool IsHitIgnoredRenderingThread(const TWeakObjectPtr<const UPrimitiveComponent>& InWeakPrimitiveComponent, const FSceneView& InSceneView)
+FActorPositionTraceResult FActorPositioning::TraceWorldForPosition(const FViewportCursorLocation& Cursor, const FSceneView& View, const TArray<AActor*>* IgnoreActors)
 {
-	// We're using the SceneProxy and ViewRelevance here, we should execute from the render thread
-	check(IsInParallelRenderingThread());
-
-	const UPrimitiveComponent* PrimitiveComponent = InWeakPrimitiveComponent.Get();
-	if (PrimitiveComponent && PrimitiveComponent->SceneProxy)
+	FCollisionQueryParams Param(SCENE_QUERY_STAT(DragDropTrace), true);
+	if (IgnoreActors)
 	{
-		const bool bConsiderInvisibleComponentForPlacement = PrimitiveComponent->bConsiderForActorPlacementWhenHidden;
-
-		// Only use this component if it is visible in the specified scene views
-		const FPrimitiveViewRelevance ViewRelevance = PrimitiveComponent->SceneProxy->GetViewRelevance(&InSceneView);
-		// BSP is a bit special in that its bDrawRelevance is false even when drawn as wireframe because InSceneView.Family->EngineShowFlags.BSPTriangles is off
-		const bool bIsRenderedOnScreen = ViewRelevance.bDrawRelevance || (PrimitiveComponent->IsA(UModelComponent::StaticClass()) && InSceneView.Family->EngineShowFlags.BSP);
-		const bool bIgnoreTranslucentPrimitive = ViewRelevance.HasTranslucency() && !GetDefault<UEditorPerProjectUserSettings>()->bAllowSelectTranslucent;
-		
-		return (!bIsRenderedOnScreen && !bConsiderInvisibleComponentForPlacement) || bIgnoreTranslucentPrimitive;
+		Param.AddIgnoredActors(*IgnoreActors);
 	}
 
-	return false;
+	return ActorPositioningLocals::ConvertToActorPositionTraceResult(
+		UE::Positioning::TraceWorldForPosition(Cursor, View, IgnoreActors ? &Param : nullptr));
 }
 
 FActorPositionTraceResult FActorPositioning::TraceWorldForPosition(const UWorld& InWorld, const FSceneView& InSceneView, const FVector& RayStart, const FVector& RayEnd, const TArray<AActor*>* IgnoreActors)
 {
-	TArray<FHitResult> Hits;
-
 	FCollisionQueryParams Param(SCENE_QUERY_STAT(DragDropTrace), true);
-	
 	if (IgnoreActors)
 	{
 		Param.AddIgnoredActors(*IgnoreActors);
 	}
 
-	FActorPositionTraceResult Results;
-	if ( InWorld.LineTraceMultiByObjectType(Hits, RayStart, RayEnd, FCollisionObjectQueryParams(FCollisionObjectQueryParams::InitType::AllObjects), Param) )
-	{
-		{
-			// Filter out anything that should be ignored based on information accessible on the game thread
-			// and build list of remaining weak primitive components that need to be filtered on the rendering thread 
-			TArray<TWeakObjectPtr<const UPrimitiveComponent>> WeakPrimitives = FilterHitsGameThread(Hits);
-			ensure(Hits.Num() == WeakPrimitives.Num());
-
-			// Send IsHitIgnoredRenderingThread on the render thread since we're accessing view relevance
-			ENQUEUE_RENDER_COMMAND(TraceWorldForPosition_FilterHitsByViewRelevance)(
-				[&Hits, &WeakPrimitives, &InSceneView](FRHICommandListImmediate& RHICmdList)
-				{
-					// Filter out anything that should be ignored
-					int32 Index = 0;
-					Hits.RemoveAll([&Index, &InSceneView, &WeakPrimitives](const FHitResult&)
-					{
-						return IsHitIgnoredRenderingThread(WeakPrimitives[Index++], InSceneView);
-					});
-				}
-			);
-			
-			// We need the result to come back before continuing
-			FRenderCommandFence Fence;
-			Fence.BeginFence();
-			Fence.Wait();
-		}
-
-		// Go through all hits and find closest
-		double ClosestHitDistanceSqr = std::numeric_limits<double>::max();
-
-		for (const FHitResult& Hit : Hits)
-		{
-			const double DistanceToHitSqr = (Hit.ImpactPoint - RayStart).SizeSquared();
-			if (DistanceToHitSqr < ClosestHitDistanceSqr)
-			{
-				ClosestHitDistanceSqr = DistanceToHitSqr;
-				Results.Location = Hit.Location;
-				Results.SurfaceNormal = Hit.Normal.GetSafeNormal();
-				Results.State = FActorPositionTraceResult::HitSuccess;
-				Results.HitActor = Hit.HitObjectHandle.GetManagingActor();
-			}
-		}
-	}
-
-	return Results;
+	return ActorPositioningLocals::ConvertToActorPositionTraceResult(
+		UE::Positioning::TraceWorldForPosition(InWorld, InSceneView, RayStart, RayEnd, IgnoreActors ? &Param : nullptr));
 }
 
 FTransform FActorPositioning::GetCurrentViewportPlacementTransform(const AActor& Actor, bool bSnap, const FViewportCursorLocation* InCursor)
diff --git a/Engine/Source/Editor/UnrealEd/Private/Editor/ObjectPositioning.cpp b/Engine/Source/Editor/UnrealEd/Private/Editor/ObjectPositioning.cpp
new file mode 100644
--- /dev/null
+++ b/Engine/Source/Editor/UnrealEd/Private/Editor/ObjectPositioning.cpp
@@ -0,0 +1,197 @@
+// Copyright Epic Games, Inc. All Rights Reserved.
+
+#include "Editor/ObjectPositioning.h"
+
+#include "Components/ModelComponent.h"
+#include "Components/ShapeComponent.h"
+#include "Editor/EditorPerProjectUserSettings.h"
+#include "EditorViewportClient.h" // FViewportCursorLocation
+#include "GameFramework/Volume.h"
+#include "LandscapeComponent.h"
+#include "LandscapeHeightfieldCollisionComponent.h"
+#include "PrimitiveSceneProxy.h"
+#include "SceneView.h"
+#include "Settings/LevelEditorViewportSettings.h"
+
+namespace ObjectPositioningLocals
+{
+	bool bCVarAllowNonPrimitiveComponentHits = true;
+	static FAutoConsoleVariableRef CVarAllowNonPrimitiveComponentHits(
+		TEXT("PlacementMode.AllowNonPrimitiveComponentHits"),
+		bCVarAllowNonPrimitiveComponentHits,
+		TEXT("When raycasting the world in placement mode, allow hits of physics objects that are not tied to a UPrimitiveComponent (to work with non-actor workflows)."));
+
+	/**
+	 * Prunes list of hit results for actor positioning calculations based on conditions that could be tested
+	 * on the game thread and returns a list of primitives for the remaining this.
+	 */
+	TArray<TWeakObjectPtr<const UPrimitiveComponent>> FilterHitsGameThread(TArray<FHitResult>& InOutHits)
+	{
+		check(IsInGameThread() || IsInParallelGameThread());
+		TArray<TWeakObjectPtr<const UPrimitiveComponent>> WeakPrimitives;
+		WeakPrimitives.Reserve(InOutHits.Num());
+
+		InOutHits.RemoveAll([&WeakPrimitives](const FHitResult& Hit)
+		{
+			if (Hit.bStartPenetrating)
+			{
+				return true;
+			}
+
+			const FActorInstanceHandle& HitObjHandle = Hit.HitObjectHandle;
+
+			// Try and find a primitive component for the hit
+			const UPrimitiveComponent* PrimitiveComponent = Cast<UPrimitiveComponent>(HitObjHandle.GetRootComponent());
+
+			if (!PrimitiveComponent)
+			{
+				PrimitiveComponent = Hit.Component.Get();
+			}
+			if (PrimitiveComponent && PrimitiveComponent->IsA(ULandscapeHeightfieldCollisionComponent::StaticClass()))
+			{
+				PrimitiveComponent = CastChecked<ULandscapeHeightfieldCollisionComponent>(PrimitiveComponent)->GetRenderComponent();
+			}
+
+			if (!PrimitiveComponent)
+			{
+				// If we don't have a primitive component, either ignore the hit, or pass it through if the CVar is set appropriately.
+				// i.e. ignoring is the inverse of the "allow non primitive hits" CVar.
+				return !bCVarAllowNonPrimitiveComponentHits;
+			}
+
+			// Ignore volumes and shapes
+			if (HitObjHandle.DoesRepresentClass(AVolume::StaticClass()))
+			{
+				return true;
+			}
+
+			if (PrimitiveComponent->IsA(UShapeComponent::StaticClass()))
+			{
+				return true;
+			}
+
+			WeakPrimitives.Add(PrimitiveComponent);
+			return false;
+		});
+
+		return MoveTemp(WeakPrimitives);
+	}
+
+	
+	/** Check to see if the specified hit result should be ignored from actor positioning calculations for the specified scene view */
+	bool IsHitIgnoredRenderingThread(const TWeakObjectPtr<const UPrimitiveComponent>& InWeakPrimitiveComponent, const FSceneView& InSceneView)
+	{
+		// We're using the SceneProxy and ViewRelevance here, we should execute from the render thread
+		check(IsInParallelRenderingThread());
+
+		const UPrimitiveComponent* PrimitiveComponent = InWeakPrimitiveComponent.Get();
+		if (PrimitiveComponent && PrimitiveComponent->SceneProxy)
+		{
+			const bool bConsiderInvisibleComponentForPlacement = PrimitiveComponent->bConsiderForActorPlacementWhenHidden;
+
+			// Only use this component if it is visible in the specified scene views
+			const FPrimitiveViewRelevance ViewRelevance = PrimitiveComponent->SceneProxy->GetViewRelevance(&InSceneView);
+			// BSP is a bit special in that its bDrawRelevance is false even when drawn as wireframe because InSceneView.Family->EngineShowFlags.BSPTriangles is off
+			const bool bIsRenderedOnScreen = ViewRelevance.bDrawRelevance || (PrimitiveComponent->IsA(UModelComponent::StaticClass()) && InSceneView.Family->EngineShowFlags.BSP);
+			const bool bIgnoreTranslucentPrimitive = ViewRelevance.HasTranslucency() && !GetDefault<UEditorPerProjectUserSettings>()->bAllowSelectTranslucent;
+		
+			return (!bIsRenderedOnScreen && !bConsiderInvisibleComponentForPlacement) || bIgnoreTranslucentPrimitive;
+		}
+
+		return false;
+	}
+}
+
+
+UE::Positioning::FObjectPositioningTraceResult UE::Positioning::TraceWorldForPositionWithDefault(const FViewportCursorLocation& Cursor, const FSceneView& View, FCollisionQueryParams* CollisionQueryParams)
+{
+	FObjectPositioningTraceResult Results = TraceWorldForPosition(Cursor, View, CollisionQueryParams);
+	if (Results.State == FObjectPositioningTraceResult::Failed)
+	{
+		Results.State = FObjectPositioningTraceResult::Default;
+
+		// And put it in front of the camera
+		const float DistanceMultiplier = (Cursor.GetViewportType() == LVT_Perspective) ? GetDefault<ULevelEditorViewportSettings>()->BackgroundDropDistance : 0.0f;
+		Results.Location = Cursor.GetOrigin() + Cursor.GetDirection() * DistanceMultiplier;
+	}
+	return Results;
+}
+
+UE::Positioning::FObjectPositioningTraceResult UE::Positioning::TraceWorldForPosition(const FViewportCursorLocation& Cursor, const FSceneView& View, FCollisionQueryParams* CollisionQueryParams)
+{
+	const auto* ViewportClient = Cursor.GetViewportClient();
+	const auto ViewportType = ViewportClient->GetViewportType();
+
+	// Start with a ray that encapsulates the entire world
+	FVector RayStart = Cursor.GetOrigin();
+	if (ViewportType == LVT_OrthoXY || ViewportType == LVT_OrthoXZ || ViewportType == LVT_OrthoYZ ||
+		ViewportType == LVT_OrthoNegativeXY || ViewportType == LVT_OrthoNegativeXZ || ViewportType == LVT_OrthoNegativeYZ)
+	{
+		RayStart -= Cursor.GetDirection() * HALF_WORLD_MAX / 2;
+	}
+
+	const FVector RayEnd = RayStart + Cursor.GetDirection() * HALF_WORLD_MAX;
+	return TraceWorldForPosition(*ViewportClient->GetWorld(), View, RayStart, RayEnd, CollisionQueryParams);
+}
+
+UE::Positioning::FObjectPositioningTraceResult UE::Positioning::TraceWorldForPosition(const UWorld& InWorld, const FSceneView& InSceneView, const FVector& RayStart, const FVector& RayEnd, FCollisionQueryParams* CollisionQueryParams)
+{
+	using namespace ObjectPositioningLocals;
+
+	TArray<FHitResult> Hits;
+
+	FCollisionQueryParams Params;
+	if (CollisionQueryParams)
+	{
+		Params = *CollisionQueryParams;
+	}
+	Params.TraceTag = SCENE_QUERY_STAT(DragDropTrace);
+	Params.bTraceComplex = true;
+
+	FObjectPositioningTraceResult Results;
+	if (InWorld.LineTraceMultiByObjectType(Hits, RayStart, RayEnd, FCollisionObjectQueryParams(FCollisionObjectQueryParams::InitType::AllObjects), Params))
+	{
+		{
+			// Filter out anything that should be ignored based on information accessible on the game thread
+			// and build list of remaining weak primitive components that need to be filtered on the rendering thread 
+			TArray<TWeakObjectPtr<const UPrimitiveComponent>> WeakPrimitives = FilterHitsGameThread(Hits);
+			ensure(Hits.Num() == WeakPrimitives.Num());
+
+			// Send IsHitIgnoredRenderingThread on the render thread since we're accessing view relevance
+			ENQUEUE_RENDER_COMMAND(TraceWorldForPosition_FilterHitsByViewRelevance)(
+				[&Hits, &WeakPrimitives, &InSceneView](FRHICommandListImmediate& RHICmdList)
+				{
+					// Filter out anything that should be ignored
+					int32 Index = 0;
+					Hits.RemoveAll([&Index, &InSceneView, &WeakPrimitives](const FHitResult&)
+					{
+						return IsHitIgnoredRenderingThread(WeakPrimitives[Index++], InSceneView);
+					});
+				}
+			);
+			
+			// We need the result to come back before continuing
+			FRenderCommandFence Fence;
+			Fence.BeginFence();
+			Fence.Wait();
+		}
+
+		// Go through all hits and find closest
+		double ClosestHitDistanceSqr = std::numeric_limits<double>::max();
+
+		for (const FHitResult& Hit : Hits)
+		{
+			const double DistanceToHitSqr = (Hit.ImpactPoint - RayStart).SizeSquared();
+			if (DistanceToHitSqr < ClosestHitDistanceSqr)
+			{
+				ClosestHitDistanceSqr = DistanceToHitSqr;
+				Results.Location = Hit.Location;
+				Results.SurfaceNormal = Hit.Normal.GetSafeNormal();
+				Results.State = FObjectPositioningTraceResult::HitSuccess;
+				Results.HitObject = Hit.HitObjectHandle.GetManagingActor();
+			}
+		}
+	}
+
+	return Results;
+}
diff --git a/Engine/Source/Editor/UnrealEd/Private/LevelEditorViewport.cpp b/Engine/Source/Editor/UnrealEd/Private/LevelEditorViewport.cpp
--- a/Engine/Source/Editor/UnrealEd/Private/LevelEditorViewport.cpp
+++ b/Engine/Source/Editor/UnrealEd/Private/LevelEditorViewport.cpp
@@ -80,6 +80,7 @@
 #include "DragAndDrop/BrushBuilderDragDropOp.h"
 #include "DynamicMeshBuilder.h"
 #include "Editor/ActorPositioning.h"
+#include "Editor/ObjectPositioning.h"
 #include "Framework/Notifications/NotificationManager.h"
 #include "Widgets/Notifications/SNotificationList.h"
 #include "Settings/EditorProjectSettings.h"
@@ -1462,7 +1463,7 @@ static bool IsDroppingOn2DLayer()
 	return ViewportSettings->bEnableLayerSnap && Settings2D->SnapLayers.IsValidIndex(ViewportSettings->ActiveSnapLayerIndex);
 }
 
-static FActorPositionTraceResult TraceForPositionOn2DLayer(const FViewportCursorLocation& Cursor)
+static UE::Positioning::FObjectPositioningTraceResult TraceForPositionOn2DLayer(const FViewportCursorLocation& Cursor)
 {
 	const ULevelEditorViewportSettings* ViewportSettings = GetDefault<ULevelEditorViewportSettings>();
 	const ULevelEditor2DSettings* Settings2D = GetDefault<ULevelEditor2DSettings>();
@@ -1479,16 +1480,16 @@ static FActorPositionTraceResult TraceForPositionOn2DLayer(const FViewportCursor
 	case ELevelEditor2DAxis::Z: PlaneCenter.Z = Offset; PlaneNormal.Z = -1; break;
 	}
 
-	FActorPositionTraceResult Result;
+	UE::Positioning::FObjectPositioningTraceResult Result;
 	const double Numerator = FVector::DotProduct(PlaneCenter - Cursor.GetOrigin(), PlaneNormal);
 	const double Denominator = FVector::DotProduct(PlaneNormal, Cursor.GetDirection());
 	if (FMath::Abs(Denominator) < SMALL_NUMBER)
 	{
-		Result.State = FActorPositionTraceResult::Failed;
+		Result.State = UE::Positioning::FObjectPositioningTraceResult::Failed;
 	}
 	else
 	{
-		Result.State = FActorPositionTraceResult::HitSuccess;
+		Result.State = UE::Positioning::FObjectPositioningTraceResult::HitSuccess;
 		Result.SurfaceNormal = PlaneNormal;
 		double D = Numerator / Denominator;
 		Result.Location = Cursor.GetOrigin() + D * Cursor.GetDirection();
@@ -1561,16 +1562,11 @@ bool FLevelEditorViewportClient::UpdateDropPreviewElements(int32 MouseX, int32 M
 	// Finish the calculation of the actors origin now that we know we are not dividing by zero
 	CombinedOrigin /= Count;
 
-	// TODO: Swap to ignored items instead of actors
-	TArray<AActor*> IgnoreActors;
-	DropPreviewElements->ForEachElement<ITypedElementObjectInterface>(
-	[&IgnoreActors](const TTypedElement<ITypedElementObjectInterface>& InElement)
+	FCollisionQueryParams CollisionQueryParams;
+	DropPreviewElements->ForEachElement<ITypedElementWorldInterface>(
+	[&CollisionQueryParams](const TTypedElement<ITypedElementWorldInterface>& InElement)
 	{
-		if (AActor* Actor = InElement.GetObjectAs<AActor>())
-		{
-			IgnoreActors.Add(Actor);
-			Actor->GetAllChildActors(IgnoreActors);
-		}
+		InElement.AddIgnoredElementToCollisionQueryParams(InElement, CollisionQueryParams);
 
 		// true means continue
 		return true;
@@ -1584,7 +1580,9 @@ bool FLevelEditorViewportClient::UpdateDropPreviewElements(int32 MouseX, int32 M
 	FSceneView* View = CalcSceneView(&ViewFamily);
 	FViewportCursorLocation Cursor(View, this, MouseX, MouseY);
 
-	const FActorPositionTraceResult TraceResult = IsDroppingOn2DLayer() ? TraceForPositionOn2DLayer(Cursor) : FActorPositioning::TraceWorldForPositionWithDefault(Cursor, *View, &IgnoreActors);
+	const UE::Positioning::FObjectPositioningTraceResult TraceResult = IsDroppingOn2DLayer() 
+		? TraceForPositionOn2DLayer(Cursor) 
+		: UE::Positioning::TraceWorldForPositionWithDefault(Cursor, *View, &CollisionQueryParams);
 
 	GEditor->UnsnappedClickLocation = TraceResult.Location;
 	GEditor->ClickLocation = TraceResult.Location;
@@ -1593,7 +1591,7 @@ bool FLevelEditorViewportClient::UpdateDropPreviewElements(int32 MouseX, int32 M
 	// Snap the new location if snapping is enabled
 	FSnappingUtils::SnapPointToGrid(GEditor->ClickLocation, FVector::ZeroVector);
 
-	AActor* DroppedOnActor = TraceResult.HitActor.Get();
+	AActor* DroppedOnActor = Cast<AActor>(TraceResult.HitObject.Get());
 
 	if (DroppedOnActor)
 	{
@@ -1847,7 +1845,9 @@ bool FLevelEditorViewportClient::DropObjectsAtCoordinates(int32 MouseX, int32 Mo
 
 		HHitProxy* HitProxy = Viewport->GetHitProxy(Cursor.GetCursorPos().X, Cursor.GetCursorPos().Y);
 
-		const FActorPositionTraceResult TraceResult = IsDroppingOn2DLayer() ? TraceForPositionOn2DLayer(Cursor) : FActorPositioning::TraceWorldForPositionWithDefault(Cursor, *View);
+		const UE::Positioning::FObjectPositioningTraceResult TraceResult = IsDroppingOn2DLayer() 
+			? TraceForPositionOn2DLayer(Cursor) 
+			: UE::Positioning::TraceWorldForPositionWithDefault(Cursor, *View, nullptr);
 		
 		GEditor->UnsnappedClickLocation = TraceResult.Location;
 		GEditor->ClickLocation = TraceResult.Location;
@@ -3455,7 +3455,7 @@ bool FLevelEditorViewportClient::InputKey(const FInputKeyEventArgs& InEventArgs)
 		return true;
 	}
 
-	
+
 	const int32	HitX = InEventArgs.Viewport->GetMouseX();
 	const int32	HitY = InEventArgs.Viewport->GetMouseY();
 
diff --git a/Engine/Source/Editor/UnrealEd/Public/Editor/ObjectPositioning.h b/Engine/Source/Editor/UnrealEd/Public/Editor/ObjectPositioning.h
new file mode 100644
--- /dev/null
+++ b/Engine/Source/Editor/UnrealEd/Public/Editor/ObjectPositioning.h
@@ -0,0 +1,86 @@
+// Copyright Epic Games, Inc. All Rights Reserved.
+
+#pragma once
+
+#include "Math/Vector.h"
+#include "UObject/WeakObjectPtr.h"
+
+struct FCollisionQueryParams;
+class FSceneView;
+struct FViewportCursorLocation;
+class UObject;
+
+namespace UE::Positioning 
+{
+	struct FObjectPositioningTraceResult
+	{
+		/** Enum representing the state of this result */
+		enum EResultState
+		{
+			/** The trace found a valid hit target */
+			HitSuccess,
+			/** The trace found no valid targets, so chose a default position */
+			Default,
+			/** The trace failed entirely */
+			Failed,
+		};
+
+		/** Constructor */
+		FObjectPositioningTraceResult() : State(Failed), Location(0.f), SurfaceNormal(0.f, 0.f, 1.f) {}
+
+		/** The state of this result */
+		EResultState State = EResultState::Failed;
+
+		/** The location of the preferred trace hit */
+		FVector Location = FVector::ZeroVector;
+
+		/** The surface normal of the trace hit */
+		FVector SurfaceNormal = FVector::UnitZ();
+
+		/** 
+		 * Pointer to the object that was hit, if any.
+		 * NOTE: This currently only works reliably for giving back an owning actor,
+		 * if there is one, so it is not yet reliable otherwise.
+		 */
+		TWeakObjectPtr<UObject> HitObject;
+	};
+
+	/** 
+	 * Trace the specified world to find a position to place an object. If nothing is hit, use a default
+	 * position in front of the camera (determined by ULevelEditorViewportSettings::BackgroundDropDistance)
+	 *
+	 *	@param	Cursor			The cursor position and direction to trace at
+	 *	@param	View			The scene view that we are tracing on
+	 *	@param	CollisionQueryParams	Optional existing collision query parameters to 
+	 *    use, except for TraceTag and bTraceComplex, which are overriden by the function.
+	 *
+	 *	@return	Result structure containing the location and normal of a trace hit, or a default position in front of the camera on failure
+	*/
+	UNREALED_API FObjectPositioningTraceResult TraceWorldForPositionWithDefault(const FViewportCursorLocation& Cursor, const FSceneView& View, FCollisionQueryParams* CollisionQueryParams);
+	
+	/** 
+	 * Trace the specified world to find a position to place an object.
+	 *
+	 *	@param	Cursor			The cursor position and direction to trace at
+	 *	@param	View			The scene view that we are tracing on
+	 *	@param	CollisionQueryParams	Optional existing collision query parameters to 
+	 *    use, except for TraceTag and bTraceComplex, which are overriden by the function.
+	 *
+	 *	@return	Result structure containing the location and normal of a trace hit, or empty on failure
+	*/
+	UNREALED_API FObjectPositioningTraceResult TraceWorldForPosition(const FViewportCursorLocation& Cursor, const FSceneView& View, FCollisionQueryParams* CollisionQueryParams);
+
+	/** 
+	 * Trace the specified world to find a position to place an object.
+	 *
+	 *	@param	World			The world to trace
+	 *	@param	InSceneView		The scene view that we are tracing on
+	 *	@param	RayStart		The start of the ray in world space
+	 *	@param	RayEnd			The end of the ray in world space
+	 *	@param	CollisionQueryParams	Optional existing collision query parameters to 
+	 *    use, except for TraceTag and bTraceComplex, which are overriden by the function.
+	 *
+	 *	@return	Result structure containing the location and normal of a trace hit, or empty on failure
+	*/
+	UNREALED_API FObjectPositioningTraceResult TraceWorldForPosition(const UWorld& InWorld, const FSceneView& InSceneView, const FVector& RayStart, const FVector& RayEnd, FCollisionQueryParams* CollisionQueryParams);
+}
\ No newline at end of file
diff --git a/Engine/Source/Runtime/Engine/Private/Elements/Actor/ActorElementWorldInterface.cpp b/Engine/Source/Runtime/Engine/Private/Elements/Actor/ActorElementWorldInterface.cpp
--- a/Engine/Source/Runtime/Engine/Private/Elements/Actor/ActorElementWorldInterface.cpp
+++ b/Engine/Source/Runtime/Engine/Private/Elements/Actor/ActorElementWorldInterface.cpp
@@ -3,6 +3,7 @@
 #include "Elements/Actor/ActorElementWorldInterface.h"
 
 #include "Elements/Actor/ActorElementData.h"
+#include "Elements/Framework/TypedElementRegistry.h"
 #include "Elements/Component/ComponentElementData.h"
 #include "Engine/World.h"
 
@@ -212,11 +213,42 @@ TArray<FTypedElementHandle> UActorElementWorldInterface::GetSelectionElementsFro
 	return {};
 }
 
-bool UActorElementWorldInterface::FindSuitableTransformAlongPath_WorldSweep(const UWorld* InWorld, const FVector& InPathStart, const FVector& InPathEnd, const FCollisionShape& InTestShape, TArrayView<const FTypedElementHandle> InElementsToIgnore, FCollisionQueryParams& InOutParams, FTransform& OutSuitableTransform)
+bool UActorElementWorldInterface::AddIgnoredElementToCollisionQueryParams(const FTypedElementHandle& InElementHandle, FCollisionQueryParams& InOutParams, bool bAlsoIgnoreSubElements)
 {
-	for (const FTypedElementHandle& ElementToIgnore : InElementsToIgnore)
+	AActor* Actor = ActorElementDataUtil::GetActorFromHandle(InElementHandle, /*bSilent*/true);
+	if (!Actor)
 	{
-		AddIgnoredCollisionQueryElement(ElementToIgnore, InOutParams);
+		return false;
+	}
+
+	TArray<UObject*> IgnoredSourceObjects{ Actor };
+
+	if (bAlsoIgnoreSubElements)
+	{
+		TArray<AActor*> ChildActors;
+		Actor->GetAllChildActors(ChildActors);
+		IgnoredSourceObjects.Append(ChildActors);
+	}
+
+	InOutParams.AddIgnoredSourceObjects(IgnoredSourceObjects);
+	return true;
+}
+
+bool UActorElementWorldInterface::FindSuitableTransformAlongPath_WorldSweep(const UWorld* InWorld, 
+	const FVector& InPathStart, const FVector& InPathEnd, 
+	const FCollisionShape& InTestShape, TArrayView<const FTypedElementHandle> InElementsToIgnore, 
+	FCollisionQueryParams& InOutParams, FTransform& OutSuitableTransform)
+{
+	UTypedElementRegistry* Registry = UTypedElementRegistry::GetInstance();
+	if (ensure(Registry))
+	{
+		for (const FTypedElementHandle& ElementToIgnore : InElementsToIgnore)
+		{
+			if (ITypedElementWorldInterface* WorldInterface = Registry->GetElementInterface<ITypedElementWorldInterface>(ElementToIgnore))
+			{
+				WorldInterface->AddIgnoredElementToCollisionQueryParams(ElementToIgnore, InOutParams);
+			}
+		}
 	}
 
 	FHitResult Hit(1.0f);
@@ -240,19 +272,17 @@ bool UActorElementWorldInterface::FindSuitableTransformAlongPath_WorldSweep(cons
 
 void UActorElementWorldInterface::AddIgnoredCollisionQueryElement(const FTypedElementHandle& InElementHandle, FCollisionQueryParams& InOutParams)
 {
-	if (const AActor* Actor = ActorElementDataUtil::GetActorFromHandle(InElementHandle, /*bSilent*/true))
+	UTypedElementRegistry* Registry = UTypedElementRegistry::GetInstance();
+	if (!ensure(Registry))
 	{
-		InOutParams.AddIgnoredActor(Actor);
 		return;
 	}
 
-	if (const UActorComponent* Component = ComponentElementDataUtil::GetComponentFromHandle(InElementHandle, /*bSilent*/true))
+	if (ITypedElementWorldInterface* WorldInterface = Registry->GetElementInterface<ITypedElementWorldInterface>(InElementHandle))
 	{
-		if (const UPrimitiveComponent* PrimComponent = Cast<UPrimitiveComponent>(Component))
-		{
-			InOutParams.AddIgnoredComponent(PrimComponent);
-		}
-		return;
+		WorldInterface->AddIgnoredElementToCollisionQueryParams(InElementHandle, InOutParams,
+			// Preserving legacy behavior, where we didn't add child elements:
+			/*bAlsoIgnoreSubElements =*/false);
 	}
 }
 
diff --git a/Engine/Source/Runtime/Engine/Private/Elements/Component/ComponentElementWorldInterface.cpp b/Engine/Source/Runtime/Engine/Private/Elements/Component/ComponentElementWorldInterface.cpp
--- a/Engine/Source/Runtime/Engine/Private/Elements/Component/ComponentElementWorldInterface.cpp
+++ b/Engine/Source/Runtime/Engine/Private/Elements/Component/ComponentElementWorldInterface.cpp
@@ -146,3 +146,41 @@ bool UComponentElementWorldInterface::FindSuitableTransformAlongPath(const FType
 	return false;
 }
 
+bool UComponentElementWorldInterface::AddIgnoredElementToCollisionQueryParams(const FTypedElementHandle& InElementHandle, FCollisionQueryParams& InOutParams, bool bAlsoIgnoreSubElements)
+{
+	UActorComponent* Component = Cast<USceneComponent>(ComponentElementDataUtil::GetComponentFromHandle(InElementHandle));
+
+	// Only primitive components are able to be ignored themselves, but scene components
+	// could have children that might be ignorable. 
+
+	TArray<UPrimitiveComponent*> ComponentsToIgnore;
+	if (UPrimitiveComponent* PrimitiveComponent = Cast<UPrimitiveComponent>(Component))
+	{
+		ComponentsToIgnore.Add(PrimitiveComponent);
+	}
+
+	if (bAlsoIgnoreSubElements)
+	{
+		TArray<USceneComponent*> ChildComponents;
+		if (USceneComponent* SceneComponent = Cast<USceneComponent>(Component))
+		{
+			SceneComponent->GetChildrenComponents(true, ChildComponents);
+		}
+
+		for (USceneComponent* Child : ChildComponents)
+		{
+			if (UPrimitiveComponent* CastChild = Cast<UPrimitiveComponent>(Child))
+			{
+				ComponentsToIgnore.Add(CastChild);
+			}
+		}
+	}
+
+	if (!ComponentsToIgnore.IsEmpty())
+	{
+		InOutParams.AddIgnoredComponents(ComponentsToIgnore);
+		return true;
+	}
+	return false;
+}
+
diff --git a/Engine/Source/Runtime/Engine/Public/Elements/Actor/ActorElementWorldInterface.h b/Engine/Source/Runtime/Engine/Public/Elements/Actor/ActorElementWorldInterface.h
--- a/Engine/Source/Runtime/Engine/Public/Elements/Actor/ActorElementWorldInterface.h
+++ b/Engine/Source/Runtime/Engine/Public/Elements/Actor/ActorElementWorldInterface.h
@@ -25,6 +25,9 @@ public:
 	ENGINE_API virtual bool FindSuitableTransformAtPoint(const FTypedElementHandle& InElementHandle, const FTransform& InPotentialTransform, FTransform& OutSuitableTransform) override;
 	ENGINE_API virtual bool FindSuitableTransformAlongPath(const FTypedElementHandle& InElementHandle, const FVector& InPathStart, const FVector& InPathEnd, const FCollisionShape& InTestShape, TArrayView<const FTypedElementHandle> InElementsToIgnore, FTransform& OutSuitableTransform) override;
 	ENGINE_API virtual TArray<FTypedElementHandle> GetSelectionElementsFromSelectionFunction(const FTypedElementHandle& InElementHandle, const FWorldSelectionElementArgs& SelectionArgs, const TFunction<bool(const FTypedElementHandle&, const FWorldSelectionElementArgs&)>& SelectionFunction) override;
+	ENGINE_API virtual bool AddIgnoredElementToCollisionQueryParams(const FTypedElementHandle& InElementHandle, FCollisionQueryParams& InOutParams, bool bAlsoIgnoreSubElements = true) override;
+
 	static ENGINE_API bool FindSuitableTransformAlongPath_WorldSweep(const UWorld* InWorld, const FVector& InPathStart, const FVector& InPathEnd, const FCollisionShape& InTestShape, TArrayView<const FTypedElementHandle> InElementsToIgnore, FCollisionQueryParams& InOutParams, FTransform& OutSuitableTransform);
+	UE_DEPRECATED(5.4, "Use ITypedElementWorldInterface::AddIgnoredCollisionQueryObjects instead.")
 	static ENGINE_API void AddIgnoredCollisionQueryElement(const FTypedElementHandle& InElementHandle, FCollisionQueryParams& InOutParams);
 };
diff --git a/Engine/Source/Runtime/Engine/Public/Elements/Component/ComponentElementWorldInterface.h b/Engine/Source/Runtime/Engine/Public/Elements/Component/ComponentElementWorldInterface.h
--- a/Engine/Source/Runtime/Engine/Public/Elements/Component/ComponentElementWorldInterface.h
+++ b/Engine/Source/Runtime/Engine/Public/Elements/Component/ComponentElementWorldInterface.h
@@ -22,4 +22,5 @@ public:
 	ENGINE_API virtual bool GetRelativeTransform(const FTypedElementHandle& InElementHandle, FTransform& OutTransform) override;
 	ENGINE_API virtual bool SetRelativeTransform(const FTypedElementHandle& InElementHandle, const FTransform& InTransform) override;
 	ENGINE_API virtual bool FindSuitableTransformAlongPath(const FTypedElementHandle& InElementHandle, const FVector& InPathStart, const FVector& InPathEnd, const FCollisionShape& InTestShape, TArrayView<const FTypedElementHandle> InElementsToIgnore, FTransform& OutSuitableTransform) override;
+	ENGINE_API virtual bool AddIgnoredElementToCollisionQueryParams(const FTypedElementHandle& InElementHandle, FCollisionQueryParams& InOutParams, bool bAlsoIgnoreSubElements = true) override;
 };
diff --git a/Engine/Source/Runtime/Engine/Public/Elements/Interfaces/TypedElementWorldInterface.h b/Engine/Source/Runtime/Engine/Public/Elements/Interfaces/TypedElementWorldInterface.h
--- a/Engine/Source/Runtime/Engine/Public/Elements/Interfaces/TypedElementWorldInterface.h
+++ b/Engine/Source/Runtime/Engine/Public/Elements/Interfaces/TypedElementWorldInterface.h
@@ -16,6 +16,7 @@ class ULevel;
 class UTypedElementSelectionSet; 
 class UWorld;
 struct FCollisionShape;
+struct FCollisionQueryParams;
 struct FConvexVolume;
 struct FEngineShowFlags;
 struct FWorldSelectionElementArgs;
@@ -172,6 +173,21 @@ public:
 		return false;
 	}
 
+	/**
+	 * Attempt to modify FCollisionQueryParams so that the given element (and optionally any attached 
+	 * sub-elements, if the element is part of a hierarchy) is ignored in the query. 
+	 *
+	 * For instance, this can be used to support dragging elements into the scene and raycasting for a
+	 * target location without hitting the preview.
+	 * 
+	 * @return true if the collision query parameters were modified.
+	 */
+	virtual bool AddIgnoredElementToCollisionQueryParams(const FTypedElementHandle& InElementHandle,
+		FCollisionQueryParams& InOutParams, bool bAlsoIgnoreSubElements = true)
+	{
+		return false;
+	}
+
 	/**
 	 * Can the given element be deleted?
 	 */
@@ -457,6 +473,7 @@ struct TTypedElement<ITypedElementWorldInterface> : public TTypedElementBase<ITy
 	void NotifyMovementEnded() const { InterfacePtr->NotifyMovementEnded(*this); }
 	bool FindSuitableTransformAtPoint(const FTransform& InPotentialTransform, FTransform& OutSuitableTransform) const { return InterfacePtr->FindSuitableTransformAtPoint(*this, InPotentialTransform, OutSuitableTransform); }
 	bool FindSuitableTransformAlongPath(const FVector& InPathStart, const FVector& InPathEnd, const FCollisionShape& InTestShape, TArrayView<const FTypedElementHandle> InElementsToIgnore, FTransform& OutSuitableTransform) const { return InterfacePtr->FindSuitableTransformAlongPath(*this, InPathStart, InPathEnd, InTestShape, InElementsToIgnore, OutSuitableTransform); }
+	bool AddIgnoredElementToCollisionQueryParams(const FTypedElementHandle& InElementHandle, FCollisionQueryParams& InOutParams, bool bAlsoIgnoreSubElements = true) const { return InterfacePtr->AddIgnoredElementToCollisionQueryParams(InElementHandle, InOutParams, bAlsoIgnoreSubElements); }
 	bool CanDeleteElement() const { return InterfacePtr->CanDeleteElement(*this); }
 	bool DeleteElement(UWorld* InWorld, UTypedElementSelectionSet* InSelectionSet, const FTypedElementDeletionOptions& InDeletionOptions) const { return InterfacePtr->DeleteElement(*this, InWorld, InSelectionSet, InDeletionOptions); }
 	bool CanDuplicateElement() const { return InterfacePtr->CanDuplicateElement(*this); }
