Editor: Fix minor comments based on feedback that slipped through the cracks in 31181626.

diff --git a/Engine/Source/Editor/UnrealEd/Private/Editor/ObjectPositioning.cpp b/Engine/Source/Editor/UnrealEd/Private/Editor/ObjectPositioning.cpp
--- a/Engine/Source/Editor/UnrealEd/Private/Editor/ObjectPositioning.cpp
+++ b/Engine/Source/Editor/UnrealEd/Private/Editor/ObjectPositioning.cpp
@@ -22,7 +22,7 @@ namespace ObjectPositioningLocals
 		TEXT("When raycasting the world in placement mode, allow hits of physics objects that are not tied to a UPrimitiveComponent (to work with non-actor workflows)."));
 
 	/**
-	 * Prunes list of hit results for actor positioning calculations based on conditions that could be tested
+	 * Prunes list of hit results for object positioning calculations based on conditions that could be tested
 	 * on the game thread and returns a list of primitives for the remaining this.
 	 */
 	TArray<TWeakObjectPtr<const UPrimitiveComponent>> FilterHitsGameThread(TArray<FHitResult>& InOutHits)
@@ -78,7 +78,7 @@ namespace ObjectPositioningLocals
 	}
 
 	
-	/** Check to see if the specified hit result should be ignored from actor positioning calculations for the specified scene view */
+	/** Check to see if the specified hit result should be ignored from object positioning calculations for the specified scene view */
 	bool IsHitIgnoredRenderingThread(const TWeakObjectPtr<const UPrimitiveComponent>& InWeakPrimitiveComponent, const FSceneView& InSceneView)
 	{
 		// We're using the SceneProxy and ViewRelevance here, we should execute from the render thread
diff --git a/Engine/Source/Runtime/Engine/Private/Elements/Actor/ActorElementWorldInterface.cpp b/Engine/Source/Runtime/Engine/Private/Elements/Actor/ActorElementWorldInterface.cpp
--- a/Engine/Source/Runtime/Engine/Private/Elements/Actor/ActorElementWorldInterface.cpp
+++ b/Engine/Source/Runtime/Engine/Private/Elements/Actor/ActorElementWorldInterface.cpp
@@ -280,9 +280,10 @@ void UActorElementWorldInterface::AddIgnoredCollisionQueryElement(const FTypedEl
 
 	if (ITypedElementWorldInterface* WorldInterface = Registry->GetElementInterface<ITypedElementWorldInterface>(InElementHandle))
 	{
-		WorldInterface->AddIgnoredElementToCollisionQueryParams(InElementHandle, InOutParams,
-			// Preserving legacy behavior, where we didn't add child elements:
-			/*bAlsoIgnoreSubElements =*/false);
+		// Preserving legacy behavior, where we didn't add child elements:
+		const bool bAlsoIgnoreSubElements = false;
+
+		WorldInterface->AddIgnoredElementToCollisionQueryParams(InElementHandle, InOutParams, bAlsoIgnoreSubElements);
 	}
 }
 
