Add the viewport to the toolbar context
The toolbar context used when extending the toolbar now has a weak pointer to the actual viewport the toolbar extension is being added to.
This provides additional information for CreateExtensionWidget in the FPanelExtensionFactory.

diff --git a/Engine/Source/Editor/LevelEditor/Private/SLevelViewportToolBar.cpp b/Engine/Source/Editor/LevelEditor/Private/SLevelViewportToolBar.cpp
--- a/Engine/Source/Editor/LevelEditor/Private/SLevelViewportToolBar.cpp
+++ b/Engine/Source/Editor/LevelEditor/Private/SLevelViewportToolBar.cpp
@@ -191,6 +191,7 @@ void SLevelViewportToolBar::Construct( const FArguments& InArgs )
 
 	UViewportToolBarContext* ExtensionContextObject = NewObject<UViewportToolBarContext>();
 	ExtensionContextObject->ViewportToolBar = SharedThis(this);
+	ExtensionContextObject->Viewport = Viewport;
 
 	const FMargin ToolbarSlotPadding(4.0f, 1.0f);
 
diff --git a/Engine/Source/Editor/UnrealEd/Public/ViewportToolBarContext.h b/Engine/Source/Editor/UnrealEd/Public/ViewportToolBarContext.h
--- a/Engine/Source/Editor/UnrealEd/Public/ViewportToolBarContext.h
+++ b/Engine/Source/Editor/UnrealEd/Public/ViewportToolBarContext.h
@@ -6,6 +6,7 @@
 #include "ViewportToolBarContext.generated.h"
 
 class SViewportToolBar;
+class SLevelViewport;
 
 UCLASS(MinimalAPI)
 class UViewportToolBarContext : public UObject
@@ -13,4 +14,5 @@ class UViewportToolBarContext : public UObject
 	GENERATED_BODY()
 public:
 	TWeakPtr<SViewportToolBar> ViewportToolBar;
+	TWeakPtr<SLevelViewport> Viewport;
 };
