Editor: Fix non-primitive component hits not being dealt with properly in object positioning.

diff --git a/Engine/Source/Editor/UnrealEd/Private/Editor/ObjectPositioning.cpp b/Engine/Source/Editor/UnrealEd/Private/Editor/ObjectPositioning.cpp
--- a/Engine/Source/Editor/UnrealEd/Private/Editor/ObjectPositioning.cpp
+++ b/Engine/Source/Editor/UnrealEd/Private/Editor/ObjectPositioning.cpp
@@ -55,8 +55,15 @@ namespace ObjectPositioningLocals
 			if (!PrimitiveComponent)
 			{
 				// If we don't have a primitive component, either ignore the hit, or pass it through if the CVar is set appropriately.
-				// i.e. ignoring is the inverse of the "allow non primitive hits" CVar.
-				return !bCVarAllowNonPrimitiveComponentHits;
+				// If we pass the hit through, we still need to add an entry to the WeakPrimitives list to make sure that
+				// we have an entry for each hit index.
+				if (bCVarAllowNonPrimitiveComponentHits)
+				{
+					WeakPrimitives.Add(nullptr);
+					return false;
+				}
+				// Filter out the hit if the CVar didn't allow it.
+				return true;
 			}
 
 			// Ignore volumes and shapes
