StaticAnalysis: Avoid error V568: ...argument of sizeof() operator is the 'Result.GetBytes()' expression by using decltype.
#rnx

diff --git a/Engine/Source/Editor/UnrealEd/Private/Cooker/CookPackageData.cpp b/Engine/Source/Editor/UnrealEd/Private/Cooker/CookPackageData.cpp
--- a/Engine/Source/Editor/UnrealEd/Private/Cooker/CookPackageData.cpp
+++ b/Engine/Source/Editor/UnrealEd/Private/Cooker/CookPackageData.cpp
@@ -2294,7 +2294,7 @@ void FCookGenerationInfo::CreatePackageHash()
 	constexpr int SizeDifference = sizeof(PackageHash) - sizeof(decltype(DeclVal<UPackage>().GetGuid()));
 	if (SizeDifference > 0)
 	{
-		FMemory::Memset(((uint8*)&PackageHash.GetBytes()) + (sizeof(PackageHash.GetBytes()) - SizeDifference),
+		FMemory::Memset(((uint8*)&PackageHash.GetBytes()) + (sizeof(decltype(PackageHash.GetBytes())) - SizeDifference),
 			0, SizeDifference);
 	}
 	PRAGMA_ENABLE_DEPRECATION_WARNINGS;
