Add r.Editor.NeverStartInPreviewMode
Disable Preview Platform automatically if Preview crashed

diff --git a/Engine/Source/Editor/UnrealEd/Private/EditorEngine.cpp b/Engine/Source/Editor/UnrealEd/Private/EditorEngine.cpp
--- a/Engine/Source/Editor/UnrealEd/Private/EditorEngine.cpp
+++ b/Engine/Source/Editor/UnrealEd/Private/EditorEngine.cpp
@@ -7753,6 +7753,29 @@ void UEditorEngine::OnEffectivePreviewShaderPlatformChange()
 	}
 }
 
+static void SaveFeatureLevelAsDisabled(FPreviewPlatformInfo& PreviewPlatform)
+{
+	auto* Settings = GetMutableDefault<UEditorPerProjectUserSettings>();
+
+	Settings->PreviewFeatureLevel = 0;
+	Settings->PreviewPlatformName = NAME_None;
+	Settings->PreviewShaderFormatName = NAME_None;
+	Settings->bPreviewFeatureLevelActive = false;
+	Settings->bPreviewFeatureLevelWasDefault = true;
+	Settings->PreviewDeviceProfileName = NAME_None;
+	Settings->PreviewShaderPlatformName = NAME_None;
+
+	Settings->SaveConfig();
+
+	Settings->PreviewFeatureLevel = (int32)PreviewPlatform.PreviewFeatureLevel;
+	Settings->PreviewPlatformName = PreviewPlatform.PreviewPlatformName;
+	Settings->PreviewShaderFormatName = PreviewPlatform.PreviewShaderFormatName;
+	Settings->bPreviewFeatureLevelActive = PreviewPlatform.bPreviewFeatureLevelActive;
+	Settings->bPreviewFeatureLevelWasDefault = (PreviewPlatform.PreviewFeatureLevel == GMaxRHIFeatureLevel);
+	Settings->PreviewDeviceProfileName = PreviewPlatform.DeviceProfileName;
+	Settings->PreviewShaderPlatformName = PreviewPlatform.PreviewShaderPlatformName;
+}
+
 void UEditorEngine::SetPreviewPlatform(const FPreviewPlatformInfo& NewPreviewPlatform, bool bSaveSettings)
 {
 	// Get the requested preview platform, make sure it is valid.
@@ -7774,6 +7797,9 @@ void UEditorEngine::SetPreviewPlatform(const FPreviewPlatformInfo& NewPreviewPla
 	// Record the new preview platform
 	PreviewPlatform = NewPreviewPlatform;
 
+	// Initially set preview as disabled in case it fails
+	SaveFeatureLevelAsDisabled(PreviewPlatform);
+
 	// If we changed the preview platform, we need to update the material quality settings
 	if (bChangedPreviewShaderPlatform)
 	{
diff --git a/Engine/Source/Editor/UnrealEd/Private/Settings/EditorPerProjectUserSettings.cpp b/Engine/Source/Editor/UnrealEd/Private/Settings/EditorPerProjectUserSettings.cpp
--- a/Engine/Source/Editor/UnrealEd/Private/Settings/EditorPerProjectUserSettings.cpp
+++ b/Engine/Source/Editor/UnrealEd/Private/Settings/EditorPerProjectUserSettings.cpp
@@ -9,6 +9,12 @@
 
 #define LOCTEXT_NAMESPACE "EditorPerProjectUserSettings"
 
+static TAutoConsoleVariable<int32> CVarNeverStartInPreviewMode(
+	TEXT("r.Editor.NeverStartInPreviewMode"),
+	0,
+	TEXT("0: Editor can start in preview mode, 1: Editor never starts in preview mode"),
+	ECVF_Default);
+
 /// @cond DOXYGEN_WARNINGS
 
 UEditorPerProjectUserSettings::UEditorPerProjectUserSettings(const FObjectInitializer& ObjectInitializer)
@@ -41,9 +47,9 @@ void UEditorPerProjectUserSettings::PostInitProperties()
 {
 	Super::PostInitProperties();
 
-	// if we last saved as the default or we somehow are loading a preview feature level higher than we can support,
+	// if we last saved as the default or we somehow are loading a preview feature level higher than we can support or we explictly disabled it with r.Editor.NeverStartInPreviewMode,
 	// fall back to the current session's maximum feature level
-	if (bPreviewFeatureLevelWasDefault || PreviewFeatureLevel > GMaxRHIFeatureLevel)
+	if (bPreviewFeatureLevelWasDefault || PreviewFeatureLevel > GMaxRHIFeatureLevel || (CVarNeverStartInPreviewMode.GetValueOnAnyThread() != 0))
 	{
 		PreviewFeatureLevel = GMaxRHIFeatureLevel;
 		PreviewShaderPlatformName = NAME_None;
