CIS fix.

diff --git a/Engine/Source/Editor/UnrealEd/Private/SourceCodeNavigation.cpp b/Engine/Source/Editor/UnrealEd/Private/SourceCodeNavigation.cpp
--- a/Engine/Source/Editor/UnrealEd/Private/SourceCodeNavigation.cpp
+++ b/Engine/Source/Editor/UnrealEd/Private/SourceCodeNavigation.cpp
@@ -997,7 +997,7 @@ void FSourceCodeNavigationImpl::GatherFunctions( const FString& ModuleName, cons
 									int32 ArgumentIndex = -1;
 									if(FunctionSymbolName.FindLastChar(TCHAR('('), ArgumentIndex))
 									{
-										FunctionSymbolName.LeftInline(ArgumentIndex, false);
+										FunctionSymbolName.LeftInline(ArgumentIndex, EAllowShrinking::No);
 										int32 TemplateNesting = 0;
 										
 										int32 Pos = FunctionSymbolName.Len();
@@ -1014,7 +1014,7 @@ void FSourceCodeNavigationImpl::GatherFunctions( const FString& ModuleName, cons
 											TCHAR Character = FunctionSymbolName[Pos - 1];
 											if(Character == TCHAR(' ') && TemplateNesting == 0)
 											{
-												FunctionSymbolName.MidInline(Pos, MAX_int32, false);
+												FunctionSymbolName.MidInline(Pos, MAX_int32, EAllowShrinking::No);
 												break;
 											}
 											else if(Character == TCHAR('>'))
