IterativeCooker: ConfigAccessTracking: Replace the limited tracking present in AccessDetection::ConfigIni with a more detailed version that calls arbitrary callbacks on each ConfigValue read. The cooker subscribes to the callback and records the ConfigValue that was read as a dependency of the current package in scope according to PackageAccessTracking scopes.
[RN] Minor, Foundation

#changelist validated
#virtualized

diff --git a/Engine/Source/Editor/UnrealEd/Classes/CookOnTheSide/CookOnTheFlyServer.h b/Engine/Source/Editor/UnrealEd/Classes/CookOnTheSide/CookOnTheFlyServer.h
--- a/Engine/Source/Editor/UnrealEd/Classes/CookOnTheSide/CookOnTheFlyServer.h
+++ b/Engine/Source/Editor/UnrealEd/Classes/CookOnTheSide/CookOnTheFlyServer.h
@@ -424,16 +424,8 @@ private:
 
 	//////////////////////////////////////////////////////////////////////////
 	// iterative ini settings checking
-	// growing list of ini settings which are accessed over the course of the cook
 
-	mutable FCriticalSection ConfigFileCS;
-	mutable UE::Cook::FIniSettingContainer AccessedIniStrings;
-	TArray<const FConfigFile*> OpenConfigFiles;
 	TArray<FString> ConfigSettingDenyList;
-	void OnFConfigDeleted(const FConfigFile* Config);
-	void OnFConfigCreated(const FConfigFile* Config);
-
-	void ProcessAccessedIniSettings(const FConfigFile* Config, UE::Cook::FIniSettingContainer& AccessedIniStrings) const;
 
 	void OnRequestClusterCompleted(const UE::Cook::FRequestCluster& RequestCluster);
 
@@ -1566,6 +1558,8 @@ private:
 	/** Helper struct for running cooking in diagnostic modes */
 	TUniquePtr<FDiffModeCookServerUtils> DiffModeHelper;
 
+	TRefCountPtr<UE::Cook::IMPCollector> ConfigCollector;
+
 	/**
 	 * Heap of Pollables to tick, ordered by NextTimeSeconds.
 	 * Only accessible by PumpPollables when bPollablesInTick==true, can be accessed outside of PollablesLock.
diff --git a/Engine/Source/Editor/UnrealEd/Private/CookOnTheFlyServer.cpp b/Engine/Source/Editor/UnrealEd/Private/CookOnTheFlyServer.cpp
--- a/Engine/Source/Editor/UnrealEd/Private/CookOnTheFlyServer.cpp
+++ b/Engine/Source/Editor/UnrealEd/Private/CookOnTheFlyServer.cpp
@@ -20,6 +20,7 @@
 #include "Containers/DirectoryTree.h"
 #include "Containers/RingBuffer.h"
 #include "Cooker/AsyncIODelete.h"
+#include "Cooker/CookConfigAccessTracker.h"
 #include "Cooker/CookDiagnostics.h"
 #include "Cooker/CookDirector.h"
 #include "Cooker/CookOnTheFlyServerInterface.h"
@@ -6839,9 +6840,6 @@ void UCookOnTheFlyServer::Initialize( ECookMode::Type DesiredCookMode, ECookInit
 		FCoreDelegates::OnTargetPlatformChangedSupportedFormats.AddUObject(this, &UCookOnTheFlyServer::OnTargetPlatformChangedSupportedFormats);
 	}
 
-	FCoreDelegates::TSOnFConfigCreated().AddUObject(this, &UCookOnTheFlyServer::OnFConfigCreated);
-	FCoreDelegates::TSOnFConfigDeleted().AddUObject(this, &UCookOnTheFlyServer::OnFConfigDeleted);
-
 	GetTargetPlatformManager()->GetOnTargetPlatformsInvalidatedDelegate().AddUObject(this, &UCookOnTheFlyServer::OnTargetPlatformsInvalidated);
 #if WITH_ADDITIONAL_CRASH_CONTEXTS
 	FGenericCrashContext::OnAdditionalCrashContextDelegate().AddUObject(this, &UCookOnTheFlyServer::DumpCrashContext);
@@ -7705,120 +7703,279 @@ void GetAdditionalCurrentIniVersionStrings( const UCookOnTheFlyServer* CookOnThe
 
 bool UCookOnTheFlyServer::GetCurrentIniVersionStrings( const ITargetPlatform* TargetPlatform, UE::Cook::FIniSettingContainer& IniVersionStrings ) const
 {
-	{
-		FScopeLock Lock(&ConfigFileCS);
-		IniVersionStrings = AccessedIniStrings;
-	}
+#if !UE_WITH_CONFIG_TRACKING
+	IniVersionStrings.Reset();
+	return true;
+#else
+	// This function should be called after the cook is finished
+	TArray<UE::ConfigAccessTracking::FConfigAccessData> AccessedRecordsArray =
+		UE::ConfigAccessTracking::FCookConfigAccessTracker::Get().GetCookRecords(TargetPlatform);
+	TConstArrayView<UE::ConfigAccessTracking::FConfigAccessData> AccessedRecords(AccessedRecordsArray);
 
-	// this should be called after the cook is finished
-	TArray<FString> IniFiles;
-	GConfig->GetConfigFilenames(IniFiles);
+	TArray<const FConfigValue*> Values;
+	int32 ConfigFileEndIndex = 0;
+	int32 EndIndex = AccessedRecords.Num();
+	TStringBuilder<128> FullConfigFileNameStr;
+	while (ConfigFileEndIndex < EndIndex)
+	{
+		int32 ConfigFileStartIndex = ConfigFileEndIndex;
+		const UE::ConfigAccessTracking::FConfigAccessData& FileStartRecord = AccessedRecords[ConfigFileStartIndex];
+		++ConfigFileEndIndex;
+		while (ConfigFileEndIndex < EndIndex && AccessedRecords[ConfigFileEndIndex].IsSameConfigFile(FileStartRecord))
+		{
+			++ConfigFileEndIndex;
+		}
 
-	TMap<FString, int32> MultiMapCounter;
+		FConfigFile Temp;
+		FName ConfigPlatform(FileStartRecord.GetConfigPlatform());
+		FName ConfigFileName(FileStartRecord.GetFileName());
+		TStringBuilder<64> ConfigFileNameStr(InPlace, ConfigFileName);
+		using UE::String::FindFirst;
 
-	for ( const FString& ConfigFilename : IniFiles )
-	{
-		if ( ConfigFilename.Contains(TEXT("CookedIniVersion.txt")) )
+		// Hardcoded additions to ConfigSettingsDenyList for ini files used by the cook. These are early-exited earlier
+		// to prevent bugs from arising if we tried to track their data and discard later.
+		if (FindFirst(ConfigFileNameStr, TEXT("CookedIniVersion.txt"), ESearchCase::IgnoreCase) != INDEX_NONE ||
+			FindFirst(ConfigFileNameStr, TEXT("CookedSettings.txt"), ESearchCase::IgnoreCase) != INDEX_NONE)
 		{
 			continue;
 		}
+		TStringBuilder<32> ConfigPlatformStr;
+		const TCHAR* SavePlatformStr = TEXT("<Editor>");
+		if (!ConfigPlatform.IsNone())
+		{
+			ConfigPlatformStr << ConfigPlatform;
+			SavePlatformStr = *ConfigPlatformStr;
+		}
+		const TCHAR* LoadTypeStr = LexToString(FileStartRecord.LoadType);
+		FullConfigFileNameStr.Reset();
+		FullConfigFileNameStr << LoadTypeStr << TEXT(".") << SavePlatformStr << TEXT(".") << ConfigFileNameStr;
+		FName FullConfigFileName(FullConfigFileNameStr);
+		const FConfigFile* ConfigFile = UE::ConfigAccessTracking::FindOrLoadConfigFile(FileStartRecord.LoadType,
+			*ConfigPlatformStr, *ConfigFileNameStr, Temp);
+		if (!ConfigFile)
+		{
+			// This is logged as Warning; it is unexpected that we were able to load a file from disk that
+			// existed previously when we received the OnConfigValueRead call.
+			UE_LOG(LogCook, Warning,
+				TEXT("Could not load config file '%s'. Changes to settings in this file will not be detected in iterative cooks."),
+				*FullConfigFileNameStr);
+			continue;
+		}
 
-		const FConfigFile *ConfigFile = GConfig->FindConfigFile(ConfigFilename);
-		ProcessAccessedIniSettings(ConfigFile, IniVersionStrings);
-		
+		TMap<FName, TMap<FName, TArray<FString>>>& FileVersionStrings = IniVersionStrings.FindOrAdd(FullConfigFileName);
+		int32 ConfigSectionEndIndex = ConfigFileStartIndex;
+		while (ConfigSectionEndIndex < ConfigFileEndIndex)
+		{
+			int32 ConfigSectionStartIndex = ConfigSectionEndIndex;
+			++ConfigSectionEndIndex;
+			const UE::ConfigAccessTracking::FConfigAccessData& SectionStartRecord = AccessedRecords[ConfigSectionStartIndex];
+			while (ConfigSectionEndIndex < ConfigFileEndIndex && AccessedRecords[ConfigSectionEndIndex].SectionName == SectionStartRecord.SectionName)
+			{
+				++ConfigSectionEndIndex;
+			}
+
+			FName SectionName(SectionStartRecord.GetSectionName());
+			const FConfigSection* ConfigSection = ConfigFile->FindSection(SectionName.ToString());
+ 			if (!ConfigSection)
+			{
+				// This is logged as Verbose rather than Warning because the section could have been added by code
+				// after loading and never existed on disk.
+				UE_LOG(LogCook, Verbose,
+					TEXT("Could not find config section %s:[%s]. Changes to settings in this section will not be detected in iterative cooks."),
+					*FullConfigFileNameStr, *WriteToString<32>(SectionName));
+				continue;
+			}
+			TMap<FName, TArray<FString>>& SectionVersionStrings = FileVersionStrings.FindOrAdd(SectionName);
+
+			for (const UE::ConfigAccessTracking::FConfigAccessData& Record :
+				AccessedRecords.Slice(ConfigSectionStartIndex, ConfigSectionEndIndex - ConfigSectionStartIndex))
+			{
+				FName ValueName(Record.GetValueName());
+				Values.Reset();
+				ConfigSection->MultiFindPointer(ValueName, Values, true /* bMaintainOrder */);
+				if (Values.IsEmpty())
+				{
+					// This is logged as Verbose rather than Warning because the value could have been added by code
+					// after loading and never existed on disk.
+					UE_LOG(LogCook, Verbose,
+						TEXT("Could not find config value %s:[%s]:%s. Changes to this value will not be detected in iterative cooks."),
+						*FullConfigFileNameStr, *WriteToString<32>(SectionName), *WriteToString<32>(ValueName));
+					continue;
+				}
+				TArray<FString>& ValueVersionStrings = SectionVersionStrings.FindOrAdd(ValueName);
+				for (const FConfigValue* Value : Values)
+				{
+					FString ValueStr = Value->GetSavedValue();
+					ValueStr.ReplaceInline(TEXT(":"), TEXT(""));
+					ValueVersionStrings.Add(MoveTemp(ValueStr));
+				}
+			}
+		}
 	}
 
+	// remove any ConfigFiles,Sections,Values which are marked as ignored by ConfigSettingDenyList
+	struct FParsedDenyEntry
 	{
-		FScopeLock Lock(&ConfigFileCS);
-		for (const FConfigFile* ConfigFile : OpenConfigFiles)
+		FStringView ConfigFileName;
+		FStringView SectionName;
+		FStringView ValueName;
+	};
+	TArray<FParsedDenyEntry> ParsedConfigSettings;
+	TArray<FStringView> Tokens;
+	for (const FString& Filter : ConfigSettingDenyList)
+	{
+		Tokens.Reset();
+		UE::String::ParseTokens(Filter, ':', Tokens, UE::String::EParseTokensOptions::Trim | UE::String::EParseTokensOptions::SkipEmpty);
+		if (Tokens.Num() >= 1)
 		{
-			ProcessAccessedIniSettings(ConfigFile, IniVersionStrings);
+			FParsedDenyEntry& DenyEntry = ParsedConfigSettings.Emplace_GetRef();
+			DenyEntry.ConfigFileName = Tokens[0];
+			if (Tokens.Num() >= 2) DenyEntry.SectionName = Tokens[1];
+			if (Tokens.Num() >= 3) DenyEntry.ValueName = Tokens[2];
 		}
 	}
 
-	// remove any which are filtered out
-	FString EditorPrefix(TEXT("Editor."));
-	for ( const FString& Filter : ConfigSettingDenyList )
+	TStringBuilder<128> FullConfigFileName;
+	for (UE::Cook::FIniSettingContainer::TIterator ConfigFile(IniVersionStrings.CreateIterator()); ConfigFile; ++ConfigFile)
 	{
-		TArray<FString> FilterArray;
-		Filter.ParseIntoArray( FilterArray, TEXT(":"));
-
-		FString *ConfigFileName = nullptr;
-		FString *SectionName = nullptr;
-		FString *ValueName = nullptr;
-		switch ( FilterArray.Num() )
-		{
-		case 3:
-			ValueName = &FilterArray[2];
-		case 2:
-			SectionName = &FilterArray[1];
-		case 1:
-			ConfigFileName = &FilterArray[0];
-			break;
-		default:
-			continue;
-		}
+		ConfigFile.Key().ToString(FullConfigFileName);
+		FStringView FullConfigFileNameView(FullConfigFileName);
+		int32 DotIndex = FullConfigFileNameView.Find(TEXTVIEW("."));
+		FString Platform = FString(FullConfigFileNameView.LeftChop(DotIndex));
+		FString PlatformAndFileName = FString(FullConfigFileNameView.RightChop(DotIndex + 1));
+		FString ConfigFileName = PlatformAndFileName.RightChop(PlatformAndFileName.Find(TEXTVIEW(".")) + 1);
+		FString BaseFileName = FPaths::GetBaseFilename(ConfigFileName);
+		FString PlatformAndBaseFileName = FString::Printf(TEXT("%s.%s"), *Platform, *BaseFileName);
 
-		if ( ConfigFileName )
+		for (const FParsedDenyEntry& DenyEntry : ParsedConfigSettings)
 		{
-			for ( auto ConfigFile = IniVersionStrings.CreateIterator(); ConfigFile; ++ConfigFile )
+			// FullConfigFileName is of the form "LoadType.Platform.ConfigFile".
+			// Wildcards are written in the form "*.ConfigFile" or "ConfigFile".
+			// We allow a match of the wildcard against either Platform.ConfigFile or just ConfigFile.
+			// We also allow a match of the wildcard against Platform.BaseName or BaseName.
+			if (PlatformAndFileName.MatchesWildcard(DenyEntry.ConfigFileName) ||
+				ConfigFileName.MatchesWildcard(DenyEntry.ConfigFileName) ||
+				PlatformAndBaseFileName.MatchesWildcard(DenyEntry.ConfigFileName) ||
+				BaseFileName.MatchesWildcard(DenyEntry.ConfigFileName))
 			{
-				// Some deny list entries are written as *.Engine, and are intended to affect the platform-less Editor Engine.ini, which is just "Engine"
-				// To make *.Engine match the editor-only config files as well, we check whether the wildcard matches either Engine or Editor.Engine for the editor files
-				FString IniVersionStringFilename = ConfigFile.Key().ToString();
-				if (IniVersionStringFilename.MatchesWildcard(*ConfigFileName) ||
-					(!IniVersionStringFilename.Contains(TEXT(".")) && (EditorPrefix + IniVersionStringFilename).MatchesWildcard(*ConfigFileName)))
+				if (!DenyEntry.SectionName.IsEmpty())
 				{
-					if ( SectionName )
+					for (TMap<FName,TMap<FName,TArray<FString>>>::TIterator Section(ConfigFile.Value().CreateIterator());
+						Section; ++Section)
 					{
-						for ( auto Section = ConfigFile.Value().CreateIterator(); Section; ++Section )
+						if (Section.Key().ToString().MatchesWildcard(DenyEntry.SectionName))
 						{
-							if ( Section.Key().ToString().MatchesWildcard(*SectionName))
+							if (!DenyEntry.ValueName.IsEmpty())
 							{
-								if (ValueName)
+								for (TMap<FName, TArray<FString>>::TIterator Value(Section.Value().CreateIterator()); Value; ++Value)
 								{
-									for ( auto Value = Section.Value().CreateIterator(); Value; ++Value )
+									if (Value.Key().ToString().MatchesWildcard(DenyEntry.ValueName))
 									{
-										if ( Value.Key().ToString().MatchesWildcard(*ValueName))
-										{
-											Value.RemoveCurrent();
-										}
+										Value.RemoveCurrent();
 									}
 								}
-								else
-								{
-									Section.RemoveCurrent();
-								}
+							}
+							else
+							{
+								Section.RemoveCurrent();
 							}
 						}
 					}
-					else
-					{
-						ConfigFile.RemoveCurrent();
-					}
+				}
+				else
+				{
+					ConfigFile.RemoveCurrent();
+					break;
 				}
 			}
 		}
 	}
 	return true;
+#endif // !UE_WITH_CONFIG_TRACKING
+}
+
+namespace
+{
+
+void EscapeUsedSettingsToken(FName Token, FStringBuilderBase& Result)
+{
+	Result.Reset();
+	Result << Token;
+	if (Result.ToView().Contains(TEXTVIEW(":")))
+	{
+		FString ReplaceText(Result);
+		ReplaceText.ReplaceInline(TEXT(":"), TEXT("::"));
+		Result.Reset();
+		Result << ReplaceText;
+	}
 }
 
+bool TryTokenizeUsedSettingsString(FStringView Text, TArrayView<FStringBuilderBase*> OutTokens)
+{
+	int32 TextLen = Text.Len();
+	if (TextLen == 0)
+	{
+		return false;
+	}
+	const TCHAR* TextData = Text.GetData();
+
+	int32 NumTokens = OutTokens.Num();
+	check(NumTokens > 0);
+	int32 TokenIndex = 0;
+	FStringBuilderBase* OutToken = OutTokens[TokenIndex++];
+	OutToken->Reset();
+	for (int32 Index = 0; Index < TextLen; ++Index)
+	{
+		TCHAR C = TextData[Index];
+		if (C != ':')
+		{
+			OutToken->AppendChar(C);
+		}
+		else if (Index < TextLen - 1 && TextData[Index + 1] == ':')
+		{
+			++Index;
+			OutToken->AppendChar(':');
+		}
+		else
+		{
+			if (OutToken->Len() == 0)
+			{
+				// Empty token
+				return false;
+			}
+			if (TokenIndex >= NumTokens)
+			{
+				// Too many tokens
+				return false;
+			}
+			OutToken = OutTokens[TokenIndex++];
+			OutToken->Reset();
+		}
+	}
+	if (OutToken->Len() == 0)
+	{
+		// Empty token
+		return false;
+	}
+	if (TokenIndex < NumTokens)
+	{
+		// too few tokens
+		return false;
+	}
+	return true;
+}
+
+} // anonymous namespace
 
 bool UCookOnTheFlyServer::GetCookedIniVersionStrings(const ITargetPlatform* TargetPlatform, UE::Cook::FIniSettingContainer& OutIniSettings, TMap<FString,FString>& OutAdditionalSettings) const
 {
 	const FString EditorIni = GetMetadataDirectory() / TEXT("CookedIniVersion.txt");
 	const FString SandboxEditorIni = ConvertToFullSandboxPath(*EditorIni, true);
-
-
 	const FString PlatformSandboxEditorIni = SandboxEditorIni.Replace(TEXT("[Platform]"), *TargetPlatform->PlatformName());
 
-	TArray<FString> SavedIniVersionedParams;
-
 	FConfigFile ConfigFile;
 	ConfigFile.Read(*PlatformSandboxEditorIni);
 
-	
-
 	const static FString NAME_UsedSettings(TEXT("UsedSettings"));
 	const FConfigSection* UsedSettings = ConfigFile.FindSection(NAME_UsedSettings);
 	if (UsedSettings == nullptr)
@@ -7826,7 +7983,6 @@ bool UCookOnTheFlyServer::GetCookedIniVersionStrings(const ITargetPlatform* Targ
 		return false;
 	}
 
-
 	const static FString NAME_AdditionalSettings(TEXT("AdditionalSettings"));
 	const FConfigSection* AdditionalSettings = ConfigFile.FindSection(NAME_AdditionalSettings);
 	if (AdditionalSettings == nullptr)
@@ -7834,40 +7990,36 @@ bool UCookOnTheFlyServer::GetCookedIniVersionStrings(const ITargetPlatform* Targ
 		return false;
 	}
 
+	TStringBuilder<256> KeyStr;
+	TStringBuilder<128> Filename;
+	TStringBuilder<64> SectionName;
+	TStringBuilder<64> ValueName;
+	TStringBuilder<64> ValueIndexStr;
+	FStringBuilderBase* TokenBuffer[] = { &Filename, &SectionName, &ValueName, &ValueIndexStr };
+	TArrayView<FStringBuilderBase*> Tokens = TokenBuffer;
 
 	for (const auto& UsedSetting : *UsedSettings )
 	{
-		FName Key = UsedSetting.Key;
-		const FConfigValue& UsedValue = UsedSetting.Value;
-
-		TArray<FString> SplitString;
-		Key.ToString().ParseIntoArray(SplitString, TEXT(":"));
-
-		if (SplitString.Num() != 4)
+		KeyStr.Reset();
+		KeyStr << UsedSetting.Key;
+		if (!TryTokenizeUsedSettingsString(KeyStr, Tokens))
 		{
-			UE_LOG(LogCook, Warning, TEXT("Found unparsable ini setting %s for platform %s, invalidating cook."), *Key.ToString(), *TargetPlatform->PlatformName());
+			UE_LOG(LogCook, Warning, TEXT("Found unparsable ini setting %s for platform %s, invalidating cook."),
+				*KeyStr, *TargetPlatform->PlatformName());
 			return false;
 		}
 
-
-		check(SplitString.Num() == 4); // We generate this ini file in SaveCurrentIniSettings
-		const FString& Filename = SplitString[0];
-		const FString& SectionName = SplitString[1];
-		const FString& ValueName = SplitString[2];
-		const int32 ValueIndex = FCString::Atoi(*SplitString[3]);
-
-		auto& OutFile = OutIniSettings.FindOrAdd(FName(*Filename));
-		auto& OutSection = OutFile.FindOrAdd(FName(*SectionName));
-		auto& ValueArray = OutSection.FindOrAdd(FName(*ValueName));
+		auto& OutFile = OutIniSettings.FindOrAdd(FName(Filename));
+		auto& OutSection = OutFile.FindOrAdd(FName(SectionName));
+		auto& ValueArray = OutSection.FindOrAdd(FName(ValueName));
+		const int32 ValueIndex = FCString::Atoi(*ValueIndexStr);
 		if ( ValueArray.Num() < (ValueIndex+1) )
 		{
 			ValueArray.AddZeroed( ValueIndex - ValueArray.Num() +1 );
 		}
-		ValueArray[ValueIndex] = UsedValue.GetSavedValue();
+		ValueArray[ValueIndex] = UsedSetting.Value.GetSavedValue();
 	}
 
-
-
 	for (const auto& AdditionalSetting : *AdditionalSettings)
 	{
 		const FName& Key = AdditionalSetting.Key;
@@ -7878,163 +8030,6 @@ bool UCookOnTheFlyServer::GetCookedIniVersionStrings(const ITargetPlatform* Targ
 	return true;
 }
 
-static thread_local bool GSuppressProcessConfigSettings = false;
-
-void UCookOnTheFlyServer::OnFConfigCreated(const FConfigFile* Config)
-{
-	if (GSuppressProcessConfigSettings)
-	{
-		return;
-	}
-
-	FScopeLock Lock(&ConfigFileCS);
-	OpenConfigFiles.Add(Config);
-}
-
-void UCookOnTheFlyServer::OnFConfigDeleted(const FConfigFile* Config)
-{
-	if (GSuppressProcessConfigSettings)
-	{
-		return;
-	}
-
-	FScopeLock Lock(&ConfigFileCS);
-	ProcessAccessedIniSettings(Config, AccessedIniStrings);
-	OpenConfigFiles.Remove(Config);
-}
-
-void UCookOnTheFlyServer::ProcessAccessedIniSettings(const FConfigFile* Config, UE::Cook::FIniSettingContainer& OutAccessedIniStrings) const
-{	
-	if (Config->Name == NAME_None)
-	{
-		return;
-	}
-
-	// try to figure out if this config file is for a specific platform 
-	FString PlatformName;
-	bool bFoundPlatformName = false;
-
-	if (GConfig->ContainsConfigFile(Config))
-	{
-		// If the ConfigFile is in GConfig, then it is the editor's config and is not platform specific
-	}
-	else if (Config->bHasPlatformName)
-	{
-		// The platform that was passed to LoadExternalIniFile
-		PlatformName = Config->PlatformName;
-		bFoundPlatformName = !PlatformName.IsEmpty();
-	}
-	else
-	{
-		// For the config files not in GConfig, we assume they were loaded from LoadConfigFile, and we match these to a platform
-		// By looking for a platform-specific filepath in their SourceIniHierarchy.
-		// Examples:
-		// (1) ROOT\Engine\Config\Windows\WindowsEngine.ini
-		// (2) ROOT\Engine\Config\Android\DataDrivePlatformInfo.ini
-		// (3) ROOT\Engine\Config\Android\AndroidWindowsCompatability.ini
-		// 
-		// Note that for config files of form #3, we want them to be matched to Android rather than windows;
-		// we assume that an exact match on a directory component is more definitive than a substring match
-		bool bFoundPlatformGuess = false;
-		for (auto It : FDataDrivenPlatformInfoRegistry::GetAllPlatformInfos())
-		{
-			const FString CurrentPlatformName = It.Key.ToString();
-			TStringBuilder<128> PlatformDirString;
-			PlatformDirString.Appendf(TEXT("/%s/"), *CurrentPlatformName);
-			for (const auto& SourceIni : Config->SourceIniHierarchy)
-			{
-				// Look for platform in the path, rating a full subdirectory name match (/Android/ or /Windows/) higher than a partial filename match (AndroidEngine.ini or WindowsEngine.ini)
-				bool bFoundPlatformDir = UE::String::FindFirst(SourceIni.Value, PlatformDirString, ESearchCase::IgnoreCase) != INDEX_NONE;
-				bool bFoundPlatformSubstring = UE::String::FindFirst(SourceIni.Value, CurrentPlatformName, ESearchCase::IgnoreCase) != INDEX_NONE;
-				if (bFoundPlatformDir)
-				{
-					PlatformName = CurrentPlatformName;
-					bFoundPlatformName = true;
-					break;
-				}
-				else if (!bFoundPlatformGuess && bFoundPlatformSubstring)
-				{
-					PlatformName = CurrentPlatformName;
-					bFoundPlatformGuess = true;
-				}
-			}
-			if (bFoundPlatformName)
-			{
-				break;
-			}
-		}
-		bFoundPlatformName = bFoundPlatformName || bFoundPlatformGuess;
-	}
-
-	TStringBuilder<128> ConfigName;
-	if (bFoundPlatformName)
-	{
-		ConfigName << PlatformName;
-		ConfigName << TEXT(".");
-	}
-	Config->Name.AppendString(ConfigName);
-	const FName& ConfigFName = FName(ConfigName);
-	TSet<FName> ProcessedValues;
-	TCHAR PlainNameString[NAME_SIZE];
-	TArray<const FConfigValue*> ValueArray;
-	for ( auto& ConfigSection : *Config )
-	{
-		ProcessedValues.Reset();
-		const FName SectionName = FName(*ConfigSection.Key);
-
-		SectionName.GetPlainNameString(PlainNameString);
-		if ( TCString<TCHAR>::Strstr(PlainNameString, TEXT(":")) )
-		{
-			UE_LOG(LogCook, Verbose, TEXT("Ignoring ini section checking for section name %s because it contains ':'"), PlainNameString);
-			continue;
-		}
-
-		for ( auto& ConfigValue : ConfigSection.Value )
-		{
-			const FName& ValueName = ConfigValue.Key;
-			if ( ProcessedValues.Contains(ValueName) )
-				continue;
-
-			ProcessedValues.Add(ValueName);
-
-			ValueName.GetPlainNameString(PlainNameString);
-			if (TCString<TCHAR>::Strstr(PlainNameString, TEXT(":")))
-			{
-				UE_LOG(LogCook, Verbose, TEXT("Ignoring ini section checking for section name %s because it contains ':'"), PlainNameString);
-				continue;
-			}
-
-			
-			ValueArray.Reset();
-			ConfigSection.Value.MultiFindPointer( ValueName, ValueArray, true );
-
-			bool bHasBeenAccessed = false;
-			for (const FConfigValue* ValueArrayEntry : ValueArray)
-			{
-				if (ValueArrayEntry->HasBeenRead())
-				{
-					bHasBeenAccessed = true;
-					break;
-				}
-			}
-
-			if ( bHasBeenAccessed )
-			{
-				auto& AccessedConfig = OutAccessedIniStrings.FindOrAdd(ConfigFName);
-				auto& AccessedSection = AccessedConfig.FindOrAdd(SectionName);
-				auto& AccessedKey = AccessedSection.FindOrAdd(ValueName);
-				AccessedKey.Empty(ValueArray.Num());
-				for (const FConfigValue* ValueArrayEntry : ValueArray )
-				{
-					FString RemovedColon = ValueArrayEntry->GetSavedValue().Replace(TEXT(":"), TEXT(""));
-					AccessedKey.Add(MoveTemp(RemovedColon));
-				}
-			}
-			
-		}
-	}
-}
-
 static const TCHAR* TEXT_CookSettings(TEXT("CookSettings"));
 static FName ExecutableHashName(TEXT("ExecutableHash"));
 static FName ExecutableHashInvalidModuleName(TEXT("ExecutableHashInvalidModule"));
@@ -8089,7 +8084,7 @@ TMap<FName, FString> UCookOnTheFlyServer::CalculateCookSettingStrings() const
 		}
 	}
 
-	CookSettingStrings.Add(FName(TEXT("Version")), TEXT("C7C76F79"));
+	CookSettingStrings.Add(FName(TEXT("Version")), TEXT("21F52B9EDD4D456AB1AF381CA172BD28"));
 	if (IsDirectorCookByTheBook())
 	{
 		CookSettingStrings.Add(NAME_CookMode, TEXT("CookByTheBook"));
@@ -8180,6 +8175,8 @@ bool UCookOnTheFlyServer::ArePreviousCookSettingsCompatible(const TMap<FName, FS
 
 void UCookOnTheFlyServer::SaveCookSettings(const TMap<FName, FString>& CurrentCookSettings, const ITargetPlatform* TargetPlatform)
 {
+	UE::ConfigAccessTracking::FIgnoreScope IgnoreScope;
+
 	FConfigFile ConfigFile;
 	for (const TPair<FName, FString>& CurrentSetting : CurrentCookSettings)
 	{
@@ -8191,13 +8188,16 @@ void UCookOnTheFlyServer::SaveCookSettings(const TMap<FName, FString>& CurrentCo
 
 bool UCookOnTheFlyServer::IniSettingsOutOfDate(const ITargetPlatform* TargetPlatform) const
 {
-	TGuardValue<bool> A(GSuppressProcessConfigSettings, true);
+#if !UE_WITH_CONFIG_TRACKING
+	return false;
+#else
+	UE::ConfigAccessTracking::FIgnoreScope IgnoreScope;
 
 	UE::Cook::FIniSettingContainer OldIniSettings;
 	TMap<FString, FString> OldAdditionalSettings;
 	if ( GetCookedIniVersionStrings(TargetPlatform, OldIniSettings, OldAdditionalSettings) == false)
 	{
-		UE_LOG(LogCook, Display, TEXT("Unable to read previous cook inisettings for platform %s invalidating cook"), *TargetPlatform->PlatformName());
+		UE_LOG(LogCook, Display, TEXT("Invalidating inisettings: Unable to read previous cook inisettings for platform %s."), *TargetPlatform->PlatformName());
 		return true;
 	}
 
@@ -8210,90 +8210,112 @@ bool UCookOnTheFlyServer::IniSettingsOutOfDate(const ITargetPlatform* TargetPlat
 		const FString* CurrentValue = CurrentAdditionalSettings.Find(OldIniSetting.Key);
 		if ( !CurrentValue )
 		{
-			UE_LOG(LogCook, Display, TEXT("Previous cook had additional ini setting: %s current cook is missing this setting."), *OldIniSetting.Key);
+			UE_LOG(LogCook, Display,
+				TEXT("Invalidating inisettings: Unable to find additional ini setting used by platform %s: %s was not found."),
+				*TargetPlatform->PlatformName(), *OldIniSetting.Key);
 			return true;
 		}
 
 		if ( *CurrentValue != OldIniSetting.Value )
 		{
-			UE_LOG(LogCook, Display, TEXT("Additional Setting from previous cook %s doesn't match %s vs %s"), *OldIniSetting.Key, **CurrentValue, *OldIniSetting.Value );
+			UE_LOG(LogCook, Display,
+				TEXT("Invalidating inisettings: Additional ini setting used by platform %s is different for %s, value '%s' != '%s'."),
+				*TargetPlatform->PlatformName(), *OldIniSetting.Key, **CurrentValue, *OldIniSetting.Value );
 			return true;
 		}
 	}
 
+	TStringBuilder<256> ConfigNameKeyStr;
+	TStringBuilder<128> Filename;
+	TStringBuilder<256> PlatformName;
 	for (const auto& OldIniFile : OldIniSettings)
 	{
-		FName ConfigNameKey = OldIniFile.Key;
+		OldIniFile.Key.ToString(ConfigNameKeyStr);
 
-		TArray<FString> ConfigNameArray;
-		ConfigNameKey.ToString().ParseIntoArray(ConfigNameArray, TEXT("."));
-		FString Filename;
-		FString PlatformName;
-		// The input NameKey is of the form 
-		//   Platform.ConfigName:Section:Key:ArrayIndex=Value
-		// The Platform is optional and will not be present if the configfile was an editor config file rather than a platform-specific config file
-		bool bFoundPlatformName = false;
-		if (ConfigNameArray.Num() <= 1)
+		// ConfigSystem.<Editor>.../../../Engine/Config/ConsoleVariables.ini
+		//   -> "ConfigSystem", "<Editor>", "../../../Engine/Config/ConsoleVariables.ini"
+		// 3rd token might have dots, first two cannot.
+		FStringView ConfigNameArray[3];
+		ConfigNameArray[0] = ConfigNameKeyStr;
+		int32 NumTokens = 1;
+		for (NumTokens = 1; NumTokens < UE_ARRAY_COUNT(ConfigNameArray); ++NumTokens)
 		{
-			Filename = ConfigNameKey.ToString();
+			FStringView& Current = ConfigNameArray[NumTokens-1];
+			FStringView& Next= ConfigNameArray[NumTokens];
+			int32 DotIndex = Current.Find(TEXTVIEW("."));
+			if (DotIndex == INDEX_NONE)
+			{
+				break;
+			}
+			Next = Current.RightChop(DotIndex+1);
+			Current.LeftInline(DotIndex);
+			if (Next.IsEmpty())
+			{
+				// break before incrementing NumTokens so that we do not count the empty Next as a token
+				break; 
+			}
 		}
-		else if (ConfigNameArray.Num() == 2)
+		UE::ConfigAccessTracking::ELoadType LoadType = UE::ConfigAccessTracking::ELoadType::Uninitialized;
+		// The input NameKey is of the form 
+		//   LoadType.Platform.ConfigName
+		// Platform is <Editor> if the configfile was an editor config file rather than a platform-specific config file
+		if (NumTokens == 3)
 		{
-			PlatformName = ConfigNameArray[0];
-			Filename = ConfigNameArray[1];
-			bFoundPlatformName = true;
+			LexFromString(LoadType, ConfigNameArray[0]);
+			if (!UE::ConfigAccessTracking::IsLoadableLoadType(LoadType))
+			{
+				LoadType = UE::ConfigAccessTracking::ELoadType::Uninitialized;
+			}
+			else
+			{
+				PlatformName.Reset();
+				PlatformName << ConfigNameArray[1];
+				Filename.Reset();
+				Filename << ConfigNameArray[2];
+			}
 		}
-		else
+		if (LoadType == UE::ConfigAccessTracking::ELoadType::Uninitialized)
 		{
-			UE_LOG(LogCook, Warning, TEXT("Found invalid file name in old ini settings file Filename %s settings file %s"), *ConfigNameKey.ToString(), *TargetPlatform->PlatformName());
+			UE_LOG(LogCook, Warning,
+				TEXT("Invalidating inisettings: Invalid filename key in old ini settings file used by platform %s: key '%s' is invalid."),
+				*TargetPlatform->PlatformName(), *ConfigNameKeyStr);
 			return true;
 		}
-		
-		const FConfigFile* ConfigFile = nullptr;
-		FConfigFile Temp;
-		if (bFoundPlatformName)
+		const TCHAR* PlatformNameToLoad = *PlatformName;
+		if (FCString::Stricmp(PlatformNameToLoad, TEXT("<Editor>")) == 0)
 		{
-			// For the platform-specific old ini files, load them using LoadLocalIniFiles; this matches the assumption in SaveCurrentIniSettings
-			// that the platform-specific ini files were loaded by LoadLocalIniFiles
-			FConfigCacheIni::LoadLocalIniFile(Temp, *Filename, true, *PlatformName);
-			ConfigFile = &Temp;
-		}
-		else
-		{
-			// For the platform-agnostic old ini files, read them from GConfig; this matches where we loaded them from in SaveCurrentIniSettings
-			// The ini files may have been saved by fullpath or by shortname; search first for a fullpath match using FindConfigFile and
-			// if that fails search for the shortname match by iterating over all files in GConfig
-			ConfigFile = GConfig->FindConfigFile(Filename);
+			PlatformNameToLoad = TEXT("");
 		}
+
+		FConfigFile Temp;
+		const FConfigFile* ConfigFile = UE::ConfigAccessTracking::FindOrLoadConfigFile(LoadType, PlatformNameToLoad,
+			*Filename, Temp);
 		if (!ConfigFile)
 		{
-			FName FileFName = FName(*Filename);
-			for (const FString& ConfigFilename : GConfig->GetFilenames())
-			{
-				FConfigFile* File = GConfig->FindConfigFile(ConfigFilename);
-				if (File->Name == FileFName)
-				{
-					ConfigFile = File;
-					break;
-				}
-			}
-			if (!ConfigFile)
-			{
-				UE_LOG(LogCook, Display, TEXT("Unable to find config file %s invalidating inisettings"), *FString::Printf(TEXT("%s %s"), *PlatformName, *Filename));
-				return true;
-			}
+			UE_LOG(LogCook, Display,
+				TEXT("Invalidating inisettings: Unable to find config file in old ini settings file used by platform %s: '%s' was not found."),
+				*TargetPlatform->PlatformName(), *ConfigNameKeyStr);
+			return true;
 		}
+
 		for ( const auto& OldIniSection : OldIniFile.Value )
 		{
-			const FName& SectionName = OldIniSection.Key;
+			FName SectionName = OldIniSection.Key;
 			const FConfigSection* IniSection = ConfigFile->FindSection( SectionName.ToString() );
-			const FString DenyListSetting = FString::Printf(TEXT("%s%s%s:%s"), *PlatformName, bFoundPlatformName ? TEXT(".") : TEXT(""), *Filename, *SectionName.ToString());
+			auto GetDenyListMessageStart = [&PlatformName, &Filename, SectionName]()
+				{
+					return FString::Printf(
+						TEXT("To avoid invalidating due to this setting, add a deny list setting")
+						TEXT("\n\tDefaultEditor.ini:[CookSettings]:+CookOnTheFlyConfigSettingDenyList=%s.%s:%s"),
+						*PlatformName, *Filename, *SectionName.ToString());
+				};
 
 			if ( IniSection == nullptr )
 			{
-				UE_LOG(LogCook, Display, TEXT("Inisetting is different for %s, Current section doesn't exist"), 
-					*FString::Printf(TEXT("%s %s %s"), *PlatformName, *Filename, *SectionName.ToString()));
-				UE_LOG(LogCook, Display, TEXT("To avoid this add a deny list setting to DefaultEditor.ini [CookSettings] %s"), *DenyListSetting);
+				UE_LOG(LogCook, Display,
+					TEXT("Invalidating inisettings: Inisetting used by platform %s is different for %s:[%s], section doesn't exist in current config."),
+					*TargetPlatform->PlatformName(), *ConfigNameKeyStr, *SectionName.ToString());
+				UE_LOG(LogCook, Display, TEXT("%s"), *GetDenyListMessageStart());
 				return true;
 			}
 
@@ -8306,9 +8328,11 @@ bool UCookOnTheFlyServer::IniSettingsOutOfDate(const ITargetPlatform* TargetPlat
 
 				if ( CurrentValues.Num() != OldIniValue.Value.Num() )
 				{
-					UE_LOG(LogCook, Display, TEXT("Inisetting is different for %s, missmatched num array elements %d != %d "), *FString::Printf(TEXT("%s %s %s %s"),
-						*PlatformName, *Filename, *SectionName.ToString(), *ValueName.ToString()), CurrentValues.Num(), OldIniValue.Value.Num());
-					UE_LOG(LogCook, Display, TEXT("To avoid this add a deny list setting to DefaultEditor.ini [CookSettings] %s"), *DenyListSetting);
+					UE_LOG(LogCook, Display,
+						TEXT("Invalidating inisettings: Inisetting used by platform %s is different for %s:[%s]:%s, missmatched num array elements %d != %d."),
+						*TargetPlatform->PlatformName(), *ConfigNameKeyStr, *SectionName.ToString(),
+						*ValueName.ToString(), CurrentValues.Num(), OldIniValue.Value.Num());
+					UE_LOG(LogCook, Display, TEXT("%s:%s"), *GetDenyListMessageStart(), *ValueName.ToString());
 					return true;
 				}
 				for ( int Index = 0; Index < CurrentValues.Num(); ++Index )
@@ -8316,10 +8340,14 @@ bool UCookOnTheFlyServer::IniSettingsOutOfDate(const ITargetPlatform* TargetPlat
 					const FString FilteredCurrentValue = CurrentValues[Index].GetSavedValue().Replace(TEXT(":"), TEXT(""));
 					if ( FilteredCurrentValue != OldIniValue.Value[Index] )
 					{
-						UE_LOG(LogCook, Display, TEXT("Inisetting is different for %s, value %s != %s invalidating cook"),
-							*FString::Printf(TEXT("%s %s %s %s %d"),*PlatformName, *Filename, *SectionName.ToString(), *ValueName.ToString(), Index),
-							*CurrentValues[Index].GetSavedValue(), *OldIniValue.Value[Index] );
-						UE_LOG(LogCook, Display, TEXT("To avoid this add a deny list setting to DefaultEditor.ini [CookSettings] %s"), *DenyListSetting);
+						UE_LOG(LogCook, Display,
+							TEXT("Invalidating inisettings: Inisetting used by platform %s is different for %s:[%s]:%s%s, value '%s' != '%s'."),
+							*TargetPlatform->PlatformName(), *ConfigNameKeyStr, *SectionName.ToString(),
+							*ValueName.ToString(),
+							(CurrentValues.Num() == 1 ? TEXT("") : *FString::Printf(TEXT(" %d"), Index)),
+							*CurrentValues[Index].GetSavedValue(),
+							*OldIniValue.Value[Index] );
+						UE_LOG(LogCook, Display, TEXT("%s:%s"), *GetDenyListMessageStart(), *ValueName.ToString());
 						return true;
 					}
 				}
@@ -8328,11 +8356,12 @@ bool UCookOnTheFlyServer::IniSettingsOutOfDate(const ITargetPlatform* TargetPlat
 	}
 
 	return false;
+#endif // UE_WITH_CONFIG_TRACKING
 }
 
 bool UCookOnTheFlyServer::SaveCurrentIniSettings(const ITargetPlatform* TargetPlatform) const
 {
-	TGuardValue<bool> S(GSuppressProcessConfigSettings, true);
+	UE::ConfigAccessTracking::FIgnoreScope IgnoreScope;
 
 	TMap<FString, FString> AdditionalIniSettings;
 	GetAdditionalCurrentIniVersionStrings(this, TargetPlatform, AdditionalIniSettings);
@@ -8356,22 +8385,30 @@ bool UCookOnTheFlyServer::SaveCurrentIniSettings(const ITargetPlatform* TargetPl
 	ConfigFile.Remove(NAME_UsedSettings);
 
 	{
+		TStringBuilder<256> NewKey;
+		TStringBuilder<128> FilenameStr;
+		TStringBuilder<64> SectionStr;
+		TStringBuilder<64> ValueNameStr;
 		UE_SCOPED_HIERARCHICAL_COOKTIMER(ProcessingAccessedStrings)
 		for (const auto& CurrentIniFilename : CurrentIniSettings)
 		{
-			const FName& Filename = CurrentIniFilename.Key;
+			FName Filename = CurrentIniFilename.Key;
+			EscapeUsedSettingsToken(Filename, FilenameStr);
 			for (const auto& CurrentSection : CurrentIniFilename.Value)
 			{
-				const FName& Section = CurrentSection.Key;
+				FName Section = CurrentSection.Key;
+				EscapeUsedSettingsToken(Section, SectionStr);
 				for (const auto& CurrentValue : CurrentSection.Value)
 				{
-					const FName& ValueName = CurrentValue.Key;
-					const TArray<FString>& Values = CurrentValue.Value;
+					FName ValueName = CurrentValue.Key;
+					EscapeUsedSettingsToken(ValueName, ValueNameStr);
 
+					const TArray<FString>& Values = CurrentValue.Value;
 					for (int Index = 0; Index < Values.Num(); ++Index)
 					{
-						FString NewKey = FString::Printf(TEXT("%s:%s:%s:%d"), *Filename.ToString(), *Section.ToString(), *ValueName.ToString(), Index);
-						ConfigFile.AddToSection(NAME_UsedSettings, *NewKey, Values[Index]);
+						NewKey.Reset();
+						NewKey.Appendf(TEXT("%s:%s:%s:%d"), *FilenameStr, *SectionStr, *ValueNameStr, Index);
+						ConfigFile.AddToSection(NAME_UsedSettings, FName(NewKey), Values[Index]);
 					}
 				}
 			}
@@ -13257,6 +13294,10 @@ void UCookOnTheFlyServer::BroadcastCookByTheBookStarted()
 	// External systems would do this during CookByTheBookStarted.Broadcast
 	if (GetProcessType() != UE::Cook::EProcessType::SingleProcess)
 	{
+#if UE_WITH_CONFIG_TRACKING
+		ConfigCollector = new UE::ConfigAccessTracking::FConfigAccessTrackingCollector();
+		RegisterCollector(ConfigCollector);
+#endif
 	}
 }
 
@@ -13265,6 +13306,10 @@ void UCookOnTheFlyServer::BroadcastCookByTheBookFinished()
 	// Unregister collectors used internally by CookOnTheFlyServer.
 	if (GetProcessType() != UE::Cook::EProcessType::SingleProcess)
 	{
+#if UE_WITH_CONFIG_TRACKING
+		UnregisterCollector(ConfigCollector);
+		ConfigCollector.SafeRelease();
+#endif
 	}
 
 	PRAGMA_DISABLE_DEPRECATION_WARNINGS;
diff --git a/Engine/Source/Editor/UnrealEd/Private/Cooker/CookConfigAccessTracker.cpp b/Engine/Source/Editor/UnrealEd/Private/Cooker/CookConfigAccessTracker.cpp
new file mode 100644
--- /dev/null
+++ b/Engine/Source/Editor/UnrealEd/Private/Cooker/CookConfigAccessTracker.cpp
@@ -0,0 +1,527 @@
+// Copyright Epic Games, Inc. All Rights Reserved.
+
+#include "CookConfigAccessTracker.h"
+
+#if UE_WITH_CONFIG_TRACKING
+
+#include "Algo/Unique.h"
+#include "Algo/Sort.h"
+#include "Async/UniqueLock.h"
+#include "CookOnTheSide/CookLog.h"
+#include "HAL/LowLevelMemTracker.h"
+#include "Logging/LogMacros.h"
+#include "Misc/CString.h"
+#include "Misc/PackageAccessTrackingOps.h"
+#include "Misc/StringBuilder.h"
+#include "Serialization/CompactBinaryWriter.h"
+
+#endif
+
+#if UE_WITH_CONFIG_TRACKING
+
+DEFINE_LOG_CATEGORY_STATIC(LogConfigBuildDependencyTracker, Log, All);
+
+namespace UE::ConfigAccessTracking
+{
+
+FCookConfigAccessTracker FCookConfigAccessTracker::Singleton;
+
+thread_local bool bIgnoreAccess = false;
+
+TStringBuilder<256> ToString(FNameEntryId FileName, FNameEntryId SectionName, FMinimalName ValueName)
+{
+	return TStringBuilder<256>(InPlace, FileName, TEXT(":["), SectionName, TEXT("]:"), FName(ValueName));
+}
+
+FIgnoreScope::FIgnoreScope()
+	: bPreviousIgnoreAccess(bIgnoreAccess)
+{
+	bIgnoreAccess = true;
+}
+
+FIgnoreScope::~FIgnoreScope()
+{
+	bIgnoreAccess = bPreviousIgnoreAccess;
+}
+
+void FCookConfigAccessTracker::Disable()
+{
+	if (bEnabled)
+	{
+		UE::ConfigAccessTracking::RemoveConfigValueReadCallback(OnConfigValueReadCallbackHandle);
+		OnConfigValueReadCallbackHandle = UE::ConfigAccessTracking::FConfigValueReadCallbackId{};
+		PackageRecords.Empty();
+		bEnabled = false;
+	}
+}
+
+bool FCookConfigAccessTracker::IsEnabled() const
+{
+	return bEnabled;
+}
+
+void FCookConfigAccessTracker::DumpStats() const
+{
+	using namespace UE::ConfigAccessTracking;
+
+	if (!IsEnabled())
+	{
+		return;
+	}
+
+	UE::TUniqueLock RecordsScopeLock(RecordsLock);
+	uint64 ReferencingPackageCount = 0;
+	uint64 ReferenceCount = 0;
+	uint64 GlobalReferenceCount = 0;
+	for (const TPair<FName, TSet<FConfigAccessData>>& PackageAccessRecord : PackageRecords)
+	{
+		if (PackageAccessRecord.Key == NAME_None)
+		{
+			for (const FConfigAccessData& AccessedData : PackageAccessRecord.Value)
+			{
+				++GlobalReferenceCount;
+			}
+		}
+		else
+		{
+			++ReferencingPackageCount;
+			for (const FConfigAccessData& AccessedData : PackageAccessRecord.Value)
+			{
+				++ReferenceCount;
+			}
+		}
+	}
+	UE_LOG(LogConfigBuildDependencyTracker, Display,
+		TEXT("Config Accesses (%u referencing packages with a total of %u unique accesses). %u unique accesses that were not associated with a package."),
+		ReferencingPackageCount, ReferenceCount, GlobalReferenceCount);
+
+	constexpr bool bDetailedDump = false;
+	if (bDetailedDump)
+	{
+		UE_LOG(LogConfigBuildDependencyTracker, Display, TEXT("========================================================================="));
+		for (const TPair<FName, TSet<FConfigAccessData>>& PackageAccessRecord : PackageRecords)
+		{
+			UE_LOG(LogConfigBuildDependencyTracker, Display, TEXT("%s:"), *PackageAccessRecord.Key.ToString());
+			for (const FConfigAccessData& AccessedData : PackageAccessRecord.Value)
+			{
+				UE_LOG(LogConfigBuildDependencyTracker, Display, TEXT("    %s"),
+					*ToString(AccessedData.FileName, AccessedData.SectionName, AccessedData.ValueName));
+			}
+		}
+	}
+}
+
+TArray<UE::ConfigAccessTracking::FConfigAccessData>
+FCookConfigAccessTracker::GetPackageRecords(FName ReferencerPackage, const ITargetPlatform* TargetPlatform) const
+{
+	using namespace UE::ConfigAccessTracking;
+
+	TArray<FConfigAccessData> Records;
+	{
+		UE::TUniqueLock RecordsScopeLock(Singleton.RecordsLock);
+		const TSet<FConfigAccessData>* ReferencerSet = Singleton.PackageRecords.Find(ReferencerPackage);
+		if (!ReferencerSet)
+		{
+			return TArray<FConfigAccessData>();
+		}
+		Records = ReferencerSet->Array();
+	}
+	SortRecordsAndFilterByPlatform(Records, TargetPlatform);
+	return Records;
+}
+
+TArray<UE::ConfigAccessTracking::FConfigAccessData>
+FCookConfigAccessTracker::GetCookRecords() const
+{
+	using namespace UE::ConfigAccessTracking;
+
+	TSet<FConfigAccessData> CookRecords;
+	{
+		UE::TUniqueLock RecordsScopeLock(Singleton.RecordsLock);
+		for (const TPair<FName, TSet<FConfigAccessData>>& Pair : Singleton.PackageRecords)
+		{
+			CookRecords.Append(Pair.Value);
+		}
+	}
+	TArray<FConfigAccessData> ResultRecords = CookRecords.Array();
+	Algo::Sort(ResultRecords);
+	return ResultRecords;
+}
+
+TArray<UE::ConfigAccessTracking::FConfigAccessData>
+FCookConfigAccessTracker::GetCookRecords(const ITargetPlatform* TargetPlatform) const
+{
+	using namespace UE::ConfigAccessTracking;
+
+	TSet<FConfigAccessData> CookRecords;
+	{
+		UE::TUniqueLock RecordsScopeLock(Singleton.RecordsLock);
+		for (const TPair<FName, TSet<FConfigAccessData>>& Pair : Singleton.PackageRecords)
+		{
+			CookRecords.Append(Pair.Value);
+		}
+	}
+	TArray<FConfigAccessData> Records = CookRecords.Array();
+	SortRecordsAndFilterByPlatform(Records, TargetPlatform);
+	return Records;
+}
+
+void FCookConfigAccessTracker::SortRecordsAndFilterByPlatform(TArray<FConfigAccessData>& Records, const ITargetPlatform* TargetPlatform)
+{
+	// Remove records not relevant to the current platform, and set the TargetPlatform to null, so we can
+	// remove records that are duplicated between platform keys.
+	for (TArray<FConfigAccessData>::TIterator Iter(Records); Iter; ++Iter)
+	{
+		FConfigAccessData& Record = *Iter;
+		if (Record.RequestingPlatform != TargetPlatform && Record.RequestingPlatform != nullptr)
+		{
+			Iter.RemoveCurrentSwap();
+		}
+		else
+		{
+			Record.RequestingPlatform = nullptr;
+		}
+	}
+	Algo::Sort(Records);
+	Records.SetNum(Algo::Unique(Records));
+}
+
+void FCookConfigAccessTracker::AddRecord(FName PackageName,
+	const UE::ConfigAccessTracking::FConfigAccessData& AccessData)
+{
+	UE::TUniqueLock RecordsScopeLock(RecordsLock);
+	PackageRecords.FindOrAdd(PackageName).Add(AccessData);
+}
+
+FCookConfigAccessTracker::FCookConfigAccessTracker()
+{
+	OnConfigValueReadCallbackHandle = UE::ConfigAccessTracking::AddConfigValueReadCallback(StaticOnConfigValueRead);
+	bEnabled = true;
+}
+
+FCookConfigAccessTracker::~FCookConfigAccessTracker()
+{
+	Disable();
+}
+
+void FCookConfigAccessTracker::StaticOnConfigValueRead(UE::ConfigAccessTracking::FSection* Section, FMinimalName ValueName,
+	const FConfigValue& ConfigValue)
+{
+	if (bIgnoreAccess)
+	{
+		return;
+	}
+	if (!Section)
+	{
+		return;
+	}
+	UE::ConfigAccessTracking::FFile* FileAccess = Section->FileAccess.GetReference();
+	if (!FileAccess)
+	{
+		return;
+	}
+	const FConfigFile* ConfigFile = FileAccess->ConfigFile;
+	if (!ConfigFile)
+	{
+		return;
+	}
+	if (!UE::ConfigAccessTracking::IsLoadableLoadType(ConfigFile->LoadType))
+	{
+		return;
+	}
+	FNameEntryId FileName = FileAccess->GetFilenameToLoad().GetComparisonIndex();
+	if (FileName.IsNone())
+	{
+		return;
+	}
+	FNameEntryId SectionName = Section->SectionName;
+	if (SectionName.IsNone())
+	{
+		return;
+	}
+	FNameEntryId ConfigPlatform = FileAccess->GetPlatformName().GetComparisonIndex();
+
+	FName Referencer = NAME_None;
+	const ITargetPlatform* RequestedPlatform = nullptr;
+#if UE_WITH_PACKAGE_ACCESS_TRACKING
+	PackageAccessTracking_Private::FTrackedData* AccumulatedScopeData = PackageAccessTracking_Private::FPackageAccessRefScope::GetCurrentThreadAccumulatedData();
+	if (AccumulatedScopeData && !AccumulatedScopeData->BuildOpName.IsNone())
+	{
+		RequestedPlatform = AccumulatedScopeData->TargetPlatform;
+		Referencer = AccumulatedScopeData->PackageName;
+
+		if (AccumulatedScopeData->OpName == PackageAccessTrackingOps::NAME_NoAccessExpected)
+		{
+			UE_LOG(LogConfigBuildDependencyTracker, Warning,
+				TEXT("Object %s is referencing config value %s inside of a NAME_NoAccessExpected scope. Programmer should narrow the scope or debug the reference."),
+				*Referencer.ToString(), *ToString(FileName, SectionName, ValueName));
+		}
+	}
+#endif
+
+	LLM_SCOPE_BYNAME(TEXTVIEW("ConfigAccessTracking"));
+	FConfigAccessData AccessData{ ConfigPlatform, FileName, SectionName, ValueName, RequestedPlatform, ConfigFile->LoadType };
+	uint32 ReferencerHash = GetTypeHash(Referencer);
+	uint32 AccessDataHash = GetTypeHash(AccessData);
+	UE::TUniqueLock RecordsScopeLock(Singleton.RecordsLock);
+	Singleton.PackageRecords.FindOrAddByHash(ReferencerHash, Referencer).AddByHash(AccessDataHash, MoveTemp(AccessData));
+}
+
+const FConfigFile* FindOrLoadConfigFile(ELoadType LoadType, const TCHAR* ConfigPlatform, const TCHAR* Filename,
+	FConfigFile& Buffer)
+{
+	bool bHasPlatformName = ConfigPlatform && ConfigPlatform[0] != '\0';
+	FName ConfigPlatformFName;
+	if (bHasPlatformName)
+	{
+		ConfigPlatformFName = FName(ConfigPlatform);
+	}
+
+	switch (LoadType)
+	{
+	case ELoadType::ConfigSystem:
+	{
+		// ForPlatform will return GConfig if ConfigPlatform is NAME_None
+		FConfigCacheIni* ConfigSystem = FConfigCacheIni::ForPlatform(ConfigPlatformFName);
+
+		// The ini files may have been recorded by fullpath or by shortname; search first for a fullpath match using
+		// FindConfigFile and if that fails search for the shortname match by iterating over all files in GConfig
+		FConfigFile* ConfigFile = ConfigSystem->FindConfigFile(FString(Filename));
+		if (ConfigFile)
+		{
+			return ConfigFile;
+		}
+		if (!ConfigFile)
+		{
+			FName FileFName = FName(Filename);
+			for (const FString& ConfigFilename : ConfigSystem->GetFilenames())
+			{
+				ConfigFile = ConfigSystem->FindConfigFile(ConfigFilename);
+				if (ConfigFile->Name == FileFName)
+				{
+					return ConfigFile;
+				}
+			}
+		}
+		return nullptr;
+	}
+	case ELoadType::LocalIniFile:
+		if (FConfigCacheIni::LoadLocalIniFile(Buffer, Filename, true /* bIsBaseIniName */, ConfigPlatform))
+		{
+			return &Buffer;
+		}
+		return nullptr;
+	case ELoadType::LocalSingleIniFile:
+		if (FConfigCacheIni::LoadLocalIniFile(Buffer, Filename, false /* bIsBaseIniName */, ConfigPlatform))
+		{
+			return &Buffer;
+		}
+		return nullptr;
+	case ELoadType::ExternalIniFile:
+		// TODO: LoadExternalIniFile is the same as LoadLocalIniFile, but with possibly redirected
+		// EngineConfigDir and ProjectConfigDir. We can not load them without that extra information.
+		// For now, assume it used the default EngineConfigDir and ProjectConfigDir.
+		if (FConfigCacheIni::LoadLocalIniFile(Buffer, Filename, true /* bIsBaseIniName */, ConfigPlatform))
+		{
+			return &Buffer;
+		}
+		return nullptr;
+	case ELoadType::ExternalSingleIniFile:
+		// TODO: Same comment as in ExternalIniFile
+		if (FConfigCacheIni::LoadLocalIniFile(Buffer, Filename, false /* bIsBaseIniName */, ConfigPlatform))
+		{
+			return &Buffer;
+		}
+		return nullptr;
+	default:
+		return nullptr;
+	}
+}
+
+/** Return whether LoadType is a type that can be loaded by FindOrLoadConfigFile. */
+bool IsLoadableLoadType(ELoadType LoadType)
+{
+	switch (LoadType)
+	{
+	case ELoadType::ConfigSystem:
+	case ELoadType::LocalIniFile:
+	case ELoadType::LocalSingleIniFile:
+	case ELoadType::ExternalIniFile:
+	case ELoadType::ExternalSingleIniFile:
+		return true;
+	default:
+		return false;
+	}
+}
+
+} // namespace UE::ConfigAccessTracking
+
+const TCHAR* LexToString(UE::ConfigAccessTracking::ELoadType LoadType)
+{
+	using namespace UE::ConfigAccessTracking;
+
+	switch (LoadType)
+	{
+	case ELoadType::ConfigSystem: return TEXT("ConfigSystem");
+	case ELoadType::LocalIniFile: return TEXT("LocalIniFile");
+	case ELoadType::LocalSingleIniFile: return TEXT("LocalSingleIniFile");
+	case ELoadType::ExternalIniFile: return TEXT("ExternalIniFile");
+	case ELoadType::ExternalSingleIniFile: return TEXT("ExternalSingleIniFile");
+	case ELoadType::Manual: return TEXT("Manual");
+	case ELoadType::SuppressReporting: return TEXT("SuppressReporting");
+	case ELoadType::Uninitialized:
+	default:
+		return TEXT("Uninitialized");
+	}
+}
+
+void LexFromString(UE::ConfigAccessTracking::ELoadType& OutLoadType, FStringView Text)
+{
+	using namespace UE::ConfigAccessTracking;
+
+	if (Text.IsEmpty())
+	{
+		OutLoadType = ELoadType::Uninitialized;
+		return;
+	}
+	switch (Text[0])
+	{
+	case 'C':
+	case 'c':
+		if (Text.Equals(TEXT("ConfigSystem"), ESearchCase::IgnoreCase))
+		{
+			OutLoadType = ELoadType::ConfigSystem;
+			return;
+		}
+		break;
+	case 'L':
+	case 'l':
+		if (Text.Equals(TEXT("LocalIniFile"), ESearchCase::IgnoreCase))
+		{
+			OutLoadType = ELoadType::LocalIniFile;
+			return;
+		}
+		if (Text.Equals(TEXT("LocalSingleIniFile"), ESearchCase::IgnoreCase))
+		{
+			OutLoadType = ELoadType::LocalSingleIniFile;
+			return;
+		}
+		break;
+	case 'E':
+	case 'e':
+		if (Text.Equals(TEXT("ExternalIniFile"), ESearchCase::IgnoreCase))
+		{
+			OutLoadType = ELoadType::ExternalIniFile;
+			return;
+		}
+		if (Text.Equals(TEXT("ExternalSingleIniFile"), ESearchCase::IgnoreCase))
+		{
+			OutLoadType = ELoadType::ExternalSingleIniFile;
+			return;
+		}
+		break;
+	case 'S':
+	case 's':
+		if (Text.Equals(TEXT("SuppressReporting"), ESearchCase::IgnoreCase))
+		{
+			OutLoadType = ELoadType::SuppressReporting;
+			return;
+		}
+		break;
+	case 'M':
+	case 'm':
+		if (Text.Equals(TEXT("Manual"), ESearchCase::IgnoreCase))
+		{
+			OutLoadType = ELoadType::Manual;
+			return;
+		}
+		break;
+	default:
+		break;
+	}
+	OutLoadType = ELoadType::Uninitialized;
+}
+
+namespace UE::ConfigAccessTracking
+{
+
+FGuid FConfigAccessTrackingCollector::MessageType(TEXT("B3F36AFEF6AE467E9E8F0DDA604856C3"));
+static const FUtf8StringView ConfigDependencyCollectorRecordsName = UTF8TEXTVIEW("R");
+
+void FConfigAccessTrackingCollector::ClientTick(UE::Cook::FMPCollectorClientTickContext& Context)
+{
+	using namespace UE::ConfigAccessTracking;
+
+	if (!Context.IsFlush())
+	{
+		return;
+	}
+
+	TArray<FConfigAccessData> Records = FCookConfigAccessTracker::Get().GetCookRecords();
+
+	FCbWriter Writer;
+	Writer.BeginObject();
+	Writer.SetName(ConfigDependencyCollectorRecordsName);
+	Writer.BeginArray();
+	for (FConfigAccessData& Record : Records)
+	{
+		Writer.BeginArray();
+		uint8 LoadType = static_cast<uint8>(Record.LoadType);
+		FName ConfigPlatform(Record.GetConfigPlatform());
+		FName FileName(Record.GetFileName());
+		FName SectionName(Record.GetSectionName());
+		FName ValueName(Record.GetValueName());
+		Writer << LoadType << ConfigPlatform << FileName << SectionName << ValueName
+			<< Context.PlatformToInt(Record.RequestingPlatform);
+		Writer.EndArray();
+	}
+	Writer.EndArray();
+	Writer.EndObject();
+	Context.AddMessage(Writer.Save().AsObject());
+}
+
+void FConfigAccessTrackingCollector::ServerReceiveMessage(UE::Cook::FMPCollectorServerMessageContext& Context,
+	FCbObjectView Message)
+{
+	using namespace UE::ConfigAccessTracking;
+
+	FCookConfigAccessTracker& Tracker = FCookConfigAccessTracker::Get();
+
+	FCbFieldView RecordsField = Message[ConfigDependencyCollectorRecordsName];
+	FCbArrayView RecordsView = RecordsField.AsArrayView();
+	if (RecordsField.HasError())
+	{
+		UE_LOG(LogCook, Error,
+			TEXT("Corrupt message received from CookWorker when replicating ConfigDependencies. IterativeFalseSkips may occur in next cook."));
+		return;
+	}
+	for (FCbFieldView RecordField : RecordsView)
+	{
+		FCbArrayView RecordArray = RecordField.AsArrayView();
+		FCbFieldViewIterator RecordIt = RecordArray.CreateViewIterator();
+		FConfigAccessData Record;
+
+		Record.LoadType = static_cast<ELoadType>(RecordIt.AsUInt8());
+		++RecordIt;
+		Record.ConfigPlatform= FName(RecordIt.AsString()).GetComparisonIndex();
+		++RecordIt;
+		Record.FileName = FName(RecordIt.AsString()).GetComparisonIndex();
+		++RecordIt;
+		Record.SectionName = FName(RecordIt.AsString(), NAME_NO_NUMBER).GetComparisonIndex();
+		++RecordIt;
+		Record.ValueName = FMinimalName(FName(RecordIt.AsString()));
+		++RecordIt;
+		Record.RequestingPlatform = Context.IntToPlatform(RecordIt.AsUInt8());
+		if (RecordIt.HasError())
+		{
+			UE_LOG(LogCook, Error,
+				TEXT("Corrupt message received from CookWorker when replicating ConfigDependencies. IterativeFalseSkips may occur in next cook."));
+			return;
+		}
+		Tracker.AddRecord(NAME_None, Record);
+	}
+}
+
+} // namespace UE::ConfigAccessTracking
+
+#endif // UE_WITH_CONFIG_TRACKING
diff --git a/Engine/Source/Editor/UnrealEd/Private/Cooker/CookConfigAccessTracker.h b/Engine/Source/Editor/UnrealEd/Private/Cooker/CookConfigAccessTracker.h
new file mode 100644
--- /dev/null
+++ b/Engine/Source/Editor/UnrealEd/Private/Cooker/CookConfigAccessTracker.h
@@ -0,0 +1,201 @@
+// Copyright Epic Games, Inc. All Rights Reserved.
+
+#pragma once
+
+#include "Async/Mutex.h"
+#include "Cooker/MPCollector.h"
+#include "Containers/Array.h"
+#include "Containers/Map.h"
+#include "Containers/Set.h"
+#include "Containers/StringView.h"
+#include "Interfaces/ITargetPlatform.h"
+#include "Misc/ConfigAccessTracking.h"
+#include "Misc/ConfigCacheIni.h"
+#include "Misc/PackageAccessTracking.h"
+#include "Templates/TypeHash.h"
+#include "Templates/UnrealTemplate.h"
+#include "UObject/NameTypes.h"
+
+#if UE_WITH_CONFIG_TRACKING
+
+namespace UE::ConfigAccessTracking
+{
+
+/** Full path of a FConfigValue that was reported read. */
+struct FConfigAccessData
+{
+	FNameEntryId ConfigPlatform;
+	FNameEntryId FileName;
+	FNameEntryId SectionName;
+	FMinimalName ValueName;
+	const ITargetPlatform* RequestingPlatform = nullptr;
+	ELoadType LoadType = ELoadType::Uninitialized;
+
+	FName GetConfigPlatform() const { return FName(ConfigPlatform, ConfigPlatform, NAME_NO_NUMBER_INTERNAL); }
+	FName GetFileName() const { return FName(FileName, FileName, NAME_NO_NUMBER_INTERNAL); }
+	FName GetSectionName() const { return FName(SectionName, SectionName, NAME_NO_NUMBER_INTERNAL); }
+	FName GetValueName() const { return FName(ValueName); }
+
+	friend uint32 GetTypeHash(const FConfigAccessData& Data)
+	{
+		uint32 Hash = static_cast<uint32>(Data.LoadType);
+		Hash = HashCombineFast(Hash, Data.ConfigPlatform.ToUnstableInt());
+		Hash = HashCombineFast(Hash, Data.FileName.ToUnstableInt());
+		Hash = HashCombineFast(Hash, Data.SectionName.ToUnstableInt());
+		Hash = HashCombineFast(Hash, GetTypeHash(Data.ValueName));
+		Hash = HashCombineFast(Hash, GetTypeHash(Data.RequestingPlatform));
+		return Hash;
+	}
+	bool IsSameConfigFile(const FConfigAccessData& Other) const
+	{
+		return LoadType == Other.LoadType && ConfigPlatform == Other.ConfigPlatform && FileName == Other.FileName;
+	}
+	bool operator==(const FConfigAccessData& Other) const
+	{
+		return LoadType == Other.LoadType && ConfigPlatform == Other.ConfigPlatform &&
+			FileName == Other.FileName && SectionName == Other.SectionName && ValueName == Other.ValueName &&
+			RequestingPlatform == Other.RequestingPlatform;
+	}
+	bool operator!=(const FConfigAccessData& Other) const
+	{
+		return !(*this == Other);
+	}
+	bool operator<(const FConfigAccessData& Other) const
+	{
+		if (LoadType != Other.LoadType) return static_cast<uint32>(LoadType) < static_cast<uint32>(Other.LoadType);
+		if (ConfigPlatform != Other.ConfigPlatform) return ConfigPlatform.LexicalLess(Other.ConfigPlatform);
+		if (FileName != Other.FileName) return FileName.LexicalLess(Other.FileName);
+		if (SectionName != Other.SectionName) return SectionName.LexicalLess(Other.SectionName);
+		if (ValueName != Other.ValueName) return FName(ValueName).LexicalLess(FName(Other.ValueName));
+		if (RequestingPlatform != Other.RequestingPlatform)
+		{
+			if (RequestingPlatform == nullptr) return true;
+			if (Other.RequestingPlatform == nullptr) return false;
+			return RequestingPlatform->PlatformName() < Other.RequestingPlatform->PlatformName();
+		}
+		return false;
+	}
+};
+
+/** Disables recording of FConfigAccessData reported as accessed on the current thread while in scope. */
+struct FIgnoreScope
+{
+	FIgnoreScope();
+	~FIgnoreScope();
+
+	bool bPreviousIgnoreAccess = false;
+};
+
+/**
+ * Tracker that subscribes to AddConfigValueReadCallback and for each access records the access associated with
+ * the package that is currently in scope according to PackageAccessTracking_Private.
+ */
+class FCookConfigAccessTracker: public FNoncopyable
+{
+public:
+	static FCookConfigAccessTracker& Get() { return Singleton; }
+
+	void Disable();
+	bool IsEnabled() const;
+	void DumpStats() const;
+	/**
+	 * Get records requested for the given package and given platform, including RequestingPlatform=nullptr.
+	 * Returned records are SORTED by FConfigAccessData::operator<.
+	 */
+	TArray<UE::ConfigAccessTracking::FConfigAccessData> GetPackageRecords(FName ReferencerPackage,
+		const ITargetPlatform* TargetPlatform) const;
+	/**
+	 * Get records for all requesting packages, including records not associated with a package.
+	 * Returned records are SORTED by FConfigAccessData::operator<.
+	 */
+	TArray<UE::ConfigAccessTracking::FConfigAccessData> GetCookRecords() const;
+	/**
+	 * Get records requested for all requesting packages, including records not associated with a package,
+	 * but filtered by the given TargetPlatform. Includes records requested with no RequestingPlatform.
+	 * TargetPlatform==nullptr returns only records requested with no RequestingPlatform.
+	 * Returned records are SORTED by FConfigAccessData::operator<.
+	 */
+	TArray<UE::ConfigAccessTracking::FConfigAccessData> GetCookRecords(const ITargetPlatform* TargetPlatform) const;
+	/** Add a record as if requested by the given package, or not associated with a package if PackageName.IsNone(). */
+	void AddRecord(FName PackageName, const UE::ConfigAccessTracking::FConfigAccessData& Data);
+
+private:
+	FCookConfigAccessTracker();
+	virtual ~FCookConfigAccessTracker();
+
+	/** Track object reference reads */
+	static void StaticOnConfigValueRead(UE::ConfigAccessTracking::FSection* Section, FMinimalName ValueName,
+		const FConfigValue& ConfigValue);
+
+	/** Helper function for GetRecords functions. */
+	static void SortRecordsAndFilterByPlatform(TArray<FConfigAccessData>& Records,
+		const ITargetPlatform* TargetPlatform);
+
+private:
+	// Use a mutex rather than a critical section for synchronization.  Calls into system libraries, such as windows critical section
+	// functions, are 50 times more expensive on build farm VMs, radically affecting cook times, which this avoids. 
+	mutable UE::FMutex RecordsLock;
+	TMap<FName, TSet<UE::ConfigAccessTracking::FConfigAccessData>> PackageRecords;
+	UE::ConfigAccessTracking::FConfigValueReadCallbackId OnConfigValueReadCallbackHandle;
+	bool bEnabled = false;
+
+private:
+	static FCookConfigAccessTracker Singleton;
+};
+
+/**
+ * Find a ConfigFile by name and ConfigPlatform, either in GConfig or loaded from disk.
+ * @param LoadType LoadType describing which ConfigCacheIni method was used to load the ConfigFile
+ * @param ConfigPlatform PlatformName used with the ConfigCacheIni method that loaded the ConfigFile, or
+ *        nullptr/emptystring for platform-agnostic
+ * @param Filename Short ('Engine') or long ('../../../QAGame/Config/DefaultEngine.ini') name of the ConfigFile
+ * @param Buffer FConfigFile Buffer that will hold the result if LoadConfigFile was called
+ * @param The discovered configfile, or nullptr.
+ */
+const FConfigFile* FindOrLoadConfigFile(ELoadType LoadType, const TCHAR* ConfigPlatform, const TCHAR* Filename,
+	FConfigFile& Buffer);
+
+/** Return whether LoadType is a type that can be loaded by FindOrLoadConfigFile. */
+bool IsLoadableLoadType(ELoadType LoadType);
+
+} // namespace UE::ConfigAccessTracking
+
+/** Convert ELoadType -> text */
+const TCHAR* LexToString(UE::ConfigAccessTracking::ELoadType LoadType);
+
+/** Convert text -> ELoadType */
+void LexFromString(UE::ConfigAccessTracking::ELoadType& OutLoadType, FStringView Text);
+
+namespace UE::ConfigAccessTracking
+{
+
+/** CookMultiprocess collector for ConfigAccess data. */
+class FConfigAccessTrackingCollector : public UE::Cook::IMPCollector
+{
+public:
+	virtual FGuid GetMessageType() const override { return MessageType; }
+	virtual const TCHAR* GetDebugName() const override { return TEXT("FConfigAccessTrackingCollector"); }
+
+	virtual void ClientTick(UE::Cook::FMPCollectorClientTickContext& Context) override;
+	virtual void ServerReceiveMessage(UE::Cook::FMPCollectorServerMessageContext& Context,
+		FCbObjectView Message) override;
+
+	static FGuid MessageType;
+};
+
+} // namespace UE::ConfigAccessTracking
+
+#else // !UE_WITH_CONFIG_TRACKING
+
+namespace UE::ConfigAccessTracking
+{
+
+struct FIgnoreScope
+{
+	FIgnoreScope() {}
+	~FIgnoreScope() {}
+};
+
+} // namespace UE::ConfigAccessTracking
+
+#endif // else !UE_WITH_CONFIG_TRACKING
diff --git a/Engine/Source/Editor/UnrealEd/Private/Cooker/PackageBuildDependencyTracker.cpp b/Engine/Source/Editor/UnrealEd/Private/Cooker/PackageBuildDependencyTracker.cpp
--- a/Engine/Source/Editor/UnrealEd/Private/Cooker/PackageBuildDependencyTracker.cpp
+++ b/Engine/Source/Editor/UnrealEd/Private/Cooker/PackageBuildDependencyTracker.cpp
@@ -4,6 +4,7 @@
 
 #if UE_WITH_PACKAGE_ACCESS_TRACKING
 #include "Async/UniqueLock.h"
+#include "Cooker/CookConfigAccessTracker.h"
 #include "HAL/Platform.h"
 #include "Logging/LogMacros.h"
 #include "Misc/PackageAccessTrackingOps.h"
@@ -158,15 +159,14 @@ void FPackageBuildDependencyTracker::StaticOnObjectHandleRead(const TArrayView<c
 		}
 	}
 }
+#endif // UE_WITH_OBJECT_HANDLE_TRACKING
 
 void DumpBuildDependencyTrackerStats()
 {
+#if UE_WITH_PACKAGE_ACCESS_TRACKING
 	FPackageBuildDependencyTracker::Get().DumpStats();
+#endif
+#if UE_WITH_CONFIG_TRACKING
+	UE::ConfigAccessTracking::FCookConfigAccessTracker::Get().DumpStats();
+#endif
 }
-
-#else
-void DumpBuildDependencyTrackerStats()
-{
-}
-
-#endif // UE_WITH_OBJECT_HANDLE_TRACKING
diff --git a/Engine/Source/Editor/UnrealEd/Private/UnrealEdMisc.cpp b/Engine/Source/Editor/UnrealEd/Private/UnrealEdMisc.cpp
--- a/Engine/Source/Editor/UnrealEd/Private/UnrealEdMisc.cpp
+++ b/Engine/Source/Editor/UnrealEd/Private/UnrealEdMisc.cpp
@@ -95,6 +95,7 @@
 #include "HAL/PlatformTime.h"
 #include "StudioAnalytics.h"
 #include "DeveloperToolSettingsDelegates.h"
+#include "Cooker/CookConfigAccessTracker.h"
 #include "Cooker/PackageBuildDependencyTracker.h"
 
 #define USE_UNIT_TESTS 0
@@ -222,11 +223,16 @@ FUnrealEdMisc::FUnrealEdMisc() :
 	NavigationBuildingNotificationHandler(NULL)
 {
 	//This is an early entry-point into the UnrealEd module to perform some editor-specific configuration
-#if UE_WITH_PACKAGE_ACCESS_TRACKING
+#if UE_WITH_PACKAGE_ACCESS_TRACKING || UE_WITH_CONFIG_TRACKING
 	const bool bBuildDependencyTrackingNeeded = GIsEditor && (IsRunningCookCommandlet() || !GetDefault<UEditorExperimentalSettings>()->bDisableCookInEditor);
 	if (!bBuildDependencyTrackingNeeded)
 	{
+#if UE_WITH_PACKAGE_ACCESS_TRACKING
 		FPackageBuildDependencyTracker::Get().Disable();
+#endif
+#if UE_WITH_CONFIG_TRACKING
+		UE::ConfigAccessTracking::FCookConfigAccessTracker::Get().Disable();
+#endif
 	}
 #endif
 }
diff --git a/Engine/Source/Runtime/Core/Private/Misc/ConfigAccessTracking.cpp b/Engine/Source/Runtime/Core/Private/Misc/ConfigAccessTracking.cpp
new file mode 100644
--- /dev/null
+++ b/Engine/Source/Runtime/Core/Private/Misc/ConfigAccessTracking.cpp
@@ -0,0 +1,324 @@
+// Copyright Epic Games, Inc. All Rights Reserved.
+
+#include "Misc/ConfigAccessTracking.h"
+
+#if UE_WITH_CONFIG_TRACKING
+#include "AutoRTFM/AutoRTFM.h"
+#include "Containers/Array.h"
+#include "Misc/ConfigCacheIni.h"
+#include "Misc/DataDrivenPlatformInfoRegistry.h"
+#include "Templates/Tuple.h"
+
+// LightWeightRWLock
+#include <atomic>
+#include "HAL/Platform.h"
+#include "HAL/PlatformProcess.h"
+#endif
+
+#if UE_WITH_CONFIG_TRACKING
+
+namespace UE::ConfigAccessTracking
+{
+
+FFile::FFile(const FConfigFile* InConfigFile)
+	: ConfigFile(InConfigFile)
+	, bSavedHasConfigFile(false)
+	, bSavedConfigFileHasPlatformName(false)
+	, bPlatformNameInitialized(false)
+	, bOverridePlatformName(false)
+	, bSavedOverridePlatformName(false)
+{
+}
+
+FName FFile::GetFilenameToLoad() const
+{
+	if (!ConfigFile)
+	{
+		return NAME_None;
+	}
+	if (!OverrideFilenameToLoad.IsNone())
+	{
+		return OverrideFilenameToLoad;
+	}
+	return ConfigFile->Name;
+}
+
+FName FFile::GetPlatformName()
+{
+	bool bDirty = false;
+	bool bHasConfigFile = ConfigFile != nullptr;
+	if (!bPlatformNameInitialized)
+	{
+		bDirty = true;
+	}
+	else if (bSavedHasConfigFile != bHasConfigFile || bSavedOverridePlatformName != bOverridePlatformName)
+	{
+		bDirty = true;
+	}
+	// if !bHasConfigFile there is nothing further to test; not dirty
+	else if (bHasConfigFile)
+	{
+		if (ConfigFile->bHasPlatformName != bSavedConfigFileHasPlatformName)
+		{
+			bDirty = true;
+		}
+		else
+		{
+			if (!ConfigFile->bHasPlatformName || ConfigFile->PlatformName.IsEmpty())
+			{
+				bDirty = !SavedConfigFilePlatformName.IsEmpty();
+			}
+			else
+			{
+				// Compare by pointer rather than by string compare, to be cheaper
+				bDirty = (SavedConfigFilePlatformName.GetData() != *ConfigFile->PlatformName ||
+					SavedConfigFilePlatformName.Len() != ConfigFile->PlatformName.Len());
+			}
+		}
+	}
+	if (!bDirty)
+	{
+		return PlatformName;
+	}
+	bSavedHasConfigFile = bHasConfigFile;
+	bSavedOverridePlatformName = bOverridePlatformName;
+	if (!bHasConfigFile)
+	{
+		bSavedConfigFileHasPlatformName = false;
+		SavedConfigFilePlatformName.Reset();
+	}
+	else
+	{
+		bSavedConfigFileHasPlatformName = ConfigFile->bHasPlatformName;
+		SavedConfigFilePlatformName = ConfigFile->PlatformName;
+	}
+	bPlatformNameInitialized = true;
+
+	if (bOverridePlatformName)
+	{
+		// Currently the only time we need to override the platformname is when we are clearing it, so we don't keep a
+		// separate variable for OverriddenPlatformName, we just interpret it as being NAME_None
+		PlatformName = NAME_None;
+	}
+	else if (!ConfigFile)
+	{
+		PlatformName = NAME_None;
+	}
+	else if (ConfigFile->bHasPlatformName)
+	{
+		// The platform that was passed to LoadExternalIniFile
+		PlatformName = FName(FStringView(ConfigFile->PlatformName));
+	}
+	else if (!GConfig || !GConfig->IsReadyForUse())
+	{
+		// All unmarked-platform config files read during config startup are platform-agnostic.
+		// That property is required because we cannot fall back to the search through
+		// GetAllPlatformInfos and SourceIniHierarchy below during startup, because they are not
+		// yet threadsafe to access.
+		PlatformName = NAME_None;
+	}
+	else
+	{
+		// For the config files not in GConfig, we assume they were loaded from LoadConfigFile, and we match these to a platform
+		// By looking for a platform-specific filepath in their SourceIniHierarchy.
+		// Examples:
+		// (1) ROOT\Engine\Config\Windows\WindowsEngine.ini
+		// (2) ROOT\Engine\Config\Android\DataDrivePlatformInfo.ini
+		// (3) ROOT\Engine\Config\Android\AndroidWindowsCompatability.ini
+		// 
+		// Note that for config files of form #3, we want them to be matched to Android rather than windows;
+		// we assume that an exact match on a directory component is more definitive than a substring match
+		bool bFoundPlatformName = false;
+		bool bFoundPlatformGuess = false;
+		for (auto It : FDataDrivenPlatformInfoRegistry::GetAllPlatformInfos())
+		{
+			const FString CurrentPlatformName = It.Key.ToString();
+			TStringBuilder<128> PlatformDirString;
+			PlatformDirString.Appendf(TEXT("/%s/"), *CurrentPlatformName);
+			for (const TPair<int32, FString>& SourceIni : ConfigFile->SourceIniHierarchy)
+			{
+				// Look for platform in the path, rating a full subdirectory name match (/Android/ or /Windows/)
+				// higher than a partial filename match (AndroidEngine.ini or WindowsEngine.ini)
+				bool bFoundPlatformDir = UE::String::FindFirst(SourceIni.Value, PlatformDirString, ESearchCase::IgnoreCase) != INDEX_NONE;
+				bool bFoundPlatformSubstring = UE::String::FindFirst(SourceIni.Value, CurrentPlatformName, ESearchCase::IgnoreCase) != INDEX_NONE;
+				if (bFoundPlatformDir)
+				{
+					PlatformName = FName(FStringView(CurrentPlatformName));
+					bFoundPlatformName = true;
+					break;
+				}
+				else if (!bFoundPlatformGuess && bFoundPlatformSubstring)
+				{
+					PlatformName = FName(FStringView(CurrentPlatformName));
+					bFoundPlatformGuess = true;
+				}
+			}
+			if (bFoundPlatformName)
+			{
+				break;
+			}
+		}
+	}
+	return PlatformName;
+}
+
+void FFile::SetAsLoadTypeConfigSystem(FConfigCacheIni& ConfigSystem, FConfigFile& InConfigFile)
+{
+	check(&InConfigFile == this->ConfigFile);
+	if (!ConfigSystem.IsGloballyRegistered())
+	{
+		// We only know how to load globally registered ConfigSystems. If this system is not globally registered leave its
+		// FConfigFiles as Uninitialized LoadType.
+		return;
+	}
+	InConfigFile.LoadType = UE::ConfigAccessTracking::ELoadType::ConfigSystem;
+	if (&ConfigSystem == GConfig)
+	{
+		// GConfig's platform is set equal to the editor's platform (e.g. Windows), but we need to mark
+		// FConfigFiles as coming from GConfig, so set the override platformname
+		bOverridePlatformName = true;
+	}
+}
+
+FSection::FSection(FFile& InFileAccess, FStringView InSectionName)
+	: FileAccess(&InFileAccess)
+	, SectionName(FName(InSectionName, NAME_NO_NUMBER).GetComparisonIndex())
+{
+}
+
+namespace Private
+{
+
+
+// Simple lock class used, copied from ObjectHandleTracking.cpp. TODO: Replace with ParkingLot's lightweight RWLock
+struct FLightweightRWLock
+{
+	static const int32 Unlocked = 0;
+	static const int32 LockedOffset = 0x7fffffff;		// Subtracted for write lock, producing a large negative number
+
+	std::atomic_int32_t Mutex = 0;
+};
+
+class FLightweightReadScopeLock
+{
+public:
+	FLightweightReadScopeLock(FLightweightRWLock& InLock)
+		: Lock(InLock)
+	{
+		int32 LockValue = Lock.Mutex.fetch_add(1, std::memory_order_acquire);
+		while (LockValue < 0)
+		{
+			FPlatformProcess::Yield();
+			LockValue = Lock.Mutex.load(std::memory_order_acquire);
+		}
+	}
+	~FLightweightReadScopeLock()
+	{
+		Lock.Mutex.fetch_sub(1, std::memory_order_release);
+	}
+
+private:
+	FLightweightRWLock& Lock;
+};
+
+class FLightweightWriteScopeLock
+{
+public:
+	FLightweightWriteScopeLock(FLightweightRWLock& InLock)
+		: Lock(InLock)
+	{
+		int32 ExpectedUnlocked = FLightweightRWLock::Unlocked;
+		while (!InLock.Mutex.compare_exchange_weak(ExpectedUnlocked, -FLightweightRWLock::LockedOffset, std::memory_order_acquire))
+		{
+			FPlatformProcess::Yield();
+			ExpectedUnlocked = FLightweightRWLock::Unlocked;
+		}
+	}
+	~FLightweightWriteScopeLock()
+	{
+		Lock.Mutex.fetch_add(FLightweightRWLock::LockedOffset, std::memory_order_release);
+	}
+private:
+	FLightweightRWLock& Lock;
+};
+
+std::atomic<int32> ConfigValueReadCallbackQuantity;
+
+struct FConfigReadCallbacks
+{
+	static FConfigReadCallbacks& Get()
+	{
+		static FConfigReadCallbacks Callbacks;
+		return Callbacks;
+	}
+
+	void OnConfigValueRead(UE::ConfigAccessTracking::FSection* Section, FMinimalName ValueName,
+		const FConfigValue& ConfigValue)
+	{
+		FLightweightReadScopeLock _(HandleLock);
+		for (auto& Pair : ConfigValueReadCallbacks)
+		{
+			Pair.Value(Section, ValueName, ConfigValue);
+		}
+	}
+
+	FConfigValueReadCallbackId AddConfigValueReadCallback(FConfigValueReadCallbackFunc Func)
+	{
+		FConfigValueReadCallbackId Result;
+		FLightweightWriteScopeLock _(HandleLock);
+		NextHandleId++;
+		ConfigValueReadCallbacks.Add({ NextHandleId, Func });
+		++ConfigValueReadCallbackQuantity;
+		Result = FConfigValueReadCallbackId{ NextHandleId };
+		return Result;
+	}
+
+	void RemoveConfigValueReadCallback(FConfigValueReadCallbackId Handle)
+	{
+		FLightweightWriteScopeLock _(HandleLock);
+		for (int32 i = ConfigValueReadCallbacks.Num() - 1; i >= 0; --i)
+		{
+			auto& Pair = ConfigValueReadCallbacks[i];
+			if (Pair.Key == Handle.Id)
+			{
+				ConfigValueReadCallbacks.RemoveAt(i);
+				--ConfigValueReadCallbackQuantity;
+			}
+		}
+	}
+
+private:
+	int32 NextHandleId = 0;
+	TArray<TTuple<int32, FConfigValueReadCallbackFunc>> ConfigValueReadCallbacks;
+	FLightweightRWLock HandleLock;
+};
+
+void OnConfigValueReadInternal(UE::ConfigAccessTracking::FSection* Section, FMinimalName ValueName,
+	const FConfigValue& ConfigValue)
+{
+	// By contract with FConfigFile::SuppressReporting we guarantee that we do not report reads of FConfigValue of
+	// suppressed ConfigFiles; we implement this by early exiting if the ConfigFile pointer is null.
+	// By contract with AddConfigValueReadCallback, we additionally guarantee that ConfigFile pointer is available in
+	// the reported information.
+	if (!Section || !Section->FileAccess || !Section->FileAccess->ConfigFile)
+	{
+		return;
+	}
+	FConfigReadCallbacks::Get().OnConfigValueRead(Section, ValueName, ConfigValue);
+}
+
+} // namespace Private
+
+FConfigValueReadCallbackId AddConfigValueReadCallback(FConfigValueReadCallbackFunc Callback)
+{
+	return UE::ConfigAccessTracking::Private::FConfigReadCallbacks::Get().AddConfigValueReadCallback(Callback);
+}
+
+void RemoveConfigValueReadCallback(FConfigValueReadCallbackId Handle)
+{
+	UE::ConfigAccessTracking::Private::FConfigReadCallbacks::Get().RemoveConfigValueReadCallback(Handle);
+}
+
+} // namespace UE::ConfigAccessTracking
+
+#endif
\ No newline at end of file
diff --git a/Engine/Source/Runtime/Core/Private/Misc/ConfigCacheIni.cpp b/Engine/Source/Runtime/Core/Private/Misc/ConfigCacheIni.cpp
--- a/Engine/Source/Runtime/Core/Private/Misc/ConfigCacheIni.cpp
+++ b/Engine/Source/Runtime/Core/Private/Misc/ConfigCacheIni.cpp
@@ -391,11 +391,11 @@ void FConfigSection::HandleAddCommand(FName ValueName, FString&& Value, bool bAp
 	{
 		if (bAppendValueIfNotArrayOfStructsKeyUsed)
 		{
-			Add(ValueName, FConfigValue(MoveTemp(Value)));
+			Add(ValueName, FConfigValue(this, ValueName, MoveTemp(Value)));
 		}
 		else
 		{
-			AddUnique(ValueName, FConfigValue(MoveTemp(Value)));
+			AddUnique(ValueName, FConfigValue(this, ValueName, MoveTemp(Value)));
 		}
 	}
 }
@@ -619,8 +619,39 @@ FConfigFile::~FConfigFile()
 
 	delete SourceConfigFile;
 	SourceConfigFile = nullptr;
+#if UE_WITH_CONFIG_TRACKING 
+	if (FileAccess)
+	{
+		FileAccess->ConfigFile = nullptr;
+	}
+#endif
 }
 
+#if UE_WITH_CONFIG_TRACKING 
+void FConfigFile::SuppressReporting()
+{
+	LoadType = UE::ConfigAccessTracking::ELoadType::SuppressReporting;
+	if (FileAccess)
+	{
+		FileAccess->ConfigFile = nullptr;
+		FileAccess.SafeRelease();
+	}
+}
+
+UE::ConfigAccessTracking::FFile* FConfigFile::GetFileAccess() const
+{
+	if (!FileAccess)
+	{
+		if (LoadType == UE::ConfigAccessTracking::ELoadType::SuppressReporting)
+		{
+			return nullptr;
+		}
+		FileAccess.Set(new UE::ConfigAccessTracking::FFile(this));
+	}
+	return FileAccess.GetReference();
+}
+#endif
+
 bool FConfigFile::operator==( const FConfigFile& Other ) const
 {
 	if ( Pairs.Num() != Other.Pairs.Num() )
@@ -653,7 +684,15 @@ FConfigSection* FConfigFile::FindOrAddSectionInternal(const FString& SectionName
 	FConfigSection* Section = FindInternal(SectionName);
 	if (Section == nullptr)
 	{
-		Section = &Add(SectionName, FConfigSection());
+		UE::ConfigAccessTracking::FSection* SectionAccess = nullptr;
+#if UE_WITH_CONFIG_TRACKING
+		UE::ConfigAccessTracking::FFile* LocalFileAccess = GetFileAccess();
+		if (LocalFileAccess)
+		{
+			SectionAccess = new UE::ConfigAccessTracking::FSection(*LocalFileAccess, FStringView(SectionName));
+		}
+#endif
+		Section = &Add(SectionName, FConfigSection(SectionAccess));
 	}
 	return Section;
 }
@@ -971,7 +1010,7 @@ void FConfigFile::CombineFromBuffer(const FString& Buffer, const FString& FileHi
 				else if( Cmd=='-' )
 				{
 					// Remove if present.
-					CurrentSection->RemoveSingle(Key, ProcessedValue);
+					CurrentSection->RemoveSingle(Key, FConfigValue(CurrentSection, Key, ProcessedValue));
 					CurrentSection->CompactStable();
 				}
 				else if ( Cmd=='.' )
@@ -1003,7 +1042,7 @@ void FConfigFile::CombineFromBuffer(const FString& Buffer, const FString& FileHi
 						if (!ConfigValue)
 						{
 							CurrentSection->Add(Key,
-								FConfigValue(MoveTemp(ProcessedValue)));
+								FConfigValue(CurrentSection, Key, MoveTemp(ProcessedValue)));
 						}
 						else
 						{
@@ -1041,6 +1080,16 @@ void FConfigFile::ProcessInputFileContents(FStringView Contents, const FString&
 	LLM_SCOPE_DYNAMIC_STAT_OBJECTPATH_FNAME(FileName, ELLMTagSet::Assets);
 	LLM_SCOPE_DYNAMIC_STAT_OBJECTPATH_FNAME(ConfigFileClassName, ELLMTagSet::AssetClasses);
 	UE_TRACE_METADATA_SCOPE_ASSET_FNAME(FileName, ConfigFileClassName, FileName);
+#if UE_WITH_CONFIG_TRACKING
+	if (LoadType == UE::ConfigAccessTracking::ELoadType::Uninitialized)
+	{
+		LoadType = UE::ConfigAccessTracking::ELoadType::LocalSingleIniFile;
+	}
+	if (Name.IsNone())
+	{
+		Name = FileName;
+	}
+#endif
 
 	const TCHAR* Ptr = Contents.Len() > 0 ? Contents.GetData() : nullptr;
 	FConfigSection* CurrentSection = nullptr;
@@ -1147,12 +1196,12 @@ void FConfigFile::ProcessInputFileContents(FStringView Contents, const FString&
 					FParse::QuotedString(Value, ProcessedValue);
 
 					// Add this pair to the current FConfigSection
-					CurrentSection->Add(KeyName, FConfigValue(MoveTemp(ProcessedValue)));
+					CurrentSection->Add(KeyName, FConfigValue(CurrentSection, KeyName, MoveTemp(ProcessedValue)));
 				}
 				else
 				{
 					// Add this pair to the current FConfigSection
-					CurrentSection->Add(KeyName, FConfigValue(Value));
+					CurrentSection->Add(KeyName, FConfigValue(CurrentSection, KeyName, Value));
 				}
 			}
 		}
@@ -1868,7 +1917,10 @@ void FConfigFile::AddMissingProperties( const FConfigFile& InSourceFile )
 					SourceSection.MultiFindPointer(SourcePropertyName, Results, true);
 					for (const FConfigValue* Result : Results)
 					{
-						DestSection->Add(SourcePropertyName, *Result);
+						FConfigValue& AddedValue = DestSection->Add(SourcePropertyName, *Result);
+#if UE_WITH_CONFIG_TRACKING 
+						AddedValue.SetSectionAccess(DestSection->SectionAccess.GetReference());
+#endif
 						Dirty = true;
 					}
 				}
@@ -2029,7 +2081,7 @@ void FConfigFile::SetString( const TCHAR* Section, const TCHAR* Key, const TCHAR
 	FConfigValue* ConfigValue = Sec->Find( Key );
 	if( ConfigValue == nullptr )
 	{
-		Sec->Add(Key, FConfigValue(Value));
+		Sec->Add(Key, FConfigValue(Sec, FName(Key), Value));
 		Dirty = true;
 	}
 	// Use GetSavedValueForWriting rather than GetSavedValue to avoid reporting the value as having been accessed for dependency tracking
@@ -2050,7 +2102,7 @@ void FConfigFile::SetText( const TCHAR* Section, const TCHAR* Key, const FText&
 	FConfigValue* ConfigValue = Sec->Find( Key );
 	if( ConfigValue == nullptr )
 	{
-		Sec->Add(Key, FConfigValue(MoveTemp(StrValue)));
+		Sec->Add(Key, FConfigValue(Sec, FName(Key), MoveTemp(StrValue)));
 		Dirty = true;
 	}
 	// Use GetSavedValueForWriting rather than GetSavedValue to avoid reporting the value as having been accessed for dependency tracking
@@ -2099,7 +2151,7 @@ void FConfigFile::SetArray(const TCHAR* Section, const TCHAR* Key, const TArray<
 
 	for (int32 i = 0; i < Value.Num(); i++)
 	{
-		Sec->Add(Key, FConfigValue(Value[i]));
+		Sec->Add(Key, FConfigValue(Sec, FName(Key), Value[i]));
 		Dirty = true;
 	}
 }
@@ -2107,7 +2159,7 @@ void FConfigFile::SetArray(const TCHAR* Section, const TCHAR* Key, const TArray<
 bool FConfigFile::AddToSection(const TCHAR* SectionName, FName Key, const FString& Value)
 {
 	FConfigSection* Section = FindOrAddSectionInternal(SectionName);
-	Section->Add(Key, FConfigValue(Value));
+	Section->Add(Key, FConfigValue(Section, Key, Value));
 	Dirty = true;
 	return true;
 }
@@ -2115,13 +2167,13 @@ bool FConfigFile::AddToSection(const TCHAR* SectionName, FName Key, const FStrin
 bool FConfigFile::AddUniqueToSection(const TCHAR* SectionName, FName Key, const FString& Value)
 {
 	FConfigSection* Section = FindOrAddSectionInternal(SectionName);
-	if (Section->FindPair(Key, FConfigValue(Value)))
+	if (Section->FindPair(Key, FConfigValue(Section, Key, Value)))
 	{
 		return false;
 	}
 	
 	// just call Add since we already checked above if it exists (AddUnique can't return whether or not it existed)
-	Section->Add(Key, FConfigValue(Value));
+	Section->Add(Key, FConfigValue(Section, Key, Value));
 	Dirty = true;
 	return true;
 }
@@ -2144,13 +2196,13 @@ bool FConfigFile::RemoveFromSection(const TCHAR* SectionName, FName Key, const F
 {
 	FConfigSection* Section = FindInternal(SectionName);
 	// if it doesn't contain the pair, do nothing
-	if (Section == nullptr || !Section->FindPair(Key, FConfigValue(Value)))
+	if (Section == nullptr || !Section->FindPair(Key, FConfigValue(Section, Key, Value)))
 	{
 		return false;
 	}
 
 	// remove any copies of the pair
-	Section->Remove(Key, FConfigValue(Value));
+	Section->Remove(Key, FConfigValue(Section, Key, Value));
 	Dirty = true;
 	return true;
 }
@@ -2196,6 +2248,9 @@ void FConfigFile::ProcessSourceAndCheckAgainstBackup()
 		FConfigFile BackupFile;
 		ProcessIniContents(*BackupFilename, *BackupFilename, &BackupFile, false, false);
 
+#if UE_WITH_CONFIG_TRACKING
+		UE::ConfigAccessTracking::FFile* LocalFileAccess = GetFileAccess();
+#endif
 		for (TMap<FString,FConfigSection>::TIterator SectionIterator(*SourceConfigFile); SectionIterator; ++SectionIterator)
 		{
 			const FString& SectionName = SectionIterator.Key();
@@ -2205,7 +2260,16 @@ void FConfigFile::ProcessSourceAndCheckAgainstBackup()
 			if (BackupSection && !UE::ConfigCacheIni::Private::FAccessor::AreSectionsEqualForWriting(SourceSection, *BackupSection))
 			{
 				this->Remove( SectionName );
-				this->Add( SectionName, SourceSection );
+				FConfigSection& AddedSection = this->Add( SectionName, SourceSection );
+#if UE_WITH_CONFIG_TRACKING
+				UE::ConfigAccessTracking::FSection* SectionAccess = LocalFileAccess ?
+					new UE::ConfigAccessTracking::FSection(*LocalFileAccess, FStringView(SectionName)) : nullptr;
+				AddedSection.SectionAccess = SectionAccess;
+				for (TPair<FName, FConfigValue>& Pair : AddedSection)
+				{
+					Pair.Value.SetSectionAccess(SectionAccess);
+				}
+#endif
 			}
 		}
 
@@ -2315,9 +2379,10 @@ static TMap<FName, TFuture<void>>& GetPlatformConfigFutures()
 }
 #endif
 
-FConfigCacheIni::FConfigCacheIni(EConfigCacheType InType)
+FConfigCacheIni::FConfigCacheIni(EConfigCacheType InType, bool bInGloballyRegistered)
 	: bAreFileOperationsDisabled(false)
 	, bIsReadyForUse(false)
+	, bGloballyRegistered(bInGloballyRegistered)
 	, Type(InType)
 {
 }
@@ -2416,6 +2481,14 @@ FConfigFile& FConfigCacheIni::Add(const FString& Filename, const FConfigFile& Fi
 		delete Result;
 	}
 	Result = new FConfigFile(File);
+#if UE_WITH_CONFIG_TRACKING
+	UE::ConfigAccessTracking::FFile* FileAccess = Result->GetFileAccess();
+	if (FileAccess)
+	{
+		FileAccess->SetAsLoadTypeConfigSystem(*this, *Result);
+		FileAccess->OverrideFilenameToLoad = FName(FStringView(Filename));
+	}
+#endif
 	return *Result;
 }
 
@@ -2824,7 +2897,12 @@ const FConfigSection* FConfigCacheIni::GetSection( const TCHAR* Section, const b
 	const FConfigSection* Sec = File->FindSection( Section );
 	if (!Sec && Force)
 	{
-		Sec = &File->Add(Section, FConfigSection());
+		UE::ConfigAccessTracking::FSection* SectionAccess = nullptr;
+#if UE_WITH_CONFIG_TRACKING
+		UE::ConfigAccessTracking::FFile* LocalFileAccess = File->GetFileAccess();
+		SectionAccess = LocalFileAccess ? new UE::ConfigAccessTracking::FSection(*LocalFileAccess, FStringView(Section)) : nullptr;
+#endif
+		Sec = &File->Add(Section, FConfigSection(SectionAccess));
 		File->Dirty = true;
 	}
 
@@ -2882,7 +2960,7 @@ void FConfigCacheIni::SetText( const TCHAR* Section, const TCHAR* Key, const FTe
 	FConfigValue* ConfigValue = Sec->Find( Key );
 	if( !ConfigValue )
 	{
-		Sec->Add(Key, FConfigValue(MoveTemp(StrValue)));
+		Sec->Add(Key, FConfigValue(Sec, FName(Key), MoveTemp(StrValue)));
 		File->Dirty = true;
 	}
 	// Use GetSavedValueForWriting rather than GetSavedValue to avoid reporting the value as having been accessed for dependency tracking
@@ -3861,6 +3939,17 @@ bool FConfigCacheIni::InitializeKnownConfigFiles(FConfigContext& Context)
 	for (uint8 KnownIndex = 0; KnownIndex < (uint8)EKnownIniFile::NumKnownFiles; KnownIndex++)
 	{
 		FConfigCacheIni::FKnownConfigFiles::FKnownConfigFile& KnownFile = Context.ConfigSystem->KnownFiles.Files[KnownIndex];
+#if UE_WITH_CONFIG_TRACKING
+		// We cannot set KnownFiles' LoadType in the FConfigCacheIni constructor because we need to compare with GConfig,
+		// which is not set during GConfig's constructor. We have to set it before calling Load on the ConfigFile, since
+		// Load can read values and LoadType must be set before any values are read.
+		UE::ConfigAccessTracking::FFile* FileAccess = KnownFile.IniFile.GetFileAccess();
+		if (FileAccess)
+		{
+			FileAccess->SetAsLoadTypeConfigSystem(*Context.ConfigSystem, KnownFile.IniFile);
+			FileAccess->OverrideFilenameToLoad = KnownFile.IniName;
+		}
+#endif
 
 		// allow for scalability to come from another platform (made above)
 		FConfigContext& ContextToUse = (KnownIndex == (uint8)EKnownIniFile::Scalability && ScalabilityPlatformOverrideContext) ? *ScalabilityPlatformOverrideContext : Context;
@@ -3963,7 +4052,7 @@ bool FConfigCacheIni::CreateGConfigFromSaved(const TCHAR* Filename)
 	// serialize right out of the preloaded data
 	FLargeMemoryReader MemoryReader((uint8*)PreloadedData, Size);
 	FKnownConfigFiles Names;
-	GConfig = new FConfigCacheIni(EConfigCacheType::Temporary);
+	GConfig = new FConfigCacheIni(EConfigCacheType::Temporary, true /* bInGloballyRegistered */);
 
 	// make an object that we can use to pass to delegates for any extra binary data they want to write
 //	FCoreDelegates::FExtraBinaryConfigData ExtraData(*GConfig, false);
@@ -4026,6 +4115,13 @@ static void InitializeConfigRemap()
 	// read in the single remap file
 	FConfigFile RemapFile;
 	FConfigContext Context = FConfigContext::ReadSingleIntoLocalFile(RemapFile);
+
+#if UE_WITH_CONFIG_TRACKING 
+	// Do not report reads of ConfigRemap. The values inside of ConfigRemap permanently affect the operation of
+	// ConfigFiles for the rest of the process lifetime, and we cannot handle rereading it for access tracking.
+	// TODO: For iterative cooks, we should instead hash RemapFile.ini and add it to a key that invalidates all packages.
+	RemapFile.SuppressReporting();
+#endif
 	
 	// read in engine and project ini files (these are not hierarchical, so it has to be done in two passes)
 	for (int Pass = 0; Pass < 2; Pass++)
@@ -4100,7 +4196,7 @@ void FConfigCacheIni::InitializeConfigSystem()
 		if (FFileHelper::LoadFileToArray(FileContent, *IniBootstrapFilename, FILEREAD_Silent))
 		{
 			FMemoryReader MemoryReader(FileContent, true);
-			GConfig = new FConfigCacheIni(EConfigCacheType::Temporary);
+			GConfig = new FConfigCacheIni(EConfigCacheType::Temporary, true /* bInGloballyRegistered */);
 			GConfig->SerializeStateForBootstrap_Impl(MemoryReader);
 			GConfig->bIsReadyForUse = true;
 			TRACE_CPUPROFILER_EVENT_SCOPE(ConfigReadyForUseBroadcast);
@@ -4117,7 +4213,7 @@ void FConfigCacheIni::InitializeConfigSystem()
 	FConfigManifest::UpgradeFromPreviousVersions();
 
 	// create GConfig
-	GConfig = new FConfigCacheIni(EConfigCacheType::DiskBacked);
+	GConfig = new FConfigCacheIni(EConfigCacheType::DiskBacked, true /* bInGloballyRegistered */);
 
 	// create a context object that we will use for all of the main ini files
 	FConfigContext Context = FConfigContext::ReadIntoGConfig();
@@ -4236,6 +4332,13 @@ bool FConfigCacheIni::LoadExternalIniFile(FConfigFile & ConfigFile, const TCHAR
 	Context.bAllowGeneratedIniWhenCooked = bAllowGeneratedIniWhenCooked;
 	Context.GeneratedConfigDir = GeneratedConfigDir;
 	Context.bWriteDestIni = bWriteDestIni;
+#if UE_WITH_CONFIG_TRACKING
+	if (ConfigFile.LoadType == UE::ConfigAccessTracking::ELoadType::Uninitialized)
+	{
+		ConfigFile.LoadType = bIsBaseIniName ? UE::ConfigAccessTracking::ELoadType::ExternalIniFile :
+			UE::ConfigAccessTracking::ELoadType::ExternalSingleIniFile;
+	}
+#endif
 	return Context.Load(IniName);
 }
 
@@ -4377,6 +4480,9 @@ FArchive& operator<<(FArchive& Ar, FConfigFile& ConfigFile)
 		ConfigFile.Dirty = bDirty;
 		ConfigFile.NoSave = bNoSave;
 		ConfigFile.bHasPlatformName = bHasPlatformName;
+#if UE_WITH_CONFIG_TRACKING
+		ConfigFile.LoadType = UE::ConfigAccessTracking::ELoadType::Manual;
+#endif
 	}
 
 	return Ar;
@@ -4730,7 +4836,7 @@ void FConfigCacheIni::AsyncInitializeConfigForPlatforms()
 	for (const TPair<FName, FDataDrivenPlatformInfo>& Pair : AllPlatformInfos)
 	{
 		GetPlatformConfigFutures().Emplace(Pair.Key);
-		GConfigForPlatform.Add(Pair.Key, new FConfigCacheIni(EConfigCacheType::Temporary));
+		GConfigForPlatform.Add(Pair.Key, new FConfigCacheIni(EConfigCacheType::Temporary, true /* bInGloballyRegistered */));
 	}
 
 	for (const TPair<FName, FDataDrivenPlatformInfo>& Pair : AllPlatformInfos)
@@ -4787,7 +4893,7 @@ FConfigCacheIni* FConfigCacheIni::ForPlatform(FName PlatformName)
 	{
 		double Start = FPlatformTime::Seconds();
 		
-		PlatformConfig = GConfigForPlatform.Add(PlatformName, new FConfigCacheIni(EConfigCacheType::Temporary));
+		PlatformConfig = GConfigForPlatform.Add(PlatformName, new FConfigCacheIni(EConfigCacheType::Temporary, true /* bInGloballyRegistered */));
 		FConfigContext Context = FConfigContext::ReadIntoConfigSystem(PlatformConfig, PlatformName.ToString());
 		InitializeKnownConfigFiles(Context);
 
diff --git a/Engine/Source/Runtime/Core/Private/Misc/ConfigContext.cpp b/Engine/Source/Runtime/Core/Private/Misc/ConfigContext.cpp
--- a/Engine/Source/Runtime/Core/Private/Misc/ConfigContext.cpp
+++ b/Engine/Source/Runtime/Core/Private/Misc/ConfigContext.cpp
@@ -355,6 +355,22 @@ bool FConfigContext::PerformLoad()
 		LLM_SCOPE_DYNAMIC_STAT_OBJECTPATH_FNAME(ConfigContextClassName, ELLMTagSet::AssetClasses);
 		UE_TRACE_METADATA_SCOPE_ASSET_FNAME(BaseName, ConfigContextClassName, BaseName);
 
+#if UE_WITH_CONFIG_TRACKING
+		if (ConfigFile->LoadType == UE::ConfigAccessTracking::ELoadType::Uninitialized)
+		{
+			ConfigFile->LoadType = UE::ConfigAccessTracking::ELoadType::LocalSingleIniFile;
+		}
+		if (ConfigFile->LoadType == UE::ConfigAccessTracking::ELoadType::LocalSingleIniFile || 
+			ConfigFile->LoadType == UE::ConfigAccessTracking::ELoadType::ExternalSingleIniFile)
+		{
+			UE::ConfigAccessTracking::FFile* FileAccess = ConfigFile->GetFileAccess();
+			if (FileAccess)
+			{
+				FileAccess->OverrideFilenameToLoad = FName(FStringView(DestIniFilename));
+			}
+		}
+#endif
+
 		// load the .ini file straight up
 		LoadAnIniFile(*DestIniFilename, *ConfigFile);
 
@@ -390,6 +406,17 @@ bool FConfigContext::PerformLoad()
 		// clear previous source config file and reset
 		delete ConfigFile->SourceConfigFile;
 		ConfigFile->SourceConfigFile = new FConfigFile();
+#if UE_WITH_CONFIG_TRACKING
+		// Ignore the SourceConfigFile for access tracking since it is not publicly accessible
+		ConfigFile->SourceConfigFile->SuppressReporting();
+
+		// Set the LoadType before calling GenerateDestIniFile, because it will set it if not
+		// already set.
+		if (ConfigFile->LoadType == UE::ConfigAccessTracking::ELoadType::Uninitialized)
+		{
+			ConfigFile->LoadType = UE::ConfigAccessTracking::ELoadType::LocalIniFile;
+		}
+#endif
 
 		// now generate and make sure it's up to date (using IniName as a Base for an ini filename)
 		// @todo This bNeedsWrite afaict is always true even if it loaded a completely valid generated/final .ini, and the write below will
@@ -788,6 +815,19 @@ bool FConfigContext::GenerateDestIniFile()
 			// this will make sure the in-memory version has what is on disk, so when Flush happens later, it does not write out an outdated value
 			// note: we only want to copy the TMap base class slice of the FConfigFile, none of the specific members of the FConfigFile class itself
 			ConfigFile->TMap<FString, FConfigSection>::operator=(*ConfigFile->SourceConfigFile);
+#if UE_WITH_CONFIG_TRACKING
+			UE::ConfigAccessTracking::FFile* LocalFileAccess = ConfigFile->GetFileAccess();
+			for (TPair<FString, FConfigSection>& SectionPair : static_cast<TMap<FString, FConfigSection>>(*ConfigFile))
+			{
+				UE::ConfigAccessTracking::FSection* SectionAccess = LocalFileAccess ?
+					new UE::ConfigAccessTracking::FSection(*LocalFileAccess, FStringView(SectionPair.Key)) : nullptr;
+				SectionPair.Value.SectionAccess = SectionAccess;
+				for (TPair<FName, FConfigValue>& ValuePair : SectionPair.Value)
+				{
+					ValuePair.Value.SetSectionAccess(SectionAccess);
+				}
+			}
+#endif
 		}
 
 		LoadAnIniFile(*DestIniFilename, *ConfigFile);
@@ -817,7 +857,10 @@ bool FConfigContext::GenerateDestIniFile()
 			InDestConfigFile.SourceConfigFile = nullptr;
 		}
 		InDestConfigFile.SourceConfigFile = new FConfigFile(InDestConfigFile);
-
+#if UE_WITH_CONFIG_TRACKING
+		// Ignore the SourceConfigFile for access tracking since it is not publicly accessible
+		InDestConfigFile.SourceConfigFile->SuppressReporting();
+#endif
 		// mark it as dirty (caller may want to save)
 		InDestConfigFile.Dirty = true;
 
@@ -947,7 +990,7 @@ bool FConfigContext::GenerateDestIniFile()
 		if (FConfigSection* DestConfigSectionIniVersion = ConfigFile->FindOrAddSectionInternal(CurrentIniVersionString))
 		{
 			// Update the version. If it's already there then good but if not, we add it.
-			DestConfigSectionIniVersion->FindOrAdd(VersionName, FConfigValue(FString::FromInt(SourceConfigVersionNum)));
+			DestConfigSectionIniVersion->FindOrAdd(VersionName, FConfigValue(DestConfigSectionIniVersion, VersionName, FString::FromInt(SourceConfigVersionNum)));
 		}
 
 		// Add back any preserved sections.
diff --git a/Engine/Source/Runtime/Core/Private/Misc/DataDrivenPlatformInfoRegistry.cpp b/Engine/Source/Runtime/Core/Private/Misc/DataDrivenPlatformInfoRegistry.cpp
--- a/Engine/Source/Runtime/Core/Private/Misc/DataDrivenPlatformInfoRegistry.cpp
+++ b/Engine/Source/Runtime/Core/Private/Misc/DataDrivenPlatformInfoRegistry.cpp
@@ -100,6 +100,9 @@ bool FDataDrivenPlatformInfoRegistry::LoadDataDrivenIniFile(int32 Index, FConfig
 			PlatformName = FPaths::GetCleanFilename(FPaths::GetPath(IniFilenames[Index]));
 		}
 
+		IniFile.PlatformName = PlatformName;
+		IniFile.bHasPlatformName = !PlatformName.IsEmpty();
+
 		return true;
 	}
 
@@ -250,7 +253,7 @@ static void ParsePreviewPlatforms(const FConfigFile& IniFile)
 	if (!FParse::Param(FCommandLine::Get(), TEXT("NoPreviewPlatforms")))
 	{
 		// walk over the file looking for PreviewPlatform sections
-		for (auto Section : IniFile)
+		for (const auto& Section : IniFile)
 		{
 			if (Section.Key.StartsWith(TEXT("PreviewPlatform ")))
 			{
diff --git a/Engine/Source/Runtime/Core/Public/Misc/AccessDetection.h b/Engine/Source/Runtime/Core/Public/Misc/AccessDetection.h
--- a/Engine/Source/Runtime/Core/Public/Misc/AccessDetection.h
+++ b/Engine/Source/Runtime/Core/Public/Misc/AccessDetection.h
@@ -11,7 +11,7 @@ enum class EType : uint32
 {
 	None			= 0,
 	File			= 1 << 0, 
-	Ini 			= 1 << 1, 
+	Ini UE_DEPRECATED(5.4, "Use FConfigAccessTracking in ConfigAccessTracking instead") = 1 << 1,
 	CVar 			= 1 << 2,
 };
 ENUM_CLASS_FLAGS(EType);
@@ -57,4 +57,4 @@ FORCEINLINE void ReportAccess(EType Type) {}
 
 #endif
 
-}}
\ No newline at end of file
+}}
diff --git a/Engine/Source/Runtime/Core/Public/Misc/ConfigAccessTracking.h b/Engine/Source/Runtime/Core/Public/Misc/ConfigAccessTracking.h
new file mode 100644
--- /dev/null
+++ b/Engine/Source/Runtime/Core/Public/Misc/ConfigAccessTracking.h
@@ -0,0 +1,159 @@
+// Copyright Epic Games, Inc. All Rights Reserved.
+
+/*---------------------------------------------------------------------------------------------------------------------
+	Access tracking for Config cache. When compiled in, adds extra data to FConfigFile structures so that we can report
+	File,Section,Value names whenever an FConfigValue is read. The data about any accessed ConfigValue is reported to
+	callbacks registered via UE::ConfigAccessTracking::AddConfigValueReadCallback.
+---------------------------------------------------------------------------------------------------------------------*/
+#pragma once
+
+#include "HAL/Platform.h"
+#include "UObject/NameTypes.h"
+
+#define UE_WITH_CONFIG_TRACKING (WITH_EDITOR || 0)
+
+#if UE_WITH_CONFIG_TRACKING
+#include <atomic>
+#include "Containers/StringView.h"
+#include "Templates/Function.h"
+#include "Templates/RefCounting.h"
+
+class FConfigCacheIni;
+class FConfigFile;
+class FConfigSection;
+struct FConfigValue;
+#endif
+
+#if UE_WITH_CONFIG_TRACKING 
+
+namespace UE::ConfigAccessTracking
+{
+
+/**
+ * The manner in which the FConfigFile was loaded, so that subscribers can record
+ * how to reload it in another process and reread the values.
+ */
+enum class ELoadType : uint8
+{
+	ConfigSystem,
+	LocalIniFile,
+	LocalSingleIniFile,
+	ExternalIniFile,
+	ExternalSingleIniFile,
+	Manual,
+	SuppressReporting,
+	Uninitialized,
+};
+
+/**
+ * A refcounted struct of data about an FConfigFile, including a backpointer to the FConfigFile if it is still 
+ * in memory. References to this data are held by all the FConfigSections in the FConfigFile so that they can lookup
+ * data about their FConfigFile when their FConfigValues are accessed. Due to FConfigSection being movable out of
+ * FConfigFiles, this struct might outlive its FConfigFile (but its backpointer will be set to null).
+ */
+struct FFile : public FRefCountBase
+{
+	CORE_API FFile(const FConfigFile* InConfigFile);
+	CORE_API FName GetFilenameToLoad() const;
+	CORE_API FName GetPlatformName();
+	CORE_API void SetAsLoadTypeConfigSystem(FConfigCacheIni& ConfigSystem, FConfigFile& ConfigFile);
+
+public:
+	const FConfigFile* ConfigFile = nullptr;
+	FName OverrideFilenameToLoad;
+private:
+	FName PlatformName = NAME_None;
+	FStringView SavedConfigFilePlatformName;
+	bool bSavedHasConfigFile : 1; // = false;
+	bool bSavedConfigFileHasPlatformName : 1; // = false;
+	bool bPlatformNameInitialized : 1; // = false;
+	bool bOverridePlatformName : 1; // = false;
+	bool bSavedOverridePlatformName : 1; // = false;
+};
+
+/**
+ * A refcounted struct of data about an FConfigSection, including a backpointer to the UE::ConfigAccessTracking::FFile
+ * that holds data about the section's FConfigFile. It does not hold a backpointer to the FConfigSection itself because
+ * FConfigSections are value types in a TMap owned by the FConfigFIle and their address frequently changes.
+ * References to this data are held by all FConfigValues in the FConfigSection. Due to FConfigSection being movable out
+ * of FConfigFiles, this struct might outlive the existence of its FConfigSection in the FConfigFile, or even outlive
+ * the FConfigFile.
+ */
+struct FSection : public FRefCountBase
+{
+	CORE_API FSection(FFile& InFileAccess, FStringView InSectionName);
+
+	TRefCountPtr<FFile> FileAccess;
+	FNameEntryId SectionName;
+};
+
+/** Function type that for subscribers to AddConfigValueReadCallback. */
+using FConfigValueReadCallbackFunc = TFunction<void(UE::ConfigAccessTracking::FSection* Section,
+	FMinimalName ValueName, const FConfigValue& ConfigValue)>;
+
+/** Handle used to remove a subscriber from AddConfigValueReadCallback. */
+struct FConfigValueReadCallbackId
+{
+	int32 Id = -1;
+	bool IsValid() { return Id != -1; }
+};
+
+/** Subscribe to events sent whenever any FConfigValue is read. */
+CORE_API FConfigValueReadCallbackId AddConfigValueReadCallback(FConfigValueReadCallbackFunc Callback);
+/** Remove a subscriber that was added via FConfigValue. */
+CORE_API void RemoveConfigValueReadCallback(FConfigValueReadCallbackId DelegateHandle);
+
+namespace Private
+{
+
+/**
+ * FFilePtrClearOnCopy: Device to workaround c++'s lack of the ability to create a copyconstructor that appends custom
+ * code to the default copyconstructor. FConfigFile uses a default copy constructor, and we don't want to replace that
+ * with a manual one that we then have to keep up to date with any property changes. But we want to change the
+ * behavior of that copy constructor; FConfigFile->FileAccess is a pointer to a new'd structure that has a backpointer
+ * to FConfigFile and so every allocated FConfigFile needs to construct its own copy of it rather than copying the
+ * pointer from a source FConfigFile. But that is not (and should not be) the behavior of TRefCountPtr. So we make
+ * this subclass of TRefCountPtr with clearing-on-copy behavior, and allocate the FFile* on demand when accessed in
+ * the new FConfigFile.
+ */
+struct FFilePtrClearOnCopy : public TRefCountPtr<FFile>
+{
+	FFilePtrClearOnCopy() = default;
+	FFilePtrClearOnCopy(FFile* In) : TRefCountPtr<FFile>(In) {}
+	FFilePtrClearOnCopy(const FFilePtrClearOnCopy& Other) {}
+	FFilePtrClearOnCopy(FFilePtrClearOnCopy&& Other) {}
+	FFilePtrClearOnCopy& operator=(const FFilePtrClearOnCopy& Other) { return *this; }
+	FFilePtrClearOnCopy& operator=(FFilePtrClearOnCopy&& Other) { return *this; }
+	void Set(FFile* Ptr) { this->TRefCountPtr<UE::ConfigAccessTracking::FFile>::operator=(Ptr); }
+};
+
+/**
+ * Stores the number of existing AddConfigValueReadCallback subscribers. We use it to avoid the cost
+ * of a function call when no subscribers are registered.
+ */
+extern CORE_API std::atomic<int32> ConfigValueReadCallbackQuantity;
+
+CORE_API void OnConfigValueReadInternal(UE::ConfigAccessTracking::FSection* Section, FMinimalName ValueName,
+	const FConfigValue& ConfigValue);
+
+inline void OnConfigValueRead(UE::ConfigAccessTracking::FSection* Section, FMinimalName ValueName,
+	const FConfigValue& ConfigValue)
+{
+	if (ConfigValueReadCallbackQuantity.load(std::memory_order_acquire) > 0)
+	{
+		OnConfigValueReadInternal(Section, ValueName, ConfigValue);
+	}
+}
+
+} // namespace Private
+
+} // namespace UE::ConfigAccessTracking
+
+#else  // UE_WITH_CONFIG_TRACKING
+
+// Declare FSection even though it is undefined so that we can pass FSection* as an argument to FConfigSection and FConfigValue
+// constructors rather than having to create separate constructors under #ifdef UE_WITH_CONFIG_TRACKING
+namespace UE::ConfigAccessTracking { struct FSection; }
+
+#endif // else !UE_WITH_CONFIG_TRACKING
+
diff --git a/Engine/Source/Runtime/Core/Public/Misc/ConfigCacheIni.h b/Engine/Source/Runtime/Core/Public/Misc/ConfigCacheIni.h
--- a/Engine/Source/Runtime/Core/Public/Misc/ConfigCacheIni.h
+++ b/Engine/Source/Runtime/Core/Public/Misc/ConfigCacheIni.h
@@ -26,8 +26,8 @@
 #include "Math/Vector.h"
 #include "Math/Vector2D.h"
 #include "Math/Vector4.h"
-#include "Misc/AccessDetection.h"
 #include "Misc/Build.h"
+#include "Misc/ConfigAccessTracking.h"
 #include "Misc/ConfigTypes.h"
 #include "Misc/Paths.h"
 #include "Serialization/Archive.h"
@@ -48,7 +48,6 @@ CORE_API DECLARE_LOG_CATEGORY_EXTERN(LogConfig, Log, All);
 
 // Server builds should be tweakable even in Shipping
 #define ALLOW_INI_OVERRIDE_FROM_COMMANDLINE			(UE_SERVER || !(UE_BUILD_SHIPPING))
-#define CONFIG_REMEMBER_ACCESS_PATTERN (WITH_EDITOR || 0)
 
 
 ///////////////////////////////////////////////////////////////////////////////
@@ -116,16 +115,34 @@ namespace UE::ConfigCacheIni::Private { struct FAccessor; }
 namespace UE::ConfigCacheIni::Private { struct FImpl; }
 
 class FConfigContext;
+#if UE_WITH_CONFIG_TRACKING
+UE::ConfigAccessTracking::FSection* GetSectionAccess(const FConfigSection* InSection);
+#endif
 
 struct FConfigValue
 {
 public:
-	FConfigValue() { }
+	FConfigValue()
+		: FConfigValue(nullptr, NAME_None)
+	{}
+
+	FConfigValue(const FConfigSection* InSection, FName InValueName)
+#if UE_WITH_CONFIG_TRACKING
+		: SectionAccess(GetSectionAccess(InSection))
+		, ValueName(InValueName)
+#endif
+	{
+	}
 
 	FConfigValue(const TCHAR* InValue)
+		: FConfigValue(nullptr, NAME_None, InValue)
+	{}
+
+	FConfigValue(const FConfigSection* InSection, FName InValueName, const TCHAR* InValue)
 		: SavedValue(InValue)
-#if CONFIG_REMEMBER_ACCESS_PATTERN 
-		, bRead(false)
+#if UE_WITH_CONFIG_TRACKING
+		, SectionAccess(GetSectionAccess(InSection))
+		, ValueName(InValueName)
 #endif
 	{
 		SavedValueHash = FTextLocalizationResource::HashString(SavedValue);
@@ -133,9 +150,14 @@ public:
 	}
 
 	FConfigValue(const FString& InValue)
+		: FConfigValue(nullptr, NAME_None, InValue)
+	{}
+
+	FConfigValue(const FConfigSection* InSection, FName InValueName, const FString& InValue)
 		: SavedValue(InValue)
-#if CONFIG_REMEMBER_ACCESS_PATTERN 
-		, bRead(false)
+#if UE_WITH_CONFIG_TRACKING
+		, SectionAccess(GetSectionAccess(InSection))
+		, ValueName(InValueName)
 #endif
 	{
 		SavedValueHash = FTextLocalizationResource::HashString(SavedValue);
@@ -143,9 +165,14 @@ public:
 	}
 
 	FConfigValue(FString&& InValue)
+		: FConfigValue(nullptr, NAME_None, MoveTemp(InValue))
+	{}
+
+	FConfigValue(const FConfigSection* InSection, FName InValueName, FString&& InValue)
 		: SavedValue(MoveTemp(InValue))
-#if CONFIG_REMEMBER_ACCESS_PATTERN 
-		, bRead(false)
+#if UE_WITH_CONFIG_TRACKING
+		, SectionAccess(GetSectionAccess(InSection))
+		, ValueName(InValueName)
 #endif
 	{
 		SavedValueHash = FTextLocalizationResource::HashString(SavedValue);
@@ -156,8 +183,9 @@ public:
 		: SavedValue(InConfigValue.SavedValue)
 		, ExpandedValue(InConfigValue.ExpandedValue)
 		, SavedValueHash(InConfigValue.SavedValueHash)
-#if CONFIG_REMEMBER_ACCESS_PATTERN 
-		, bRead(InConfigValue.bRead)
+#if UE_WITH_CONFIG_TRACKING
+		, SectionAccess(InConfigValue.SectionAccess)
+		, ValueName(InConfigValue.ValueName)
 #endif
 	{
 		// shouldn't need to expand value it's assumed that the other FConfigValue has done this already
@@ -167,8 +195,9 @@ public:
 		: SavedValue(MoveTemp(InConfigValue.SavedValue))
 		, ExpandedValue(MoveTemp(InConfigValue.ExpandedValue))
 		, SavedValueHash(InConfigValue.SavedValueHash)
-#if CONFIG_REMEMBER_ACCESS_PATTERN 
-		, bRead(InConfigValue.bRead)
+#if UE_WITH_CONFIG_TRACKING
+		, SectionAccess(MoveTemp(InConfigValue.SectionAccess))
+		, ValueName(InConfigValue.ValueName)
 #endif
 	{
 		// shouldn't need to expand value it's assumed that the other FConfigValue has done this already
@@ -179,8 +208,9 @@ public:
 		SavedValue = MoveTemp(RHS.SavedValue);
 		ExpandedValue = MoveTemp(RHS.ExpandedValue);
 		SavedValueHash = RHS.SavedValueHash;
-#if CONFIG_REMEMBER_ACCESS_PATTERN 
-		bRead = RHS.bRead;
+#if UE_WITH_CONFIG_TRACKING
+		SectionAccess = MoveTemp(RHS.SectionAccess);
+		ValueName = RHS.ValueName;
 #endif
 
 		return *this;
@@ -191,8 +221,9 @@ public:
 		SavedValue = RHS.SavedValue;
 		ExpandedValue = RHS.ExpandedValue;
 		SavedValueHash = RHS.SavedValueHash;
-#if CONFIG_REMEMBER_ACCESS_PATTERN 
-		bRead = RHS.bRead;
+#if UE_WITH_CONFIG_TRACKING
+		SectionAccess = RHS.SectionAccess;
+		ValueName = RHS.ValueName;
 #endif
 
 		return *this;
@@ -220,9 +251,8 @@ public:
 	// Returns the ini setting with any macros expanded out
 	const FString& GetValue() const 
 	{
-		UE::AccessDetection::ReportAccess(UE::AccessDetection::EType::Ini);
-#if CONFIG_REMEMBER_ACCESS_PATTERN 
-		bRead = true; 
+#if UE_WITH_CONFIG_TRACKING
+		UE::ConfigAccessTracking::Private::OnConfigValueRead(SectionAccess, ValueName, *this);
 #endif
 		return (ExpandedValue.Len() > 0 ? ExpandedValue : SavedValue); 
 	}
@@ -230,20 +260,25 @@ public:
 	// Returns the original ini setting without macro expansion
 	const FString& GetSavedValue() const 
 	{
-		UE::AccessDetection::ReportAccess(UE::AccessDetection::EType::Ini);
-#if CONFIG_REMEMBER_ACCESS_PATTERN 
-		bRead = true; 
+#if UE_WITH_CONFIG_TRACKING
+		UE::ConfigAccessTracking::Private::OnConfigValueRead(SectionAccess, ValueName, *this);
 #endif
 		return SavedValue; 
 	}
-#if CONFIG_REMEMBER_ACCESS_PATTERN 
+#if UE_WITH_CONFIG_TRACKING
+	UE_DEPRECATED(5.4, "No longer written. Use UE::ConfigAccessTracking::AddConfigValueReadCallback instead")
 	inline const bool HasBeenRead() const
 	{
-		return bRead;
+		return false;
 	}
+	UE_DEPRECATED(5.4, "No longer read.")
 	inline void SetHasBeenRead(bool InBRead ) const
 	{
-		bRead = InBRead;
+	}
+
+	void SetSectionAccess(UE::ConfigAccessTracking::FSection* InSectionAccess)
+	{
+		SectionAccess = InSectionAccess;
 	}
 #endif
 
@@ -263,9 +298,6 @@ public:
 		if (Slot.GetUnderlyingArchive().IsLoading())
 		{
 			ConfigSection.ExpandValueInternal();
-#if CONFIG_REMEMBER_ACCESS_PATTERN 
-			ConfigSection.bRead = false;
-#endif
 		}
 	}
 
@@ -327,8 +359,9 @@ private:
 	FString SavedValue;
 	FString ExpandedValue;
 	uint32 SavedValueHash;
-#if CONFIG_REMEMBER_ACCESS_PATTERN 
-	mutable bool bRead; // has this value been read since the config system started
+#if UE_WITH_CONFIG_TRACKING
+	TRefCountPtr<UE::ConfigAccessTracking::FSection> SectionAccess;
+	FMinimalName ValueName;
 #endif
 };
 
@@ -338,6 +371,18 @@ typedef TMultiMap<FName,FConfigValue> FConfigSectionMap;
 class FConfigSection : public FConfigSectionMap
 {
 public:
+	FConfigSection()
+		: FConfigSection(nullptr)
+	{
+	}
+
+	FConfigSection(UE::ConfigAccessTracking::FSection* InSectionAccess)
+#if UE_WITH_CONFIG_TRACKING
+		: SectionAccess(InSectionAccess)
+#endif
+	{
+	}
+
 	/**
 	* Check whether the input string is surrounded by quotes
 	*
@@ -376,6 +421,9 @@ public:
 
 	// look for "array of struct" keys for overwriting single entries of an array
 	TMap<FName, FString> ArrayOfStructKeys;
+#if UE_WITH_CONFIG_TRACKING
+	TRefCountPtr<UE::ConfigAccessTracking::FSection> SectionAccess;
+#endif
 
 	friend struct UE::ConfigCacheIni::Private::FAccessor;
 	friend FArchive& operator<<(FArchive& Ar, FConfigSection& ConfigSection);
@@ -427,6 +475,9 @@ public:
 	// by default, we allow saving - this is going to be applied to config files that are not loaded from disk
 	// (when loading, this will get set to false, and then the ini sections will be checked)
 	bool bCanSaveAllSections : 1; // = true;
+#if UE_WITH_CONFIG_TRACKING
+	UE::ConfigAccessTracking::ELoadType LoadType = UE::ConfigAccessTracking::ELoadType::Uninitialized;
+#endif
 
 	/** The name of this config file */
 	FName Name;
@@ -453,6 +504,10 @@ private:
 	// we can't track it just in that section. This is expected to be empty/small
 	TMap<FString, TMap<FName, FString> > PerObjectConfigArrayOfStructKeys;
 
+#if UE_WITH_CONFIG_TRACKING
+	mutable UE::ConfigAccessTracking::Private::FFilePtrClearOnCopy FileAccess;
+#endif
+
 public:
 	CORE_API FConfigFile();
 	FConfigFile( int32 ) {}	// @todo UE-DLL: Workaround for instantiated TMap template during DLLExport (TMap::FindRef)
@@ -722,6 +777,12 @@ public:
 	void AddDynamicLayerToHeirarchy(const FString& Filename) { AddDynamicLayerToHierarchy(Filename); }
 
 	friend FArchive& operator<<(FArchive& Ar, FConfigFile& ConfigFile);
+
+#if UE_WITH_CONFIG_TRACKING
+	CORE_API void SuppressReporting();
+	CORE_API UE::ConfigAccessTracking::FFile* GetFileAccess() const;
+#endif
+
 private:
 	/** 
 	 * Save the source hierarchy which was loaded out to a backup file so we can check future changes in the base/default configs
@@ -773,7 +834,7 @@ class FConfigCacheIni
 {
 public:
 	// Basic functions.
-	CORE_API FConfigCacheIni(EConfigCacheType Type);
+	CORE_API FConfigCacheIni(EConfigCacheType Type, bool bInGloballyRegistered=false);
 
 	/** DO NOT USE. This constructor is for internal usage only for hot-reload purposes. */
 	CORE_API FConfigCacheIni();
@@ -803,6 +864,12 @@ public:
 		return bIsReadyForUse;
 	}
 
+	/** Return whether this is one of the global ConfigSystems: GConfig or FConfigCacheIni::ForPlatform. */
+	bool IsGloballyRegistered() const
+	{
+		return bGloballyRegistered;
+	}
+
 	/**
 	* Prases apart an ini section that contains a list of 1-to-N mappings of strings in the following format
 	*	 [PerMapPackages]
@@ -1480,6 +1547,8 @@ private:
 
 	/** true after the base .ini files have been loaded, and GConfig is generally "ready for use" */
 	bool bIsReadyForUse;
+
+	bool bGloballyRegistered;
 	
 	/** The type of the cache (basically, do we call Flush in the destructor) */
 	EConfigCacheType Type;
@@ -1558,3 +1627,10 @@ CORE_API void DeleteRecordedConfigReadsFromIni();
  */
 UE_DEPRECATED(5.1, "Use UE::ConfigUtilities::ConvertValueFromHumanFriendlyValue")
 CORE_API const TCHAR* ConvertValueFromHumanFriendlyValue(const TCHAR* Value);
+
+#if UE_WITH_CONFIG_TRACKING
+inline UE::ConfigAccessTracking::FSection* GetSectionAccess(const FConfigSection* InSection)
+{
+	return InSection ? InSection->SectionAccess.GetReference() : nullptr;
+}
+#endif
diff --git a/Engine/Source/Runtime/Launch/Private/LaunchEngineLoop.cpp b/Engine/Source/Runtime/Launch/Private/LaunchEngineLoop.cpp
--- a/Engine/Source/Runtime/Launch/Private/LaunchEngineLoop.cpp
+++ b/Engine/Source/Runtime/Launch/Private/LaunchEngineLoop.cpp
@@ -3578,51 +3578,6 @@ int32 FEngineLoop::PreInitPostStartupScreen(const TCHAR* CmdLine)
 				}
 				UE_LOG(LogInit, Display, TEXT("\n%s"), *FileList);
 #endif
-
-#if 0 && CONFIG_REMEMBER_ACCESS_PATTERN
-
-				TArray<FString> ConfigFilenames;
-				GConfig->GetConfigFilenames(ConfigFilenames);
-				for (const FString& ConfigFilename : ConfigFilenames)
-				{
-					const FConfigFile* Config = GConfig->FindConfigFile(ConfigFilename);
-
-					for (auto& ConfigSection : *Config)
-					{
-						TSet<FName> ProcessedValues;
-						const FName SectionName = FName(*ConfigSection.Key);
-
-						for (auto& ConfigValue : ConfigSection.Value)
-						{
-							const FName& ValueName = ConfigValue.Key;
-							if (ProcessedValues.Contains(ValueName))
-								continue;
-
-							ProcessedValues.Add(ValueName);
-
-							TArray<FConfigValue> ValueArray;
-							ConfigSection.Value.MultiFind(ValueName, ValueArray, true);
-
-							bool bHasBeenAccessed = false;
-							for (const auto& ValueArrayEntry : ValueArray)
-							{
-								if (ValueArrayEntry.HasBeenRead())
-								{
-									bHasBeenAccessed = true;
-									break;
-								}
-							}
-
-							if (bHasBeenAccessed)
-							{
-								UE_LOG(LogInit, Display, TEXT("Accessed Ini Setting %s %s %s"), *ConfigFilename, *SectionName.ToString(), *ValueName.ToString());
-							}
-
-						}
-					}
-
-				}
-#endif
 			}
 			else
 			{
