[FBX Legacy] Re-Import ensure fix.

#rnx

#changelist validated
#virtualized

diff --git a/Engine/Source/Editor/UnrealEd/Private/Editor.cpp b/Engine/Source/Editor/UnrealEd/Private/Editor.cpp
--- a/Engine/Source/Editor/UnrealEd/Private/Editor.cpp
+++ b/Engine/Source/Editor/UnrealEd/Private/Editor.cpp
@@ -75,6 +75,8 @@
 #include "InterchangeManager.h"
 #include "InterchangeResultsContainer.h"
 
+#include "AssetCompilingManager.h"
+
 #if WITH_EDITOR
 #include "Subsystems/AssetEditorSubsystem.h"
 #endif
@@ -566,6 +568,8 @@ UE::Interchange::FAssetImportResultRef FReimportManager::ReimportAsync(UObject*
 		}
 	}
 
+	FAssetCompilingManager::Get().FinishCompilationForObjects({Obj});
+
 	// Let listeners know whether the reimport was successful or not
 	PostReimport.Broadcast( Obj, bSuccess );
 
@@ -577,6 +581,7 @@ UE::Interchange::FAssetImportResultRef FReimportManager::ReimportAsync(UObject*
 		ImportResultSynchronous->GetResults()->Add<UInterchangeResultError_ReimportFail>();
 	}
 	ImportResultSynchronous->SetDone();
+
 	return ImportResultSynchronous;
 }
 
