Hide scaling gizmo for non-scalable elements

* Add bool ITypedElementWorldInterface::CanScaleElement(FTypedElementHandle) which returns true by default to keep current behavior. Overrides should return whether or not an element can be scaled.
* When the in-editor selection to manipulate is cached by FLevelEditorViewportClient::CacheElementsToManipulate, skip non-scalable elements if we're using the scaling tool.
* Recalculate the cached selection to manipulate if the user changes tools. This gives a chance for non-scalable elements to be included in the selection to manipulate if the user changed away from the scaling tool, for example.
* Override CanScaleElement(FTypedElementHandle) in UActorElementWorldInterface. Return what CanEditChange returns for the ReleativeScale3D property on the actor's root component.
* Override CanScaleElement(FTypedElementHandle) in UComponentElementWorldInterface. Return what CanEditChange returns for the ReleativeScale3D property on the component.

Note that this change does not override CanScaleElement in USMInstanceElementWorldInterface so static meshes keep their current behavior by relying on ITypedElementWorldInterface::CanScaleElement(FTypedElementHandle)'s default return value of true.

#changelist validated
#virtualized

diff --git a/Engine/Source/Editor/UnrealEd/Private/LevelEditorViewport.cpp b/Engine/Source/Editor/UnrealEd/Private/LevelEditorViewport.cpp
--- a/Engine/Source/Editor/UnrealEd/Private/LevelEditorViewport.cpp
+++ b/Engine/Source/Editor/UnrealEd/Private/LevelEditorViewport.cpp
@@ -2384,6 +2384,8 @@ FLevelEditorViewportClient::FLevelEditorViewportClient(const TSharedPtr<SLevelVi
 	ModeTools->SetWidgetMode(UE::Widget::WM_Translate);
 	Widget->SetUsesEditorModeTools(ModeTools.Get());
 
+	ModeTools->OnWidgetModeChanged().AddRaw(this, &FLevelEditorViewportClient::OnWidgetModeChanged);
+
 	// Register for editor cleanse events so we can release references to hovered actors
 	FEditorSupportDelegates::CleanseEditor.AddRaw(this, &FLevelEditorViewportClient::OnEditorCleanse);
 
@@ -2431,6 +2433,8 @@ FLevelEditorViewportClient::~FLevelEditorViewportClient()
 	}
 	ResetElementsToManipulate();
 
+	ModeTools->OnWidgetModeChanged().RemoveAll(this);
+
 	// Unregister for all global callbacks to this object
 	FEditorSupportDelegates::CleanseEditor.RemoveAll(this);
 	FEditorDelegates::PreBeginPIE.RemoveAll(this);
@@ -4232,14 +4236,21 @@ void FLevelEditorViewportClient::CacheElementsToManipulate(const bool bForceRefr
 		const UTypedElementSelectionSet* SelectionSet = GetSelectionSet();
 		SelectionSet->GetNormalizedSelection(NormalizationOptions, CachedElementsToManipulate);
 
+		const UE::Widget::EWidgetMode WidgetMode = GetWidgetMode();
+
 		// Remove any elements that cannot be moved
-		CachedElementsToManipulate->RemoveAll<ITypedElementWorldInterface>([this](const TTypedElement<ITypedElementWorldInterface>& InWorldElement)
+		CachedElementsToManipulate->RemoveAll<ITypedElementWorldInterface>([this, WidgetMode](const TTypedElement<ITypedElementWorldInterface>& InWorldElement)
 		{
 			if (!InWorldElement.CanMoveElement(bIsSimulateInEditorViewport ? ETypedElementWorldType::Game : ETypedElementWorldType::Editor))
 			{
 				return true;
 			}
 
+			if (WidgetMode == UE::Widget::WM_Scale && !InWorldElement.CanScaleElement())
+			{
+				return true;
+			}
+
 			// This element must belong to the current viewport world
 			if (GEditor->PlayWorld)
 			{
@@ -5683,6 +5694,12 @@ void FLevelEditorViewportClient::ClearHoverFromObjects()
 	}
 }
 
+void FLevelEditorViewportClient::OnWidgetModeChanged(UE::Widget::EWidgetMode NewMode)
+{
+	// We need this in case the current selection allowed the previous widget mode but not the new one or vice versa.
+	ResetElementsToManipulate();
+}
+
 void FLevelEditorViewportClient::OnEditorCleanse()
 {
 	ClearHoverFromObjects();
diff --git a/Engine/Source/Editor/UnrealEd/Public/LevelEditorViewport.h b/Engine/Source/Editor/UnrealEd/Public/LevelEditorViewport.h
--- a/Engine/Source/Editor/UnrealEd/Public/LevelEditorViewport.h
+++ b/Engine/Source/Editor/UnrealEd/Public/LevelEditorViewport.h
@@ -702,6 +702,9 @@ protected:
 	 */
 	UNREALED_API bool CanDropBlueprintAsset ( const struct FSelectedAssetInfo& );
 
+	/** Called when the widget mode changes. */
+	void OnWidgetModeChanged(UE::Widget::EWidgetMode NewMode);
+
 	/** Called when editor cleanse event is triggered */
 	UNREALED_API void OnEditorCleanse();
 
diff --git a/Engine/Source/Runtime/Engine/Private/Elements/Actor/ActorElementWorldInterface.cpp b/Engine/Source/Runtime/Engine/Private/Elements/Actor/ActorElementWorldInterface.cpp
--- a/Engine/Source/Runtime/Engine/Private/Elements/Actor/ActorElementWorldInterface.cpp
+++ b/Engine/Source/Runtime/Engine/Private/Elements/Actor/ActorElementWorldInterface.cpp
@@ -57,7 +57,7 @@ bool UActorElementWorldInterface::CanMoveElement(const FTypedElementHandle& InEl
 		{
 			return false;
 		}
-#endif	// WITH_EDITOR
+#endif // WITH_EDITOR
 
 		// If the actor has a root component, but it cannot be moved, then the actor cannot move.
 		if (InWorldType == ETypedElementWorldType::Game && Actor->GetRootComponent() && !Actor->IsRootComponentMovable())
@@ -71,6 +71,20 @@ bool UActorElementWorldInterface::CanMoveElement(const FTypedElementHandle& InEl
 	return false;
 }
 
+bool UActorElementWorldInterface::CanScaleElement(const FTypedElementHandle& InElementHandle)
+{
+#if WITH_EDITOR
+	if (const AActor* Actor = ActorElementDataUtil::GetActorFromHandle(InElementHandle))
+	{
+		// The actor must be scalable.
+		FProperty* const RelativeScale3DProperty = USceneComponent::StaticClass()->FindPropertyByName(USceneComponent::GetRelativeScale3DPropertyName());
+		return Actor->CanEditChangeComponent(Actor->GetRootComponent(), RelativeScale3DProperty);
+	}
+#endif // WITH_EDITOR
+
+	return true;
+}
+
 bool UActorElementWorldInterface::GetWorldTransform(const FTypedElementHandle& InElementHandle, FTransform& OutTransform)
 {
 	if (const AActor* Actor = ActorElementDataUtil::GetActorFromHandle(InElementHandle))
diff --git a/Engine/Source/Runtime/Engine/Private/Elements/Component/ComponentElementWorldInterface.cpp b/Engine/Source/Runtime/Engine/Private/Elements/Component/ComponentElementWorldInterface.cpp
--- a/Engine/Source/Runtime/Engine/Private/Elements/Component/ComponentElementWorldInterface.cpp
+++ b/Engine/Source/Runtime/Engine/Private/Elements/Component/ComponentElementWorldInterface.cpp
@@ -67,6 +67,22 @@ bool UComponentElementWorldInterface::CanMoveElement(const FTypedElementHandle&
 	return false;
 }
 
+bool UComponentElementWorldInterface::CanScaleElement(const FTypedElementHandle& InElementHandle)
+{
+#if WITH_EDITOR
+	if (const UActorComponent* Component = ComponentElementDataUtil::GetComponentFromHandle(InElementHandle))
+	{
+		if (const USceneComponent* SceneComponent = Cast<USceneComponent>(Component))
+		{
+			FProperty* const RelativeScale3DProperty = USceneComponent::StaticClass()->FindPropertyByName(USceneComponent::GetRelativeScale3DPropertyName());
+			return SceneComponent->CanEditChange(RelativeScale3DProperty);
+		}
+	}
+#endif // WITH_EDITOR
+
+	return false;
+}
+
 bool UComponentElementWorldInterface::GetWorldTransform(const FTypedElementHandle& InElementHandle, FTransform& OutTransform)
 {
 	if (const UActorComponent* Component = ComponentElementDataUtil::GetComponentFromHandle(InElementHandle))
diff --git a/Engine/Source/Runtime/Engine/Private/Elements/Interfaces/TypedElementWorldInterface.cpp b/Engine/Source/Runtime/Engine/Private/Elements/Interfaces/TypedElementWorldInterface.cpp
--- a/Engine/Source/Runtime/Engine/Private/Elements/Interfaces/TypedElementWorldInterface.cpp
+++ b/Engine/Source/Runtime/Engine/Private/Elements/Interfaces/TypedElementWorldInterface.cpp
@@ -198,7 +198,6 @@ bool ITypedElementWorldInterface::GetBounds(const FScriptTypedElementHandle& InE
 	return GetBounds(NativeHandle, OutBounds);
 }
 
-
 bool ITypedElementWorldInterface::CanMoveElement(const FScriptTypedElementHandle& InElementHandle, const ETypedElementWorldType InWorldType)
 {
 	FTypedElementHandle NativeHandle = InElementHandle.GetTypedElementHandle();
@@ -211,6 +210,18 @@ bool ITypedElementWorldInterface::CanMoveElement(const FScriptTypedElementHandle
 	return CanMoveElement(NativeHandle, InWorldType);
 }
 
+bool ITypedElementWorldInterface::CanScaleElement(const FScriptTypedElementHandle& InElementHandle)
+{
+	FTypedElementHandle NativeHandle = InElementHandle.GetTypedElementHandle();
+	if (!NativeHandle)
+	{
+		FFrame::KismetExecutionMessage(TEXT("InElementHandle is not a valid handle."), ELogVerbosity::Error);
+		return false;
+	}
+
+	return CanScaleElement(NativeHandle);
+}
+
 bool ITypedElementWorldInterface::GetWorldTransform(const FScriptTypedElementHandle& InElementHandle, FTransform& OutTransform)
 {
 	FTypedElementHandle NativeHandle = InElementHandle.GetTypedElementHandle();
diff --git a/Engine/Source/Runtime/Engine/Public/Elements/Actor/ActorElementWorldInterface.h b/Engine/Source/Runtime/Engine/Public/Elements/Actor/ActorElementWorldInterface.h
--- a/Engine/Source/Runtime/Engine/Public/Elements/Actor/ActorElementWorldInterface.h
+++ b/Engine/Source/Runtime/Engine/Public/Elements/Actor/ActorElementWorldInterface.h
@@ -18,6 +18,7 @@ public:
 	ENGINE_API virtual UWorld* GetOwnerWorld(const FTypedElementHandle& InElementHandle) override;
 	ENGINE_API virtual bool GetBounds(const FTypedElementHandle& InElementHandle, FBoxSphereBounds& OutBounds) override;
 	ENGINE_API virtual bool CanMoveElement(const FTypedElementHandle& InElementHandle, const ETypedElementWorldType InWorldType) override;
+	ENGINE_API virtual bool CanScaleElement(const FTypedElementHandle& InElementHandle) override;
 	ENGINE_API virtual bool GetWorldTransform(const FTypedElementHandle& InElementHandle, FTransform& OutTransform) override;
 	ENGINE_API virtual bool SetWorldTransform(const FTypedElementHandle& InElementHandle, const FTransform& InTransform) override;
 	ENGINE_API virtual bool GetRelativeTransform(const FTypedElementHandle& InElementHandle, FTransform& OutTransform) override;
diff --git a/Engine/Source/Runtime/Engine/Public/Elements/Component/ComponentElementWorldInterface.h b/Engine/Source/Runtime/Engine/Public/Elements/Component/ComponentElementWorldInterface.h
--- a/Engine/Source/Runtime/Engine/Public/Elements/Component/ComponentElementWorldInterface.h
+++ b/Engine/Source/Runtime/Engine/Public/Elements/Component/ComponentElementWorldInterface.h
@@ -17,6 +17,7 @@ public:
 	ENGINE_API virtual UWorld* GetOwnerWorld(const FTypedElementHandle& InElementHandle) override;
 	ENGINE_API virtual bool GetBounds(const FTypedElementHandle& InElementHandle, FBoxSphereBounds& OutBounds) override;
 	ENGINE_API virtual bool CanMoveElement(const FTypedElementHandle& InElementHandle, const ETypedElementWorldType InWorldType) override;
+	ENGINE_API virtual bool CanScaleElement(const FTypedElementHandle& InElementHandle) override;
 	ENGINE_API virtual bool GetWorldTransform(const FTypedElementHandle& InElementHandle, FTransform& OutTransform) override;
 	ENGINE_API virtual bool SetWorldTransform(const FTypedElementHandle& InElementHandle, const FTransform& InTransform) override;
 	ENGINE_API virtual bool GetRelativeTransform(const FTypedElementHandle& InElementHandle, FTransform& OutTransform) override;
diff --git a/Engine/Source/Runtime/Engine/Public/Elements/Interfaces/TypedElementWorldInterface.h b/Engine/Source/Runtime/Engine/Public/Elements/Interfaces/TypedElementWorldInterface.h
--- a/Engine/Source/Runtime/Engine/Public/Elements/Interfaces/TypedElementWorldInterface.h
+++ b/Engine/Source/Runtime/Engine/Public/Elements/Interfaces/TypedElementWorldInterface.h
@@ -111,6 +111,11 @@ public:
 	 */
 	virtual bool CanMoveElement(const FTypedElementHandle& InElementHandle, const ETypedElementWorldType InWorldType) { return false; }
 
+	/**
+	 * Can the given element be scaled?
+	 */
+	virtual bool CanScaleElement(const FTypedElementHandle& InElementHandle) { return true; }
+
 	/**
 	 * Get the transform of this element within its owner world, if any.
 	 */
@@ -348,6 +353,12 @@ public:
 	UFUNCTION(BlueprintCallable, Category="TypedElementInterfaces|World")
 	ENGINE_API virtual bool CanMoveElement(const FScriptTypedElementHandle& InElementHandle, const ETypedElementWorldType InWorldType);
 
+	/**
+	 * Can the given element be scaled?
+	 */
+	UFUNCTION(BlueprintCallable, Category="TypedElementInterfaces|World")
+	ENGINE_API virtual bool CanScaleElement(const FScriptTypedElementHandle& InElementHandle);
+
 	/**
 	 * Get the transform of this element within its owner world, if any.
 	 */
@@ -462,6 +473,7 @@ struct TTypedElement<ITypedElementWorldInterface> : public TTypedElementBase<ITy
 	UWorld* GetOwnerWorld() const { return InterfacePtr->GetOwnerWorld(*this); }
 	bool GetBounds(FBoxSphereBounds& OutBounds) const { return InterfacePtr->GetBounds(*this, OutBounds); }
 	bool CanMoveElement(const ETypedElementWorldType InWorldType) const { return InterfacePtr->CanMoveElement(*this, InWorldType); }
+	bool CanScaleElement() const { return InterfacePtr->CanScaleElement(*this); }
 	bool GetWorldTransform(FTransform& OutTransform) const { return InterfacePtr->GetWorldTransform(*this, OutTransform); }
 	bool SetWorldTransform(const FTransform& InTransform) const { return InterfacePtr->SetWorldTransform(*this, InTransform); }
 	bool GetRelativeTransform(FTransform& OutTransform) const { return InterfacePtr->GetRelativeTransform(*this, OutTransform); }
