Fix for use after destruction issue at AssetEditorModeManager, if the PreviewScene has been already destroyed

#changelist validated
#virtualized

diff --git a/Engine/Source/Editor/UnrealEd/Private/AssetEditorModeManager.cpp b/Engine/Source/Editor/UnrealEd/Private/AssetEditorModeManager.cpp
--- a/Engine/Source/Editor/UnrealEd/Private/AssetEditorModeManager.cpp
+++ b/Engine/Source/Editor/UnrealEd/Private/AssetEditorModeManager.cpp
@@ -74,12 +74,17 @@ USelection* FAssetEditorModeManager::GetSelectedComponents() const
 
 UWorld* FAssetEditorModeManager::GetWorld() const
 {
-	return (PreviewScene != nullptr) ? PreviewScene->GetWorld() : GEditor->GetEditorWorldContext().World();
+	return (PreviewSceneWorld.IsValid()) ? PreviewSceneWorld.Get() : GEditor->GetEditorWorldContext().World();
 }
 
 void FAssetEditorModeManager::SetPreviewScene(FPreviewScene* NewPreviewScene)
 {
 	PreviewScene = NewPreviewScene;
+
+	// Due to destruction order, we might get a call from FEditorModeTools::OnWorldCleanup with an already destoyed PreviewScene,
+	// but we sill have to return the correct World the editor was using to perform a correct shutdown,
+	// so caching the PreviewScene World when assigning (and clearing if null preview is set)
+	PreviewSceneWorld = (NewPreviewScene != nullptr) ? NewPreviewScene->GetWorld() : nullptr;
 }
 
 FPreviewScene* FAssetEditorModeManager::GetPreviewScene() const
diff --git a/Engine/Source/Editor/UnrealEd/Public/AssetEditorModeManager.h b/Engine/Source/Editor/UnrealEd/Public/AssetEditorModeManager.h
--- a/Engine/Source/Editor/UnrealEd/Public/AssetEditorModeManager.h
+++ b/Engine/Source/Editor/UnrealEd/Public/AssetEditorModeManager.h
@@ -6,6 +6,7 @@
 #include "EditorModeManager.h"
 
 class FPreviewScene;
+class UWorld;
 
 //////////////////////////////////////////////////////////////////////////
 // FAssetEditorModeManager
@@ -31,4 +32,5 @@ protected:
 	USelection* ObjectSet = nullptr;
 	USelection* ComponentSet = nullptr;
 	FPreviewScene* PreviewScene = nullptr;
+	TWeakObjectPtr<UWorld> PreviewSceneWorld = nullptr;
 };
