Cooker: downgrade severity of warning Directories to never cook' has invalid element '%s'.
#rnx

diff --git a/Engine/Source/Editor/UnrealEd/Private/CookOnTheFlyServer.cpp b/Engine/Source/Editor/UnrealEd/Private/CookOnTheFlyServer.cpp
--- a/Engine/Source/Editor/UnrealEd/Private/CookOnTheFlyServer.cpp
+++ b/Engine/Source/Editor/UnrealEd/Private/CookOnTheFlyServer.cpp
@@ -12225,7 +12225,11 @@ TArray<FName> UCookOnTheFlyServer::GetNeverCookPackageFileNames(TArrayView<const
 			}
 			else
 			{
-				UE_LOG(LogCook, Warning, TEXT("'%s' has invalid element '%s'"), SettingName, *DirToNotCook.Path);
+				// An unmounted directory that we try to add to nevercook settings is not an error case; since the
+				// directory is unmounted nothing in it can be cooked. And no plugins should be loading after the first
+				// call to this function (which is after CookCommandlet::Main or after editor startup), so we shouldn't
+				// have the problem of a plugin possibly loading later. So downgrade this warning message to verbose.
+				UE_LOG(LogCook, Verbose, TEXT("'%s' has invalid element '%s'"), SettingName, *DirToNotCook.Path);
 			}
 		}
 
