Avoid creating a default sub-object that can end up dangling if the codec is changed

diff --git a/Engine/Source/Editor/UnrealEd/Private/Factories/AnimCurveCompressionSettingsFactory.cpp b/Engine/Source/Editor/UnrealEd/Private/Factories/AnimCurveCompressionSettingsFactory.cpp
--- a/Engine/Source/Editor/UnrealEd/Private/Factories/AnimCurveCompressionSettingsFactory.cpp
+++ b/Engine/Source/Editor/UnrealEd/Private/Factories/AnimCurveCompressionSettingsFactory.cpp
@@ -6,6 +6,7 @@ AnimCurveCompressionSettingsFactory.cpp: Factory for animation curve compression
 
 #include "Factories/AnimCurveCompressionSettingsFactory.h"
 #include "Animation/AnimCurveCompressionSettings.h"
+#include "Animation/AnimCurveCompressionCodec_CompressedRichCurve.h"
 
 UAnimCurveCompressionSettingsFactory::UAnimCurveCompressionSettingsFactory(const FObjectInitializer& ObjectInitializer)
 	: Super(ObjectInitializer)
@@ -16,5 +17,9 @@ UAnimCurveCompressionSettingsFactory::UAnimCurveCompressionSettingsFactory(const
 
 UObject* UAnimCurveCompressionSettingsFactory::FactoryCreateNew(UClass* Class, UObject* InParent, FName Name, EObjectFlags Flags, UObject* Context, FFeedbackContext* Warn)
 {
-	return NewObject<UAnimCurveCompressionSettings>(InParent, Class, Name, Flags);
+	UAnimCurveCompressionSettings* Settings = NewObject<UAnimCurveCompressionSettings>(InParent, Class, Name, Flags);
+	Settings->Codec = NewObject<UAnimCurveCompressionCodec_CompressedRichCurve>(Settings);
+	Settings->Codec->SetFlags(RF_Transactional);
+
+	return Settings;
 }
diff --git a/Engine/Source/Runtime/Engine/Private/Animation/AnimCurveCompressionSettings.cpp b/Engine/Source/Runtime/Engine/Private/Animation/AnimCurveCompressionSettings.cpp
--- a/Engine/Source/Runtime/Engine/Private/Animation/AnimCurveCompressionSettings.cpp
+++ b/Engine/Source/Runtime/Engine/Private/Animation/AnimCurveCompressionSettings.cpp
@@ -2,15 +2,13 @@
 
 #include "Animation/AnimCurveCompressionSettings.h"
 #include "Animation/AnimCompressionTypes.h"
-#include "Animation/AnimCurveCompressionCodec_CompressedRichCurve.h"
+#include "Animation/AnimCurveCompressionCodec.h"
 
 #include UE_INLINE_GENERATED_CPP_BY_NAME(AnimCurveCompressionSettings)
 
 UAnimCurveCompressionSettings::UAnimCurveCompressionSettings(const FObjectInitializer& ObjectInitializer)
 	: Super(ObjectInitializer)
 {
-	Codec = CreateDefaultSubobject<UAnimCurveCompressionCodec_CompressedRichCurve>(TEXT("CurveCompressionCodec"));
-	Codec->SetFlags(RF_Transactional);
 }
 
 UAnimCurveCompressionCodec* UAnimCurveCompressionSettings::GetCodec(const FString& Path)
