World Partition ConvertCommandlet: when duplicating engine default HLOD layer setup, replace all references to the layers hierarchy from the runtime hash.

diff --git a/Engine/Source/Editor/UnrealEd/Private/Commandlets/WorldPartitionConvertCommandlet.cpp b/Engine/Source/Editor/UnrealEd/Private/Commandlets/WorldPartitionConvertCommandlet.cpp
--- a/Engine/Source/Editor/UnrealEd/Private/Commandlets/WorldPartitionConvertCommandlet.cpp
+++ b/Engine/Source/Editor/UnrealEd/Private/Commandlets/WorldPartitionConvertCommandlet.cpp
@@ -46,6 +46,7 @@
 #include "FoliageHelper.h"
 #include "Engine/WorldComposition.h"
 #include "ActorPartition/ActorPartitionSubsystem.h"
+#include "Serialization/ArchiveReplaceObjectRef.h"
 #include "InstancedFoliage.h"
 #include "LandscapeStreamingProxy.h"
 #include "LandscapeInfo.h"
@@ -244,16 +245,24 @@ UWorldPartition* UWorldPartitionConvertCommandlet::CreateWorldPartition(AWorldSe
 		}
 	}
 
+	// Duplicate the default HLOD setup
 	if ((WorldPartition->DefaultHLODLayer == UHLODLayer::GetEngineDefaultHLODLayersSetup()) && !bDisableStreaming)
 	{
-		WorldPartition->DefaultHLODLayer = UHLODLayer::DuplicateHLODLayersSetup(UHLODLayer::GetEngineDefaultHLODLayersSetup(), WorldPartition->GetPackage()->GetName(), WorldPartition->GetWorld()->GetName());
+		UHLODLayer* CurHLODLayer = WorldPartition->GetDefaultHLODLayer();
+		UHLODLayer* NewHLODLayer = UHLODLayer::DuplicateHLODLayersSetup(CurHLODLayer, WorldPartition->GetPackage()->GetName(), WorldPartition->GetWorld()->GetName());
+		
+		WorldPartition->SetDefaultHLODLayer(NewHLODLayer);
 
-		UHLODLayer* CurrentHLODLayer = WorldPartition->DefaultHLODLayer;
-		while (CurrentHLODLayer)
+		TMap<UHLODLayer*, UHLODLayer*> ReplacementMap;
+		while (NewHLODLayer)
 		{
-			PackagesToSave.Add(CurrentHLODLayer->GetPackage());
-			CurrentHLODLayer = CurrentHLODLayer->GetParentLayer();
+			ReplacementMap.Add(CurHLODLayer, NewHLODLayer);
+			PackagesToSave.Add(NewHLODLayer->GetPackage());
+			CurHLODLayer = CurHLODLayer->GetParentLayer();
+			NewHLODLayer = NewHLODLayer->GetParentLayer();
 		}
+							
+		FArchiveReplaceObjectRef<UHLODLayer> ReplaceObjectRefAr(WorldPartition->RuntimeHash, ReplacementMap, EArchiveReplaceObjectFlags::IgnoreOuterRef | EArchiveReplaceObjectFlags::IgnoreArchetypeRef);							
 	}
 	
 	return WorldPartition;
