Actor Coloration Framework: add support for creating property coloration shortcuts from the editor config file.

diff --git a/Engine/Source/Editor/UnrealEd/Classes/Editor/EditorEngine.h b/Engine/Source/Editor/UnrealEd/Classes/Editor/EditorEngine.h
--- a/Engine/Source/Editor/UnrealEd/Classes/Editor/EditorEngine.h
+++ b/Engine/Source/Editor/UnrealEd/Classes/Editor/EditorEngine.h
@@ -761,7 +761,7 @@ public:
 	UNREALED_API virtual bool ShouldDrawBrushWireframe(AActor* InActor) override;
 	UNREALED_API virtual void NotifyToolsOfObjectReplacement(const TMap<UObject*, UObject*>& OldToNewInstanceMap) override;
 	UNREALED_API virtual bool ShouldThrottleCPUUsage() const override;
-	UNREALED_API virtual bool GetPropertyColorationColor(class UObject* Object, FColor& OutColor) override;
+	UNREALED_API virtual bool GetPropertyColorationMatch(class UObject* Object);
 	UNREALED_API virtual bool WorldIsPIEInNewViewport(UWorld* InWorld) override;
 	UNREALED_API virtual void FocusNextPIEWorld(UWorld* CurrentPieWorld, bool previous = false) override;
 	UNREALED_API virtual void ResetPIEAudioSetting(UWorld *CurrentPieWorld) override;
diff --git a/Engine/Source/Editor/UnrealEd/Private/EditorEngine.cpp b/Engine/Source/Editor/UnrealEd/Private/EditorEngine.cpp
--- a/Engine/Source/Editor/UnrealEd/Private/EditorEngine.cpp
+++ b/Engine/Source/Editor/UnrealEd/Private/EditorEngine.cpp
@@ -110,6 +110,7 @@
 #include "UncontrolledChangelistsModule.h"
 #include "SceneView.h"
 #include "StaticBoundShaderState.h"
+#include "PropertyColorSettings.h"
 
 // needed for the RemotePropagator
 #include "AudioDevice.h"
@@ -508,15 +509,74 @@ UEditorEngine::UEditorEngine(const FObjectInitializer& ObjectInitializer)
 #if ENABLE_ACTOR_PRIMITIVE_COLOR_HANDLER
 	if (HasAnyFlags(RF_ClassDefaultObject) && ExactCast<UEditorEngine>(this))
 	{
-		FActorPrimitiveColorHandler::Get().RegisterPrimitiveColorHandler(TEXT("PropertyColor"), LOCTEXT("PropertyColor", "Property Color"), [this](const UPrimitiveComponent* InPrimitiveComponent)
-		{
-			FColor PropertyColor(FColor::White);
-			if (AActor* Actor = InPrimitiveComponent->GetOwner())
+		FActorPrimitiveColorHandler::Get().RegisterPrimitiveColorHandler(TEXT("PropertyColor"), LOCTEXT("PropertyColor", "Property Color"), 
+			[this](const UPrimitiveComponent* InPrimitiveComponent)
 			{
-				GetPropertyColorationColor(Actor, PropertyColor);
-			}
-			return PropertyColor;
-		});
+				FColor PropertyColor(FColor::White);
+				if (AActor* Actor = InPrimitiveComponent->GetOwner())
+				{
+					if (GetPropertyColorationMatch(Actor))
+					{
+						PropertyColor = FColor::Red;
+					}
+				}
+				return PropertyColor;
+			},
+			[this]()
+			{
+				const FString EmptyString;
+				SetPropertyColorationTarget(GWorld, EmptyString, nullptr, nullptr, nullptr);
+			});
+
+		for (const FPropertyColorCustomProperty& PropertyColorCustomProperty : GetDefault<UPropertyColorSettings>()->CustomProperties)
+		{
+			FActorPrimitiveColorHandler::Get().RegisterPrimitiveColorHandler(PropertyColorCustomProperty.Name, 
+				FInternationalization::ForUseOnlyByLocMacroAndGraphNodeTextLiterals_CreateText(*PropertyColorCustomProperty.Text, TEXT("PropertyColor"), *PropertyColorCustomProperty.Name.ToString()),
+				[this, PropertyColorCustomProperty](const UPrimitiveComponent* InPrimitiveComponent)
+				{
+					FColor PropertyColor(FColor::White);
+					if (AActor* Actor = InPrimitiveComponent->GetOwner())
+					{
+						if (GetPropertyColorationMatch(Actor))
+						{
+							PropertyColor = PropertyColorCustomProperty.PropertyColor;
+						}
+					}
+					return PropertyColor;
+				},
+				[this, PropertyColorCustomProperty]()
+				{
+					TArray<FString> PropertyChainNames;
+					UStruct* PropertyContainer = AActor::StaticClass();
+					if (PropertyColorCustomProperty.PropertyChain.ParseIntoArray(PropertyChainNames, TEXT(".")))
+					{
+						TSharedRef<FEditPropertyChain> PropertyChain = MakeShared<FEditPropertyChain>();
+
+						for (const FString& PropertyName : PropertyChainNames)
+						{
+							if (FProperty* Property = PropertyContainer->FindPropertyByName(*PropertyName))
+							{
+								PropertyChain->AddTail(Property);
+
+								if (FObjectProperty* ObjectProperty = CastField<FObjectProperty>(Property))
+								{
+									PropertyContainer = ObjectProperty->PropertyClass;
+								}
+							}
+							else
+							{
+								UE_LOG(LogEditor, Warning, TEXT("Invalid custom property color %s (%s)"), *PropertyColorCustomProperty.Name.ToString(), *PropertyColorCustomProperty.PropertyChain);
+								break;
+							}
+						}
+
+						if (PropertyChain->Num() == PropertyChainNames.Num())
+						{
+							SetPropertyColorationTarget(GWorld, PropertyColorCustomProperty.PropertyValue, PropertyChain->GetTail()->GetValue(), AActor::StaticClass(), &PropertyChain);
+						}
+					}
+				});
+		}
 	}
 #endif
 }
@@ -974,7 +1034,7 @@ void UEditorEngine::InitEditor(IEngineLoop* InEngineLoop)
 				bNewSetEditMenusMode = (Args[0] == TEXT("1")) || FCString::ToBool(*Args[0]);
 			}
 
-			UE_LOG(LogEditor, Log, TEXT("%s menu editing"), bNewSetEditMenusMode ? TEXT("Enable") : TEXT("Disable"));			
+			UE_LOG(LogEditor, Log, TEXT("%s menu editing"), bNewSetEditMenusMode ? TEXT("Enable") : TEXT("Disable"));
 			UToolMenus::Get()->SetEditMenusMode(bNewSetEditMenusMode);
 		}));
 
diff --git a/Engine/Source/Editor/UnrealEd/Private/EditorServer.cpp b/Engine/Source/Editor/UnrealEd/Private/EditorServer.cpp
--- a/Engine/Source/Editor/UnrealEd/Private/EditorServer.cpp
+++ b/Engine/Source/Editor/UnrealEd/Private/EditorServer.cpp
@@ -3681,18 +3681,20 @@ void UEditorEngine::SetPropertyColorationTarget(UWorld* InWorld, const FString&
 		GPropertyColorationProperty = Property;
 		GPropertyColorationClass = CommonBaseClass;
 		
-		GPropertyColorationChain = PropertyChain ? new TSharedRef<FEditPropertyChain>(*PropertyChain) : nullptr;
-
-		GbColorationClassIsActor = GPropertyColorationClass->IsChildOf( AActor::StaticClass() );
-		GbColorationPropertyIsObjectProperty = CastField<FObjectPropertyBase>(GPropertyColorationProperty) != NULL;
-
-		FActorPrimitiveColorHandler::Get().RefreshPrimitiveColorHandler(TEXT("PropertyColor"), InWorld);
+		if (PropertyChain)
+		{
+			GPropertyColorationChain = new TSharedRef<FEditPropertyChain>(*PropertyChain);
+			GbColorationClassIsActor = GPropertyColorationClass->IsChildOf( AActor::StaticClass() );
+			GbColorationPropertyIsObjectProperty = CastField<FObjectPropertyBase>(GPropertyColorationProperty) != NULL;
+			
+			FActorPrimitiveColorHandler::Get().RefreshPrimitiveColorHandler(TEXT("PropertyColor"), InWorld);
+		}
 
 		RedrawLevelEditingViewports();
 	}
 }
 
-bool UEditorEngine::GetPropertyColorationColor(UObject* Object, FColor& OutColor)
+bool UEditorEngine::GetPropertyColorationMatch(UObject* Object)
 {
 	bool bResult = false;
 	if (GPropertyColorationChain)
@@ -3761,7 +3763,6 @@ bool UEditorEngine::GetPropertyColorationColor(UObject* Object, FColor& OutColor
 				if ( PropertyValue == GPropertyColorationValue )
 				{
 					bResult  = true;
-					OutColor = FColor::Red;
 
 					// Collect actor references.
 					if ( GPropertyColorationActorCollector && Owner )
diff --git a/Engine/Source/Editor/UnrealEd/Private/PropertyColorSettings.h b/Engine/Source/Editor/UnrealEd/Private/PropertyColorSettings.h
new file mode 100644
--- /dev/null
+++ b/Engine/Source/Editor/UnrealEd/Private/PropertyColorSettings.h
@@ -0,0 +1,45 @@
+// Copyright Epic Games, Inc. All Rights Reserved.
+
+#pragma once
+
+#include "CoreMinimal.h"
+#include "UObject/ObjectMacros.h"
+#include "UObject/Object.h"
+#include "PropertyColorSettings.generated.h"
+
+USTRUCT()
+struct FPropertyColorCustomProperty
+{
+	GENERATED_BODY()
+
+	UPROPERTY(Config)
+	FName Name;
+
+	UPROPERTY(Config)
+	FString Text;
+
+	UPROPERTY(Config)
+	FString PropertyChain;
+
+	UPROPERTY(Config)
+	FString PropertyValue;
+
+	UPROPERTY(Config)
+	FColor PropertyColor = FColor::Red;
+};
+
+/**
+* Implements the settings for Property Color.
+*/
+UCLASS(config = EditorSettings)
+class UNREALED_API UPropertyColorSettings : public UObject
+{
+	GENERATED_BODY()
+
+public:
+	UPropertyColorSettings()
+	{}
+
+	UPROPERTY(Config)
+	TArray<FPropertyColorCustomProperty> CustomProperties;
+};
\ No newline at end of file
diff --git a/Engine/Source/Runtime/Engine/Classes/Engine/Engine.h b/Engine/Source/Runtime/Engine/Classes/Engine/Engine.h
--- a/Engine/Source/Runtime/Engine/Classes/Engine/Engine.h
+++ b/Engine/Source/Runtime/Engine/Classes/Engine/Engine.h
@@ -2546,15 +2546,6 @@ public:
 		// Intentionally empty.
 	}
 
-	/**
-	 * Computes a color to use for property coloration for the given object.
-	 *
-	 * @param	Object		The object for which to compute a property color.
-	 * @param	OutColor	[out] The returned color.
-	 * @return				true if a color was successfully set on OutColor, false otherwise.
-	 */
-	ENGINE_API virtual bool GetPropertyColorationColor(class UObject* Object, FColor& OutColor);
-
 	/** Uses StatColorMappings to find a color for this stat's value. */
 	ENGINE_API bool GetStatValueColoration(const FString& StatName, float Value, FColor& OutColor);
 
diff --git a/Engine/Source/Runtime/Engine/Classes/GameFramework/ActorPrimitiveColorHandler.h b/Engine/Source/Runtime/Engine/Classes/GameFramework/ActorPrimitiveColorHandler.h
--- a/Engine/Source/Runtime/Engine/Classes/GameFramework/ActorPrimitiveColorHandler.h
+++ b/Engine/Source/Runtime/Engine/Classes/GameFramework/ActorPrimitiveColorHandler.h
@@ -16,26 +16,29 @@ class UPrimitiveComponent;
  */
 class ENGINE_API FActorPrimitiveColorHandler
 {
-	using FFunc = TFunction<FLinearColor(const UPrimitiveComponent*)>;
+	using FGetColorFunc = TFunction<FLinearColor(const UPrimitiveComponent*)>;
+	using FActivateFunc = TFunction<void(void)>;
 
 public:
 	struct FPrimitiveColorHandler
 	{
-		FPrimitiveColorHandler(FName InHandlerName, FText InHandlerText, const FFunc& InHandlerFunc)
+		FPrimitiveColorHandler(FName InHandlerName, FText InHandlerText, const FGetColorFunc& InGetColorFunc, const FActivateFunc& InActivateFunc)
 			: HandlerName(InHandlerName)
 			, HandlerText(InHandlerText)
-			, HandlerFunc(InHandlerFunc)
+			, GetColorFunc(InGetColorFunc)
+			, ActivateFunc(InActivateFunc)
 		{}
 
 		FName HandlerName;
 		FText HandlerText;
-		FFunc HandlerFunc;
+		FGetColorFunc GetColorFunc;
+		FActivateFunc ActivateFunc;
 	};	
 
 	FActorPrimitiveColorHandler();
 	static FActorPrimitiveColorHandler& Get();
 
-	void RegisterPrimitiveColorHandler(FName InHandlerName, FText InHandlerText, const FFunc& InHandlerFunc);
+	void RegisterPrimitiveColorHandler(FName InHandlerName, FText InHandlerText, const FGetColorFunc& InHandlerFunc, const FActivateFunc& InActivateFunc = []() {});
 	void UnregisterPrimitiveColorHandler(FName InHandlerName);
 	void GetRegisteredPrimitiveColorHandlers(TArray<FPrimitiveColorHandler>& OutPrimitiveColorHandlers) const;
 
diff --git a/Engine/Source/Runtime/Engine/Private/ActorPrimitiveColorHandler.cpp b/Engine/Source/Runtime/Engine/Private/ActorPrimitiveColorHandler.cpp
--- a/Engine/Source/Runtime/Engine/Private/ActorPrimitiveColorHandler.cpp
+++ b/Engine/Source/Runtime/Engine/Private/ActorPrimitiveColorHandler.cpp
@@ -25,11 +25,11 @@ FActorPrimitiveColorHandler& FActorPrimitiveColorHandler::Get()
 	return TLazySingleton<FActorPrimitiveColorHandler>::Get();
 }
 
-void FActorPrimitiveColorHandler::RegisterPrimitiveColorHandler(FName InHandlerName, FText InHandlerText, const FFunc& InHandlerFunc)
+void FActorPrimitiveColorHandler::RegisterPrimitiveColorHandler(FName InHandlerName, FText InHandlerText, const FGetColorFunc& InGetColorFunc, const FActivateFunc& InActivateFunc)
 {
 #if ENABLE_ACTOR_PRIMITIVE_COLOR_HANDLER
 	check(!Handlers.Contains(InHandlerName));
-	Handlers.Add(InHandlerName, { InHandlerName, InHandlerText, InHandlerFunc });	
+	Handlers.Add(InHandlerName, { InHandlerName, InHandlerText, InGetColorFunc, InActivateFunc });
 	ActivePrimitiveColorHandler = &Handlers.FindChecked(ActivePrimitiveColorHandlerName);
 #endif
 }
@@ -57,6 +57,7 @@ bool FActorPrimitiveColorHandler::SetActivePrimitiveColorHandler(FName InHandler
 	{
 		ActivePrimitiveColorHandlerName = InHandlerName;
 		ActivePrimitiveColorHandler = NewActivePrimitiveColorHandler;
+		NewActivePrimitiveColorHandler->ActivateFunc();
 		RefreshPrimitiveColorHandler(InHandlerName, InWorld);
 		return true;
 	}
@@ -138,7 +139,7 @@ void FActorPrimitiveColorHandler::GetRegisteredPrimitiveColorHandlers(TArray<FPr
 FLinearColor FActorPrimitiveColorHandler::GetPrimitiveColor(const UPrimitiveComponent* InPrimitiveComponent) const
 {
 #if ENABLE_ACTOR_PRIMITIVE_COLOR_HANDLER
-	return ActivePrimitiveColorHandler->HandlerFunc(InPrimitiveComponent);
+	return ActivePrimitiveColorHandler->GetColorFunc(InPrimitiveComponent);
 #else
 	return FLinearColor::White;
 #endif
diff --git a/Engine/Source/Runtime/Engine/Private/UnrealEngine.cpp b/Engine/Source/Runtime/Engine/Private/UnrealEngine.cpp
--- a/Engine/Source/Runtime/Engine/Private/UnrealEngine.cpp
+++ b/Engine/Source/Runtime/Engine/Private/UnrealEngine.cpp
@@ -11074,19 +11074,6 @@ bool UEngine::PerformError(const TCHAR* Cmd, FOutputDevice& Ar)
 	return false;
 }
 
-
-/**
-* Computes a color to use for property coloration for the given object.
-*
-* @param	Object		The object for which to compute a property color.
-* @param	OutColor	[out] The returned color.
-* @return				true if a color was successfully set on OutColor, false otherwise.
-*/
-bool UEngine::GetPropertyColorationColor(UObject* Object, FColor& OutColor)
-{
-	return false;
-}
-
 /** Uses StatColorMappings to find a color for this stat's value. */
 bool UEngine::GetStatValueColoration(const FString& StatName, float Value, FColor& OutColor)
 {
