[CrashFix] Crash when attempting to merge files on mac

diff --git a/Engine/Source/Editor/UnrealEd/Private/MergeUtils.cpp b/Engine/Source/Editor/UnrealEd/Private/MergeUtils.cpp
--- a/Engine/Source/Editor/UnrealEd/Private/MergeUtils.cpp
+++ b/Engine/Source/Editor/UnrealEd/Private/MergeUtils.cpp
@@ -83,11 +83,15 @@ FString MergeUtils::LoadSCFileForMerge(const FString& SCFile, const FString& Rev
 	const FString DownloadPath = FPaths::ConvertRelativePathToFull(FPaths::DiffDir() / FPaths::GetCleanFilename(FileWithRevision));
 
 	// move downloaded file to renamed path so it meets ue asset name requirements
-	FString ResultPath = DownloadPath;
-	ResultPath.ReplaceInline(TEXT(".uasset"), TEXT(""));
-	ResultPath.ReplaceCharInline('#', '-');
-	ResultPath.ReplaceCharInline('.', '-');
-	ResultPath = FPaths::CreateTempFilename(*FPaths::GetPath(ResultPath), *FPaths::GetBaseFilename(ResultPath), TEXT(".uasset"));
+	const FString Directory = FPaths::GetPath(DownloadPath);
+	FString Filename = FPaths::GetCleanFilename(DownloadPath);
+	Filename.ReplaceInline(TEXT(".uasset"), TEXT(""));
+	Filename.ReplaceCharInline('#', '-');
+	Filename.ReplaceCharInline('.', '-');
+	Filename += TEXT("-");
+	FString ResultPath = FPaths::CreateTempFilename(*Directory, *Filename, TEXT(".uasset"));
+	checkf(FPaths::DirectoryExists(Directory), TEXT("Tried to move file to a directory that doesn't exist"));
+	checkf(!FPaths::FileExists(ResultPath), TEXT("Tried to rename file to a name that's already taken"));
 	if (ensure(FPlatformFileManager::Get().GetPlatformFile().MoveFile(*ResultPath, *DownloadPath)))
 	{
 		return ResultPath;
