Actor Coloration: when disabling actor coloration, don't automatically fallback to level coloration, keep the active one instead.

diff --git a/Engine/Source/Editor/UnrealEd/Private/EditorViewportClient.cpp b/Engine/Source/Editor/UnrealEd/Private/EditorViewportClient.cpp
--- a/Engine/Source/Editor/UnrealEd/Private/EditorViewportClient.cpp
+++ b/Engine/Source/Editor/UnrealEd/Private/EditorViewportClient.cpp
@@ -2685,11 +2685,6 @@ void FEditorViewportClient::HandleToggleShowFlag(FEngineShowFlags::EShowFlag Eng
 	{
 		UpdateHiddenCollisionDrawing();
 	}
-	// If changing level coloration flag, always go back to the legacy handler.
-	else if (EngineShowFlagIndex == FEngineShowFlags::EShowFlag::SF_ActorColoration)
-	{
-		FActorPrimitiveColorHandler::Get().SetActivePrimitiveColorHandler(TEXT("LevelColor"), GetWorld());
-	}
 
 	// Invalidate clients which aren't real-time so we see the changes.
 	Invalidate();
