New TRS Gizmos: avoid useless hit proxies invalidation while manipulating

NOTE:
    There should be a better API for disabling hit proxy requests. The current implementation is based on checking FEditorViewportClient::bWidgetAxisControlledByDrag or FModeTools::bHasOngoingTransform to launch a request or not (via updating the cursor...)
    Having a way to deactivate it more globally would improve both performance and UX (I don't think the other ITF gizmos manage this case for example, this concern has been highlighted because rigs are heavier to evaluate than a classic transform)

#rnx

diff --git a/Engine/Source/Editor/UnrealEd/Private/EditorModeManager.cpp b/Engine/Source/Editor/UnrealEd/Private/EditorModeManager.cpp
--- a/Engine/Source/Editor/UnrealEd/Private/EditorModeManager.cpp
+++ b/Engine/Source/Editor/UnrealEd/Private/EditorModeManager.cpp
@@ -2028,4 +2028,9 @@ bool FEditorModeTools::EndTransform(const FGizmoState& InState) const
 	// See header file more more explanations
 	
 	return bHandled;
+}
+
+bool FEditorModeTools::HasOngoingTransform() const
+{
+	return bHasOngoingTransform;
 }
\ No newline at end of file
diff --git a/Engine/Source/Editor/UnrealEd/Private/EditorViewportClient.cpp b/Engine/Source/Editor/UnrealEd/Private/EditorViewportClient.cpp
--- a/Engine/Source/Editor/UnrealEd/Private/EditorViewportClient.cpp
+++ b/Engine/Source/Editor/UnrealEd/Private/EditorViewportClient.cpp
@@ -1651,7 +1651,7 @@ EMouseCursor::Type FEditorViewportClient::GetCursor(FViewport* InViewport,int32
 		MouseCursor = EMouseCursor::CardinalCross;
 	}
 	// Don't select widget axes by mouse over while they're being controlled by a mouse drag.
-	else if( InViewport->IsCursorVisible() && !bWidgetAxisControlledByDrag )
+	else if( InViewport->IsCursorVisible() && !bWidgetAxisControlledByDrag && !ModeTools->HasOngoingTransform())
 	{
 		// allow editor modes to override cursor
 		EMouseCursor::Type EditorModeCursor = EMouseCursor::Default;
diff --git a/Engine/Source/Editor/UnrealEd/Private/LevelEditorViewport.cpp b/Engine/Source/Editor/UnrealEd/Private/LevelEditorViewport.cpp
--- a/Engine/Source/Editor/UnrealEd/Private/LevelEditorViewport.cpp
+++ b/Engine/Source/Editor/UnrealEd/Private/LevelEditorViewport.cpp
@@ -4623,7 +4623,7 @@ EMouseCursor::Type FLevelEditorViewportClient::GetCursor(FViewport* InViewport,i
 	}
 
 	// Don't select widget axes by mouse over while they're being controlled by a mouse drag.
-	if( InViewport->IsCursorVisible() && !bWidgetAxisControlledByDrag)
+	if( InViewport->IsCursorVisible() && !bWidgetAxisControlledByDrag && !ModeTools->HasOngoingTransform())
 	{
 		HHitProxy* HitProxy = InViewport->GetHitProxy(X, Y);
 		if( !HitProxy && HoveredObjects.Num() > 0 )
diff --git a/Engine/Source/Editor/UnrealEd/Public/EditorModeManager.h b/Engine/Source/Editor/UnrealEd/Public/EditorModeManager.h
--- a/Engine/Source/Editor/UnrealEd/Public/EditorModeManager.h
+++ b/Engine/Source/Editor/UnrealEd/Public/EditorModeManager.h
@@ -540,6 +540,7 @@ public:
 	UNREALED_API IGizmoStateTarget* GetGizmoStateTarget();
 	UNREALED_API bool BeginTransform(const FGizmoState& InState);
 	UNREALED_API bool EndTransform(const FGizmoState& InState) const;
+	UNREALED_API bool HasOngoingTransform() const;
 	
 protected:
 	/** 
