World Partition - HLOD: During distributed HLOD builds, ensure new HLOD actors created during setup are treated as Added during the build step, rather than Edited

diff --git a/Engine/Source/Editor/UnrealEd/Private/WorldPartition/WorldPartitionHLODsBuilder.cpp b/Engine/Source/Editor/UnrealEd/Private/WorldPartition/WorldPartitionHLODsBuilder.cpp
--- a/Engine/Source/Editor/UnrealEd/Private/WorldPartition/WorldPartitionHLODsBuilder.cpp
+++ b/Engine/Source/Editor/UnrealEd/Private/WorldPartition/WorldPartitionHLODsBuilder.cpp
@@ -38,8 +38,9 @@ DEFINE_LOG_CATEGORY_STATIC(LogWorldPartitionHLODsBuilder, All, All);
 class FSourceControlHelper : public ISourceControlHelper
 {
 public:
-	FSourceControlHelper(FPackageSourceControlHelper& InPackageHelper)
+	FSourceControlHelper(FPackageSourceControlHelper& InPackageHelper, FHLODModifiedFiles& InModifiedFiles)
 		: PackageHelper(InPackageHelper)
+		, ModifiedFiles(InModifiedFiles)
 	{}
 
 	virtual ~FSourceControlHelper()
@@ -60,7 +61,12 @@ public:
 		bool bCheckedOut = PackageHelper.Checkout(Package);
 		if (bCheckedOut)
 		{
-			ModifiedFiles.Add(FHLODModifiedFiles::EFileOperation::FileEdited, GetFilename(Package));
+			const FString Filename = GetFilename(Package);
+			const bool bAdded = ModifiedFiles.Get(FHLODModifiedFiles::EFileOperation::FileAdded).Contains(Filename);
+			if (!bAdded)
+			{
+				ModifiedFiles.Add(FHLODModifiedFiles::EFileOperation::FileEdited, Filename);
+			}
 		}
 		return bCheckedOut;
 	}
@@ -136,15 +142,9 @@ public:
 		return true;
 	}
 
-	const FHLODModifiedFiles& GetModifiedFiles() const
-	{
-		UPackage::WaitForAsyncFileWrites();
-		return ModifiedFiles;
-	}
-
 private:
 	FPackageSourceControlHelper& PackageHelper;
-	mutable FHLODModifiedFiles ModifiedFiles;
+	FHLODModifiedFiles& ModifiedFiles;
 };
 
 static const FString DistributedBuildWorkingDirName = TEXT("HLODTemp");
@@ -278,6 +278,8 @@ bool UWorldPartitionHLODsBuilder::ShouldProcessWorld(UWorld* InWorld) const
 
 bool UWorldPartitionHLODsBuilder::PreWorldInitialization(UWorld* InWorld, FPackageSourceControlHelper& PackageHelper)
 {
+	ModifiedFiles.Empty();
+
 	if (bDistributedBuild)
 	{
 		DistributedBuildWorkingDir = GetDistributedBuildWorkingDir(InWorld);
@@ -316,7 +318,7 @@ bool UWorldPartitionHLODsBuilder::RunInternal(UWorld* InWorld, const FCellInfo&
 	// Allows HLOD Streaming levels to be GCed properly
 	FLevelStreamingGCHelper::EnableForCommandlet();
 
-	SourceControlHelper = new FSourceControlHelper(PackageHelper);
+	SourceControlHelper = new FSourceControlHelper(PackageHelper, ModifiedFiles);
 
 	bool bRet = true;
 
@@ -425,8 +427,6 @@ bool UWorldPartitionHLODsBuilder::SetupHLODActors()
 				FWorldPartitionHelpers::DoCollectGarbage();
 			}
 
-			ModifiedFiles.Append(SourceControlHelper->GetModifiedFiles());
-
 			TArray<FHLODModifiedFiles> BuildersFiles;
 			BuildersFiles.SetNum(BuilderCount);
 
@@ -468,8 +468,6 @@ bool UWorldPartitionHLODsBuilder::SetupHLODActors()
 			{
 				return false;
 			}
-
-			ModifiedFiles.Empty();
 		}
 	}
 
@@ -593,7 +591,8 @@ bool UWorldPartitionHLODsBuilder::BuildHLODActors()
 			FWorldPartitionHelpers::DoCollectGarbage();
 		}
 
-		ModifiedFiles.Append(SourceControlHelper->GetModifiedFiles());
+		// Wait for pending async file writes before copying to working dir
+		UPackage::WaitForAsyncFileWrites();
 
 		TArray<FString> BuildProducts;
 
@@ -607,8 +606,6 @@ bool UWorldPartitionHLODsBuilder::BuildHLODActors()
 		{
 			return false;
 		}
-
-		ModifiedFiles.Empty();
 	}
 
 	return true;
@@ -663,8 +660,8 @@ bool UWorldPartitionHLODsBuilder::DeleteHLODActors()
 
 bool UWorldPartitionHLODsBuilder::SubmitHLODActors()
 {
-	// Ensure all files modified by the source control helper are taken into account
-	ModifiedFiles.Append(SourceControlHelper->GetModifiedFiles());
+	// Wait for pending async file writes before submitting
+	UPackage::WaitForAsyncFileWrites();
 
 	// Check in all modified files
 	const FString ChangeDescription = FString::Printf(TEXT("Rebuilt HLODs for %s"), *World->GetPackage()->GetName());
