VR Editor: Disable PIE/VR PIE/SIE while VR editor is active.

diff --git a/Engine/Source/Editor/UnrealEd/Private/Kismet2/DebuggerCommands.cpp b/Engine/Source/Editor/UnrealEd/Private/Kismet2/DebuggerCommands.cpp
--- a/Engine/Source/Editor/UnrealEd/Private/Kismet2/DebuggerCommands.cpp
+++ b/Engine/Source/Editor/UnrealEd/Private/Kismet2/DebuggerCommands.cpp
@@ -45,6 +45,7 @@
 #include "PlatformInfo.h"
 
 #include "IHeadMountedDisplay.h"
+#include "IVREditorModule.h"
 #include "IXRTrackingSystem.h"
 #include "Editor.h"
 
@@ -92,6 +93,16 @@ namespace DebuggerCommands
 	static bool bAllowPlayWorldFeature = true;
 	static FAutoConsoleVariableRef AllowPlayWorldFeatureCVar(TEXT("Editor.AllowPlayWorldFeature"), bAllowPlayWorldFeature, TEXT("When true play world is allowed."));
 	static bool AllowPlayWorldFeature();
+
+	static bool IsVREditorActive()
+	{
+		if (IVREditorModule::IsAvailable() && IVREditorModule::Get().GetVRModeBase() != nullptr)
+		{
+			return true;
+		}
+
+		return false;
+	}
 }
 
 void SGlobalPlayWorldActions::Construct(const FArguments& InArgs)
@@ -970,7 +981,7 @@ void FPlayWorldCommandCallbacks::StartPlayFromHere()
 
 void FPlayWorldCommandCallbacks::StartPlayFromHere(const TOptional<FVector>& Location, const TOptional<FRotator>& Rotation, const TSharedPtr<IAssetViewport>& ActiveLevelViewport)
 {
-	if (FInternalPlayWorldCommandCallbacks::IsStoppedAtBreakpoint_InEngineMode())
+	if (FInternalPlayWorldCommandCallbacks::IsStoppedAtBreakpoint_InEngineMode() || DebuggerCommands::IsVREditorActive())
 	{
 		return;
 	}
@@ -1164,7 +1175,7 @@ bool FInternalPlayWorldCommandCallbacks::CanShowLateJoinButton()
 
 void FInternalPlayWorldCommandCallbacks::Simulate_Clicked()
 {
-	if (IsStoppedAtBreakpoint_InEngineMode())
+	if (IsStoppedAtBreakpoint_InEngineMode() || DebuggerCommands::IsVREditorActive())
 	{
 		return;
 	}
@@ -1204,7 +1215,7 @@ void FInternalPlayWorldCommandCallbacks::Simulate_Clicked()
 
 bool FInternalPlayWorldCommandCallbacks::Simulate_CanExecute()
 {
-	if (IsStoppedAtBreakpoint_InEngineMode())
+	if (IsStoppedAtBreakpoint_InEngineMode() || DebuggerCommands::IsVREditorActive())
 	{
 		return false;
 	}
@@ -1438,7 +1449,7 @@ FSlateIcon FInternalPlayWorldCommandCallbacks::GetRepeatLastPlayIcon()
 
 void FInternalPlayWorldCommandCallbacks::PlayInViewport_Clicked()
 {
-	if (IsStoppedAtBreakpoint_InEngineMode())
+	if (IsStoppedAtBreakpoint_InEngineMode() || DebuggerCommands::IsVREditorActive())
 	{
 		return;
 	}
@@ -1490,7 +1501,7 @@ void FInternalPlayWorldCommandCallbacks::PlayInViewport_Clicked()
 
 bool FInternalPlayWorldCommandCallbacks::PlayInViewport_CanExecute()
 {
-	if (IsStoppedAtBreakpoint_InEngineMode())
+	if (IsStoppedAtBreakpoint_InEngineMode() || DebuggerCommands::IsVREditorActive())
 	{
 		return false;
 	}
@@ -1508,7 +1519,7 @@ bool FInternalPlayWorldCommandCallbacks::PlayInViewport_CanExecute()
 
 void FInternalPlayWorldCommandCallbacks::PlayInEditorFloating_Clicked()
 {
-	if (IsStoppedAtBreakpoint_InEngineMode())
+	if (IsStoppedAtBreakpoint_InEngineMode() || DebuggerCommands::IsVREditorActive())
 	{
 		return;
 	}
@@ -1554,7 +1565,7 @@ void FInternalPlayWorldCommandCallbacks::PlayInEditorFloating_Clicked()
 
 bool FInternalPlayWorldCommandCallbacks::PlayInEditorFloating_CanExecute()
 {
-	if (IsStoppedAtBreakpoint_InEngineMode())
+	if (IsStoppedAtBreakpoint_InEngineMode() || DebuggerCommands::IsVREditorActive())
 	{
 		return false;
 	}
@@ -1564,7 +1575,7 @@ bool FInternalPlayWorldCommandCallbacks::PlayInEditorFloating_CanExecute()
 
 void FInternalPlayWorldCommandCallbacks::PlayInVR_Clicked()
 {
-	if (IsStoppedAtBreakpoint_InEngineMode())
+	if (IsStoppedAtBreakpoint_InEngineMode() || DebuggerCommands::IsVREditorActive())
 	{
 		return;
 	}
@@ -1606,10 +1617,11 @@ void FInternalPlayWorldCommandCallbacks::PlayInVR_Clicked()
 
 bool FInternalPlayWorldCommandCallbacks::PlayInVR_CanExecute()
 {
-	if (IsStoppedAtBreakpoint_InEngineMode())
+	if (IsStoppedAtBreakpoint_InEngineMode() || DebuggerCommands::IsVREditorActive())
 	{
 		return false;
 	}
+
 	return (!HasPlayWorld() || !GUnrealEd->bIsSimulatingInEditor) && !GEditor->IsLightingBuildCurrentlyRunning() &&
 		GEngine && GEngine->XRSystem.IsValid() && GEngine->XRSystem->GetHMDDevice() && GEngine->XRSystem->GetHMDDevice()->IsHMDConnected();
 }
@@ -1634,7 +1646,7 @@ void FInternalPlayWorldCommandCallbacks::PlayInNewProcess_Clicked(EPlayModeType
 	check(PlayModeType == PlayMode_InNewProcess || PlayModeType == PlayMode_InMobilePreview
 		|| PlayModeType == PlayMode_InTargetedMobilePreview || PlayModeType == PlayMode_InVulkanPreview);
 
-	if (IsStoppedAtBreakpoint_InEngineMode())
+	if (IsStoppedAtBreakpoint_InEngineMode() || DebuggerCommands::IsVREditorActive())
 	{
 		return;
 	}
@@ -1687,7 +1699,7 @@ void FInternalPlayWorldCommandCallbacks::PlayInNewProcess_Clicked(EPlayModeType
 
 bool FInternalPlayWorldCommandCallbacks::PlayInNewProcess_CanExecute()
 {
-	if (IsStoppedAtBreakpoint_InEngineMode())
+	if (IsStoppedAtBreakpoint_InEngineMode() || DebuggerCommands::IsVREditorActive())
 	{
 		return false;
 	}
@@ -1713,7 +1725,7 @@ bool FInternalPlayWorldCommandCallbacks::PlayInModeIsChecked(EPlayModeType PlayM
 
 bool FInternalPlayWorldCommandCallbacks::PlayInLocation_CanExecute(EPlayModeLocations Location)
 {
-	if (IsStoppedAtBreakpoint_InEngineMode())
+	if (IsStoppedAtBreakpoint_InEngineMode() || DebuggerCommands::IsVREditorActive())
 	{
 		return false;
 	}
