Explicit mark dirty to Modify for Brush, Poly, and Model.

This change restores the behavior prior to the change to the default argument of the Modify function of (Brush, Poly, and Model) from false to true.

diff --git a/Engine/Plugins/Editor/GeometryMode/Source/GeometryMode/Private/GeometryModifiers.cpp b/Engine/Plugins/Editor/GeometryMode/Source/GeometryMode/Private/GeometryModifiers.cpp
--- a/Engine/Plugins/Editor/GeometryMode/Source/GeometryMode/Private/GeometryModifiers.cpp
+++ b/Engine/Plugins/Editor/GeometryMode/Source/GeometryMode/Private/GeometryModifiers.cpp
@@ -370,7 +370,7 @@ void UGeomModifier::StartTrans()
 		FGeomObjectPtr go = *Itor;
 		ABrush* Actor = go->GetActualBrush();
 
-		Actor->Modify();
+		Actor->Modify(false);
 	}
 }
 
@@ -1973,8 +1973,8 @@ static ABrush* ClipBrushAgainstPlane( const FPlane& InPlane, ABrush* InBrush)
 	// Intersect the giant brush with the source brush's BSP.  This will give us the finished, clipping brush
 	// contained inside of the giant brush.
 
-	ClippedBrush->Modify();
-	InBrush->Brush->Modify();
+	ClippedBrush->Modify(false);
+	InBrush->Brush->Modify(false);
 	GEditor->bspBrushCSG( ClippedBrush, InBrush->Brush, 0, Brush_MAX, CSG_Intersect, false, false, true );
 	FBSPOps::bspUnlinkPolys( ClippedBrush->Brush );
 
@@ -2051,7 +2051,7 @@ static ABrush* ClipBrushAgainstPlane( const FPlane& InPlane, ABrush* InBrush)
 		// Special-case behavior for the builder brush.
 
 		// Copy the temporary brush back over onto the builder brush (keeping object flags)
-		BuilderBrush->Modify();
+		BuilderBrush->Modify(false);
 		FBSPOps::csgCopyBrush( BuilderBrush, ClippedBrush, BuilderBrush->PolyFlags, BuilderBrush->GetFlags(), 0, true );
 		ULayersSubsystem* Layers = GEditor->GetEditorSubsystem<ULayersSubsystem>();
 		Layers->DisassociateActorFromLayers( ClippedBrush );
diff --git a/Engine/Plugins/Experimental/MeshModelingToolsetExp/Source/MeshModelingToolsEditorOnlyExp/Private/BspConversionTool.cpp b/Engine/Plugins/Experimental/MeshModelingToolsetExp/Source/MeshModelingToolsEditorOnlyExp/Private/BspConversionTool.cpp
--- a/Engine/Plugins/Experimental/MeshModelingToolsetExp/Source/MeshModelingToolsEditorOnlyExp/Private/BspConversionTool.cpp
+++ b/Engine/Plugins/Experimental/MeshModelingToolsetExp/Source/MeshModelingToolsEditorOnlyExp/Private/BspConversionTool.cpp
@@ -466,7 +466,7 @@ bool UBspConversionTool::ConvertThenCombine(FText *ErrorMessage)
 			// depending on the settings.
 			// Despite the undoing, we need to call Modify now because it gets called in the conversion code, which
 			// would cause the changed property to be saved.
-			NextBrush->Modify();
+			NextBrush->Modify(false);
 			NextBrush->BrushType = EBrushType::Brush_Add;
 		}
 
diff --git a/Engine/Plugins/Experimental/MeshModelingToolsetExp/Source/MeshModelingToolsEditorOnlyExp/Private/MeshToVolumeTool.cpp b/Engine/Plugins/Experimental/MeshModelingToolsetExp/Source/MeshModelingToolsEditorOnlyExp/Private/MeshToVolumeTool.cpp
--- a/Engine/Plugins/Experimental/MeshModelingToolsetExp/Source/MeshModelingToolsEditorOnlyExp/Private/MeshToVolumeTool.cpp
+++ b/Engine/Plugins/Experimental/MeshModelingToolsetExp/Source/MeshModelingToolsEditorOnlyExp/Private/MeshToVolumeTool.cpp
@@ -241,7 +241,7 @@ void UMeshToVolumeTool::OnShutdown(EToolShutdownType ShutdownType)
 		{
 			TargetVolume = Settings->TargetVolume.Get();
 			SetTransform = TargetVolume->GetActorTransform();
-			TargetVolume->Modify();
+			TargetVolume->Modify(true);
 			TargetVolume->GetBrushComponent()->Modify();
 		}
 
diff --git a/Engine/Plugins/Runtime/MeshModelingToolset/Source/ModelingComponentsEditorOnly/Private/ConversionUtils/DynamicMeshToVolume.cpp b/Engine/Plugins/Runtime/MeshModelingToolset/Source/ModelingComponentsEditorOnly/Private/ConversionUtils/DynamicMeshToVolume.cpp
--- a/Engine/Plugins/Runtime/MeshModelingToolset/Source/ModelingComponentsEditorOnly/Private/ConversionUtils/DynamicMeshToVolume.cpp
+++ b/Engine/Plugins/Runtime/MeshModelingToolset/Source/ModelingComponentsEditorOnly/Private/ConversionUtils/DynamicMeshToVolume.cpp
@@ -741,7 +741,7 @@ void DynamicMeshToVolume(const FDynamicMesh3&, TArray<FDynamicMeshFace>& Faces,
 
 	UModel* Model = TargetVolume->Brush;
 
-	Model->Modify();
+	Model->Modify(false);
 
 	Model->Initialize(TargetVolume);
 	UPolys* Polys = Model->Polys;
diff --git a/Engine/Plugins/Runtime/MeshModelingToolset/Source/ModelingComponentsEditorOnly/Private/Selection/VolumeSelector.cpp b/Engine/Plugins/Runtime/MeshModelingToolset/Source/ModelingComponentsEditorOnly/Private/Selection/VolumeSelector.cpp
--- a/Engine/Plugins/Runtime/MeshModelingToolset/Source/ModelingComponentsEditorOnly/Private/Selection/VolumeSelector.cpp
+++ b/Engine/Plugins/Runtime/MeshModelingToolset/Source/ModelingComponentsEditorOnly/Private/Selection/VolumeSelector.cpp
@@ -227,7 +227,7 @@ void FVolumeSelector::CommitMeshTransform()
 {
 	FTransform SetTransform = ParentVolume->GetActorTransform();
 
-	ParentVolume->Modify();
+	ParentVolume->Modify(true);
 	BrushComponent->Modify();
 
 	UE::Conversion::FMeshToVolumeOptions ConvertOptions;
diff --git a/Engine/Plugins/Runtime/nDisplay/Source/DisplayClusterColorGrading/Private/Drawer/SDisplayClusterColorGradingDrawer.cpp b/Engine/Plugins/Runtime/nDisplay/Source/DisplayClusterColorGrading/Private/Drawer/SDisplayClusterColorGradingDrawer.cpp
--- a/Engine/Plugins/Runtime/nDisplay/Source/DisplayClusterColorGrading/Private/Drawer/SDisplayClusterColorGradingDrawer.cpp
+++ b/Engine/Plugins/Runtime/nDisplay/Source/DisplayClusterColorGrading/Private/Drawer/SDisplayClusterColorGradingDrawer.cpp
@@ -604,7 +604,7 @@ void SDisplayClusterColorGradingDrawer::UnbindBlueprintCompiledDelegate(const UC
 #define CREATE_ON_ENABLED_CHANGED_LAMBDA(Object, IsEnabledProperty) FOnColorGradingItemEnabledChanged::CreateLambda([Object](FDisplayClusterColorGradingListItemRef ListItem, bool bIsEnabled) \
 	{ \
 		FScopedTransaction Transaction(LOCTEXT("ColorGradingToggledTransaction", "Color Grading Toggled")); \
-		Object->Modify(); \
+		Object->Modify(!Object->IsA<ABrush>()); \
 		IsEnabledProperty = bIsEnabled; \
 	})
 
diff --git a/Engine/Source/Developer/BSPUtils/Private/BSPOps.cpp b/Engine/Source/Developer/BSPUtils/Private/BSPOps.cpp
--- a/Engine/Source/Developer/BSPUtils/Private/BSPOps.cpp
+++ b/Engine/Source/Developer/BSPUtils/Private/BSPOps.cpp
@@ -1028,7 +1028,7 @@ void FBSPOps::bspBuildBounds( UModel* Model )
 void FBSPOps::bspValidateBrush( UModel* Brush, bool ForceValidate, bool DoStatusUpdate )
 {
 	check(Brush != nullptr);
-	Brush->Modify();
+	Brush->Modify(false);
 	if( ForceValidate || !Brush->Linked )
 	{
 		Brush->Linked = 1;
@@ -1072,7 +1072,7 @@ void FBSPOps::bspValidateBrush( UModel* Brush, bool ForceValidate, bool DoStatus
 
 void FBSPOps::bspUnlinkPolys( UModel* Brush )
 {
-	Brush->Modify();
+	Brush->Modify(false);
 	Brush->Linked = 1;
 	for( int32 i=0; i<Brush->Polys->Element.Num(); i++ )
 	{
@@ -1272,7 +1272,7 @@ int32	FBSPOps::bspAddNode( UModel* Model, int32 iParent, ENodePlace NodePlace, u
  */
 void FBSPOps::RebuildBrush(UModel* Brush)
 {
-	Brush->Modify();
+	Brush->Modify(false);
 	Brush->EmptyModel(1, 0);
 
 	// Build bounding box.
diff --git a/Engine/Source/Editor/UnrealEd/Private/EditorActor.cpp b/Engine/Source/Editor/UnrealEd/Private/EditorActor.cpp
--- a/Engine/Source/Editor/UnrealEd/Private/EditorActor.cpp
+++ b/Engine/Source/Editor/UnrealEd/Private/EditorActor.cpp
@@ -1179,7 +1179,7 @@ bool UUnrealEdEngine::DeleteActors(const TArray<AActor*>& InActorsToDelete, UWor
 			ULevel* BrushLevel = Actor->GetLevel();
 			if (BrushLevel && !Brush->IsVolumeBrush())
 			{
-				BrushLevel->Model->Modify();
+				BrushLevel->Model->Modify(false);
 				LevelsToRebuildBSP.Add(BrushLevel);
 				// Rebuilding bsp will also take care of navigation
 				LevelsToRebuildNavigation.Remove(BrushLevel);
@@ -1345,7 +1345,7 @@ void UUnrealEdEngine::edactReplaceSelectedBrush( UWorld* InWorld )
 
 			LevelDirtyCallback.Request();
 
-			NewBrush->Modify();
+			NewBrush->Modify(false);
 
 			NewBrush->Layers.Append( SrcBrush->Layers );
 
@@ -1652,7 +1652,7 @@ void UUnrealEdEngine::edactHideSelectedStartup( UWorld* InWorld )
 				const bool bSelected = CurSurface.Actor->IsActorOrSelectionParentSelected() || (CurSurface.PolyFlags & PF_Selected);
 				if (bSelected && !CurSurface.IsHiddenEdAtStartup() && !CurSurface.IsHiddenEd())
 				{
-					CurLevelModel.Modify();
+					CurLevelModel.Modify(false);
 					CurLevelModel.ModifySurf( SurfaceIterator.GetIndex(), false );
 					CurSurface.PolyFlags |= PF_HiddenEd;
 					LevelDirtyCallback.Request();
@@ -1697,7 +1697,7 @@ void UUnrealEdEngine::edactUnHideAllStartup( UWorld* InWorld )
 				// If the BSP surface is set to be hidden at editor startup, change it so that it will be shown at startup
 				if ( CurSurface.IsHiddenEdAtStartup() )
 				{
-					CurLevelModel.Modify();
+					CurLevelModel.Modify(false);
 					CurLevelModel.ModifySurf( SurfaceIterator.GetIndex(), false );
 					CurSurface.PolyFlags &= ~PF_HiddenEd;
 					LevelDirtyCallback.Request();
@@ -1743,7 +1743,7 @@ void UUnrealEdEngine::edactUnHideSelectedStartup( UWorld* InWorld )
 				const bool bSelected = CurSurface.Actor->IsActorOrSelectionParentSelected() || (CurSurface.PolyFlags & PF_Selected);
 				if (bSelected && CurSurface.IsHiddenEdAtStartup())
 				{
-					CurLevelModel.Modify();
+					CurLevelModel.Modify(false);
 					CurLevelModel.ModifySurf( SurfaceIterator.GetIndex(), false );
 					CurSurface.PolyFlags &= ~PF_HiddenEd;
 					LevelDirtyCallback.Request();
@@ -2676,7 +2676,7 @@ void UUnrealEdEngine::edactAlignOrigin()
 			LevelDirtyCallback.Request();
 
 			Brush->PreEditChange(NULL);
-			Brush->Modify();
+			Brush->Modify(false);
 
 			//Snap the location of the brush to the grid
 			FVector BrushLocation = Brush->GetActorLocation();
@@ -2715,7 +2715,7 @@ void UUnrealEdEngine::edactAlignVertices()
 			LevelDirtyCallback.Request();
 
 			Brush->PreEditChange(NULL);
-			Brush->Modify();
+			Brush->Modify(false);
 			FVector BrushLocation = Brush->GetActorLocation();
 			const FTransform BrushTransform = Brush->GetRootComponent()->GetComponentTransform();
 
diff --git a/Engine/Source/Editor/UnrealEd/Private/EditorBrushBuilder.cpp b/Engine/Source/Editor/UnrealEd/Private/EditorBrushBuilder.cpp
--- a/Engine/Source/Editor/UnrealEd/Private/EditorBrushBuilder.cpp
+++ b/Engine/Source/Editor/UnrealEd/Private/EditorBrushBuilder.cpp
@@ -80,8 +80,8 @@ bool UEditorBrushBuilder::EndBrush( UWorld* InWorld, ABrush* InBrush )
 		return true;
 	}
 
-	Brush->Modify();
-	BuilderBrush->Modify();
+	Brush->Modify(false);
+	BuilderBrush->Modify(false);
 
 	FRotator Temp(0.0f,0.0f,0.0f);
 	FSnappingUtils::SnapToBSPVertex( Location, FVector::ZeroVector, Temp );
diff --git a/Engine/Source/Editor/UnrealEd/Private/EditorCsg.cpp b/Engine/Source/Editor/UnrealEd/Private/EditorCsg.cpp
--- a/Engine/Source/Editor/UnrealEd/Private/EditorCsg.cpp
+++ b/Engine/Source/Editor/UnrealEd/Private/EditorCsg.cpp
@@ -144,7 +144,7 @@ void UEditorEngine::csgRebuild( UWorld* InWorld )
 	FinishAllSnaps();
 
 	// Empty the model out.
-	InWorld->GetModel()->Modify();
+	InWorld->GetModel()->Modify(false);
 	InWorld->GetModel()->EmptyModel(1, 1);
 
 	// Count brushes.
@@ -201,7 +201,7 @@ void UEditorEngine::csgRebuild( UWorld* InWorld )
 				Args.Add( TEXT("BrushTotal"), BrushTotal );
 				GWarn->StatusUpdate( BrushCount, BrushTotal, FText::Format( NSLOCTEXT("UnrealEd", "ApplyingStructuralBrushF", "Applying structural brush {BrushCount} of {BrushTotal}"), Args ) );
 
-				Brush->Modify();
+				Brush->Modify(false);
 				bspBrushCSG( Brush, InWorld->GetModel(), Brush->PolyFlags, (EBrushType)Brush->BrushType, CSG_None, false, true, false );
 			}
 		}
@@ -245,7 +245,7 @@ void UEditorEngine::csgRebuild( UWorld* InWorld )
 			Args.Add( TEXT("BrushTotal"), BrushTotal );
 			GWarn->StatusUpdate( BrushCount, BrushTotal, FText::Format( NSLOCTEXT("UnrealEd", "ApplyingDetailBrushF", "Applying detail brush {BrushCount} of {BrushTotal}"), Args ) );
 
-			Brush->Modify();
+			Brush->Modify(false);
 			bspBrushCSG( Brush, InWorld->GetModel(), Brush->PolyFlags, (EBrushType)Brush->BrushType, CSG_None, false, true, false );
 		}
 	}
@@ -1100,7 +1100,7 @@ void UEditorEngine::MapBrushGet(UWorld* InWorld)
 			check( BrushActor->GetWorld() );			
 			ABrush* WorldBrush = BrushActor->GetWorld()->GetDefaultBrush();
 			check( WorldBrush );
-			WorldBrush->Modify();
+			WorldBrush->Modify(false);
 			WorldBrush->Brush->Polys->Element = BrushActor->Brush->Polys->Element;
 			WorldBrush->CopyPosRotScaleFrom( BrushActor );
 
@@ -1128,7 +1128,7 @@ void UEditorEngine::mapBrushPut()
 			ABrush* WorldBrush = BrushActor->GetWorld()->GetDefaultBrush();
 			check( WorldBrush );
 
-			BrushActor->Modify();
+			BrushActor->Modify(false);
 			BrushActor->Brush->Polys->Element = WorldBrush->Brush->Polys->Element;
 			BrushActor->CopyPosRotScaleFrom( WorldBrush );
 			BrushActor->SetNeedRebuild(BrushActor->GetLevel());
@@ -1239,7 +1239,7 @@ void UEditorEngine::MapSetBrush( UWorld* InWorld, EMapSetBrushFlags	PropertiesMa
 		{
 			if( PropertiesMask & MSB_PolyFlags )
 			{
-				Brush->Modify();
+				Brush->Modify(false);
 				Brush->PolyFlags = (Brush->PolyFlags & ~ClearPolyFlags) | SetPolyFlags;
 				Brush->UpdateComponentTransforms();
 				Brush->MarkPackageDirty();
@@ -1247,7 +1247,7 @@ void UEditorEngine::MapSetBrush( UWorld* InWorld, EMapSetBrushFlags	PropertiesMa
 			}
 			if( PropertiesMask & MSB_BrushType )
 			{
-				Brush->Modify();
+				Brush->Modify(false);
 				Brush->BrushType = EBrushType(BrushType);
 				Brush->UpdateComponentTransforms();
 				Brush->MarkPackageDirty();
diff --git a/Engine/Source/Editor/UnrealEd/Private/EditorEngine.cpp b/Engine/Source/Editor/UnrealEd/Private/EditorEngine.cpp
--- a/Engine/Source/Editor/UnrealEd/Private/EditorEngine.cpp
+++ b/Engine/Source/Editor/UnrealEd/Private/EditorEngine.cpp
@@ -3599,7 +3599,7 @@ void UEditorEngine::ConvertSelectedBrushesToVolumes( UClass* VolumeClass )
 
 				NewVolume->PostEditChange();
 				NewVolume->PostEditMove( true );
-				NewVolume->Modify();
+				NewVolume->Modify(false);
 
 				// Make the actor visible as the brush is hidden by default
 				NewVolume->SetActorHiddenInGame(false);
@@ -5504,7 +5504,7 @@ void UEditorEngine::ReplaceActors(UActorFactory* Factory, const FAssetData& Asse
 				ULevel* BrushLevel = OldActor->GetLevel();
 				if (BrushLevel && !Brush->IsVolumeBrush())
 				{
-					BrushLevel->Model->Modify();
+					BrushLevel->Model->Modify(false);
 					LevelsToRebuildBSP.Add(BrushLevel);
 				}
 			}
diff --git a/Engine/Source/Editor/UnrealEd/Private/EditorLevelUtils.cpp b/Engine/Source/Editor/UnrealEd/Private/EditorLevelUtils.cpp
--- a/Engine/Source/Editor/UnrealEd/Private/EditorLevelUtils.cpp
+++ b/Engine/Source/Editor/UnrealEd/Private/EditorLevelUtils.cpp
@@ -624,7 +624,7 @@ bool UEditorLevelUtils::PrivateRemoveInvalidLevelFromWorld(ULevelStreaming* InLe
 		{
 			if (LevelStreamingVolume)
 			{
-				LevelStreamingVolume->Modify();
+				LevelStreamingVolume->Modify(true);
 				LevelStreamingVolume->StreamingLevelNames.Remove(InLevelStreaming->GetWorldAssetPackageFName());
 			}
 		}
diff --git a/Engine/Source/Editor/UnrealEd/Private/EditorServer.cpp b/Engine/Source/Editor/UnrealEd/Private/EditorServer.cpp
--- a/Engine/Source/Editor/UnrealEd/Private/EditorServer.cpp
+++ b/Engine/Source/Editor/UnrealEd/Private/EditorServer.cpp
@@ -530,7 +530,7 @@ bool UEditorEngine::Exec_StaticMesh( UWorld* InWorld, const TCHAR* Str, FOutputD
 		if(FParse::Command(&Str,TEXT("BRUSH")))
 		{
 			const FScopedTransaction Transaction( NSLOCTEXT("UnrealEd", "StaticMeshToBrush", "StaticMesh to Brush") );
-			WorldBrush->Brush->Modify();
+			WorldBrush->Brush->Modify(false);
 
 			// Find the first selected static mesh actor.
 			AStaticMeshActor* SelectedActor = NULL;
@@ -625,7 +625,7 @@ bool UEditorEngine::Exec_Brush( UWorld* InWorld, const TCHAR* Str, FOutputDevice
 			ABrush* DefaultBrush = InWorld->GetDefaultBrush();
 			if (DefaultBrush != NULL)
 			{
-				DefaultBrush->Brush->Modify();
+				DefaultBrush->Brush->Modify(false);
 				SnapLocation = DefaultBrush->GetActorLocation();
 				PrePivot = DefaultBrush->GetPivotOffset();
 			}
@@ -647,7 +647,7 @@ bool UEditorEngine::Exec_Brush( UWorld* InWorld, const TCHAR* Str, FOutputDevice
 	else if( FParse::Command(&Str,TEXT("RESET")) )
 	{
 		const FScopedTransaction Transaction( NSLOCTEXT("UnrealEd", "BrushReset", "Brush Reset") );
-		WorldBrush->Modify();
+		WorldBrush->Modify(false);
 		WorldBrush->InitPosRotScale();
 		RedrawLevelEditingViewports();
 		return true;
@@ -677,7 +677,7 @@ bool UEditorEngine::Exec_Brush( UWorld* InWorld, const TCHAR* Str, FOutputDevice
 			{
 				if ( Brush->Brush )
 				{
-					Brush->Brush->Modify();
+					Brush->Brush->Modify(false);
 					for( int32 poly = 0 ; poly < Brush->Brush->Polys->Element.Num() ; poly++ )
 					{
 						FPoly* Poly = &(Brush->Brush->Polys->Element[poly]);
@@ -708,7 +708,7 @@ bool UEditorEngine::Exec_Brush( UWorld* InWorld, const TCHAR* Str, FOutputDevice
 	else if( FParse::Command(&Str,TEXT("MOVETO")) )
 	{
 		const FScopedTransaction Transaction( NSLOCTEXT("UnrealEd", "BrushMoveTo", "Brush MoveTo") );
-		WorldBrush->Modify();
+		WorldBrush->Modify(false);
 		FVector TempVector(0.f);
 		GetFVECTOR( Str, TempVector );
 		WorldBrush->SetActorLocation(TempVector, false);
@@ -718,7 +718,7 @@ bool UEditorEngine::Exec_Brush( UWorld* InWorld, const TCHAR* Str, FOutputDevice
 	else if( FParse::Command(&Str,TEXT("MOVEREL")) )
 	{
 		const FScopedTransaction Transaction( NSLOCTEXT("UnrealEd", "BrushMoveRel", "Brush MoveRel") );
-		WorldBrush->Modify();
+		WorldBrush->Modify(false);
 		FVector TempVector( 0, 0, 0 );
 		GetFVECTOR( Str, TempVector );
 		FVector NewLocation = WorldBrush->GetActorLocation();
@@ -757,8 +757,8 @@ bool UEditorEngine::Exec_Brush( UWorld* InWorld, const TCHAR* Str, FOutputDevice
 					LoadAndSelectAssets( SelectedAssets, UMaterial::StaticClass() );
 				}
 
-				InWorld->GetModel()->Modify();
-				NewBrush->Modify();
+				InWorld->GetModel()->Modify(false);
+				NewBrush->Modify(false);
 				bspBrushCSG( NewBrush, InWorld->GetModel(), DWord1, Brush_Add, CSG_None, true, true, true );
 
 				if (FParse::Command(&Str, TEXT("SELECTNEWBRUSH")))
@@ -847,8 +847,8 @@ bool UEditorEngine::Exec_Brush( UWorld* InWorld, const TCHAR* Str, FOutputDevice
 					FActorLabelUtilities::SetActorLabelUnique(NewBrush, FText::Format(NSLOCTEXT("UnrealEd", "BrushName", "{0} Brush"), FText::FromString(NewBrush->GetBrushBuilder()->GetClass()->GetDescription())).ToString());
 				}
 
-				NewBrush->Modify();
-				InWorld->GetModel()->Modify();
+				NewBrush->Modify(false);
+				InWorld->GetModel()->Modify(false);
 				bspBrushCSG( NewBrush, InWorld->GetModel(), 0, Brush_Subtract, CSG_None, true, true, true );
 			}
 			InWorld->InvalidateModelGeometry( InWorld->GetCurrentLevel() );
@@ -916,7 +916,7 @@ bool UEditorEngine::Exec_Brush( UWorld* InWorld, const TCHAR* Str, FOutputDevice
 	else if( FParse::Command (&Str,TEXT("NEW")) )
 	{
 		const FScopedTransaction Transaction( NSLOCTEXT("UnrealEd", "BrushNew", "Brush New") );
-		WorldBrush->Brush->Modify();
+		WorldBrush->Brush->Modify(false);
 		WorldBrush->Brush->Polys->Element.Empty();
 		RedrawLevelEditingViewports();
 		return true;
@@ -965,7 +965,7 @@ bool UEditorEngine::Exec_Brush( UWorld* InWorld, const TCHAR* Str, FOutputDevice
 
 			GWarn->BeginSlowTask( NSLOCTEXT("UnrealEd", "ImportingBrush", "Importing brush"), true );
 
-			WorldBrush->Brush->Polys->Modify();
+			WorldBrush->Brush->Polys->Modify(false);
 			WorldBrush->Brush->Polys->Element.Empty();
 			uint32 Flags=0;
 			bool Merge=0;
@@ -1652,7 +1652,7 @@ void UEditorEngine::RebuildModelFromBrushes(UModel* Model, bool bSelectedBrushes
 	const int32 NumVectors = Model->Vectors.Num();
 	const int32 NumSurfs = Model->Surfs.Num();
 
-	Model->Modify();
+	Model->Modify(false);
 	Model->EmptyModel(1, 1);
 
 	// Reserve arrays an eighth bigger than the previous allocation
@@ -1725,7 +1725,7 @@ void UEditorEngine::RebuildModelFromBrushes(UModel* Model, bool bSelectedBrushes
 	for (ABrush* Brush : StaticBrushes)
 	{
 		SlowTask.EnterProgressFrame(1);
-		Brush->Modify();
+		Brush->Modify(false);
 		bspBrushCSG(Brush, Model, Brush->PolyFlags, (EBrushType)Brush->BrushType, CSG_None, false, true, false, false);
 	}
 
@@ -1759,7 +1759,7 @@ void UEditorEngine::RebuildModelFromBrushes(TArray<ABrush*> &BrushesToBuild, UMo
 	const int32 NumVectors = Model->Vectors.Num();
 	const int32 NumSurfs = Model->Surfs.Num();
 
-	Model->Modify();
+	Model->Modify(false);
 	Model->EmptyModel(1, 1);
 
 	// Reserve arrays an eighth bigger than the previous allocation
@@ -1776,7 +1776,7 @@ void UEditorEngine::RebuildModelFromBrushes(TArray<ABrush*> &BrushesToBuild, UMo
 	for (ABrush* Brush : BrushesToBuild)
 	{
 		SlowTask.EnterProgressFrame(1);
-		Brush->Modify();
+		Brush->Modify(false);
 		bspBrushCSG(Brush, Model, Brush->PolyFlags, (EBrushType)Brush->BrushType, CSG_None, false, true, false, false);
 	}
 
@@ -1879,8 +1879,8 @@ void UEditorEngine::BSPIntersectionHelper(UWorld* InWorld, ECsgOper Operation)
 	ABrush* DefaultBrush = InWorld->GetDefaultBrush();
 	if (DefaultBrush != NULL)
 	{
-		DefaultBrush->Modify();
-		InWorld->GetModel()->Modify();
+		DefaultBrush->Modify(false);
+		InWorld->GetModel()->Modify(false);
 		FinishAllSnaps();
 		bspBrushCSG(DefaultBrush, InWorld->GetModel(), 0, Brush_MAX, Operation, false, true, true);
 	}
@@ -4184,7 +4184,7 @@ namespace {
 		{
 			FBspSurf* Surf = *It;
 			UModel* Model = It.GetModel();
-			Model->Modify();
+			Model->Modify(false);
 			const FVector3f TextureU( Model->Vectors[Surf->vTextureU] );
 			const FVector3f TextureV( Model->Vectors[Surf->vTextureV] );
 			Surf->vTextureU = Model->Vectors.Add(TextureU);
@@ -4200,7 +4200,7 @@ namespace {
 		for( FConstLevelIterator Iterator = InWorld->GetLevelIterator(); Iterator; ++Iterator )
 		{
 			UModel* Model = (*Iterator)->Model;
-			Model->Modify();
+			Model->Modify(false);
 			GEditor->polyTexScale( Model, UU, UV, VU, VV, !!Word2 );
 		}
 	}
@@ -4495,7 +4495,7 @@ bool UEditorEngine::Exec_Poly( UWorld* InWorld, const TCHAR* Str, FOutputDevice&
 			{
 				FBspSurf* Surf = *It;
 				UModel* Model = It.GetModel();
-				Model->Modify();
+				Model->Modify(false);
 				const FVector3f Base( Model->Points[Surf->pBase] );
 				Surf->pBase = Model->Points.Add(Base);
 			}
@@ -4505,7 +4505,7 @@ bool UEditorEngine::Exec_Poly( UWorld* InWorld, const TCHAR* Str, FOutputDevice&
 				for( FConstLevelIterator Iterator = InWorld->GetLevelIterator(); Iterator; ++Iterator )
 				{
 					UModel* Model = (*Iterator)->Model;
-					Model->Modify();
+					Model->Modify(false);
 					polyTexPan( Model, 0, 0, 1 );
 				}
 			}
@@ -4515,7 +4515,7 @@ bool UEditorEngine::Exec_Poly( UWorld* InWorld, const TCHAR* Str, FOutputDevice&
 			for( FConstLevelIterator Iterator = InWorld->GetLevelIterator(); Iterator; ++Iterator )
 			{
 				UModel* Model = (*Iterator)->Model;
-				Model->Modify();
+				Model->Modify(false);
 				polyTexPan( Model, PanU, PanV, 0 );
 			}
 		}
diff --git a/Engine/Source/Editor/UnrealEd/Private/UnrealEdSrv.cpp b/Engine/Source/Editor/UnrealEd/Private/UnrealEdSrv.cpp
--- a/Engine/Source/Editor/UnrealEd/Private/UnrealEdSrv.cpp
+++ b/Engine/Source/Editor/UnrealEd/Private/UnrealEdSrv.cpp
@@ -852,10 +852,10 @@ bool UUnrealEdEngine::Exec( UWorld* InWorld, const TCHAR* Stream, FOutputDevice&
 						{
 							const FVector OldLocation = Brush->GetActorLocation();
 							const FVector NewLocation = OldLocation * ScaleVec;
-							Brush->Modify();
+							Brush->Modify(false);
 							Brush->SetActorLocation( NewLocation );
 							
-							Brush->Brush->Modify();
+							Brush->Brush->Modify(false);
 							for( int32 poly = 0 ; poly < Brush->Brush->Polys->Element.Num() ; poly++ )
 							{
 								FPoly* Poly = &(Brush->Brush->Polys->Element[poly]);
@@ -2004,7 +2004,7 @@ bool UUnrealEdEngine::Exec_Actor( UWorld* InWorld, const TCHAR* Str, FOutputDevi
 	else if( FParse::Command(&Str,TEXT("CREATE_BV_BOUNDINGBOX")) )
 	{
 		const FScopedTransaction Transaction( NSLOCTEXT("UnrealEd", "CreateBoundingBoxBlockingVolume", "Create Bounding Box Blocking Volume") );
-		InWorld->GetDefaultBrush()->Modify();
+		InWorld->GetDefaultBrush()->Modify(false);
 
 		bool bSnapToGrid=0;
 		FParse::Bool( Str, TEXT("SNAPTOGRID="), bSnapToGrid );
@@ -2035,7 +2035,7 @@ bool UUnrealEdEngine::Exec_Actor( UWorld* InWorld, const TCHAR* Str, FOutputDevi
 	else if( FParse::Command(&Str,TEXT("CREATE_BV_CONVEXVOLUME")) )
 	{
 		const FScopedTransaction Transaction( NSLOCTEXT("UnrealEd", "CreateConvexBlockingVolume", "Create Convex Blocking Volume") );
-		InWorld->GetDefaultBrush()->Modify();
+		InWorld->GetDefaultBrush()->Modify(false);
 
 		bool bSnapToGrid=0;
 		FParse::Bool( Str, TEXT("SNAPTOGRID="), bSnapToGrid );
diff --git a/Engine/Source/Editor/WorldPartitionEditor/Private/WorldPartition/SWorldPartitionEditorGrid2D.cpp b/Engine/Source/Editor/WorldPartitionEditor/Private/WorldPartition/SWorldPartitionEditorGrid2D.cpp
--- a/Engine/Source/Editor/WorldPartitionEditor/Private/WorldPartition/SWorldPartitionEditorGrid2D.cpp
+++ b/Engine/Source/Editor/WorldPartitionEditor/Private/WorldPartition/SWorldPartitionEditorGrid2D.cpp
@@ -738,7 +738,7 @@ void SWorldPartitionEditorGrid2D::ConvertSelectedRegionsToActors()
 			const FBox ActorVolumeBox(FVector(LoaderVolumeBox.Min.X, LoaderVolumeBox.Min.Y, WorldBounds.Min.Z), FVector(LoaderVolumeBox.Max.X, LoaderVolumeBox.Max.Y, WorldBounds.Max.Z));
 
 			ALocationVolume* LocationVolume = GetWorld()->SpawnActor<ALocationVolume>(ActorVolumeBox.GetCenter(), FRotator::ZeroRotator);
-			LocationVolume->Modify();
+			LocationVolume->Modify(true);
 			FActorLabelUtilities::SetActorLabelUnique(LocationVolume, LocationVolume->GetActorLabel());
 			
 			UCubeBuilder* Builder = NewObject<UCubeBuilder>();
diff --git a/Engine/Source/Runtime/Engine/Private/Level.cpp b/Engine/Source/Runtime/Engine/Private/Level.cpp
--- a/Engine/Source/Runtime/Engine/Private/Level.cpp
+++ b/Engine/Source/Runtime/Engine/Private/Level.cpp
@@ -2642,7 +2642,7 @@ bool ULevel::GetPromptWhenAddingToLevelBeforeCheckout() const
 void ULevel::InvalidateModelGeometry()
 {
 	// Save the level/model state for transactions.
-	Model->Modify();
+	Model->Modify(true);
 	Modify();
 
 	// Remove existing model components.
diff --git a/Engine/Source/Runtime/Engine/Private/Model.cpp b/Engine/Source/Runtime/Engine/Private/Model.cpp
--- a/Engine/Source/Runtime/Engine/Private/Model.cpp
+++ b/Engine/Source/Runtime/Engine/Private/Model.cpp
@@ -474,18 +474,18 @@ void UModel::PostEditUndo()
 
 void UModel::ModifySurf( int32 InIndex, bool UpdateBrushes )
 {
-	Modify();
+	Modify(false);
 
 	FBspSurf& Surf = Surfs[InIndex];
 	if( UpdateBrushes && Surf.Actor )
 	{
-		Surf.Actor->Brush->Modify();
+		Surf.Actor->Brush->Modify(false);
 	}
 }
 
 void UModel::ModifyAllSurfs( bool UpdateBrushes )
 {
-	Modify();
+	Modify(false);
 
 	if (UpdateBrushes)
 	{
@@ -503,14 +503,14 @@ void UModel::ModifyAllSurfs( bool UpdateBrushes )
 
 		for (UModel* Brush: Brushes)
 		{
-			Brush->Modify();
+			Brush->Modify(false);
 		}
 	}
 }
 
 void UModel::ModifySelectedSurfs( bool UpdateBrushes )
 {
-	Modify();
+	Modify(false);
 
 	if (UpdateBrushes)
 	{
@@ -528,7 +528,7 @@ void UModel::ModifySelectedSurfs( bool UpdateBrushes )
 
 		for (UModel* Brush : Brushes)
 		{
-			Brush->Modify();
+			Brush->Modify(false);
 		}
 	}
 }
@@ -621,7 +621,7 @@ IMPLEMENT_INTRINSIC_CLASS(UModel, ENGINE_API, UObject, CORE_API, "/Script/Engine
 ---------------------------------------------------------------------------------------*/
 
 #if WITH_EDITOR
-bool UModel::Modify( bool bAlwaysMarkDirty/*=false*/ )
+bool UModel::Modify( bool bAlwaysMarkDirty/*=true*/ )
 {
 	bool bSavedToTransactionBuffer = Super::Modify(bAlwaysMarkDirty);
 
