-Make HLOD SourceControl helper usable by other WorldPartition builders

diff --git a/Engine/Source/Editor/UnrealEd/Private/WorldPartition/WorldPartitionBuilderHelpers.cpp b/Engine/Source/Editor/UnrealEd/Private/WorldPartition/WorldPartitionBuilderHelpers.cpp
new file mode 100644
--- /dev/null
+++ b/Engine/Source/Editor/UnrealEd/Private/WorldPartition/WorldPartitionBuilderHelpers.cpp
@@ -0,0 +1,159 @@
+// Copyright Epic Games, Inc. All Rights Reserved.
+
+#include "WorldPartition/WorldPartitionBuilderHelpers.h"
+
+#include "CoreMinimal.h"
+#include "HAL/FileManager.h"
+#include "Logging/LogMacros.h"
+#include "WorldPartition/WorldPartitionBuilder.h"
+#include "Misc/FileHelper.h"
+#include "Algo/ForEach.h"
+#include "UObject/SavePackage.h"
+
+#include "ISourceControlModule.h"
+#include "ISourceControlProvider.h"
+#include "SourceControlOperations.h"
+
+
+void FBuilderModifiedFiles::Add(EFileOperation FileOp, const FString& File)
+{
+	Files[FileOp].Add(File);
+}
+
+const TSet<FString>& FBuilderModifiedFiles::Get(EFileOperation FileOp) const
+{
+	return Files[FileOp];
+}
+
+void FBuilderModifiedFiles::Append(EFileOperation FileOp, const TArray<FString>& InFiles)
+{
+	Files[FileOp].Append(InFiles);
+}
+
+void FBuilderModifiedFiles::Append(const FBuilderModifiedFiles& Other)
+{
+	Files[EFileOperation::FileAdded].Append(Other.Files[EFileOperation::FileAdded]);
+	Files[EFileOperation::FileEdited].Append(Other.Files[EFileOperation::FileEdited]);
+	Files[EFileOperation::FileDeleted].Append(Other.Files[EFileOperation::FileDeleted]);
+}
+
+void FBuilderModifiedFiles::Empty()
+{
+	Files[EFileOperation::FileAdded].Empty();
+	Files[EFileOperation::FileEdited].Empty();
+	Files[EFileOperation::FileDeleted].Empty();
+}
+
+TArray<FString> FBuilderModifiedFiles::GetAllFiles() const
+{
+	TArray<FString> AllFiles;
+	AllFiles.Append(Files[EFileOperation::FileAdded].Array());
+	AllFiles.Append(Files[EFileOperation::FileEdited].Array());
+	AllFiles.Append(Files[EFileOperation::FileDeleted].Array());
+	return AllFiles;
+}
+
+
+FSourceControlHelper::FSourceControlHelper(FPackageSourceControlHelper& InPackageHelper, FBuilderModifiedFiles& InModifiedFiles)
+	: PackageHelper(InPackageHelper)
+	, ModifiedFiles(InModifiedFiles)
+{}
+
+FSourceControlHelper::~FSourceControlHelper()
+{}
+
+FString FSourceControlHelper::GetFilename(const FString& PackageName) const 
+{
+	return SourceControlHelpers::PackageFilename(PackageName);
+}
+
+FString FSourceControlHelper::GetFilename(UPackage* Package) const 
+{
+	return SourceControlHelpers::PackageFilename(Package);
+}
+
+bool FSourceControlHelper::Checkout(UPackage* Package) const 
+{
+	bool bCheckedOut = PackageHelper.Checkout(Package);
+	if (bCheckedOut)
+	{
+		const FString Filename = GetFilename(Package);
+		const bool bAdded = ModifiedFiles.Get(FBuilderModifiedFiles::EFileOperation::FileAdded).Contains(Filename);
+		if (!bAdded)
+		{
+			ModifiedFiles.Add(FBuilderModifiedFiles::EFileOperation::FileEdited, Filename);
+		}		
+	}
+	return bCheckedOut;
+}
+
+bool FSourceControlHelper::Add(UPackage* Package) const 
+{
+	bool bAdded = PackageHelper.AddToSourceControl(Package);
+	if (bAdded)
+	{
+		ModifiedFiles.Add(FBuilderModifiedFiles::EFileOperation::FileAdded, GetFilename(Package));
+	}
+	return bAdded;
+}
+
+bool FSourceControlHelper::Delete(const FString& PackageName) const 
+{
+	bool bDeleted = PackageHelper.Delete(PackageName);
+	if (bDeleted)
+	{
+		ModifiedFiles.Add(FBuilderModifiedFiles::EFileOperation::FileDeleted, PackageName);
+	}
+	return bDeleted;
+}
+
+bool FSourceControlHelper::Delete(UPackage* Package) const 
+{
+	FString PackageName = GetFilename(Package);
+	bool bDeleted = PackageHelper.Delete(Package);
+	if (bDeleted)
+	{
+		ModifiedFiles.Add(FBuilderModifiedFiles::EFileOperation::FileDeleted, PackageName);
+	}
+	return bDeleted;
+}
+
+bool FSourceControlHelper::Save(UPackage* Package) const 
+{
+	bool bFileExists = IPlatformFile::GetPlatformPhysical().FileExists(*GetFilename(Package));
+
+	// Checkout package
+	Package->MarkAsFullyLoaded();
+
+	if (bFileExists)
+	{
+		if (!Checkout(Package))
+		{
+			UE_LOG(LogWorldPartitionBuilderSourceControlHelper, Error, TEXT("Error checking out package %s."), *Package->GetName());
+			return false;
+		}
+	}
+
+	// Save package
+	FString PackageFileName = GetFilename(Package);
+	FSavePackageArgs SaveArgs;
+	SaveArgs.TopLevelFlags = RF_Standalone;
+	SaveArgs.SaveFlags = PackageHelper.UseSourceControl() ? ESaveFlags::SAVE_None : ESaveFlags::SAVE_Async;
+	if (!UPackage::SavePackage(Package, nullptr, *PackageFileName, SaveArgs))
+	{
+		UE_LOG(LogWorldPartitionBuilderSourceControlHelper, Error, TEXT("Error saving package %s."), *Package->GetName());
+		return false;
+	}
+
+	// Add new package to source control
+	if (!bFileExists)
+	{
+		if (!Add(Package))
+		{
+			UE_LOG(LogWorldPartitionBuilderSourceControlHelper, Error, TEXT("Error adding package %s to revision control."), *Package->GetName());
+			return false;
+		}
+	}
+
+	return true;
+}
\ No newline at end of file
diff --git a/Engine/Source/Editor/UnrealEd/Private/WorldPartition/WorldPartitionHLODsBuilder.cpp b/Engine/Source/Editor/UnrealEd/Private/WorldPartition/WorldPartitionHLODsBuilder.cpp
--- a/Engine/Source/Editor/UnrealEd/Private/WorldPartition/WorldPartitionHLODsBuilder.cpp
+++ b/Engine/Source/Editor/UnrealEd/Private/WorldPartition/WorldPartitionHLODsBuilder.cpp
@@ -35,118 +35,6 @@
 
 DEFINE_LOG_CATEGORY_STATIC(LogWorldPartitionHLODsBuilder, All, All);
 
-class FSourceControlHelper : public ISourceControlHelper
-{
-public:
-	FSourceControlHelper(FPackageSourceControlHelper& InPackageHelper, FHLODModifiedFiles& InModifiedFiles)
-		: PackageHelper(InPackageHelper)
-		, ModifiedFiles(InModifiedFiles)
-	{}
-
-	virtual ~FSourceControlHelper()
-	{}
-
-	virtual FString GetFilename(const FString& PackageName) const override
-	{
-		return SourceControlHelpers::PackageFilename(PackageName);
-	}
-
-	virtual FString GetFilename(UPackage* Package) const override
-	{
-		return SourceControlHelpers::PackageFilename(Package);
-	}
-
-	virtual bool Checkout(UPackage* Package) const override
-	{
-		bool bCheckedOut = PackageHelper.Checkout(Package);
-		if (bCheckedOut)
-		{
-			const FString Filename = GetFilename(Package);
-			const bool bAdded = ModifiedFiles.Get(FHLODModifiedFiles::EFileOperation::FileAdded).Contains(Filename);
-			if (!bAdded)
-			{
-				ModifiedFiles.Add(FHLODModifiedFiles::EFileOperation::FileEdited, Filename);
-			}
-		}
-		return bCheckedOut;
-	}
-
-	virtual bool Add(UPackage* Package) const override
-	{
-		bool bAdded = PackageHelper.AddToSourceControl(Package);
-		if (bAdded)
-		{
-			ModifiedFiles.Add(FHLODModifiedFiles::EFileOperation::FileAdded, GetFilename(Package));
-		}
-		return bAdded;
-	}
-
-	virtual bool Delete(const FString& PackageName) const override
-	{
-		bool bDeleted = PackageHelper.Delete(PackageName);
-		if (bDeleted)
-		{
-			ModifiedFiles.Add(FHLODModifiedFiles::EFileOperation::FileDeleted, PackageName);
-		}
-		return bDeleted;
-	}
-
-	virtual bool Delete(UPackage* Package) const override
-	{
-		FString PackageName = GetFilename(Package);
-		bool bDeleted = PackageHelper.Delete(Package);
-		if (bDeleted)
-		{
-			ModifiedFiles.Add(FHLODModifiedFiles::EFileOperation::FileDeleted, PackageName);
-		}
-		return bDeleted;
-	}
-
-	virtual bool Save(UPackage* Package) const override
-	{
-		bool bFileExists = IPlatformFile::GetPlatformPhysical().FileExists(*GetFilename(Package));
-
-		// Checkout package
-		Package->MarkAsFullyLoaded();
-
-		if (bFileExists)
-		{
-			if (!Checkout(Package))
-			{
-				UE_LOG(LogWorldPartitionHLODsBuilder, Error, TEXT("Error checking out package %s."), *Package->GetName());
-				return false;
-			}
-		}
-
-		// Save package
-		FString PackageFileName = GetFilename(Package);
-		FSavePackageArgs SaveArgs;
-		SaveArgs.TopLevelFlags = RF_Standalone;
-		SaveArgs.SaveFlags = PackageHelper.UseSourceControl() ? ESaveFlags::SAVE_None : ESaveFlags::SAVE_Async;
-		if (!UPackage::SavePackage(Package, nullptr, *PackageFileName, SaveArgs))
-		{
-			UE_LOG(LogWorldPartitionHLODsBuilder, Error, TEXT("Error saving package %s."), *Package->GetName());
-			return false;
-		}
-
-		// Add new package to source control
-		if (!bFileExists)
-		{
-			if (!Add(Package))
-			{
-				UE_LOG(LogWorldPartitionHLODsBuilder, Error, TEXT("Error adding package %s to revision control."), *Package->GetName());
-				return false;
-			}
-		}
-
-		return true;
-	}
-
-private:
-	FPackageSourceControlHelper& PackageHelper;
-	FHLODModifiedFiles& ModifiedFiles;
-};
-
 static const FString DistributedBuildWorkingDirName = TEXT("HLODTemp");
 static const FString DistributedBuildManifestName = TEXT("HLODBuildManifest.ini");
 static const FString BuildProductsFileName = TEXT("BuildProducts.txt");
@@ -278,8 +166,6 @@ bool UWorldPartitionHLODsBuilder::ShouldProcessWorld(UWorld* InWorld) const
 
 bool UWorldPartitionHLODsBuilder::PreWorldInitialization(UWorld* InWorld, FPackageSourceControlHelper& PackageHelper)
 {
-	ModifiedFiles.Empty();
-
 	if (bDistributedBuild)
 	{
 		DistributedBuildWorkingDir = GetDistributedBuildWorkingDir(InWorld);
@@ -427,12 +313,12 @@ bool UWorldPartitionHLODsBuilder::SetupHLODActors()
 				FWorldPartitionHelpers::DoCollectGarbage();
 			}
 
-			TArray<FHLODModifiedFiles> BuildersFiles;
+			TArray<FBuilderModifiedFiles> BuildersFiles;
 			BuildersFiles.SetNum(BuilderCount);
 
-			for (int32 i = 0; i < FHLODModifiedFiles::EFileOperation::NumFileOperations; i++)
+			for (int32 i = 0; i < FBuilderModifiedFiles::EFileOperation::NumFileOperations; i++)
 			{
-				FHLODModifiedFiles::EFileOperation FileOp = (FHLODModifiedFiles::EFileOperation)i;
+				FBuilderModifiedFiles::EFileOperation FileOp = (FBuilderModifiedFiles::EFileOperation)i;
 				for (const FString& ModifiedFile : ModifiedFiles.Get(FileOp))
 				{
 					int32* Idx = FilesToBuilderMap.Find(ModifiedFile);
@@ -941,7 +827,7 @@ const FName FileAction_Add(TEXT("Add"));
 const FName FileAction_Edit(TEXT("Edit"));
 const FName FileAction_Delete(TEXT("Delete"));
 
-bool UWorldPartitionHLODsBuilder::CopyFilesToWorkingDir(const FString& TargetDir, const FHLODModifiedFiles& Files, TArray<FString>& BuildProducts)
+bool UWorldPartitionHLODsBuilder::CopyFilesToWorkingDir(const FString& TargetDir, const FBuilderModifiedFiles& Files, TArray<FString>& BuildProducts)
 {
 	const FString AbsoluteTargetDir = DistributedBuildWorkingDir / TargetDir / TEXT("");
 
@@ -978,9 +864,9 @@ bool UWorldPartitionHLODsBuilder::CopyFilesToWorkingDir(const FString& TargetDir
 		}
 	};
 
-	Algo::ForEach(Files.Get(FHLODModifiedFiles::EFileOperation::FileAdded), [&](const FString& SourceFilename) { CopyFileToWorkingDir(SourceFilename, FileAction_Add); });
-	Algo::ForEach(Files.Get(FHLODModifiedFiles::EFileOperation::FileEdited), [&](const FString& SourceFilename) { CopyFileToWorkingDir(SourceFilename, FileAction_Edit); });
-	Algo::ForEach(Files.Get(FHLODModifiedFiles::EFileOperation::FileDeleted), [&](const FString& SourceFilename) { CopyFileToWorkingDir(SourceFilename, FileAction_Delete); });
+	Algo::ForEach(Files.Get(FBuilderModifiedFiles::EFileOperation::FileAdded), [&](const FString& SourceFilename) { CopyFileToWorkingDir(SourceFilename, FileAction_Add); });
+	Algo::ForEach(Files.Get(FBuilderModifiedFiles::EFileOperation::FileEdited), [&](const FString& SourceFilename) { CopyFileToWorkingDir(SourceFilename, FileAction_Edit); });
+	Algo::ForEach(Files.Get(FBuilderModifiedFiles::EFileOperation::FileDeleted), [&](const FString& SourceFilename) { CopyFileToWorkingDir(SourceFilename, FileAction_Delete); });
 	if (!bSuccess)
 	{
 		return false;
@@ -998,7 +884,7 @@ bool UWorldPartitionHLODsBuilder::CopyFilesToWorkingDir(const FString& TargetDir
 	}
 
 	// Delete files we added
-	for (const FString& FileToDelete : Files.Get(FHLODModifiedFiles::EFileOperation::FileAdded))
+	for (const FString& FileToDelete : Files.Get(FBuilderModifiedFiles::EFileOperation::FileAdded))
 	{
 		if (!IFileManager::Get().Delete(*FileToDelete, false, true))
 		{
@@ -1152,9 +1038,9 @@ bool UWorldPartitionHLODsBuilder::CopyFilesFromWorkingDir(const FString& SourceD
 	}
 
 	// Keep track of all modified files
-	ModifiedFiles.Append(FHLODModifiedFiles::EFileOperation::FileAdded, ToAdd);
-	ModifiedFiles.Append(FHLODModifiedFiles::EFileOperation::FileDeleted, FilesToDelete);
-	ModifiedFiles.Append(FHLODModifiedFiles::EFileOperation::FileEdited, ToEdit);
+	ModifiedFiles.Append(FBuilderModifiedFiles::EFileOperation::FileAdded, ToAdd);
+	ModifiedFiles.Append(FBuilderModifiedFiles::EFileOperation::FileDeleted, FilesToDelete);
+	ModifiedFiles.Append(FBuilderModifiedFiles::EFileOperation::FileEdited, ToEdit);
 
 	// Force a rescan of the updated files
 	IAssetRegistry& AssetRegistry = FModuleManager::LoadModuleChecked<FAssetRegistryModule>(TEXT("AssetRegistry")).Get();
diff --git a/Engine/Source/Editor/UnrealEd/Public/WorldPartition/WorldPartitionBuilderHelpers.h b/Engine/Source/Editor/UnrealEd/Public/WorldPartition/WorldPartitionBuilderHelpers.h
new file mode 100644
--- /dev/null
+++ b/Engine/Source/Editor/UnrealEd/Public/WorldPartition/WorldPartitionBuilderHelpers.h
@@ -0,0 +1,49 @@
+// Copyright Epic Games, Inc. All Rights Reserved.
+
+#pragma once
+
+#include "SourceControlHelpers.h"
+#include "PackageSourceControlHelper.h" 
+#include "WorldPartition/WorldPartition.h" // for ISourceControlHelper
+
+DEFINE_LOG_CATEGORY_STATIC(LogWorldPartitionBuilderSourceControlHelper, All, All);
+
+struct FBuilderModifiedFiles
+{
+	enum EFileOperation
+	{
+		FileAdded,
+		FileEdited,
+		FileDeleted,
+		NumFileOperations
+	};
+
+	void Add(EFileOperation FileOp, const FString& File);
+	const TSet<FString>& Get(EFileOperation FileOp) const;
+	void Append(EFileOperation FileOp, const TArray<FString>& InFiles);
+	void Append(const FBuilderModifiedFiles& Other);
+	void Empty();
+	TArray<FString> GetAllFiles() const;
+
+private:
+	TSet<FString> Files[NumFileOperations];
+};
+
+class FSourceControlHelper : public ISourceControlHelper
+{
+public:
+	FSourceControlHelper(FPackageSourceControlHelper& InPackageHelper, FBuilderModifiedFiles& InModifiedFiles);
+
+	virtual ~FSourceControlHelper();
+	virtual FString GetFilename(const FString& PackageName) const override;
+	virtual FString GetFilename(UPackage* Package) const override;
+	virtual bool Checkout(UPackage* Package) const override;
+	virtual bool Add(UPackage* Package) const override;
+	virtual bool Delete(const FString& PackageName) const override;
+	virtual bool Delete(UPackage* Package) const override;
+	virtual bool Save(UPackage* Package) const override;
+
+private:
+	FPackageSourceControlHelper& PackageHelper;
+	FBuilderModifiedFiles& ModifiedFiles;
+};
diff --git a/Engine/Source/Editor/UnrealEd/Public/WorldPartition/WorldPartitionHLODsBuilder.h b/Engine/Source/Editor/UnrealEd/Public/WorldPartition/WorldPartitionHLODsBuilder.h
--- a/Engine/Source/Editor/UnrealEd/Public/WorldPartition/WorldPartitionHLODsBuilder.h
+++ b/Engine/Source/Editor/UnrealEd/Public/WorldPartition/WorldPartitionHLODsBuilder.h
@@ -4,66 +4,9 @@
 
 #include "Misc/EnumClassFlags.h"
 #include "WorldPartition/WorldPartitionBuilder.h"
+#include "WorldPartition/WorldPartitionBuilderHelpers.h"
 #include "WorldPartitionHLODsBuilder.generated.h"
 
-
-class UWorld;
-class UWorldPartition;
-class FSourceControlHelper;
-
-
-struct FHLODModifiedFiles
-{
-	enum EFileOperation
-	{
-		FileAdded,
-		FileEdited,
-		FileDeleted,
-		NumFileOperations
-	};
-
-	void Add(EFileOperation FileOp, const FString& File)
-	{
-		Files[FileOp].Add(File);
-	}
-
-	const TSet<FString>& Get(EFileOperation FileOp) const
-	{
-		return Files[FileOp];
-	}
-
-	void Append(EFileOperation FileOp, const TArray<FString>& InFiles)
-	{
-		Files[FileOp].Append(InFiles);
-	}
-
-	void Append(const FHLODModifiedFiles& Other)
-	{
-		Files[EFileOperation::FileAdded].Append(Other.Files[EFileOperation::FileAdded]);
-		Files[EFileOperation::FileEdited].Append(Other.Files[EFileOperation::FileEdited]);
-		Files[EFileOperation::FileDeleted].Append(Other.Files[EFileOperation::FileDeleted]);
-	}
-
-	void Empty()
-	{
-		Files[EFileOperation::FileAdded].Empty();
-		Files[EFileOperation::FileEdited].Empty();
-		Files[EFileOperation::FileDeleted].Empty();
-	}
-
-	TArray<FString> GetAllFiles() const
-	{
-		TArray<FString> AllFiles;
-		AllFiles.Append(Files[EFileOperation::FileAdded].Array());
-		AllFiles.Append(Files[EFileOperation::FileEdited].Array());
-		AllFiles.Append(Files[EFileOperation::FileDeleted].Array());
-		return AllFiles;
-	}
-
-private:
-	TSet<FString> Files[NumFileOperations];
-};
-
 enum class EHLODBuildStep : uint8
 {
 	None		= 0,
@@ -107,7 +50,7 @@ protected:
 	UNREALED_API TArray<TArray<FGuid>> GetHLODWorkloads(int32 NumWorkloads) const;
 	UNREALED_API bool ValidateWorkload(const TArray<FGuid>& Workload) const;
 
-	UNREALED_API bool CopyFilesToWorkingDir(const FString& TargetDir, const FHLODModifiedFiles& ModifiedFiles, TArray<FString>& BuildProducts);
+	UNREALED_API bool CopyFilesToWorkingDir(const FString& TargetDir, const FBuilderModifiedFiles& ModifiedFiles, TArray<FString>& BuildProducts);
 	UNREALED_API bool CopyFilesFromWorkingDir(const FString& SourceDir);
 
 	UNREALED_API bool ShouldRunStep(const EHLODBuildStep BuildStep) const;
@@ -116,8 +59,8 @@ protected:
 
 private:
 	UWorld* World;
-	UWorldPartition* WorldPartition;
-	FSourceControlHelper* SourceControlHelper;
+	class UWorldPartition* WorldPartition;
+	class FSourceControlHelper* SourceControlHelper;
 
 	// Options --
 	EHLODBuildStep BuildOptions;
@@ -136,5 +79,5 @@ private:
 	FString DistributedBuildWorkingDir;
 	FString DistributedBuildManifest;
 	
-	FHLODModifiedFiles ModifiedFiles;
-};
+	FBuilderModifiedFiles ModifiedFiles;
+};
\ No newline at end of file
