Use cached transforms for Pilot transactions to allow for other transactions to occur while piloting.

diff --git a/Engine/Source/Editor/UnrealEd/Private/LevelEditorViewport.cpp b/Engine/Source/Editor/UnrealEd/Private/LevelEditorViewport.cpp
--- a/Engine/Source/Editor/UnrealEd/Private/LevelEditorViewport.cpp
+++ b/Engine/Source/Editor/UnrealEd/Private/LevelEditorViewport.cpp
@@ -2367,6 +2367,7 @@ FLevelEditorViewportClient::FLevelEditorViewportClient(const TSharedPtr<SLevelVi
 	, SpriteCategoryVisibility()
 	, World(nullptr)
 	, TrackingTransaction()
+	, CachedPilotTransform(TOptional<FTransform>())
 	, DropPreviewMouseX(0)
 	, DropPreviewMouseY(0)
 	, bWasControlledByOtherViewport(false)
@@ -3587,22 +3588,6 @@ void FLevelEditorViewportClient::TrackingStarted( const FInputEventState& InInpu
 	// Begin transacting.  Give the current editor mode an opportunity to do the transacting.
 	const bool bTrackingHandledExternally = ModeTools->StartTracking(this, Viewport);
 
-	// If we are have a lock due to piloting, don't modify the tracking transaction until the piloting is stopped in 'SetActorLock'
-	AActor* ActiveActorLock = GetActiveActorLock().Get();
-	if (ActiveActorLock && PilotingTransaction.IsActive() && PilotingTransaction.TransCount > 0)
-	{
-		PreDragActorTransforms.Empty();
-		PreDragElementTransforms.Empty();
-
-		if (PilotingTransaction.IsActive() || PilotingTransaction.IsPending())
-		{
-			// Suspend actor/component modification during each delta step to avoid recording unnecessary overhead into the transaction buffer
-			GEditor->DisableDeltaModification(true);
-		}
-
-		return;
-	}
-
 	TrackingTransaction.End();
 
 	// Re-initialize new tracking only if a new button was pressed, otherwise we continue the previous one.
@@ -3712,11 +3697,11 @@ void FLevelEditorViewportClient::TrackingStarted( const FInputEventState& InInpu
 		}
 		else
 		{
-			if (ActiveActorLock && !ActiveActorLock->IsLockLocation() && TrackingTransaction.TransCount == 0 && PilotingTransaction.TransCount == 0)
+			AActor* ActiveActorLock = GetActiveActorLock().Get();
+			if (ActiveActorLock && !ActiveActorLock->IsLockLocation() && TrackingTransaction.TransCount == 0 && !CachedPilotTransform.IsSet())
 			{
-				// Open a tracking transaction to contain the locked actor changes
-				PilotingTransaction.TransCount++;
-				PilotingTransaction.Begin(LOCTEXT("PilotTransaction", "Pilot Actor"));
+				// Cache Pilot Actor transform before piloting ends
+				CachedPilotTransform = ActiveActorLock->GetTransform();
 			}
 		}
 
@@ -3735,16 +3720,6 @@ void FLevelEditorViewportClient::TrackingStopped()
 	// Only disable the duplicate on next drag flag if we actually dragged the mouse.
 	bDuplicateOnNextDrag = false;
 
-	// If we are have a lock due to piloting, don't stop the transaction until the piloting is stopped in 'SetActorLock'
-	AActor* ActiveActorLock = GetActiveActorLock().Get();
-	if (ActiveActorLock && PilotingTransaction.IsActive() && PilotingTransaction.TransCount > 0)
-	{
-		PreDragActorTransforms.Empty();
-		PreDragElementTransforms.Empty();
-
-		return;
-	}
-
 	// here we check to see if anything of worth actually changed when ending our MouseMovement
 	// If the TransCount > 0 (we changed something of value) so we need to call PostEditMove() on stuff
 	// if we didn't change anything then don't call PostEditMove()
@@ -4411,12 +4386,17 @@ void FLevelEditorViewportClient::SetActorLock(AActor* Actor)
 	if (!Actor)
 	{
 		AActor* ActiveActorLock = GetActiveActorLock().Get();
-		if (ActiveActorLock && !ActiveActorLock->IsLockLocation() && PilotingTransaction.TransCount > 0)
+		if (ActiveActorLock && !ActiveActorLock->IsLockLocation() && CachedPilotTransform.IsSet())
 		{
-			PilotingTransaction.TransCount--;
-			PilotingTransaction.End();
+			FTransform EndPosition = ActiveActorLock->GetTransform();
+			ActiveActorLock->SetActorTransform(CachedPilotTransform.GetValue());
+			{
+				const FScopedTransaction Transaction(LOCTEXT("PilotTransaction", "Pilot Actor"));
+				ActiveActorLock->Modify();
+				ActiveActorLock->SetActorTransform(EndPosition);
+			}
 
-			ModeTools->ActorMoveNotify();
+			CachedPilotTransform = TOptional<FTransform>();
 		}
 	}
 	SetActorLock(FLevelViewportActorLock(Actor));
diff --git a/Engine/Source/Editor/UnrealEd/Public/LevelEditorViewport.h b/Engine/Source/Editor/UnrealEd/Public/LevelEditorViewport.h
--- a/Engine/Source/Editor/UnrealEd/Public/LevelEditorViewport.h
+++ b/Engine/Source/Editor/UnrealEd/Public/LevelEditorViewport.h
@@ -964,8 +964,8 @@ private:
 	/** Global shared transaction for all mouse interactions. */
 	FTrackingTransaction TrackingTransaction;
 
-	/** Transaction excluse to pilot operations. Seperate as it externally handles & prevents tracking transactions */
-	FTrackingTransaction PilotingTransaction;
+	/** Cached transform of pilot actor before transaction, used to allow other transactions while piloting by performing a single end-transaction */
+	TOptional<FTransform> CachedPilotTransform;
 
 	/** Represents the last known drop preview mouse position. */
 	int32 DropPreviewMouseX;
