[Skeletal mesh] submitting a file missed from a previous change adding "Build Half Edge" to skeletal mesh build settings

diff --git a/Engine/Source/Editor/UnrealEd/Private/ImportUtils/SkeletalMeshImportUtils.cpp b/Engine/Source/Editor/UnrealEd/Private/ImportUtils/SkeletalMeshImportUtils.cpp
--- a/Engine/Source/Editor/UnrealEd/Private/ImportUtils/SkeletalMeshImportUtils.cpp
+++ b/Engine/Source/Editor/UnrealEd/Private/ImportUtils/SkeletalMeshImportUtils.cpp
@@ -480,6 +480,7 @@ void SkeletalMesUtilsImpl::RestoreLODInfo(const TSharedPtr<const FExistingSkelMe
 	ImportedLODInfo.SkinCacheUsage = ExistingLODInfo.SkinCacheUsage;
 	ImportedLODInfo.MorphTargetPositionErrorTolerance = ExistingLODInfo.MorphTargetPositionErrorTolerance;
 	ImportedLODInfo.bAllowCPUAccess = ExistingLODInfo.bAllowCPUAccess;
+	ImportedLODInfo.bBuildHalfEdgeBuffers = ExistingLODInfo.bBuildHalfEdgeBuffers;
 	ImportedLODInfo.bSupportUniformlyDistributedSampling = ExistingLODInfo.bSupportUniformlyDistributedSampling;
 }
 
