Modal dialog base class

diff --git a/Engine/Source/Editor/UnrealEd/Private/Dialogs/Dialogs.cpp b/Engine/Source/Editor/UnrealEd/Private/Dialogs/Dialogs.cpp
--- a/Engine/Source/Editor/UnrealEd/Private/Dialogs/Dialogs.cpp
+++ b/Engine/Source/Editor/UnrealEd/Private/Dialogs/Dialogs.cpp
@@ -1111,4 +1111,33 @@ FReply SGenericDialogWidget::OnOK_Clicked(void)
 	return FReply::Handled();
 }
 
+TSharedRef<SWindow> UE::Private::CreateModalDialogWindow(
+	const FText& InTitle,
+	TSharedRef<SWidget> Contents,
+	ESizingRule Sizing,
+	FVector2D MinDimensions)
+{
+	return SNew(SWindow)
+		.Title(InTitle)
+		.SizingRule(Sizing)
+		.MinWidth(MinDimensions.X)
+		.MinHeight(MinDimensions.Y)
+		.SupportsMaximize(false)
+		.SupportsMinimize(false)
+		.HasCloseButton(false)
+		[
+			SNew(SBorder)
+			.Padding(4.f)
+			.BorderImage(FAppStyle::GetBrush("ToolPanel.GroupBorder"))
+			[
+				MoveTemp(Contents)
+			]
+		];
+}
+
+void UE::Private::ShowModalDialogWindow(TSharedRef<SWindow> Window)
+{
+	GEditor->EditorAddModalWindow(Window);
+}
+
 #undef LOCTEXT_NAMESPACE 
diff --git a/Engine/Source/Editor/UnrealEd/Public/Dialogs/Dialogs.h b/Engine/Source/Editor/UnrealEd/Public/Dialogs/Dialogs.h
--- a/Engine/Source/Editor/UnrealEd/Public/Dialogs/Dialogs.h
+++ b/Engine/Source/Editor/UnrealEd/Public/Dialogs/Dialogs.h
@@ -5,8 +5,9 @@
 #include "CoreMinimal.h"
 #include "Input/Reply.h"
 #include "Widgets/DeclarativeSyntaxSupport.h"
-#include "Widgets/SWidget.h"
+#include "Widgets/Layout/SBorder.h"
 #include "Widgets/SCompoundWidget.h"
+#include "Widgets/SWidget.h"
 #include "Widgets/SWindow.h"
 
 class SModalDialogWithCheckbox;
@@ -250,6 +251,60 @@ private:
 	FSimpleDelegate OkPressedDelegate;
 };
 
+namespace UE::Private
+{
+UNREALED_API TSharedRef<SWindow> CreateModalDialogWindow(const FText& InTitle, TSharedRef<SWidget> Contents, ESizingRule Sizing, FVector2D MinDimensions);
+UNREALED_API void ShowModalDialogWindow(TSharedRef<SWindow> Window);
+} // namespace UE::Private
+
+/**
+ * Base class for a dialog which can be shown modally and returns a user's selection after it is closed.
+ */
+template<typename ResultType>
+class SModalEditorDialog : public SCompoundWidget
+{
+public:
+	SLATE_BEGIN_ARGS(SGenericDialogWidget)
+	{
+	}
+	SLATE_END_ARGS()
+
+	ResultType ShowModalDialog(const FText& InTitle)
+	{
+		static_assert(std::is_default_constructible_v<ResultType>, "ResultType must be default constructable");
+		Window = UE::Private::CreateModalDialogWindow(InTitle, AsShared(), Sizing, MinDimensions);
+		Window->SetWidgetToFocusOnActivate(GetWidgetToFocusOnActivate());
+		ResultType Result;
+		ResultPointer = &Result;
+		UE::Private::ShowModalDialogWindow(Window.ToSharedRef());
+		Window.Reset();
+		ResultPointer = nullptr;
+		return MoveTemp(Result);
+	}
+
+protected:
+	// Derived classes call this function from their widget events to close the dialog and return the result to the calling context
+	void ProvideResult(ResultType InResult)
+	{
+		// Close owning window and move result into space where ShowDialog can return it
+		*ResultPointer = MoveTemp(InResult);
+		Window->RequestDestroyWindow();
+	}
+	
+	virtual TSharedPtr<SWidget> GetWidgetToFocusOnActivate() 
+	{
+		return {};
+	}
+
+	// Child classes can modify these
+	ESizingRule Sizing = ESizingRule::Autosized;
+	FVector2D MinDimensions = FVector2D(400.0f, 300.f);
+
+private:
+	TSharedPtr<SWindow> Window;
+	ResultType* ResultPointer = nullptr;
+};
+
 UE_DEPRECATED(4.26, "Creating groups (nested packages) is no longer supported. Use PromptUserIfExistingObject overload that does not take the Group paramater.")
 UNREALED_API bool PromptUserIfExistingObject(const FString& Name, const FString& Package, const FString& Group, class UPackage* &Pkg );
 
