Interchange: Fix alternate skinning import to use the proper build setting so the number of vertex match the imported mesh LOD.

#rnx

diff --git a/Engine/Source/Editor/UnrealEd/Private/SkinWeightsUtilities.cpp b/Engine/Source/Editor/UnrealEd/Private/SkinWeightsUtilities.cpp
--- a/Engine/Source/Editor/UnrealEd/Private/SkinWeightsUtilities.cpp
+++ b/Engine/Source/Editor/UnrealEd/Private/SkinWeightsUtilities.cpp
@@ -47,13 +47,18 @@ DEFINE_LOG_CATEGORY_STATIC(LogSkinWeightsUtilities, Log, All);
 bool FSkinWeightsUtilities::ImportAlternateSkinWeight(USkeletalMesh* SkeletalMesh, const FString& Path, int32 TargetLODIndex, const FName& ProfileName, const bool bIsReimport)
 {
 	check(SkeletalMesh);
-	check(SkeletalMesh->GetLODInfo(TargetLODIndex));
 	FSkeletalMeshLODInfo* LODInfo = SkeletalMesh->GetLODInfo(TargetLODIndex);
-	
-	if (LODInfo && LODInfo->bHasBeenSimplified && LODInfo->ReductionSettings.BaseLOD != TargetLODIndex)
+
+	if (!LODInfo)
+	{
+		UE_LOG(LogSkinWeightsUtilities, Error, TEXT("Cannot import Skin Weight Profile. No valid LOD info."));
+		return false;
+	}
+
+	if (LODInfo->bHasBeenSimplified && LODInfo->ReductionSettings.BaseLOD != TargetLODIndex)
 	{
 		//We cannot import alternate skin weights profile for a generated LOD
-		UE_LOG(LogSkinWeightsUtilities, Error, TEXT("Cannot import Skin Weight Profile for a generated LOD."));
+		UE_LOG(LogSkinWeightsUtilities, Error, TEXT("Cannot import Skin Weight Profile for a generated LOD. Skin weight profile are transfert from the source LOD during the simplification process."));
 		return false;
 	}
 
@@ -101,10 +106,8 @@ bool FSkinWeightsUtilities::ImportAlternateSkinWeight(USkeletalMesh* SkeletalMes
 	DeletePathAssets();
 
 	UObject* ImportedObject = nullptr;
-	UnFbx::FBXImportOptions ImportOptions;
 
 	bool bCreateTransaction = false;
-
 	//Only use interchange if the base skeletal mesh was imported with interchange
 	const UInterchangeAssetImportData* SelectedInterchangeAssetImportData = Cast<UInterchangeAssetImportData>(SkeletalMesh->GetAssetImportData());
 	if (bUseInterchangeFramework && SelectedInterchangeAssetImportData)
@@ -116,8 +119,27 @@ bool FSkinWeightsUtilities::ImportAlternateSkinWeight(USkeletalMesh* SkeletalMes
 
 		if (const UClass* GenericPipelineClass = InterchangeProjectSettings->GenericPipelineClass.LoadSynchronous())
 		{
-			if (UInterchangePipelineBase* GenericPipeline = NewObject<UInterchangePipelineBase>(GetTransientPackage(), GenericPipelineClass))
+			UInterchangePipelineBase* GenericPipeline = nullptr;
+			for (UObject* PipelineObject : SelectedInterchangeAssetImportData->GetPipelines())
 			{
+				if (PipelineObject->GetClass()->IsChildOf(GenericPipelineClass))
+				{
+					if (UInterchangePipelineBase* ImportPipeline = Cast<UInterchangePipelineBase>(PipelineObject))
+					{
+						GenericPipeline = Cast<UInterchangePipelineBase>(StaticDuplicateObject(ImportPipeline, GetTransientPackage()));
+						break;
+					}
+				}
+			}
+			
+			if (!GenericPipeline)
+			{
+				GenericPipeline = NewObject<UInterchangePipelineBase>(GetTransientPackage(), GenericPipelineClass);
+			}
+
+			if (GenericPipeline)
+			{
+				GenericPipeline->ClearFlags(EObjectFlags::RF_Standalone | EObjectFlags::RF_Public);
 				GenericPipeline->AdjustSettingsForContext(bIsReimport ? EInterchangePipelineContext::AssetAlternateSkinningReimport : EInterchangePipelineContext::AssetAlternateSkinningImport, nullptr);
 				ImportAssetParameters.OverridePipelines.Add(GenericPipeline);
 			}
@@ -178,8 +200,6 @@ bool FSkinWeightsUtilities::ImportAlternateSkinWeight(USkeletalMesh* SkeletalMes
 			FbxFactory->ImportUI->TextureImportData->BaseMaterialName.Reset();
 		}
 
-		ApplyImportUIToImportOptions(FbxFactory->ImportUI, ImportOptions);
-
 		TArray<FString> ImportFilePaths;
 		ImportFilePaths.Add(AbsoluteFilePath);
 
@@ -234,7 +254,7 @@ bool FSkinWeightsUtilities::ImportAlternateSkinWeight(USkeletalMesh* SkeletalMes
 			{
 				//Prepare the profile data
 				FSkeletalMeshLODModel& TargetLODModel = TargetModel->LODModels[TargetLODIndex];
-				
+
 				// Prepare the profile data
 				FSkinWeightProfileInfo* Profile = SkeletalMesh->GetSkinWeightProfiles().FindByPredicate([ProfileName](FSkinWeightProfileInfo Profile) { return Profile.Name == ProfileName; });
 
@@ -263,11 +283,15 @@ bool FSkinWeightsUtilities::ImportAlternateSkinWeight(USkeletalMesh* SkeletalMes
 				FImportedSkinWeightProfileData PreviousProfileData = ProfileData;
 				
 				IMeshUtilities::MeshBuildOptions BuildOptions;
-				BuildOptions.OverlappingThresholds = ImportOptions.OverlappingThresholds;
-				BuildOptions.bComputeNormals = !ImportOptions.ShouldImportNormals();
-				BuildOptions.bComputeTangents = !ImportOptions.ShouldImportTangents();
-				BuildOptions.bUseMikkTSpace = ImportOptions.NormalGenerationMethod == EFBXNormalGenerationMethod::MikkTSpace;
-				BuildOptions.bComputeWeightedNormals = ImportOptions.bComputeWeightedNormals;
+				//Use the Lod info build settings
+				BuildOptions.OverlappingThresholds.ThresholdPosition = LODInfo->BuildSettings.ThresholdPosition;
+				BuildOptions.OverlappingThresholds.ThresholdTangentNormal = LODInfo->BuildSettings.ThresholdTangentNormal;
+				BuildOptions.OverlappingThresholds.ThresholdUV = LODInfo->BuildSettings.ThresholdUV;
+				BuildOptions.OverlappingThresholds.MorphThresholdPosition = LODInfo->BuildSettings.MorphThresholdPosition;
+				BuildOptions.bComputeNormals = LODInfo->BuildSettings.bRecomputeNormals;
+				BuildOptions.bComputeTangents = LODInfo->BuildSettings.bRecomputeTangents;
+				BuildOptions.bUseMikkTSpace = LODInfo->BuildSettings.bUseMikkTSpace;
+				BuildOptions.bComputeWeightedNormals = LODInfo->BuildSettings.bComputeWeightedNormals;
 				// There's currently no import option for this. We could add one if needed.
 				BuildOptions.BoneInfluenceLimit = 0;
 
@@ -370,8 +394,24 @@ bool FSkinWeightsUtilities::RemoveSkinnedWeightProfileData(USkeletalMesh* Skelet
 	SkeletalMesh->LoadLODImportedData(LODIndex, ImportDataDest);
 	PRAGMA_ENABLE_DEPRECATION_WARNINGS
 
-	//Rechunk the skeletal mesh since we remove it, we rebuild the skeletal mesh to achieve rechunking
-	UFbxSkeletalMeshImportData* OriginalSkeletalMeshImportData = UFbxSkeletalMeshImportData::GetImportDataForSkeletalMesh(SkeletalMesh, nullptr);
+	//If we have a LOD info we use the build settings to be sure we are rechunking the LOD with the existing options
+	IMeshUtilities::MeshBuildOptions BuildOptions;
+	if (FSkeletalMeshLODInfo* LODInfo = SkeletalMesh->GetLODInfo(LODIndex))
+	{
+		BuildOptions.OverlappingThresholds.ThresholdPosition = LODInfo->BuildSettings.ThresholdPosition;
+		BuildOptions.OverlappingThresholds.ThresholdTangentNormal = LODInfo->BuildSettings.ThresholdTangentNormal;
+		BuildOptions.OverlappingThresholds.ThresholdUV = LODInfo->BuildSettings.ThresholdUV;
+		BuildOptions.OverlappingThresholds.MorphThresholdPosition = LODInfo->BuildSettings.MorphThresholdPosition;
+		BuildOptions.bComputeNormals = LODInfo->BuildSettings.bRecomputeNormals;
+		BuildOptions.bComputeTangents = LODInfo->BuildSettings.bRecomputeTangents;
+		BuildOptions.bUseMikkTSpace = LODInfo->BuildSettings.bUseMikkTSpace;
+		BuildOptions.bComputeWeightedNormals = LODInfo->BuildSettings.bComputeWeightedNormals;
+	}
+
+	BuildOptions.bRemoveDegenerateTriangles = false;
+	BuildOptions.TargetPlatform = GetTargetPlatformManagerRef().GetRunningTargetPlatform();
+	// There's currently no import option for this. We could add one if needed.
+	BuildOptions.BoneInfluenceLimit = 0;
 
 	TArray<FVector3f> LODPointsDest;
 	TArray<SkeletalMeshImportData::FMeshWedge> LODWedgesDest;
@@ -380,23 +420,6 @@ bool FSkinWeightsUtilities::RemoveSkinnedWeightProfileData(USkeletalMesh* Skelet
 	TArray<int32> LODPointToRawMapDest;
 	ImportDataDest.CopyLODImportData(LODPointsDest, LODWedgesDest, LODFacesDest, LODInfluencesDest, LODPointToRawMapDest);
 
-	const bool bShouldImportNormals = OriginalSkeletalMeshImportData->NormalImportMethod == FBXNIM_ImportNormals || OriginalSkeletalMeshImportData->NormalImportMethod == FBXNIM_ImportNormalsAndTangents;
-	const bool bShouldImportTangents = OriginalSkeletalMeshImportData->NormalImportMethod == FBXNIM_ImportNormalsAndTangents;
-	//Set the options with the current asset build options
-	IMeshUtilities::MeshBuildOptions BuildOptions;
-	BuildOptions.OverlappingThresholds.ThresholdPosition = OriginalSkeletalMeshImportData->ThresholdPosition;
-	BuildOptions.OverlappingThresholds.ThresholdTangentNormal = OriginalSkeletalMeshImportData->ThresholdTangentNormal;
-	BuildOptions.OverlappingThresholds.ThresholdUV = OriginalSkeletalMeshImportData->ThresholdUV;
-	BuildOptions.OverlappingThresholds.MorphThresholdPosition = OriginalSkeletalMeshImportData->MorphThresholdPosition;
-	BuildOptions.bComputeNormals = !bShouldImportNormals || !ImportDataDest.bHasNormals;
-	BuildOptions.bComputeTangents = !bShouldImportTangents || !ImportDataDest.bHasTangents;
-	BuildOptions.bUseMikkTSpace = (OriginalSkeletalMeshImportData->NormalGenerationMethod == EFBXNormalGenerationMethod::MikkTSpace) && (!bShouldImportNormals || !bShouldImportTangents);
-	BuildOptions.bComputeWeightedNormals = OriginalSkeletalMeshImportData->bComputeWeightedNormals;
-	BuildOptions.bRemoveDegenerateTriangles = false;
-	BuildOptions.TargetPlatform = GetTargetPlatformManagerRef().GetRunningTargetPlatform();
-	// There's currently no import option for this. We could add one if needed.
-	BuildOptions.BoneInfluenceLimit = 0;
-
 	//Build the skeletal mesh asset
 	IMeshUtilities& MeshUtilities = FModuleManager::Get().LoadModuleChecked<IMeshUtilities>("MeshUtilities");
 	TArray<FText> WarningMessages;
