[Material Editor] Material Instance Editor now automatically disable overrides of CurveAtlas parameters when the new parameter in the new parent references a different Atlas texture (as the original curve would no longer make sense).

diff --git a/Engine/Source/Editor/UnrealEd/Classes/MaterialEditor/MaterialEditorInstanceConstant.h b/Engine/Source/Editor/UnrealEd/Classes/MaterialEditor/MaterialEditorInstanceConstant.h
--- a/Engine/Source/Editor/UnrealEd/Classes/MaterialEditor/MaterialEditorInstanceConstant.h
+++ b/Engine/Source/Editor/UnrealEd/Classes/MaterialEditor/MaterialEditorInstanceConstant.h
@@ -291,12 +291,18 @@ class UMaterialEditorInstanceConstant : public UObject
 
 	/** Regenerates the parameter arrays. */
 	UNREALED_API void RegenerateArrays();
+
 #if WITH_EDITOR
 	/** Sets back to zero the overrides for any parameters copied out of the layer stack */
 	UNREALED_API void CleanParameterStack(int32 Index, EMaterialParameterAssociation MaterialType);
+
 	/** Copies the overrides for any parameters copied out of the layer stack from the layer or blend */
 	UNREALED_API void ResetOverrides(int32 Index, EMaterialParameterAssociation MaterialType);
+
+	/** Arrays and clears parameters no longer valid (e.g. curve atlases). It has the potential effect of regenerating parameter arrays. */
+	UNREALED_API void ClearInvalidParameterOverrides();
 #endif
+
 	/** Copies the parameter array values back to the source instance. */
 	UNREALED_API void CopyToSourceInstance(const bool bForceStaticPermutationUpdate = false);
 
diff --git a/Engine/Source/Editor/UnrealEd/Private/PreviewMaterial.cpp b/Engine/Source/Editor/UnrealEd/Private/PreviewMaterial.cpp
--- a/Engine/Source/Editor/UnrealEd/Private/PreviewMaterial.cpp
+++ b/Engine/Source/Editor/UnrealEd/Private/PreviewMaterial.cpp
@@ -17,6 +17,7 @@
 #include "MaterialEditor/MaterialEditorPreviewParameters.h"
 #include "MaterialEditor/MaterialEditorMeshComponent.h"
 #include "MaterialEditorModule.h"
+#include "MaterialCachedData.h"
 #include "Materials/MaterialInstance.h"
 #include "Materials/MaterialInstanceConstant.h"
 #include "Materials/MaterialFunctionInstance.h"
@@ -501,6 +502,8 @@ void UMaterialEditorInstanceConstant::PostEditChangeProperty(FPropertyChangedEve
 
 				// Fully update static parameters before recreating render state for all components
 				SetSourceInstance(SourceInstance);
+				
+				ClearInvalidParameterOverrides();
 			}
 		}
 		else if (!bIsFunctionPreviewMaterial)
@@ -719,6 +722,7 @@ void UMaterialEditorInstanceConstant::CleanParameterStack(int32 Index, EMaterial
 	ParameterGroups = CleanedGroups;
 	CopyToSourceInstance(true);
 }
+
 void UMaterialEditorInstanceConstant::ResetOverrides(int32 Index, EMaterialParameterAssociation MaterialType)
 {
 	check(GIsEditor);
@@ -747,6 +751,43 @@ void UMaterialEditorInstanceConstant::ResetOverrides(int32 Index, EMaterialParam
 	CopyToSourceInstance(true);
 
 }
+
+void UMaterialEditorInstanceConstant::ClearInvalidParameterOverrides()
+{
+	const FMaterialCachedExpressionData& CachedExpressionData = Parent->GetCachedExpressionData();
+
+	// Look for all Atlas Scalar parameters in each parameter group, then if a parameter has
+	// an override, disable it unless the atlas texture matches that originally set in the parent Material. 
+	for (int32 GroupIdx = 0; GroupIdx < ParameterGroups.Num(); GroupIdx++)
+	{
+		FEditorParameterGroup& Group = ParameterGroups[GroupIdx];
+		for (int32 ParameterIdx = 0; ParameterIdx < Group.Parameters.Num(); ParameterIdx++)
+		{
+			UDEditorParameterValue* Parameter = Group.Parameters[ParameterIdx];
+			if (UDEditorScalarParameterValue* ScalarParameter = Cast<UDEditorScalarParameterValue>(Parameter))
+			{
+				// Ignore parameters without override.
+				if (!Parameter->bOverride)
+				{
+					continue;
+				}
+
+				// Get parent's parameter atlas texture. If identical to the one the editor parameter is using,
+				// we can keep the override on (as the selected atlas curve will still make sense).
+				FMaterialParameterMetadata Value;
+				if (CachedExpressionData.GetParameterValue(EMaterialParameterType::Scalar, ScalarParameter->ParameterInfo, Value)
+					&& Value.ScalarAtlas == ScalarParameter->AtlasData.Atlas)
+				{
+					continue;
+				}
+
+				// The atlas texture of the newly bound material are different. Disable the override.
+				Parameter->bOverride = false;
+			}
+		}
+	}
+}
+
 #endif
 
 void UMaterialEditorInstanceConstant::CopyToSourceInstance(const bool bForceStaticPermutationUpdate)
diff --git a/Engine/Source/Runtime/Engine/Public/MaterialCachedData.h b/Engine/Source/Runtime/Engine/Public/MaterialCachedData.h
--- a/Engine/Source/Runtime/Engine/Public/MaterialCachedData.h
+++ b/Engine/Source/Runtime/Engine/Public/MaterialCachedData.h
@@ -233,7 +233,7 @@ struct FMaterialCachedExpressionData
 
 	inline int32 GetNumParameters(EMaterialParameterType Type) const { return GetParameterTypeEntry(Type).ParameterInfoSet.Num(); }
 	int32 FindParameterIndex(EMaterialParameterType Type, const FMemoryImageMaterialParameterInfo& HashedParameterInfo) const;
-	bool GetParameterValue(EMaterialParameterType Type, const FMemoryImageMaterialParameterInfo& ParameterInfo, FMaterialParameterMetadata& OutResult) const;
+	ENGINE_API bool GetParameterValue(EMaterialParameterType Type, const FMemoryImageMaterialParameterInfo& ParameterInfo, FMaterialParameterMetadata& OutResult) const;
 	void GetParameterValueByIndex(EMaterialParameterType Type, int32 ParameterIndex, FMaterialParameterMetadata& OutResult) const;
 	const FGuid& GetExpressionGuid(EMaterialParameterType Type, int32 Index) const;
 	void GetAllParametersOfType(EMaterialParameterType Type, TMap<FMaterialParameterInfo, FMaterialParameterMetadata>& OutParameters) const;
