Outliner: Pasting / Duplicating an actor now moves the new one to the "Current Folder" (if it's set) in World Partitioned levels too (follow up to CL 30537280).

diff --git a/Engine/Source/Editor/UnrealEd/Private/Factories/EditorFactories.cpp b/Engine/Source/Editor/UnrealEd/Private/Factories/EditorFactories.cpp
--- a/Engine/Source/Editor/UnrealEd/Private/Factories/EditorFactories.cpp
+++ b/Engine/Source/Editor/UnrealEd/Private/Factories/EditorFactories.cpp
@@ -1007,10 +1007,17 @@ UObject* ULevelFactory::FactoryCreateText
 									bGrouped = true;
 									FParse::Value(Str, TEXT("GroupFolder="), GroupFolder);
 								}
-								FString ActorFolderPath;
-								if (FParse::Value(Str, TEXT("ActorFolderPath="), ActorFolderPath))
+
+								// Don't overwrite if its already been set during actor creation
+								// (i.e. by UActorEditorContextSubsystem::ApplyContext).
+								// It mirrors the behavior of the task "Importing Actors" below.
+								if (NewActor->GetFolder().IsNone())
 								{
-									NewActor->SetFolderPath(*ActorFolderPath);
+									FString ActorFolderPath;
+									if (FParse::Value(Str, TEXT("ActorFolderPath="), ActorFolderPath))
+									{
+										NewActor->SetFolderPath(*ActorFolderPath);
+									}
 								}
 
 								uint32 CopyPasteId;
