Catch infinite recursion in PatchCDOSubobjectsIntoExport.
Temp workaround for GDC. This is caused by verse classes incorrectly marking all object properties as subobjects.

diff --git a/Engine/Source/Editor/UnrealEd/Private/Kismet2/BlueprintEditorUtils.cpp b/Engine/Source/Editor/UnrealEd/Private/Kismet2/BlueprintEditorUtils.cpp
--- a/Engine/Source/Editor/UnrealEd/Private/Kismet2/BlueprintEditorUtils.cpp
+++ b/Engine/Source/Editor/UnrealEd/Private/Kismet2/BlueprintEditorUtils.cpp
@@ -1413,7 +1413,7 @@ void FBlueprintEditorUtils::PatchCDOSubobjectsIntoExport(UObject* PreviousCDO, U
 	{
 		struct PatchCDOSubobjectsIntoExport_Impl
 		{
-			static void PatchSubObjects(UObject* OldObj, UObject* NewObj)
+			static void PatchSubObjects(UObject* OldObj, UObject* NewObj, TSet<UObject*>& AlreadyPatched)
 			{
 				TArray<UObject*> OldSubObjects;
 				GetObjectsWithOuter(OldObj, OldSubObjects, /*bIncludeNestedSubObjects =*/false);
@@ -1446,8 +1446,13 @@ void FBlueprintEditorUtils::PatchCDOSubobjectsIntoExport(UObject* PreviousCDO, U
 							{
 								FLinkerLoad::PRIVATE_PatchNewObjectIntoExport(OldSubObj, NewSubObj);
 
-								// Recursively find and patch any instances nested within the current subobject.
-								PatchSubObjects(OldSubObj, NewSubObj);
+								bool bAlreadyPatched;
+								AlreadyPatched.Add(OldSubObj, &bAlreadyPatched);
+								if (!bAlreadyPatched)
+								{
+									// Recursively find and patch any instances nested within the current subobject.
+									PatchSubObjects(OldSubObj, NewSubObj, AlreadyPatched);
+								}
 
 								// Track the old instanced reference so we don't attempt to patch it again below.
 								PatchedAsInstancedReferenceSet.Add(OldSubObj);
@@ -1488,13 +1493,15 @@ void FBlueprintEditorUtils::PatchCDOSubobjectsIntoExport(UObject* PreviousCDO, U
 								FLinkerLoad::PRIVATE_PatchNewObjectIntoExport(OldSubObj, NewSubObj);
 							}
 
-							PatchSubObjects(OldSubObj, NewSubObj);
+							PatchSubObjects(OldSubObj, NewSubObj, AlreadyPatched);
 						}
 					}
 				}
 			}
 		};
-		PatchCDOSubobjectsIntoExport_Impl::PatchSubObjects(PreviousCDO, NewCDO);
+
+		TSet<UObject*> AlreadyPatched; 
+		PatchCDOSubobjectsIntoExport_Impl::PatchSubObjects(PreviousCDO, NewCDO, AlreadyPatched);
 		NewCDO->CheckDefaultSubobjects();
 	}
 }
