Fix memory leak when the level editor viewport is set to realtime but it's not visible.

The code was skipping the explicit render resource flush when the level editing viewport was set to realtime, assuming that realtime meant the editor is redrawn each frame. However, if another editor is docked and active in the main window, the level editor will not be drawn, and render resources accumulated because nothing was calling FlushPendingDeletes(). This change checks if any level editors are refreshed during the frame, and pumps the deletion queue if none are.

diff --git a/Engine/Source/Editor/UnrealEd/Private/EditorEngine.cpp b/Engine/Source/Editor/UnrealEd/Private/EditorEngine.cpp
--- a/Engine/Source/Editor/UnrealEd/Private/EditorEngine.cpp
+++ b/Engine/Source/Editor/UnrealEd/Private/EditorEngine.cpp
@@ -2170,7 +2170,7 @@ void UEditorEngine::Tick( float DeltaSeconds, bool bIdleMode )
 
 	// Do not redraw if the application is hidden
 	const bool bAllWindowsHidden = !bHasFocus && AreAllWindowsHidden();
-	bool bEditorFrameNonRealtimeViewportDrawn = false;
+	bool bAnyLevelEditorsDrawn = false;
 	if (!bAllWindowsHidden || bRunDrawWithEditorHidden)
 	{
 		FPixelInspectorModule& PixelInspectorModule = FModuleManager::LoadModuleChecked<FPixelInspectorModule>(TEXT("PixelInspectorModule"));
@@ -2180,16 +2180,18 @@ void UEditorEngine::Tick( float DeltaSeconds, bool bIdleMode )
 		}
 
 		// Render view parents, then view children.
+		bool bEditorFrameNonRealtimeViewportDrawn = false;
 		if (GCurrentLevelEditingViewportClient && GCurrentLevelEditingViewportClient->IsVisible())
 		{
 			if (!bAllWindowsHidden || GCurrentLevelEditingViewportClient->WantsDrawWhenAppIsHidden())
 			{
 				bool bAllowNonRealtimeViewports = true;
 				GCurrentLevelEditingViewportClient->SetIsCurrentLevelEditingFocus(true);
-				bool bWasNonRealtimeViewportDraw = UpdateSingleViewportClient(GCurrentLevelEditingViewportClient, bAllowNonRealtimeViewports, bUpdateLinkedOrthoViewports);
+				bool bWasNonRealtimeViewportDrawn = UpdateSingleViewportClient(GCurrentLevelEditingViewportClient, bAllowNonRealtimeViewports, bUpdateLinkedOrthoViewports);
 				if (GCurrentLevelEditingViewportClient->IsLevelEditorClient())
 				{
-					bEditorFrameNonRealtimeViewportDrawn |= bWasNonRealtimeViewportDraw;
+					bEditorFrameNonRealtimeViewportDrawn |= bWasNonRealtimeViewportDrawn;
+					bAnyLevelEditorsDrawn |= GCurrentLevelEditingViewportClient->IsRealtime() || bWasNonRealtimeViewportDrawn;
 				}
 			}
 		}
@@ -2217,6 +2219,7 @@ void UEditorEngine::Tick( float DeltaSeconds, bool bIdleMode )
 						if (ViewportClient->IsLevelEditorClient())
 						{
 							bEditorFrameNonRealtimeViewportDrawn |= bWasNonRealtimeViewportDrawn;
+							bAnyLevelEditorsDrawn |= ViewportClient->IsRealtime() || bWasNonRealtimeViewportDrawn;
 						}
 					}
 				}
@@ -2224,8 +2227,12 @@ void UEditorEngine::Tick( float DeltaSeconds, bool bIdleMode )
 		}
 	}
 
-	// If we're not Realtime and no NonRealtime viewports are drawn, or if everything is hidden, make sure RHI gets its flush to prevent memory from accumulating
-	if ((bAllWindowsHidden || !IsRealtime) && !bEditorFrameNonRealtimeViewportDrawn && !bHasPIEViewport && IsRunningRHIInSeparateThread())
+	// Rendering resources are normally flushed when a 3D viewport is drawn. If no viewports are updated (because the editor is hidden, or no realtime viewports are visible),
+	// we need to force-flush resources here, since nothing else will. Note that this condition checks if any *level* viewports have been drawn in the block above; other
+	// editors can contain 3D viewports and refreshing them will also flush resources, but it's difficult to know when 3D rendering has happened in general. Instead, we
+	// optimize for the level editor case, since that's performance-sensitive. If there's no level editor, but there are other 3D editors, this will do an unnecessary
+	// flush, but the small performance impact in that case should not affect users.
+	if (!bAnyLevelEditorsDrawn && !bHasPIEViewport && IsRunningRHIInSeparateThread())
 	{
 		ENQUEUE_RENDER_COMMAND(FlushPendingDeleteRHIResources_NonRealtime)(
 			[](FRHICommandListImmediate& RHICmdList)
