FStaticLightingBuildContext & UMapBuildDataRegistry access changes

-Added FStaticLightingBuildContext to regroup all code pertaining to UMapBuildDataRegistry, LightingScenario levels & Level GUIDs during static lighting computation into a single class
-Added static methods to UMapBuildDataRegistry to also regroup all code pertaining to UMapBuildDataRegistry access during resource initialization
-This change is a preparation change for supporting static lighting in WorldPartition maps which introduces different access patterns, but shouldn't introduce any different behavior.

diff --git a/Engine/Plugins/Experimental/GPULightmass/Source/GPULightmass/Private/Scene/Scene.cpp b/Engine/Plugins/Experimental/GPULightmass/Source/GPULightmass/Private/Scene/Scene.cpp
--- a/Engine/Plugins/Experimental/GPULightmass/Source/GPULightmass/Private/Scene/Scene.cpp
+++ b/Engine/Plugins/Experimental/GPULightmass/Source/GPULightmass/Private/Scene/Scene.cpp
@@ -31,6 +31,7 @@
 #include "ReflectionEnvironment.h"
 #include "RHIStaticStates.h"
 #include "InstancedStaticMeshSceneProxyDesc.h"
+#include "StaticLightingBuildContext.h"
 
 #define LOCTEXT_NAMESPACE "StaticLightingSystem"
 
@@ -2991,8 +2992,11 @@ void FScene::ApplyFinishedLightmapsToWorld()
 
 		GCompressLightmaps = Settings->bCompressLightmaps;
 
-		FLightMap2D::EncodeTextures(World, LightingScenario, true, true);
-		FShadowMap2D::EncodeTextures(World, LightingScenario, true, true);
+		// Would be better to acquire this from the IStaticLightingSystem interface
+		FStaticLightingBuildContext LightingContext (World, LightingScenario);
+
+		FLightMap2D::EncodeTextures(&LightingContext, true, true);
+		FShadowMap2D::EncodeTextures(&LightingContext, true, true);
 
 		SlowTask.EnterProgressFrame(1, LOCTEXT("ApplyingNewLighting", "Applying new lighting"));
 
diff --git a/Engine/Source/Editor/UnrealEd/Private/Lightmass/ImportVolumetricLightmap.cpp b/Engine/Source/Editor/UnrealEd/Private/Lightmass/ImportVolumetricLightmap.cpp
--- a/Engine/Source/Editor/UnrealEd/Private/Lightmass/ImportVolumetricLightmap.cpp
+++ b/Engine/Source/Editor/UnrealEd/Private/Lightmass/ImportVolumetricLightmap.cpp
@@ -1058,8 +1058,7 @@ void FLightmassProcessor::ImportVolumetricLightmap()
 		}
 	}
 
-	ULevel* StorageLevel = System.LightingScenario ? System.LightingScenario : ToRawPtr(System.GetWorld()->PersistentLevel);
-	UMapBuildDataRegistry* Registry = StorageLevel->GetOrCreateMapBuildData();
+	UMapBuildDataRegistry* Registry = System.LightingContext.GetOrCreateGlobalRegistry();
 	FPrecomputedVolumetricLightmapData CurrentLevelData;
 
 	{
@@ -1330,9 +1329,8 @@ void FLightmassProcessor::ImportVolumetricLightmap()
 			SubLevelBrickLayoutDimension.Z = FMath::Min(SubLevelBrickTextureLinearAllocator, MaxBricksInLayoutOneDim);
 
 			{
-				ULevel* SubLevelStorageLevel = System.LightingScenario ? System.LightingScenario : FindLevel(Pair.Key);
-				UMapBuildDataRegistry* SubLevelRegistry = SubLevelStorageLevel->GetOrCreateMapBuildData();
-				FPrecomputedVolumetricLightmapData& SubLevelData = SubLevelRegistry->AllocateLevelPrecomputedVolumetricLightmapBuildData(FindLevel(Pair.Key)->LevelBuildDataId);
+				UMapBuildDataRegistry* SubLevelRegistry = System.LightingContext.GetOrCreateRegistryForLevelGuid(Pair.Key);
+				FPrecomputedVolumetricLightmapData& SubLevelData = SubLevelRegistry->AllocateLevelPrecomputedVolumetricLightmapBuildData(System.LightingContext.GetLevelBuildDataID(Pair.Key));
 
 				SubLevelData.InitializeOnImport(FBox(VolumetricLightmapSettings.VolumeMin, VolumetricLightmapSettings.VolumeMin + VolumetricLightmapSettings.VolumeSize), BrickSize);
 				{
@@ -1390,9 +1388,8 @@ void FLightmassProcessor::ImportVolumetricLightmap()
 					LevelGuid = PersistentLevelGuid;
 				}
 
-				ULevel* SubLevelStorageLevel = System.LightingScenario ? System.LightingScenario : FindLevel(LevelGuid);
-				UMapBuildDataRegistry* SubLevelRegistry = SubLevelStorageLevel->GetOrCreateMapBuildData();
-				FPrecomputedVolumetricLightmapData& SubLevelData = *SubLevelRegistry->GetLevelPrecomputedVolumetricLightmapBuildData(FindLevel(LevelGuid)->LevelBuildDataId);
+				UMapBuildDataRegistry* SubLevelRegistry = System.LightingContext.GetOrCreateRegistryForLevelGuid(LevelGuid);
+				FPrecomputedVolumetricLightmapData& SubLevelData = *SubLevelRegistry->GetLevelPrecomputedVolumetricLightmapBuildData(System.LightingContext.GetLevelBuildDataID(LevelGuid));
 
 				const FIntVector BrickLayoutPosition = ComputeBrickLayoutPosition(BrickStartAllocation + BrickIndex, BrickLayoutDimensions) * PaddedBrickSize;
 				const FIntVector SubLevelBrickLayoutPosition = ComputeBrickLayoutPosition(SubLevelBrickAllocator[LevelGuid], SubLevelBrickLayoutDimensions[LevelGuid]) * PaddedBrickSize;
@@ -1460,9 +1457,8 @@ void FLightmassProcessor::ImportVolumetricLightmap()
 			CurrentLevelData.IndirectionTexture.Resize(TotalIndirectionTextureSize * IndirectionTextureDataStride);
 			BuildIndirectionTexture(BricksByDepth, VolumetricLightmapSettings, MaxBricksInLayoutOneDim, SubLevelBrickLayoutDimensions[FGuid()], IndirectionTextureDataStride, CurrentLevelData, true);
 
-			ULevel* SubLevelStorageLevel = System.LightingScenario ? System.LightingScenario : FindLevel(PersistentLevelGuid);
-			UMapBuildDataRegistry* SubLevelRegistry = SubLevelStorageLevel->GetOrCreateMapBuildData();
-			FPrecomputedVolumetricLightmapData& PersistentLevelData = *SubLevelRegistry->GetLevelPrecomputedVolumetricLightmapBuildData(FindLevel(PersistentLevelGuid)->LevelBuildDataId);
+			UMapBuildDataRegistry* SubLevelRegistry = System.LightingContext.GetOrCreateRegistryForLevelGuid(PersistentLevelGuid);
+			FPrecomputedVolumetricLightmapData& PersistentLevelData = *SubLevelRegistry->GetLevelPrecomputedVolumetricLightmapBuildData(System.LightingContext.GetLevelBuildDataID(PersistentLevelGuid));
 
 			PersistentLevelData.InitializeOnImport(FBox(VolumetricLightmapSettings.VolumeMin, VolumetricLightmapSettings.VolumeMin + VolumetricLightmapSettings.VolumeSize), BrickSize);
 			PersistentLevelData.IndirectionTexture = CurrentLevelData.IndirectionTexture;
@@ -1492,9 +1488,8 @@ void FLightmassProcessor::ImportVolumetricLightmap()
 
 					if (LevelGuid != PersistentLevelGuid)
 					{
-						ULevel* SubLevelStorageLevel = System.LightingScenario ? System.LightingScenario : FindLevel(LevelGuid);
-						UMapBuildDataRegistry* SubLevelRegistry = SubLevelStorageLevel->GetOrCreateMapBuildData();
-						FPrecomputedVolumetricLightmapData& SubLevelData = *SubLevelRegistry->GetLevelPrecomputedVolumetricLightmapBuildData(FindLevel(LevelGuid)->LevelBuildDataId);
+						UMapBuildDataRegistry* SubLevelRegistry = System.LightingContext.GetOrCreateRegistryForLevelGuid(LevelGuid);
+						FPrecomputedVolumetricLightmapData& SubLevelData = *SubLevelRegistry->GetLevelPrecomputedVolumetricLightmapBuildData(System.LightingContext.GetLevelBuildDataID(LevelGuid));
 
 						FIntVector IndirectionBrickOffset;
 						int32 IndirectionBrickSize;
@@ -1513,9 +1508,8 @@ void FLightmassProcessor::ImportVolumetricLightmap()
 			{
 				FGuid LevelGuid = Pair.Key;
 
-				ULevel* SubLevelStorageLevel = System.LightingScenario ? System.LightingScenario : FindLevel(LevelGuid);
-				UMapBuildDataRegistry* SubLevelRegistry = SubLevelStorageLevel->GetOrCreateMapBuildData();
-				FPrecomputedVolumetricLightmapData& SubLevelData = *SubLevelRegistry->GetLevelPrecomputedVolumetricLightmapBuildData(FindLevel(LevelGuid)->LevelBuildDataId);
+				UMapBuildDataRegistry* SubLevelRegistry = System.LightingContext.GetOrCreateRegistryForLevelGuid(LevelGuid);
+				FPrecomputedVolumetricLightmapData& SubLevelData = *SubLevelRegistry->GetLevelPrecomputedVolumetricLightmapBuildData(System.LightingContext.GetLevelBuildDataID(LevelGuid));
 
 				ConvertBGRA8ToRGBA8ForLayer(SubLevelData.BrickData.SkyBentNormal);
 
@@ -1544,9 +1538,8 @@ void FLightmassProcessor::ImportVolumetricLightmap()
 
 		for (auto Pair : SubLevelTotalBricks)
 		{
-			ULevel* SubLevelStorageLevel = System.LightingScenario ? System.LightingScenario : FindLevel(Pair.Key);
-			UMapBuildDataRegistry* SubLevelRegistry = SubLevelStorageLevel->GetOrCreateMapBuildData();
-			FPrecomputedVolumetricLightmapData& SubLevelData = *SubLevelRegistry->GetLevelPrecomputedVolumetricLightmapBuildData(FindLevel(Pair.Key)->LevelBuildDataId);
+			UMapBuildDataRegistry* SubLevelRegistry = System.LightingContext.GetOrCreateRegistryForLevelGuid(Pair.Key);
+			FPrecomputedVolumetricLightmapData& SubLevelData = *SubLevelRegistry->GetLevelPrecomputedVolumetricLightmapBuildData(System.LightingContext.GetLevelBuildDataID(Pair.Key));
 
 			SubLevelData.FinalizeImport();
 
diff --git a/Engine/Source/Editor/UnrealEd/Private/Lightmass/Lightmass.cpp b/Engine/Source/Editor/UnrealEd/Private/Lightmass/Lightmass.cpp
--- a/Engine/Source/Editor/UnrealEd/Private/Lightmass/Lightmass.cpp
+++ b/Engine/Source/Editor/UnrealEd/Private/Lightmass/Lightmass.cpp
@@ -54,6 +54,7 @@
 #include "Modules/ModuleManager.h"
 #include "ImageCoreUtils.h"
 #include "Misc/FileHelper.h"
+#include "StaticLightingBuildContext.h"
 
 extern FSwarmDebugOptions GSwarmDebugOptions;
 
@@ -487,13 +488,14 @@ void FLightmassProcessor::SwarmCallback( NSwarm::FMessage* CallbackMessage, void
 /*-----------------------------------------------------------------------------
 	FLightmassExporter
 -----------------------------------------------------------------------------*/
-FLightmassExporter::FLightmassExporter( UWorld* InWorld )
+FLightmassExporter::FLightmassExporter(const FStaticLightingBuildContext& Context)
 	: Swarm( NSwarm::FSwarmInterface::Get() ) 
 	, SkyAtmosphereComponent(nullptr)
 	, ExportStage(NotRunning)
 	, CurrentAmortizationIndex(0)
 	, OpenedMaterialExportChannels()
-	, World( InWorld )
+	, World( Context.World )
+	, LightingContext(Context)
 {
 	// We must have a valid world
 	check( World );
@@ -1518,14 +1520,13 @@ void FLightmassExporter::WriteBaseMeshInstanceData( int32 Channel, int32 MeshInd
 	MeshInstanceData.NumVertices = Mesh->NumVertices;
 	MeshInstanceData.NumShadingVertices = Mesh->NumShadingVertices;
 	MeshInstanceData.MeshIndex = MeshIndex;
-	MeshInstanceData.LevelGuid = *LevelGuids.FindKey(World->PersistentLevel);
+	MeshInstanceData.LevelGuid = LightingContext.GetPersistentLevelGuid();
 	check(Mesh->Component);
 	bool bFoundLevel = false;
 	AActor* ComponentOwner = Mesh->Component->GetOwner();
 	if (ComponentOwner && ComponentOwner->GetLevel())
 	{
-		ULevel* MeshLevel = Mesh->Component->GetOwner()->GetLevel();
-		MeshInstanceData.LevelGuid = *LevelGuids.FindKey(MeshLevel);
+		MeshInstanceData.LevelGuid = LightingContext.GetLevelGuidForActor(ComponentOwner);
 		bFoundLevel = true;
 	}
 	else if (Mesh->Component->IsA(UModelComponent::StaticClass()))
@@ -1535,7 +1536,7 @@ void FLightmassExporter::WriteBaseMeshInstanceData( int32 Channel, int32 MeshInd
 		{
 			if (ModelComponent->GetModel() == World->GetLevel(LevelIndex)->Model)
 			{
-				MeshInstanceData.LevelGuid = *LevelGuids.FindKey(World->GetLevel(LevelIndex));
+				MeshInstanceData.LevelGuid = LightingContext.GetLevelGuidForLevel(World->GetLevel(LevelIndex));
 				bFoundLevel = true;
 				break;
 			}
@@ -2660,7 +2661,7 @@ FLightmassProcessor::FLightmassProcessor(const FStaticLightingSystem& InSystem,
 	OptionsFolder = IFileManager::Get().ConvertToAbsolutePathForExternalAppForRead(*OptionsFolder);
 	int32 ConnectionHandle = Swarm.OpenConnection( SwarmCallback, this,  LogFlags, *OptionsFolder );
 	bSwarmConnectionIsValid = (ConnectionHandle >= 0);
-	Exporter = new FLightmassExporter( System.GetWorld() );
+	Exporter = new FLightmassExporter( System.GetLightingContext() );
 	check(Exporter);
 	Exporter->bSwarmConnectionIsValid = bSwarmConnectionIsValid;
 
@@ -2768,15 +2769,6 @@ void FLightmassProcessor::InitiateExport()
 	int32 NumCellDistributionBuckets;
 	VERIFYLIGHTMASSINI(GConfig->GetInt(TEXT("DevOptions.PrecomputedVisibility"), TEXT("NumCellDistributionBuckets"), NumCellDistributionBuckets, GLightmassIni));
 	
-	for ( int32 LevelIndex=0; LevelIndex < System.GetWorld()->GetNumLevels(); LevelIndex++ )
-	{
-		ULevel* Level = System.GetWorld()->GetLevel(LevelIndex);
-		FGuid LevelGuid = FGuid(0,0,0,LevelIndex);
-		Exporter->LevelGuids.Add(LevelGuid, Level);
-	}
-	auto FirstGuid = FGuid(0,0,0,0);
-	check(FindLevel(FirstGuid) == System.GetWorld()->PersistentLevel);
-
 	if (System.GetWorld()->GetWorldSettings()->bPrecomputeVisibility)
 	{
 		for (int32 DistributionBucketIndex = 0; DistributionBucketIndex < NumCellDistributionBuckets; DistributionBucketIndex++)
@@ -3491,13 +3483,12 @@ void FLightmassProcessor::ImportVolumeSamples()
 				Swarm.ReadChannel(Channel, &LevelGuid, sizeof(LevelGuid));
 				TArray<Lightmass::FVolumeLightingSampleData> VolumeSamples;
 				ReadArray(Channel, VolumeSamples);
-				ULevel* CurrentLevel = FindLevel(LevelGuid);
+				ULevel* CurrentLevel = System.LightingContext.GetLevelForGuid(LevelGuid).Get();
 
 				// Only build precomputed light for visible streamed levels
 				if (CurrentLevel && CurrentLevel->bIsVisible)
 				{
-					ULevel* CurrentStorageLevel = System.LightingScenario ? System.LightingScenario : CurrentLevel;
-					UMapBuildDataRegistry* CurrentRegistry = CurrentStorageLevel->GetOrCreateMapBuildData();
+					UMapBuildDataRegistry* CurrentRegistry = System.LightingContext.GetOrCreateRegistryForLevel(CurrentLevel);
 					FPrecomputedLightVolumeData& CurrentLevelData = CurrentRegistry->AllocateLevelPrecomputedLightVolumeBuildData(CurrentLevel->LevelBuildDataId);
 
 					FBox3f LevelVolumeBounds(ForceInit);
@@ -3934,7 +3925,7 @@ void FLightmassProcessor::ImportMeshAreaLightData()
 			{
 				Lightmass::FMeshAreaLightData LMCurrentLightData;
 				Swarm.ReadChannel(Channel, &LMCurrentLightData, sizeof(LMCurrentLightData));
-				const ULevel* CurrentLevel = FindLevel(LMCurrentLightData.LevelGuid);
+				const ULevel* CurrentLevel = System.LightingContext.GetLevelForGuid(LMCurrentLightData.LevelGuid).Get();
 				if (CurrentLevel && CurrentLevel->Actors.Num() > 0)
 				{
 					// Find the level that the mesh area light was in
@@ -4180,8 +4171,7 @@ void FLightmassProcessor::ImportStaticShadowDepthMap(ULightComponent* Light)
 	const int32 Channel = Swarm.OpenChannel( *ChannelName, LM_DOMINANTSHADOW_CHANNEL_FLAGS );
 	if (Channel >= 0)
 	{
-		ULevel* CurrentStorageLevel = System.LightingScenario ? System.LightingScenario : Light->GetOwner()->GetLevel();
-		UMapBuildDataRegistry* CurrentRegistry = CurrentStorageLevel->GetOrCreateMapBuildData();
+		UMapBuildDataRegistry* CurrentRegistry = System.LightingContext.GetOrCreateRegistryForActor(Light->GetOwner());
 		FLightComponentMapBuildData& CurrentLightData = CurrentRegistry->FindOrAllocateLightBuildData(Light->LightGuid, true);
 
 		Lightmass::FStaticShadowDepthMapData ShadowMapData;
@@ -4462,8 +4452,8 @@ ULevel* FLightmassProcessor::FindLevel(const FGuid& Guid)
 {
 	if (Exporter)
 	{
-		const TWeakObjectPtr<ULevel>* Level = Exporter->LevelGuids.Find(Guid);
-		return Level && Level->IsValid() ? Level->Get() : NULL;
+		const TWeakObjectPtr<ULevel> Level = Exporter->LightingContext.GetLevelForGuid(Guid);
+		return Level.IsValid() ? Level.Get() : NULL;
 	}
 	return NULL;
 }
diff --git a/Engine/Source/Editor/UnrealEd/Private/Lightmass/Lightmass.h b/Engine/Source/Editor/UnrealEd/Private/Lightmass/Lightmass.h
--- a/Engine/Source/Editor/UnrealEd/Private/Lightmass/Lightmass.h
+++ b/Engine/Source/Editor/UnrealEd/Private/Lightmass/Lightmass.h
@@ -62,7 +62,7 @@ class FLightmassExporter
 {
 public:
 
-	FLightmassExporter( UWorld* InWorld );
+	FLightmassExporter(const FStaticLightingBuildContext& LightingContext);
 	virtual ~FLightmassExporter();
 
 	void SetLevelSettings(FLightmassWorldInfoSettings& InLevelSettings)
@@ -232,8 +232,6 @@ private:
 
 	FString LevelName;
 
-	TMap<FGuid, const TWeakObjectPtr<ULevel>> LevelGuids;
-
 	// lights objects
 	TArray<const class UDirectionalLightComponent*> DirectionalLights;
 	TArray<const class UPointLightComponent*> PointLights;
@@ -272,6 +270,9 @@ private:
 	/** The world we are exporting from */
 	UWorld* World;
 
+	/** The ligthing context we are exporting for, used to acquire some mappings & GUIDs that are contextual  */
+	const FStaticLightingBuildContext&	LightingContext;
+
 	/** Friends */
 	friend class FBSPSurfaceStaticLighting;
 	friend class FStaticMeshStaticLightingMesh;
diff --git a/Engine/Source/Editor/UnrealEd/Private/StaticLightingSystem/StaticLightingPrivate.h b/Engine/Source/Editor/UnrealEd/Private/StaticLightingSystem/StaticLightingPrivate.h
--- a/Engine/Source/Editor/UnrealEd/Private/StaticLightingSystem/StaticLightingPrivate.h
+++ b/Engine/Source/Editor/UnrealEd/Private/StaticLightingSystem/StaticLightingPrivate.h
@@ -14,6 +14,8 @@
 #include "StaticLighting.h"
 #include "LightingBuildOptions.h"
 #include "Templates/UniquePtr.h"
+#include "Engine/MapBuildDataRegistry.h"
+#include "StaticLightingBuildContext.h"
 
 class FCanvas;
 class FLightmassProcessor;
@@ -265,9 +267,9 @@ public:
 	/**
 	 * Initializes this static lighting system, and builds static lighting based on the provided options.
 	 * @param InOptions - The static lighting build options.
-	 * @param InWorld -   The world we wish to build the lighting for
+	 * @param InContext - The context (world, lighting scenario, world subsection, data layers)  we wish to build the lighting for
 	 */
-	FStaticLightingSystem(const FLightingBuildOptions& InOptions, UWorld* InWorld, ULevel* InLightingScenario);
+	FStaticLightingSystem(const FLightingBuildOptions& InOptions, const FStaticLightingBuildContext& InContext);
 	~FStaticLightingSystem();
 
 	bool CheckLightmassExecutableVersion();
@@ -293,6 +295,9 @@ public:
 	/** Get the UWorld this light system was created with */
 	UWorld*	 GetWorld() const;
 
+	/** Get the Lighting Context this light system was created with */
+	const FStaticLightingBuildContext& GetLightingContext() const;
+
 	/** True if the current stage of building is asynchronous (lightmass in flight) */
 	bool IsAsyncBuilding() const;
 
@@ -300,7 +305,7 @@ public:
 
 	bool ShouldOperateOnLevel(ULevel* InLevel) const
 	{
-		return InLevel && (!InLevel->bIsLightingScenario || InLevel == LightingScenario) && InLevel->bIsVisible;
+		return InLevel && LightingContext.ShouldIncludeLevel(InLevel) && InLevel->bIsVisible;
 	}
 
 private:
@@ -428,16 +433,13 @@ private:
 	double ProcessingStartTime;
 	double WaitForUserAcceptStartTime;
 	
-	/** The world this light system was created with */
-	UWorld*	 World;
-
-	/** The lighting scenario that's currently being built, if any.  When valid, any outputs of the lighting build should go into this level's MapBuildData. */
-	ULevel* LightingScenario;
+	/** The lighting context (world, scenario, cells, datalayer) that's currently being built, if any.  */
+	FStaticLightingBuildContext	LightingContext;
 
 	/** The resource guid for all hidden/excluded levels. Used to keep those level data valid. */
 	TSet<FGuid> BuildDataResourcesToKeep;
 
-	/** A handle on the processor that actually interfacets with Lightmass */
+	/** A handle on the processor that actually interfaces with Lightmass */
 	class FLightmassProcessor* LightmassProcessor;
 
 	friend FStaticLightingManager;
diff --git a/Engine/Source/Editor/UnrealEd/Private/StaticLightingSystem/StaticLightingSystem.cpp b/Engine/Source/Editor/UnrealEd/Private/StaticLightingSystem/StaticLightingSystem.cpp
--- a/Engine/Source/Editor/UnrealEd/Private/StaticLightingSystem/StaticLightingSystem.cpp
+++ b/Engine/Source/Editor/UnrealEd/Private/StaticLightingSystem/StaticLightingSystem.cpp
@@ -49,6 +49,7 @@
 #include "EditorModeManager.h"
 #include "EditorModes.h"
 #include "Dialogs/Dialogs.h"
+#include "WorldPartition/WorldPartition.h"
 
 FSwarmDebugOptions GSwarmDebugOptions;
 
@@ -252,9 +253,9 @@ void FStaticLightingManager::SendBuildDoneNotification( bool AutoApplyFailed )
 {
 	FText CompletedText = LOCTEXT("LightBuildDoneMessage", "Lighting build completed");
 
-	if (ActiveStaticLightingSystem != StaticLightingSystems.Last().Get() && ActiveStaticLightingSystem->LightingScenario)
+	if (ActiveStaticLightingSystem != StaticLightingSystems.Last().Get() && ActiveStaticLightingSystem->LightingContext.LightingScenario)
 	{
-		FString PackageName = FPackageName::GetShortName(ActiveStaticLightingSystem->LightingScenario->GetOutermost()->GetName());
+		FString PackageName = FPackageName::GetShortName(ActiveStaticLightingSystem->LightingContext.LightingScenario->GetOutermost()->GetName());
 		CompletedText = FText::Format(LOCTEXT("LightScenarioBuildDoneMessage", "{0} Lighting Scenario completed"), FText::FromString(PackageName));
 	}
 
@@ -300,13 +301,13 @@ void FStaticLightingManager::CreateStaticLightingSystem(const FLightingBuildOpti
 		{
 			if (Level->bIsLightingScenario && Level->bIsVisible)
 			{
-				StaticLightingSystems.Emplace(new FStaticLightingSystem(Options, World, Level));
+				StaticLightingSystems.Emplace(new FStaticLightingSystem(Options, FStaticLightingBuildContext(World, Level)));
 			}
 		}
 
 		if (StaticLightingSystems.Num() == 0)
 		{
-			StaticLightingSystems.Emplace(new FStaticLightingSystem(Options, World, nullptr));
+			StaticLightingSystems.Emplace(new FStaticLightingSystem(Options, FStaticLightingBuildContext(World, nullptr)));
 		}
 
 		ActiveStaticLightingSystem = StaticLightingSystems[0].Get();
@@ -473,14 +474,13 @@ bool FStaticLightingManager::IsLightingBuildCurrentlyExporting() const
 	return ActiveStaticLightingSystem != NULL && ActiveStaticLightingSystem->IsAmortizedExporting();
 }
 
-FStaticLightingSystem::FStaticLightingSystem(const FLightingBuildOptions& InOptions, UWorld* InWorld, ULevel* InLightingScenario)
+FStaticLightingSystem::FStaticLightingSystem(const FLightingBuildOptions& InOptions, const FStaticLightingBuildContext& context)
 	: Options(InOptions)
 	, bBuildCanceled(false)
 	, DeterministicIndex(0)
 	, NextVisibilityId(0)
 	, CurrentBuildStage(FStaticLightingSystem::NotRunning)
-	, World(InWorld)
-	, LightingScenario(InLightingScenario)
+	, LightingContext(context)
 	, LightmassProcessor(NULL)
 {
 }
@@ -560,9 +560,9 @@ bool FStaticLightingSystem::BeginLightmassProcess()
 		{
 			TMap<FGuid, ULevel*> LevelGuids;
 
-			for (int32 LevelIndex = 0; LevelIndex < World->GetNumLevels(); LevelIndex++)
+			for (int32 LevelIndex = 0; LevelIndex < LightingContext.World->GetNumLevels(); LevelIndex++)
 			{
-				ULevel* Level = World->GetLevel(LevelIndex);
+				ULevel* Level = LightingContext.World->GetLevel(LevelIndex);
 
 				if (ShouldOperateOnLevel(Level) && Options.ShouldBuildLightingForLevel(Level))
 				{
@@ -590,16 +590,16 @@ bool FStaticLightingSystem::BeginLightmassProcess()
 		}
 
 		FString SkippedLevels;
-		for ( int32 LevelIndex=0; LevelIndex < World->GetNumLevels(); LevelIndex++ )
+		for ( int32 LevelIndex=0; LevelIndex < LightingContext.World->GetNumLevels(); LevelIndex++ )
 		{
-			ULevel* Level = World->GetLevel(LevelIndex);
+			ULevel* Level = LightingContext.World->GetLevel(LevelIndex);
 
 			if (ShouldOperateOnLevel(Level))
 			{
 				Level->LightmapTotalSize = 0.0f;
 				Level->ShadowmapTotalSize = 0.0f;
 				ULevelStreaming* LevelStreaming = NULL;
-				if ( World->PersistentLevel != Level )
+				if ( LightingContext.World->PersistentLevel != Level )
 				{
 					LevelStreaming = FLevelUtils::FindStreamingLevel( Level );
 				}
@@ -619,7 +619,7 @@ bool FStaticLightingSystem::BeginLightmassProcess()
 			}
 		}
 
-		for (ULevelStreaming* CurStreamingLevel : World->GetStreamingLevels())
+		for (ULevelStreaming* CurStreamingLevel : LightingContext.World->GetStreamingLevels())
 		{
 			if (CurStreamingLevel && CurStreamingLevel->GetLoadedLevel() && !CurStreamingLevel->GetShouldBeVisibleInEditor())
 			{
@@ -647,7 +647,7 @@ bool FStaticLightingSystem::BeginLightmassProcess()
 		}
 
 		const bool bAllowStaticLighting = IsStaticLightingAllowed();
-		bForceNoPrecomputedLighting = World->GetWorldSettings()->bForceNoPrecomputedLighting || !bAllowStaticLighting;
+		bForceNoPrecomputedLighting = LightingContext.World->GetWorldSettings()->bForceNoPrecomputedLighting || !bAllowStaticLighting;
 		GConfig->GetFloat( TEXT("TextureStreaming"), TEXT("MaxLightmapRadius"), GMaxLightmapRadius, GEngineIni );
 		GConfig->GetBool( TEXT("TextureStreaming"), TEXT("AllowStreamingLightmaps"), GAllowStreamingLightmaps, GEngineIni );
 		
@@ -668,7 +668,7 @@ bool FStaticLightingSystem::BeginLightmassProcess()
 		verify(GConfig->GetBool(TEXT("DevOptions.StaticLighting"), TEXT("bRebuildDirtyGeometryForLighting"), bRebuildDirtyGeometryForLighting, GLightmassIni));
 		verify(GConfig->GetBool(TEXT("DevOptions.StaticLighting"), TEXT("bCompressLightmaps"), GCompressLightmaps, GLightmassIni));
 
-		GCompressLightmaps = GCompressLightmaps && World->GetWorldSettings()->LightmassSettings.bCompressLightmaps;
+		GCompressLightmaps = GCompressLightmaps && LightingContext.World->GetWorldSettings()->LightmassSettings.bCompressLightmaps;
 
 		GAllowLightmapPadding = true;
 		FMemory::Memzero(&LightingMeshBounds, sizeof(FBox));
@@ -705,7 +705,7 @@ bool FStaticLightingSystem::BeginLightmassProcess()
 				{
 					ULightComponentBase* const Light = *LightIt;
 					const bool bLightIsInWorld = IsValid(Light->GetOwner()) 
-						&& World->ContainsActor(Light->GetOwner());
+						&& LightingContext.World->ContainsActor(Light->GetOwner());
 
 					if (bLightIsInWorld && ShouldOperateOnLevel(Light->GetOwner()->GetLevel()))
 					{
@@ -835,10 +835,10 @@ void FStaticLightingSystem::InvalidateStaticLighting()
 	FLightmassStatistics::FScopedGather InvalidationScopeStat(LightmassStatistics.InvalidationTime);
 	FGlobalComponentRecreateRenderStateContext Context;
 
-	for( int32 LevelIndex=0; LevelIndex<World->GetNumLevels(); LevelIndex++ )
+	for( int32 LevelIndex=0; LevelIndex < LightingContext.World->GetNumLevels(); LevelIndex++ )
 	{
 		bool bMarkLevelDirty = false;
-		ULevel* Level = World->GetLevel(LevelIndex);
+		ULevel* Level = LightingContext.World->GetLevel(LevelIndex);
 		
 		if (!ShouldOperateOnLevel(Level))
 		{
@@ -855,13 +855,13 @@ void FStaticLightingSystem::InvalidateStaticLighting()
 
 				if (Level->MapBuildData)
 				{
-					Level->MapBuildData->InvalidateStaticLighting(World, false, &BuildDataResourcesToKeep);
+					Level->MapBuildData->InvalidateStaticLighting(LightingContext.World, false, &BuildDataResourcesToKeep);
 				}
 			}
-			if (Level == World->PersistentLevel)
+			if (Level == LightingContext.World->PersistentLevel)
 			{
-				Level->PrecomputedVisibilityHandler.Invalidate(World->Scene);
-				Level->PrecomputedVolumeDistanceField.Invalidate(World->Scene);
+				Level->PrecomputedVisibilityHandler.Invalidate(LightingContext.World->Scene);
+				Level->PrecomputedVolumeDistanceField.Invalidate(LightingContext.World->Scene);
 			}
 
 			// Mark any existing cached lightmap data as transient. This allows the derived data cache to purge it more aggressively.
@@ -938,9 +938,9 @@ void FStaticLightingSystem::GatherStaticLightingInfo(bool bRebuildDirtyGeometryF
 {
 	uint32 ActorsInvalidated = 0;
 	uint32 ActorsToInvalidate = 0;
-	for( int32 LevelIndex=0; LevelIndex<World->GetNumLevels(); LevelIndex++ )
+	for( int32 LevelIndex=0; LevelIndex< LightingContext.World->GetNumLevels(); LevelIndex++ )
 	{
-		ActorsToInvalidate += World->GetLevel(LevelIndex)->Actors.Num();
+		ActorsToInvalidate += LightingContext.World->GetLevel(LevelIndex)->Actors.Num();
 	}
 	const int32 ProgressUpdateFrequency = FMath::Max<int32>(ActorsToInvalidate / 20, 1);
 
@@ -948,11 +948,10 @@ void FStaticLightingSystem::GatherStaticLightingInfo(bool bRebuildDirtyGeometryF
 	
 	bool bObjectsToBuildLightingForFound = false;
 	// Gather static lighting info from actor components.
-	for (int32 LevelIndex = 0; LevelIndex < World->GetNumLevels(); LevelIndex++)
+	for (int32 LevelIndex = 0; LevelIndex < LightingContext.World->GetNumLevels(); LevelIndex++)
 	{
 		TSet<UPackage*> PackagesToDirty;
-
-		ULevel* Level = World->GetLevel(LevelIndex);
+		ULevel* Level = LightingContext.World->GetLevel(LevelIndex);
 
 		if (!ShouldOperateOnLevel(Level))
 		{
@@ -967,7 +966,7 @@ void FStaticLightingSystem::GatherStaticLightingInfo(bool bRebuildDirtyGeometryF
 			{
 				// This will go ahead and clean up lighting on all dirty levels (not just this one)
 				UE_LOG(LogStaticLightingSystem, Warning, TEXT("WARNING: Lighting build automatically rebuilding geometry.") );
-				GEditor->Exec(World, TEXT("MAP REBUILD ALLDIRTYFORLIGHTING"));
+				GEditor->Exec(LightingContext.World, TEXT("MAP REBUILD ALLDIRTYFORLIGHTING"));
 			}
 		}
 
@@ -1203,7 +1202,7 @@ void FStaticLightingSystem::GatherStaticLightingInfo(bool bRebuildDirtyGeometryF
 						Primitive->GetStaticLightingInfo(PrimitiveInfo, PrimitiveRelevantLights, Options);
 						if (PrimitiveInfo.Meshes.Num() > 0 && (Primitive->Mobility == EComponentMobility::Static))
 						{
-							if (World->GetWorldSettings()->bPrecomputeVisibility)
+							if (LightingContext.World->GetWorldSettings()->bPrecomputeVisibility)
 							{
 								// Make sure packages gets dirtied since we are changing the visibility Id of a component in them
 								PackagesToDirty.Add(Primitive->GetPackage());
@@ -1282,13 +1281,13 @@ void FStaticLightingSystem::EncodeTextures(bool bLightingSuccessful)
 		FLightmassStatistics::FScopedGather EncodeStatScope2(LightmassStatistics.EncodingLightmapsTime);
 		// Flush pending shadow-map and light-map encoding.
 		SlowTask.EnterProgressFrame(1, LOCTEXT("EncodingImportedStaticLightMapsStatusMessage", "Encoding imported static light maps."));
-		FLightMap2D::EncodeTextures(World, LightingScenario, bLightingSuccessful, GMultithreadedLightmapEncode ? true : false);
+		FLightMap2D::EncodeTextures(&LightingContext, bLightingSuccessful, GMultithreadedLightmapEncode ? true : false);
 	}
 
 	{
 		FLightmassStatistics::FScopedGather EncodeStatScope2(LightmassStatistics.EncodingShadowMapsTime);
 		SlowTask.EnterProgressFrame(1, LOCTEXT("EncodingImportedStaticShadowMapsStatusMessage", "Encoding imported static shadow maps."));
-		FShadowMap2D::EncodeTextures(World, LightingScenario, bLightingSuccessful, GMultithreadedShadowmapEncode ? true : false);
+		FShadowMap2D::EncodeTextures(&LightingContext, bLightingSuccessful, GMultithreadedShadowmapEncode ? true : false);
 	}
 }
 
@@ -1305,25 +1304,22 @@ void FStaticLightingSystem::ApplyNewLightingData(bool bLightingSuccessful)
 		FLightmassStatistics::FScopedGather FinishStatScope(LightmassStatistics.FinishingTime);
 
 		// Mark lights of the computed level to have valid precomputed lighting.
-		for (int32 LevelIndex = 0; LevelIndex < World->GetNumLevels(); LevelIndex++)
+		for (int32 LevelIndex = 0; LevelIndex < LightingContext.World->GetNumLevels(); LevelIndex++)
 		{
-			ULevel* Level = World->GetLevel(LevelIndex);
+			ULevel* Level = LightingContext.World->GetLevel(LevelIndex);
 
 			if (!ShouldOperateOnLevel(Level))
 			{
 				continue;
 			}
 
-			ULevel* StorageLevel = LightingScenario ? LightingScenario : Level;
-			UMapBuildDataRegistry* Registry = StorageLevel->GetOrCreateMapBuildData();
-			
 			// Notify level about new lighting data
 			Level->OnApplyNewLightingData(bLightingSuccessful);
 
-			if (World->PersistentLevel == Level)
+			if (LightingContext.World->PersistentLevel == Level)
 			{
-				Level->PrecomputedVisibilityHandler.UpdateScene(World->Scene);
-				Level->PrecomputedVolumeDistanceField.UpdateScene(World->Scene);
+				Level->PrecomputedVisibilityHandler.UpdateScene(LightingContext.World->Scene);
+				Level->PrecomputedVolumeDistanceField.UpdateScene(LightingContext.World->Scene);
 			}
 
 			uint32 ActorCount = Level->Actors.Num();
@@ -1334,6 +1330,8 @@ void FStaticLightingSystem::ApplyNewLightingData(bool bLightingSuccessful)
 
 				if (Actor && bLightingSuccessful && !Options.bOnlyBuildSelected)
 				{
+					UMapBuildDataRegistry* Registry = LightingContext.GetOrCreateRegistryForActor(Actor);
+
 					TInlineComponentArray<ULightComponent*> LightComponents;
 					Actor->GetComponents(LightComponents);
 
@@ -1369,6 +1367,8 @@ void FStaticLightingSystem::ApplyNewLightingData(bool bLightingSuccessful)
 
 			const bool bBuildLightingForLevel = Options.ShouldBuildLightingForLevel( Level );
 
+			UMapBuildDataRegistry* Registry = LightingContext.GetRegistryForLevel(Level);
+
 			// Store off the quality of the lighting for the level if lighting was successful and we build lighting for this level.
 			if( bLightingSuccessful && bBuildLightingForLevel )
 			{
@@ -1409,7 +1409,7 @@ void FStaticLightingSystem::ApplyNewLightingData(bool bLightingSuccessful)
 		CollectGarbage( GARBAGE_COLLECTION_KEEPFLAGS );
 
 		// Commit the changes to the world's BSP surfaces.
-		World->CommitModelSurfaces();
+		LightingContext.World->CommitModelSurfaces();
 	}
 
 	// Report failed lighting build (don't count cancelled builds as failure).
@@ -1561,9 +1561,9 @@ void FStaticLightingSystem::ReportStatistics()
 			, 100.0f * float(GNumShadowmapTotalTexels - GNumShadowmapMappedTexels - GNumShadowmapUnmappedTexels) / NumShadowmapTotalTexels
 			);
 
-		for ( int32 LevelIndex=0; LevelIndex < World->GetNumLevels(); LevelIndex++ )
+		for ( int32 LevelIndex=0; LevelIndex < LightingContext.World->GetNumLevels(); LevelIndex++ )
 		{
-			ULevel* Level = World->GetLevel(LevelIndex);
+			ULevel* Level = LightingContext.World->GetLevel(LevelIndex);
 			UE_LOG(LogStaticLightingSystem, Log,  TEXT("Level %2d - Lightmaps: %.1f MB. Shadowmaps: %.1f MB."), LevelIndex, Level->LightmapTotalSize/1024.0f, Level->ShadowmapTotalSize/1024.0f );
 		}
 	}
@@ -1741,7 +1741,7 @@ void FStaticLightingSystem::AddBSPStaticLightingInfo(ULevel* Level, bool bBuildL
 				UModelComponent* Component = Level->ModelComponents[Node.ComponentIndex];
 				if (Component->VisibilityId == INDEX_NONE)
 				{
-					if (World->GetWorldSettings()->bPrecomputeVisibility)
+					if (LightingContext.World->GetWorldSettings()->bPrecomputeVisibility)
 					{
 						// Make sure packages gets dirtied since we are changing the visibility Id of a component in them
 						PackagesToDirty.Add(Component->GetPackage());
@@ -2006,7 +2006,7 @@ bool FStaticLightingSystem::CreateLightmassProcessor()
 	
 	GWarn->StatusForceUpdate( -1, -1, LOCTEXT("StartingSwarmConnectionStatus", "Starting up Swarm Connection...") );
 	
-	if (Options.bOnlyBuildVisibility && !World->GetWorldSettings()->bPrecomputeVisibility)
+	if (Options.bOnlyBuildVisibility && !LightingContext.World->GetWorldSettings()->bPrecomputeVisibility)
 	{
 		FMessageDialog::Open( EAppMsgType::Ok, NSLOCTEXT("UnrealEd", "BuildFailed_VisibilityOnlyButVisibilityDisabled", "'Build Only Visibility' option was enabled but precomputed visibility is disabled!  Aborting build."));
 		return false;
@@ -2062,7 +2062,7 @@ void FStaticLightingSystem::GatherScene()
 	check(LightmassExporter);
 
 	// The Level settings...
-	AWorldSettings* WorldSettings = World->GetWorldSettings();
+	AWorldSettings* WorldSettings = LightingContext.World->GetWorldSettings();
 	if (WorldSettings)
 	{
 		LightmassExporter->SetLevelSettings(WorldSettings->LightmassSettings);
@@ -2075,16 +2075,16 @@ void FStaticLightingSystem::GatherScene()
 	LightmassExporter->SetNumUnusedLocalCores(Options.NumUnusedLocalCores);
 	LightmassExporter->SetQualityLevel(Options.QualityLevel);
 
-	if (World->PersistentLevel && Options.ShouldBuildLightingForLevel( World->PersistentLevel ))
+	if (LightingContext.World->PersistentLevel && Options.ShouldBuildLightingForLevel(LightingContext.World->PersistentLevel ))
 	{
-		LightmassExporter->SetLevelName(World->PersistentLevel->GetPathName());
+		LightmassExporter->SetLevelName(LightingContext.World->PersistentLevel->GetPathName());
 	}
 
 	LightmassExporter->ClearImportanceVolumes();
 	for( TObjectIterator<ALightmassImportanceVolume> It ; It ; ++It )
 	{
 		ALightmassImportanceVolume* LMIVolume = *It;
-		if (World->ContainsActor(LMIVolume) && IsValid(LMIVolume) && ShouldOperateOnLevel(LMIVolume->GetLevel()))
+		if (LightingContext.World->ContainsActor(LMIVolume) && IsValid(LMIVolume) && ShouldOperateOnLevel(LMIVolume->GetLevel()))
 		{
 			LightmassExporter->AddImportanceVolume(LMIVolume);
 		}
@@ -2093,7 +2093,7 @@ void FStaticLightingSystem::GatherScene()
 	for( TObjectIterator<ALightmassCharacterIndirectDetailVolume> It ; It ; ++It )
 	{
 		ALightmassCharacterIndirectDetailVolume* LMDetailVolume = *It;
-		if (World->ContainsActor(LMDetailVolume) && IsValid(LMDetailVolume) && ShouldOperateOnLevel(LMDetailVolume->GetLevel()))
+		if (LightingContext.World->ContainsActor(LMDetailVolume) && IsValid(LMDetailVolume) && ShouldOperateOnLevel(LMDetailVolume->GetLevel()))
 		{
 			LightmassExporter->AddCharacterIndirectDetailVolume(LMDetailVolume);
 		}
@@ -2102,7 +2102,7 @@ void FStaticLightingSystem::GatherScene()
 	for (TObjectIterator<AVolumetricLightmapDensityVolume> It; It; ++It)
 	{
 		AVolumetricLightmapDensityVolume* DetailVolume = *It;
-		if (World->ContainsActor(DetailVolume) && IsValid(DetailVolume) && ShouldOperateOnLevel(DetailVolume->GetLevel()))
+		if (LightingContext.World->ContainsActor(DetailVolume) && IsValid(DetailVolume) && ShouldOperateOnLevel(DetailVolume->GetLevel()))
 		{
 			LightmassExporter->VolumetricLightmapDensityVolumes.Add(DetailVolume);
 		}
@@ -2111,7 +2111,7 @@ void FStaticLightingSystem::GatherScene()
 	for( TObjectIterator<ULightmassPortalComponent> It ; It ; ++It )
 	{
 		ULightmassPortalComponent* LMPortal = *It;
-		if (LMPortal->GetOwner() && World->ContainsActor(LMPortal->GetOwner()) && IsValid(LMPortal) && ShouldOperateOnLevel(LMPortal->GetOwner()->GetLevel()))
+		if (LMPortal->GetOwner() && LightingContext.World->ContainsActor(LMPortal->GetOwner()) && IsValid(LMPortal) && ShouldOperateOnLevel(LMPortal->GetOwner()->GetLevel()))
 		{
 			LightmassExporter->AddPortal(LMPortal);
 		}
@@ -2120,7 +2120,7 @@ void FStaticLightingSystem::GatherScene()
 	for (TObjectIterator<USkyAtmosphereComponent> It; It; ++It)
 	{
 		USkyAtmosphereComponent* SkyAtmosphere = *It;
-		if (SkyAtmosphere->GetOwner() && World->ContainsActor(SkyAtmosphere->GetOwner()) && IsValid(SkyAtmosphere) && ShouldOperateOnLevel(SkyAtmosphere->GetOwner()->GetLevel()))
+		if (SkyAtmosphere->GetOwner() && LightingContext.World->ContainsActor(SkyAtmosphere->GetOwner()) && IsValid(SkyAtmosphere) && ShouldOperateOnLevel(SkyAtmosphere->GetOwner()->GetLevel()))
 		{
 			LightmassExporter->SetSkyAtmosphereComponent(SkyAtmosphere);
 			break;	// We only register the first we find
@@ -2268,7 +2268,7 @@ bool FStaticLightingSystem::FinishLightmassProcess()
 			if (!Options.bOnlyBuildVisibility)
 			{
 				FLightmassStatistics::FScopedGather FinishStatScope(LightmassStatistics.FinishingTime);
-				ULightComponent::ReassignStationaryLightChannels(GWorld, true, LightingScenario);
+				ULightComponent::ReassignStationaryLightChannels(GWorld, true, &LightingContext);
 			}
 		}
 	
@@ -2353,9 +2353,9 @@ void FStaticLightingSystem::UpdateLightingBuild()
 		
 		FString ScenarioString;
 
-		if (LightingScenario)
+		if (LightingContext.LightingScenario)
 		{
-			FString PackageName = FPackageName::GetShortName(LightingScenario->GetOutermost()->GetName());
+			FString PackageName = FPackageName::GetShortName(LightingContext.LightingScenario->GetOutermost()->GetName());
 			ScenarioString = FString(TEXT(" for ")) + PackageName;
 		}
 
@@ -2414,8 +2414,8 @@ void FStaticLightingSystem::UpdateAutomaticImportanceVolumeBounds( const FBox& M
 
 void FStaticLightingSystem::GatherBuildDataResourcesToKeep(const ULevel* InLevel)
 {
-	// This is only required is using a lighting scenario, otherwise the build data is saved within the level itself and follows it's inclusion in the lighting build.
-	if (InLevel && LightingScenario)
+	// This is only required is using a lighting scenario, otherwise the build data is saved within the level itself and follows it's inclusion in the lighting build.	
+	if (InLevel && LightingContext.LightingScenario)
 	{
 		BuildDataResourcesToKeep.Add(InLevel->LevelBuildDataId);
 
@@ -2494,12 +2494,17 @@ void FStaticLightingSystem::ApplyMapping(
 	FQuantizedLightmapData* QuantizedData,
 	const TMap<ULightComponent*,FShadowMapData2D*>& ShadowMapData) const
 {
-	TextureMapping->Apply(QuantizedData, ShadowMapData, LightingScenario);
+	TextureMapping->Apply(QuantizedData, ShadowMapData, &LightingContext);
 }
 
 UWorld* FStaticLightingSystem::GetWorld() const
 {
-	return World;
+	return LightingContext.World;
+}
+
+const FStaticLightingBuildContext& FStaticLightingSystem::GetLightingContext() const
+{
+	return LightingContext;
 }
 
 bool FStaticLightingSystem::IsAsyncBuilding() const
diff --git a/Engine/Source/Runtime/Engine/Classes/Components/InstancedStaticMeshComponent.h b/Engine/Source/Runtime/Engine/Classes/Components/InstancedStaticMeshComponent.h
--- a/Engine/Source/Runtime/Engine/Classes/Components/InstancedStaticMeshComponent.h
+++ b/Engine/Source/Runtime/Engine/Classes/Components/InstancedStaticMeshComponent.h
@@ -18,6 +18,7 @@
 
 class FLightingBuildOptions;
 class FPrimitiveSceneProxy;
+class FStaticLightingBuildContext;
 class FStaticLightingTextureMapping_InstancedStaticMesh;
 class ULightComponent;
 struct FNavigableGeometryExport;
@@ -683,7 +684,7 @@ protected:
 	UPROPERTY(Transient, DuplicateTransient, TextExportTransient)
 	TArray<FInstancedStaticMeshMappingInfo> CachedMappings;
 
-	ENGINE_API void ApplyLightMapping(FStaticLightingTextureMapping_InstancedStaticMesh* InMapping, ULevel* LightingScenario);
+	ENGINE_API void ApplyLightMapping(FStaticLightingTextureMapping_InstancedStaticMesh* InMapping, const FStaticLightingBuildContext* LightingContext);
 	
 	ENGINE_API void CreateHitProxyData(TArray<TRefCountPtr<HHitProxy>>& HitProxies);
 
diff --git a/Engine/Source/Runtime/Engine/Classes/Components/LightComponent.h b/Engine/Source/Runtime/Engine/Classes/Components/LightComponent.h
--- a/Engine/Source/Runtime/Engine/Classes/Components/LightComponent.h
+++ b/Engine/Source/Runtime/Engine/Classes/Components/LightComponent.h
@@ -13,6 +13,7 @@
 #include "LightComponent.generated.h"
 
 class FLightComponentMapBuildData;
+class FStaticLightingBuildContext;
 class FStaticShadowDepthMapData;
 class ULevel;
 class UMaterialInterface;
@@ -486,11 +487,13 @@ public:
 	/** Get the color temperature in the working color space. */
 	ENGINE_API FLinearColor GetColorTemperature() const;
 
+#if WITH_EDITOR	
 	/** 
 	 * Iterates over ALL stationary light components in the target world and assigns their preview shadowmap channel, and updates light icons accordingly.
 	 * Also handles assignment after a lighting build, so that the same algorithm is used for previewing and static lighting.
 	 */
-	static ENGINE_API void ReassignStationaryLightChannels(UWorld* TargetWorld, bool bAssignForLightingBuild, ULevel* LightingScenario);
+	static ENGINE_API void ReassignStationaryLightChannels(UWorld* TargetWorld, bool bAssignForLightingBuild, FStaticLightingBuildContext* LightingContext);
+#endif
 
 	DECLARE_MULTICAST_DELEGATE_OneParam(FOnUpdateColorAndBrightness, ULightComponent&);
 
diff --git a/Engine/Source/Runtime/Engine/Classes/Engine/InstancedStaticMesh.h b/Engine/Source/Runtime/Engine/Classes/Engine/InstancedStaticMesh.h
--- a/Engine/Source/Runtime/Engine/Classes/Engine/InstancedStaticMesh.h
+++ b/Engine/Source/Runtime/Engine/Classes/Engine/InstancedStaticMesh.h
@@ -611,7 +611,7 @@ public:
 	}
 
 	// FStaticLightingTextureMapping interface
-	virtual void Apply(FQuantizedLightmapData* InQuantizedData, const TMap<ULightComponent*, FShadowMapData2D*>& InShadowMapData, ULevel* LightingScenario) override
+	virtual void Apply(FQuantizedLightmapData* InQuantizedData, const TMap<ULightComponent*, FShadowMapData2D*>& InShadowMapData, const FStaticLightingBuildContext* LightingContext) override
 	{
 		check(bComplete == false);
 
@@ -627,7 +627,7 @@ public:
 				ShadowMapData.Add(ShadowDataPair.Key, TUniquePtr<FShadowMapData2D>(ShadowDataPair.Value));
 			}
 
-			InstancedComponent->ApplyLightMapping(this, LightingScenario);
+			InstancedComponent->ApplyLightMapping(this, LightingContext);
 		}
 
 		bComplete = true;
diff --git a/Engine/Source/Runtime/Engine/Classes/Engine/MapBuildDataRegistry.h b/Engine/Source/Runtime/Engine/Classes/Engine/MapBuildDataRegistry.h
--- a/Engine/Source/Runtime/Engine/Classes/Engine/MapBuildDataRegistry.h
+++ b/Engine/Source/Runtime/Engine/Classes/Engine/MapBuildDataRegistry.h
@@ -377,6 +377,18 @@ public:
 		or called by PostLoad for legacy BuildData with old EncodedData
 	*/
 	ENGINE_API void HandleLegacyEncodedCubemapData();
+
+	/** 
+	    Various helpers to access the UUMapBuildDataRegistry, necessary since the access patterns
+		are different in WorldPartition maps than in map. 
+	*/
+
+	/** Returns the associated UUMapBuildDataRegistry for an Actor Component */
+	ENGINE_API static UMapBuildDataRegistry* Get(const UActorComponent* Component);
+
+	/** Returns the associated UUMapBuildDataRegistry for an Actor */
+	ENGINE_API static UMapBuildDataRegistry* Get(const AActor* Actor);
+	
 private:
 #if WITH_EDITOR
 	void HandleAssetPostCompileEvent(const TArray<FAssetCompileData>& CompiledAssets);
diff --git a/Engine/Source/Runtime/Engine/Private/Components/LightComponent.cpp b/Engine/Source/Runtime/Engine/Private/Components/LightComponent.cpp
--- a/Engine/Source/Runtime/Engine/Private/Components/LightComponent.cpp
+++ b/Engine/Source/Runtime/Engine/Private/Components/LightComponent.cpp
@@ -9,6 +9,7 @@
 #include "ColorSpace.h"
 #include "Engine/Level.h"
 #include "Engine/MapBuildDataRegistry.h"
+#include "StaticLightingBuildContext.h"
 #include "Engine/World.h"
 #include "Materials/Material.h"
 #include "MaterialDomain.h"
@@ -1264,6 +1265,7 @@ void ULightComponent::InvalidateLightingCacheDetailed(bool bInvalidateBuildEnque
 		// Create new guids for light.
 		UpdateLightGUIDs();
 
+#if WITH_EDITOR
 		if (GIsEditor)
 		{
 			UWorld* World = GetWorld();
@@ -1273,6 +1275,7 @@ void ULightComponent::InvalidateLightingCacheDetailed(bool bInvalidateBuildEnque
 				ReassignStationaryLightChannels(World, false, NULL);
 			}
 		}
+#endif
 
 		MarkRenderStateDirty();
 
@@ -1349,17 +1352,7 @@ const FLightComponentMapBuildData* ULightComponent::GetLightComponentMapBuildDat
 			}
 #endif
 
-			ULevel* ActiveLightingScenario = OwnerLevel->OwningWorld->GetActiveLightingScenario();
-			UMapBuildDataRegistry* MapBuildData = NULL;
-
-			if (ActiveLightingScenario && ActiveLightingScenario->MapBuildData)
-			{
-				MapBuildData = ActiveLightingScenario->MapBuildData;
-			}
-			else if (OwnerLevel->MapBuildData)
-			{
-				MapBuildData = OwnerLevel->MapBuildData;
-			}
+			UMapBuildDataRegistry* MapBuildData = UMapBuildDataRegistry::Get(this);
 
 			if (MapBuildData)
 			{
@@ -1480,7 +1473,9 @@ struct FCompareLightsByArrayCount
  * - finishing a lighting build
  * If you're adding more call sites to this function, make sure not to break GPULightmass as it is based on the above assumption
  */
-void ULightComponent::ReassignStationaryLightChannels(UWorld* TargetWorld, bool bAssignForLightingBuild, ULevel* LightingScenario)
+
+#if WITH_EDITOR
+void ULightComponent::ReassignStationaryLightChannels(UWorld* TargetWorld, bool bAssignForLightingBuild, FStaticLightingBuildContext* LightingContext)
 {
 	TMap<FLightAndChannel*, TArray<FLightAndChannel*> > LightToOverlapMap;
 
@@ -1497,9 +1492,7 @@ void ULightComponent::ReassignStationaryLightChannels(UWorld* TargetWorld, bool
 			&& LightComponent->HasStaticShadowing()
 			&& !LightComponent->HasStaticLighting())
 		{
-			ULevel* LightLevel = LightOwner->GetLevel();
-
-			if (!LightingScenario || !LightLevel->bIsLightingScenario || LightLevel == LightingScenario)
+			if (!LightingContext || LightingContext->ShouldIncludeActor(LightOwner))
 			{				
 				if (LightComponent->bAffectsWorld
 					&& (LightComponent->CastShadows || LightComponent->LightFunctionMaterial)
@@ -1529,8 +1522,7 @@ void ULightComponent::ReassignStationaryLightChannels(UWorld* TargetWorld, bool
 
 		if (bAssignForLightingBuild)
 		{
-			ULevel* StorageLevel = LightingScenario ? LightingScenario : CurrentLight->GetOwner()->GetLevel();
-			UMapBuildDataRegistry* Registry = StorageLevel->GetOrCreateMapBuildData();
+			UMapBuildDataRegistry* Registry = LightingContext->GetOrCreateRegistryForActor(CurrentLight->GetOwner());
 			FLightComponentMapBuildData& LightBuildData = Registry->FindOrAllocateLightBuildData(CurrentLight->LightGuid, true);
 			LightBuildData.ShadowMapChannel = INDEX_NONE;
 		}
@@ -1627,8 +1619,7 @@ void ULightComponent::ReassignStationaryLightChannels(UWorld* TargetWorld, bool
 
 		if (bAssignForLightingBuild)
 		{
-			ULevel* StorageLevel = LightingScenario ? LightingScenario : CurrentLight->Light->GetOwner()->GetLevel();
-			UMapBuildDataRegistry* Registry = StorageLevel->GetOrCreateMapBuildData();
+			UMapBuildDataRegistry* Registry = LightingContext->GetOrCreateRegistryForActor(CurrentLight->Light->GetOwner());
 			FLightComponentMapBuildData& LightBuildData = Registry->FindOrAllocateLightBuildData(CurrentLight->Light->LightGuid, true);
 			LightBuildData.ShadowMapChannel = CurrentLight->Channel;
 
@@ -1643,6 +1634,7 @@ void ULightComponent::ReassignStationaryLightChannels(UWorld* TargetWorld, bool
 		delete CurrentLight;
 	}
 }
+#endif
 
 static void ToggleLight(const TArray<FString>& Args)
 {
diff --git a/Engine/Source/Runtime/Engine/Private/Components/ModelComponent.cpp b/Engine/Source/Runtime/Engine/Private/Components/ModelComponent.cpp
--- a/Engine/Source/Runtime/Engine/Private/Components/ModelComponent.cpp
+++ b/Engine/Source/Runtime/Engine/Private/Components/ModelComponent.cpp
@@ -55,26 +55,10 @@ FModelElement::~FModelElement()
 const FMeshMapBuildData* FModelElement::GetMeshMapBuildData() const
 {
 	check(Component);
-	ULevel* OwnerLevel = Cast<ULevel>(Component->GetModel()->GetOuter());
 
-	if (OwnerLevel && OwnerLevel->OwningWorld)
+	if (UMapBuildDataRegistry* MapBuildData = UMapBuildDataRegistry::Get(Component))
 	{
-		ULevel* ActiveLightingScenario = OwnerLevel->OwningWorld->GetActiveLightingScenario();
-		UMapBuildDataRegistry* MapBuildData = NULL;
-
-		if (ActiveLightingScenario && ActiveLightingScenario->MapBuildData)
-		{
-			MapBuildData = ActiveLightingScenario->MapBuildData;
-		}
-		else if (OwnerLevel->MapBuildData)
-		{
-			MapBuildData = OwnerLevel->MapBuildData;
-		}
-
-		if (MapBuildData)
-		{
-			return MapBuildData->GetMeshBuildData(MapBuildDataId);
-		}
+		return MapBuildData->GetMeshBuildData(MapBuildDataId);		
 	}
 	
 	return NULL;
diff --git a/Engine/Source/Runtime/Engine/Private/Components/StaticMeshComponent.cpp b/Engine/Source/Runtime/Engine/Private/Components/StaticMeshComponent.cpp
--- a/Engine/Source/Runtime/Engine/Private/Components/StaticMeshComponent.cpp
+++ b/Engine/Source/Runtime/Engine/Private/Components/StaticMeshComponent.cpp
@@ -645,31 +645,9 @@ const FMeshMapBuildData* UStaticMeshComponent::GetMeshMapBuildData(const FStatic
 		return LODInfo.OverrideMapBuildData.Get();
 	}
 
-	AActor* Owner = GetOwner();
-
-	if (Owner)
+	if (UMapBuildDataRegistry* MapBuildData = UMapBuildDataRegistry::Get(this))
 	{
-		ULevel* OwnerLevel = Owner->GetLevel();
-
-		if (OwnerLevel && OwnerLevel->OwningWorld)
-		{
-			ULevel* ActiveLightingScenario = OwnerLevel->OwningWorld->GetActiveLightingScenario();
-			UMapBuildDataRegistry* MapBuildData = NULL;
-
-			if (ActiveLightingScenario && ActiveLightingScenario->MapBuildData)
-			{
-				MapBuildData = ActiveLightingScenario->MapBuildData;
-			}
-			else if (OwnerLevel->MapBuildData)
-			{
-				MapBuildData = OwnerLevel->MapBuildData;
-			}
-
-			if (MapBuildData)
-			{
-				return bCheckForResourceCluster ? MapBuildData->GetMeshBuildData(LODInfo.MapBuildDataId) : MapBuildData->GetMeshBuildDataDuringBuild(LODInfo.MapBuildDataId);
-			}
-		}
+		return bCheckForResourceCluster ? MapBuildData->GetMeshBuildData(LODInfo.MapBuildDataId) : MapBuildData->GetMeshBuildDataDuringBuild(LODInfo.MapBuildDataId);
 	}
 	
 	return NULL;
diff --git a/Engine/Source/Runtime/Engine/Private/InstancedStaticMesh.cpp b/Engine/Source/Runtime/Engine/Private/InstancedStaticMesh.cpp
--- a/Engine/Source/Runtime/Engine/Private/InstancedStaticMesh.cpp
+++ b/Engine/Source/Runtime/Engine/Private/InstancedStaticMesh.cpp
@@ -10,6 +10,7 @@
 #include "Engine/OverlapResult.h"
 #include "AI/NavigationSystemBase.h"
 #include "Engine/MapBuildDataRegistry.h"
+#include "StaticLightingBuildContext.h"
 #include "Components/LightComponent.h"
 #include "EngineLogs.h"
 #include "Logging/MessageLog.h"
@@ -2802,7 +2803,7 @@ void UInstancedStaticMeshComponent::GetStaticLightingInfo(FStaticLightingPrimiti
 	}
 }
 
-void UInstancedStaticMeshComponent::ApplyLightMapping(FStaticLightingTextureMapping_InstancedStaticMesh* InMapping, ULevel* LightingScenario)
+void UInstancedStaticMeshComponent::ApplyLightMapping(FStaticLightingTextureMapping_InstancedStaticMesh* InMapping, const FStaticLightingBuildContext* LightingContext)
 {
 	static const auto CVar = IConsoleManager::Get().FindTConsoleVariableDataInt(TEXT("r.VirtualTexturedLightmaps"));
 	const bool bUseVirtualTextures = (CVar->GetValueOnAnyThread() != 0) && UseVirtualTexturing(GMaxRHIShaderPlatform);
@@ -2847,8 +2848,7 @@ void UInstancedStaticMeshComponent::ApplyLightMapping(FStaticLightingTextureMapp
 			MarkPackageDirty();
 		}
 
-		ULevel* StorageLevel = LightingScenario ? LightingScenario : GetOwner()->GetLevel();
-		UMapBuildDataRegistry* Registry = StorageLevel->GetOrCreateMapBuildData();
+		UMapBuildDataRegistry* Registry = LightingContext->GetOrCreateRegistryForActor(GetOwner());
 		FMeshMapBuildData& MeshBuildData = Registry->AllocateMeshBuildData(LODInfo.MapBuildDataId, true);
 
 		MeshBuildData.PerInstanceLightmapData.Empty(AllQuantizedData.Num());
diff --git a/Engine/Source/Runtime/Engine/Private/LightMap.cpp b/Engine/Source/Runtime/Engine/Private/LightMap.cpp
--- a/Engine/Source/Runtime/Engine/Private/LightMap.cpp
+++ b/Engine/Source/Runtime/Engine/Private/LightMap.cpp
@@ -21,6 +21,7 @@
 #include "Misc/FeedbackContext.h"
 #include "GameFramework/WorldSettings.h"
 #include "Engine/MapBuildDataRegistry.h"
+#include "StaticLightingBuildContext.h"
 #include "VT/LightmapVirtualTexture.h"
 #include "VT/VirtualTexture.h"
 #include "EngineModule.h"
@@ -443,7 +444,7 @@ struct FLightMapPendingTexture : public FTextureLayout
 	/**
 	 * Processes the textures and starts asynchronous compression tasks for all mip-levels.
 	 */
-	void StartEncoding(ULevel* Unused, ITextureCompressorModule* UnusedCompressor);
+	void StartEncoding(const FStaticLightingBuildContext* LightingContext, ITextureCompressorModule* UnusedCompressor);
 
 	void EncodeCoefficientTexture(int32 CoefficientIndex, UTexture* Texture, uint32 LayerIndex, const FColor& TextureColor, bool bEncodeVirtualTexture);
 	void EncodeSkyOcclusionTexture(UTexture* Texture, uint32 LayerIndex, const FColor& TextureColor);
@@ -1741,7 +1742,7 @@ void FLightMapPendingTexture::EncodeCoefficientTexture(int32 CoefficientIndex, U
 	}
 }
 
-void FLightMapPendingTexture::StartEncoding(ULevel* LightingScenario, ITextureCompressorModule* UnusedCompressor)
+void FLightMapPendingTexture::StartEncoding(const FStaticLightingBuildContext* LightingContext, ITextureCompressorModule* UnusedCompressor)
 {
 	if (!bUObjectsCreated)
 	{
@@ -1754,7 +1755,7 @@ void FLightMapPendingTexture::StartEncoding(ULevel* LightingScenario, ITextureCo
 	int32 NumShadowChannelsUsed = 0;
 	if (NeedsStaticShadowTexture())
 	{
-		NumShadowChannelsUsed = FLightMap2D::EncodeShadowTexture(LightingScenario, *this, ShadowMipData);
+		NumShadowChannelsUsed = FLightMap2D::EncodeShadowTexture(LightingContext, *this, ShadowMipData);
 	}
 
 	FColor TextureColor;
@@ -2383,7 +2384,7 @@ TRefCountPtr<FLightMap2D> FLightMap2D::AllocateInstancedLightMap(UObject* LightM
  * @param	bLightingSuccessful	Whether the lighting build was successful or not.
  * @param	bMultithreadedEncode encode textures on different threads ;)
  */
-void FLightMap2D::EncodeTextures( UWorld* InWorld, ULevel* LightingScenario, bool bLightingSuccessful, bool bMultithreadedEncode)
+void FLightMap2D::EncodeTextures(const FStaticLightingBuildContext* LightingContext, bool bLightingSuccessful, bool bMultithreadedEncode)
 {
 #if WITH_EDITOR
 	if (bLightingSuccessful)
@@ -2392,7 +2393,7 @@ void FLightMap2D::EncodeTextures( UWorld* InWorld, ULevel* LightingScenario, boo
 		const bool bIncludeNonVirtualTextures = !bUseVirtualTextures || (CVarIncludeNonVirtualTexturedLightMaps.GetValueOnAnyThread() != 0);
 
 		GWarn->BeginSlowTask( NSLOCTEXT("LightMap2D", "BeginEncodingLightMapsTask", "Encoding light-maps"), false );
-		int32 PackedLightAndShadowMapTextureSizeX = InWorld->GetWorldSettings()->PackedLightAndShadowMapTextureSize;
+		int32 PackedLightAndShadowMapTextureSizeX = LightingContext->World->GetWorldSettings()->PackedLightAndShadowMapTextureSize;
 		int32 PackedLightAndShadowMapTextureSizeY = PackedLightAndShadowMapTextureSizeX / 2;
 
 		if (!bIncludeNonVirtualTextures)
@@ -2499,7 +2500,7 @@ void FLightMap2D::EncodeTextures( UWorld* InWorld, ULevel* LightingScenario, boo
 				// If there is no existing appropriate texture, create a new one.
 				// If we have non-VT, need 2to1 aspect ratio to handle packing lightmap into top/bottom texture region
 				// With only VT, better to use square lightmaps (better page table usage)
-				Texture = new FLightMapPendingTexture(InWorld, NewTextureSizeX, NewTextureSizeY, bIncludeNonVirtualTextures ? ETextureLayoutAspectRatio::Force2To1 : ETextureLayoutAspectRatio::ForceSquare);
+				Texture = new FLightMapPendingTexture(LightingContext->World, NewTextureSizeX, NewTextureSizeY, bIncludeNonVirtualTextures ? ETextureLayoutAspectRatio::Force2To1 : ETextureLayoutAspectRatio::ForceSquare);
 				PendingTextures.Add(Texture);
 				Texture->Outer = PendingGroup.Outer;
 				Texture->Bounds = PendingGroup.Bounds;
@@ -2525,7 +2526,7 @@ void FLightMap2D::EncodeTextures( UWorld* InWorld, ULevel* LightingScenario, boo
 				// precreate the UObjects then give them to some threads to process
 				// need to precreate Uobjects 
 				Texture->CreateUObjects();
-				auto AsyncEncodeTask = new (AsyncEncodeTasks)FAsyncEncode<FLightMapPendingTexture>(Texture, LightingScenario, Counter, nullptr);
+				auto AsyncEncodeTask = new (AsyncEncodeTasks)FAsyncEncode<FLightMapPendingTexture>(Texture, LightingContext, Counter, nullptr);
 				GLargeThreadPool->AddQueuedWork(AsyncEncodeTask);
 			}
 
@@ -2545,7 +2546,7 @@ void FLightMap2D::EncodeTextures( UWorld* InWorld, ULevel* LightingScenario, boo
 					GWarn->UpdateProgress(TextureIndex, PendingTextures.Num());
 				}
 				FLightMapPendingTexture* PendingTexture = PendingTextures[TextureIndex];
-				PendingTexture->StartEncoding(nullptr,nullptr);
+				PendingTexture->StartEncoding(LightingContext,nullptr);
 			}
 		}
 
@@ -2591,7 +2592,7 @@ void FLightMap2D::EncodeTextures( UWorld* InWorld, ULevel* LightingScenario, boo
 }
 
 #if WITH_EDITOR
-int32 FLightMap2D::EncodeShadowTexture(ULevel* LightingScenario, struct FLightMapPendingTexture& PendingTexture, TArray< TArray<FFourDistanceFieldSamples>>& MipData)
+int32 FLightMap2D::EncodeShadowTexture(const FStaticLightingBuildContext* LightingContext, struct FLightMapPendingTexture& PendingTexture, TArray< TArray<FFourDistanceFieldSamples>>& MipData)
 {
 	TArray<FFourDistanceFieldSamples>* TopMipData = new(MipData) TArray<FFourDistanceFieldSamples>();
 	TopMipData->Empty(PendingTexture.GetSizeX() * PendingTexture.GetSizeY());
@@ -2611,9 +2612,7 @@ int32 FLightMap2D::EncodeShadowTexture(ULevel* LightingScenario, struct FLightMa
 			for (const auto& ShadowMapPair : Allocation.ShadowMapData)
 			{
 				ULightComponent* CurrentLight = ShadowMapPair.Key;
-				ULevel* StorageLevel = LightingScenario ? LightingScenario : CurrentLight->GetOwner()->GetLevel();
-				UMapBuildDataRegistry* Registry = StorageLevel->MapBuildData;
-				const FLightComponentMapBuildData* LightBuildData = Registry->GetLightBuildData(CurrentLight->LightGuid);
+				const FLightComponentMapBuildData* LightBuildData = LightingContext->GetRegistryForActor(CurrentLight->GetOwner())->GetLightBuildData(CurrentLight->LightGuid);
 
 				// Should have been setup by ReassignStationaryLightChannels
 				check(LightBuildData);
diff --git a/Engine/Source/Runtime/Engine/Private/MapBuildData.cpp b/Engine/Source/Runtime/Engine/Private/MapBuildData.cpp
--- a/Engine/Source/Runtime/Engine/Private/MapBuildData.cpp
+++ b/Engine/Source/Runtime/Engine/Private/MapBuildData.cpp
@@ -33,6 +33,8 @@ MapBuildData.cpp
 #endif
 #include "Engine/TextureCube.h"
 #include "Components/HierarchicalInstancedStaticMeshComponent.h"
+#include "UnrealEngine.h"
+
 
 DECLARE_MEMORY_STAT(TEXT("Stationary Light Static Shadowmap"),STAT_StationaryLightBuildData,STATGROUP_MapBuildData);
 DECLARE_MEMORY_STAT(TEXT("Reflection Captures"),STAT_ReflectionCaptureBuildData,STATGROUP_MapBuildData);
@@ -1153,6 +1155,41 @@ void UMapBuildDataRegistry::CleanupTransientOverrideMapBuildData()
 	}
 }
 
+UMapBuildDataRegistry* UMapBuildDataRegistry::Get(const UActorComponent* Component)
+{
+	AActor* Owner = Component->GetOwner();
+
+	if (Owner)
+	{
+		return Get(Owner);
+	}
+
+	return nullptr;
+}
+
+UMapBuildDataRegistry* UMapBuildDataRegistry::Get(const AActor* Actor)
+{
+	ULevel* OwnerLevel = Actor->GetLevel();
+	UWorld* World = OwnerLevel->GetWorld();
+	UMapBuildDataRegistry* MapBuildData = nullptr;
+	
+	if (OwnerLevel && World)
+	{
+		ULevel* ActiveLightingScenario = World->GetActiveLightingScenario();
+		
+		if (ActiveLightingScenario && ActiveLightingScenario->MapBuildData)
+		{
+			MapBuildData = ActiveLightingScenario->MapBuildData;
+		}
+		else if (OwnerLevel->MapBuildData)
+		{
+			MapBuildData = OwnerLevel->MapBuildData;
+		}
+	}
+
+	return MapBuildData;
+}
+
 FUObjectAnnotationSparse<FMeshMapBuildLegacyData, true> GComponentsWithLegacyLightmaps;
 FUObjectAnnotationSparse<FLevelLegacyMapBuildData, true> GLevelsWithLegacyBuildData;
 FUObjectAnnotationSparse<FLightComponentLegacyMapBuildData, true> GLightComponentsWithLegacyBuildData;
diff --git a/Engine/Source/Runtime/Engine/Private/ModelLight.cpp b/Engine/Source/Runtime/Engine/Private/ModelLight.cpp
--- a/Engine/Source/Runtime/Engine/Private/ModelLight.cpp
+++ b/Engine/Source/Runtime/Engine/Private/ModelLight.cpp
@@ -7,6 +7,7 @@
 #include "ModelLight.h"
 #include "Engine/Level.h"
 #include "Engine/MapBuildDataRegistry.h"
+#include "StaticLightingBuildContext.h"
 #include "Components/LightComponent.h"
 #include "Misc/ScopedSlowTask.h"
 #include "ComponentReregisterContext.h"
@@ -155,7 +156,7 @@ FLightRayIntersection FBSPSurfaceStaticLighting::IntersectLightRay(const FVector
 }
 
 #if WITH_EDITOR
-void FBSPSurfaceStaticLighting::Apply(FQuantizedLightmapData* InQuantizedData, const TMap<ULightComponent*,FShadowMapData2D*>& InShadowMapData, ULevel* LightingScenario)
+void FBSPSurfaceStaticLighting::Apply(FQuantizedLightmapData* InQuantizedData, const TMap<ULightComponent*,FShadowMapData2D*>& InShadowMapData, const FStaticLightingBuildContext* LightingContext)
 {
 	if(!bComplete)
 	{
@@ -171,7 +172,7 @@ void FBSPSurfaceStaticLighting::Apply(FQuantizedLightmapData* InQuantizedData, c
 	// If all the surfaces have complete static lighting, apply the component's static lighting.
 	if(Model->NumIncompleteNodeGroups == 0)
 	{
-		Model->ApplyStaticLighting(LightingScenario);
+		Model->ApplyStaticLighting(LightingContext);
 	}
 }
 
@@ -994,7 +995,7 @@ void UModel::GroupAllNodes(ULevel* Level, const TArray<ULightComponentBase*>& Li
 /**
  * Applies all of the finished lighting cached in the NodeGroups 
  */
-void UModel::ApplyStaticLighting(ULevel* LightingScenario)
+void UModel::ApplyStaticLighting(const FStaticLightingBuildContext* LightingContext)
 {
 #if WITH_EDITOR
 	static const auto CVar = IConsoleManager::Get().FindTConsoleVariableDataInt(TEXT("r.VirtualTexturedLightmaps"));
@@ -1216,8 +1217,7 @@ void UModel::ApplyStaticLighting(ULevel* LightingScenario)
 		const bool bHasRelevantLights = SurfaceGroup.Surfaces.ContainsByPredicate([](const FSurfaceStaticLightingGroup::FSurfaceInfo& SurfaceInfo) { return SurfaceInfo.SurfaceStaticLighting->RelevantLights.Num() > 0; });
 		const bool bNeedsLightMap = bHasNonZeroData || SurfaceGroup.ShadowMappedLights.Num() > 0 || bHasRelevantLights || GroupQuantizedData->bHasSkyShadowing;
 
-		ULevel* StorageLevel = LightingScenario ? LightingScenario : LightingLevel;
-		UMapBuildDataRegistry* Registry = StorageLevel->GetOrCreateMapBuildData();
+		UMapBuildDataRegistry* Registry = LightingContext->GetOrCreateRegistryForLevel(LightingLevel);
 
 		// Allocate merged shadow-map data.
 		TMap<ULightComponent*,FShadowMapData2D*> GroupShadowMapData;
diff --git a/Engine/Source/Runtime/Engine/Private/ShadowMap.cpp b/Engine/Source/Runtime/Engine/Private/ShadowMap.cpp
--- a/Engine/Source/Runtime/Engine/Private/ShadowMap.cpp
+++ b/Engine/Source/Runtime/Engine/Private/ShadowMap.cpp
@@ -3,6 +3,7 @@
 #include "ShadowMap.h"
 #include "Engine/Level.h"
 #include "Engine/MapBuildDataRegistry.h"
+#include "StaticLightingBuildContext.h"
 #include "Components/LightComponent.h"
 
 #include "Engine/World.h"
@@ -192,7 +193,7 @@ struct FShadowMapPendingTexture : FTextureLayout
 	 *
 	 * @param	InWorld	World in which the textures exist
 	 */
-	void StartEncoding(ULevel* LightingScenario, ITextureCompressorModule* Compressor);
+	void StartEncoding(const FStaticLightingBuildContext* LightingContext, ITextureCompressorModule* Compressor);
 
 	/**
 	 * Create UObjects required in the encoding step, this is so we can multithread teh encode step
@@ -304,7 +305,7 @@ void FShadowMapPendingTexture::CreateUObjects()
 	bCreatedUObjects = true;
 }
 
-void FShadowMapPendingTexture::StartEncoding(ULevel* LightingScenario, ITextureCompressorModule* Compressor)
+void FShadowMapPendingTexture::StartEncoding(const FStaticLightingBuildContext* LightingContext, ITextureCompressorModule* Compressor)
 {
 	FOptionalTaskTagScope Scope(ETaskTag::EParallelGameThread);
 
@@ -322,7 +323,7 @@ void FShadowMapPendingTexture::StartEncoding(ULevel* LightingScenario, ITextureC
 	{
 		// Create the uncompressed top mip-level.
 		TArray< TArray<FFourDistanceFieldSamples> > MipData;
-		const int32 NumChannelsUsed = FShadowMap2D::EncodeSingleTexture(LightingScenario, *this, Texture, MipData);
+		const int32 NumChannelsUsed = FShadowMap2D::EncodeSingleTexture(LightingContext, *this, Texture, MipData);
 
 		Texture->Source.Init2DWithMipChain(GetSizeX(), GetSizeY(), NumChannelsUsed == 1 ? TSF_G8 : TSF_BGRA8);
 		Texture->MipGenSettings = TMGS_LeaveExistingMips;
@@ -718,12 +719,12 @@ TRefCountPtr<FShadowMap2D> FShadowMap2D::AllocateInstancedShadowMap(UObject* Lig
  * @param	InWorld				World in which the textures exist
  * @param	bLightingSuccessful	Whether the lighting build was successful or not.
  */
-void FShadowMap2D::EncodeTextures(UWorld* InWorld, ULevel* LightingScenario, bool bLightingSuccessful, bool bMultithreadedEncode)
+void FShadowMap2D::EncodeTextures(const FStaticLightingBuildContext* LightingContext, bool bLightingSuccessful, bool bMultithreadedEncode)
 {
 	if ( bLightingSuccessful )
 	{
 		GWarn->BeginSlowTask( NSLOCTEXT("ShadowMap2D", "BeginEncodingShadowMapsTask", "Encoding shadow-maps"), false );
-		const int32 PackedLightAndShadowMapTextureSize = InWorld->GetWorldSettings()->PackedLightAndShadowMapTextureSize;
+		const int32 PackedLightAndShadowMapTextureSize = LightingContext->World->GetWorldSettings()->PackedLightAndShadowMapTextureSize;
 
 		ITextureCompressorModule* TextureCompressorModule = &FModuleManager::LoadModuleChecked<ITextureCompressorModule>(TEXTURE_COMPRESSOR_MODULENAME);
 
@@ -810,7 +811,7 @@ void FShadowMap2D::EncodeTextures(UWorld* InWorld, ULevel* LightingScenario, boo
 			for (auto& PendingTexture : PendingTextures)
 			{
 				PendingTexture.CreateUObjects();
-				auto AsyncEncodeTask = new (AsyncEncodeTasks)FAsyncEncode<FShadowMapPendingTexture>(&PendingTexture, LightingScenario, Counter, TextureCompressorModule);
+				auto AsyncEncodeTask = new (AsyncEncodeTasks)FAsyncEncode<FShadowMapPendingTexture>(&PendingTexture, LightingContext, Counter, TextureCompressorModule);
 				GThreadPool->AddQueuedWork(AsyncEncodeTask);
 			}
 
@@ -826,7 +827,7 @@ void FShadowMap2D::EncodeTextures(UWorld* InWorld, ULevel* LightingScenario, boo
 			for (int32 TextureIndex = 0; TextureIndex < PendingTextures.Num(); TextureIndex++)
 			{
 				FShadowMapPendingTexture& PendingTexture = PendingTextures[TextureIndex];
-				PendingTexture.StartEncoding(LightingScenario, TextureCompressorModule);
+				PendingTexture.StartEncoding(LightingContext, TextureCompressorModule);
 			}
 		}
 
@@ -853,7 +854,7 @@ void FShadowMap2D::EncodeTextures(UWorld* InWorld, ULevel* LightingScenario, boo
 		for (int32 TextureIndex = 0; TextureIndex < PendingTextures.Num(); TextureIndex++)
 		{
 			FShadowMapPendingTexture& PendingTexture = PendingTextures[TextureIndex];
-			PendingTexture.FinishCachingTextures(InWorld);
+			PendingTexture.FinishCachingTextures(LightingContext->World);
 			if (bUpdateStatus && ((TextureIndex % 20) == 0))
 			{
 				GWarn->UpdateProgress(TextureIndex, PendingTextures.Num());
@@ -869,7 +870,7 @@ void FShadowMap2D::EncodeTextures(UWorld* InWorld, ULevel* LightingScenario, boo
 	}
 }
 
-int32 FShadowMap2D::EncodeSingleTexture(ULevel* LightingScenario, FShadowMapPendingTexture& PendingTexture, UShadowMapTexture2D* Texture, TArray< TArray<FFourDistanceFieldSamples> >& MipData)
+int32 FShadowMap2D::EncodeSingleTexture(const FStaticLightingBuildContext* LightingContext, FShadowMapPendingTexture& PendingTexture, UShadowMapTexture2D* Texture, TArray< TArray<FFourDistanceFieldSamples> >& MipData)
 {
 	TArray<FFourDistanceFieldSamples>* TopMipData = new(MipData) TArray<FFourDistanceFieldSamples>();
 	TopMipData->Empty(PendingTexture.GetSizeX() * PendingTexture.GetSizeY());
@@ -889,8 +890,7 @@ int32 FShadowMap2D::EncodeSingleTexture(ULevel* LightingScenario, FShadowMapPend
 			for (const auto& ShadowMapPair : Allocation.ShadowMapData)
 			{
 				ULightComponent* CurrentLight = ShadowMapPair.Key;
-				ULevel* StorageLevel = LightingScenario ? LightingScenario : CurrentLight->GetOwner()->GetLevel();
-				UMapBuildDataRegistry* Registry = StorageLevel->MapBuildData;
+				UMapBuildDataRegistry* Registry = LightingContext->GetRegistryForActor(CurrentLight->GetOwner());
 				const FLightComponentMapBuildData* LightBuildData = Registry->GetLightBuildData(CurrentLight->LightGuid);
 
 				// Should have been setup by ReassignStationaryLightChannels
diff --git a/Engine/Source/Runtime/Engine/Private/StaticLightingBuildContext.cpp b/Engine/Source/Runtime/Engine/Private/StaticLightingBuildContext.cpp
new file mode 100644
--- /dev/null
+++ b/Engine/Source/Runtime/Engine/Private/StaticLightingBuildContext.cpp
@@ -0,0 +1,205 @@
+// Copyright Epic Games, Inc. All Rights Reserved.
+
+#include "StaticLightingBuildContext.h"
+
+#include "Engine/Level.h"
+#include "Engine/World.h"
+#include "GameFramework/Actor.h"
+#include "Engine/MapBuildDataRegistry.h"
+#include "WorldPartition/WorldPartition.h"
+
+#if WITH_EDITOR
+
+FStaticLightingBuildContext::FStaticLightingBuildContext(UWorld* InWorld, ULevel* InLightingScenario) 	
+{
+	checkf(InWorld, TEXT("Constructing a FStaticLightingBuildContext requires a valid World"));
+
+	World = InWorld;
+	LightingScenario = InLightingScenario;
+
+	if (LightingScenario)
+	{
+		MapBuildDataRegistry = InLightingScenario->MapBuildData;
+	}
+	else
+	{
+		MapBuildDataRegistry = InWorld->PersistentLevel->MapBuildData;
+	}
+
+	for (int32 LevelIndex = 0; LevelIndex < World->GetNumLevels(); LevelIndex++)
+	{
+		ULevel* Level = World->GetLevel(LevelIndex);
+		FGuid LevelGuid = FGuid(0, 0, 0, LevelIndex);
+		LevelGuids.Add(LevelGuid, Level);
+	}
+	
+	FGuid FirstGuid = FGuid(0, 0, 0, 0);
+	check(GetLevelForGuid(FirstGuid) == World->PersistentLevel);
+}
+
+FStaticLightingBuildContext::FStaticLightingBuildContext(const FStaticLightingBuildContext& InFrom)
+{
+	MapBuildDataRegistry = InFrom.MapBuildDataRegistry;
+	LevelGuids = InFrom.LevelGuids;
+	World = InFrom.World;
+	LightingScenario = InFrom.LightingScenario;
+}
+
+bool FStaticLightingBuildContext::ShouldIncludeActor(AActor* Actor) const
+{
+	check(Actor);
+
+	ULevel* ActorLevel = Actor->GetLevel();
+	check(ActorLevel);
+
+	bool IncludeActor = false;
+
+	if (!LightingScenario || !ActorLevel->bIsLightingScenario || ActorLevel == LightingScenario)
+	{
+		IncludeActor = true;
+	}
+
+	return IncludeActor;
+}
+
+bool FStaticLightingBuildContext::ShouldIncludeLevel(ULevel* Level) const
+{
+	check(Level);
+
+	bool IncludeLevel = false;
+
+	if (!LightingScenario || !Level->bIsLightingScenario || Level == LightingScenario)
+	{
+		IncludeLevel = true;
+	}
+
+	return IncludeLevel;
+}
+
+UMapBuildDataRegistry* FStaticLightingBuildContext::GetGlobalRegistry() const
+{
+	return MapBuildDataRegistry;
+}
+
+UMapBuildDataRegistry* FStaticLightingBuildContext::GetOrCreateGlobalRegistry() const
+{
+	if (!MapBuildDataRegistry)
+	{
+		if (LightingScenario)
+		{
+			MapBuildDataRegistry = LightingScenario->GetOrCreateMapBuildData();
+		}
+		else
+		{
+			MapBuildDataRegistry = World->PersistentLevel->GetOrCreateMapBuildData();
+		}
+	}
+	
+	return MapBuildDataRegistry;
+}
+
+ULevel* FStaticLightingBuildContext::GetLightingStorageLevel(ULevel* Level) const
+{
+	if (LightingScenario)
+	{
+		return LightingScenario;
+	}
+	else
+	{
+		return Level;
+	}
+}
+
+UMapBuildDataRegistry* FStaticLightingBuildContext::GetRegistryForActor(AActor* Actor) const
+{	
+	check(Actor);
+
+	ULevel* Level = Actor->GetLevel();
+
+	return GetLightingStorageLevel(Level)->MapBuildData;	
+}
+
+FGuid FStaticLightingBuildContext::GetPersistentLevelGuid() const
+{
+	return *LevelGuids.FindKey(World->PersistentLevel);
+}
+
+FGuid FStaticLightingBuildContext::GetLevelGuidForLevel(ULevel* Level) const
+{
+	return *LevelGuids.FindKey(Level);
+}
+
+FGuid FStaticLightingBuildContext::GetLevelGuidForActor(AActor* Actor) const
+{		
+	check(Actor);
+		
+	if (!World->IsPartitionedWorld())
+	{
+		return *LevelGuids.FindKey(Actor->GetLevel());
+	}
+
+	FGuid LevelGuid = FGuid(0, 0, 0, 0);		
+	return LevelGuid;
+}
+
+
+UMapBuildDataRegistry* FStaticLightingBuildContext::GetOrCreateRegistryForLevelGuid(const FGuid& Guid) const
+{	
+	if (Guid.IsValid())
+	{	
+		if (!World->IsPartitionedWorld())
+		{
+			ULevel* Level = GetLevelForGuid(Guid).Get();
+			return Level->GetOrCreateMapBuildData();
+		}
+		else
+		{
+		}
+	}
+
+	UMapBuildDataRegistry* Registry = GetOrCreateGlobalRegistry();
+	return Registry;
+}
+
+const TWeakObjectPtr<ULevel> FStaticLightingBuildContext::GetLevelForGuid(const FGuid& Guid) const
+{	
+	return LevelGuids.FindRef(Guid);
+}
+
+FGuid FStaticLightingBuildContext::GetLevelBuildDataID(const FGuid& LevelGuid) const
+{
+	if (!World->IsPartitionedWorld())
+	{
+		return GetLevelForGuid(LevelGuid)->LevelBuildDataId;
+	}
+	else
+	{
+		if (!LevelGuid.IsValid())
+		{
+			return World->PersistentLevel->LevelBuildDataId;
+		}
+		
+		return LevelGuid;			
+	}
+}
+
+UMapBuildDataRegistry* FStaticLightingBuildContext::GetOrCreateRegistryForActor(AActor* Actor) const
+{
+	check(Actor);
+
+	ULevel* Level = Actor->GetLevel();	
+
+	return GetLightingStorageLevel(Level)->GetOrCreateMapBuildData();
+}
+
+UMapBuildDataRegistry* FStaticLightingBuildContext::GetRegistryForLevel(ULevel* Level) const
+{
+	return GetLightingStorageLevel(Level)->MapBuildData;
+}
+
+UMapBuildDataRegistry* FStaticLightingBuildContext::GetOrCreateRegistryForLevel(ULevel* Level) const
+{
+	return GetLightingStorageLevel(Level)->GetOrCreateMapBuildData();
+}
+
+#endif
\ No newline at end of file
diff --git a/Engine/Source/Runtime/Engine/Private/StaticMeshLight.cpp b/Engine/Source/Runtime/Engine/Private/StaticMeshLight.cpp
--- a/Engine/Source/Runtime/Engine/Private/StaticMeshLight.cpp
+++ b/Engine/Source/Runtime/Engine/Private/StaticMeshLight.cpp
@@ -12,6 +12,7 @@
 #include "StaticMeshResources.h"
 #include "LightMap.h"
 #include "Engine/MapBuildDataRegistry.h"
+#include "StaticLightingBuildContext.h"
 #include "Components/LightComponent.h"
 #include "ShadowMap.h"
 #include "Engine/StaticMesh.h"
@@ -208,8 +209,9 @@ FStaticMeshStaticLightingTextureMapping::FStaticMeshStaticLightingTextureMapping
 	LODIndex(InLODIndex)
 {}
 
+#if WITH_EDITOR
 // FStaticLightingTextureMapping interface
-void FStaticMeshStaticLightingTextureMapping::Apply(FQuantizedLightmapData* QuantizedData, const TMap<ULightComponent*,FShadowMapData2D*>& ShadowMapData, ULevel* LightingScenario)
+void FStaticMeshStaticLightingTextureMapping::Apply(FQuantizedLightmapData* QuantizedData, const TMap<ULightComponent*,FShadowMapData2D*>& ShadowMapData, const FStaticLightingBuildContext* LightingContext)
 {
 	static const auto CVar = IConsoleManager::Get().FindTConsoleVariableDataInt(TEXT("r.VirtualTexturedLightmaps"));
 	const bool bUseVirtualTextures = (CVar->GetValueOnAnyThread() != 0) && UseVirtualTexturing(GMaxRHIShaderPlatform);
@@ -243,8 +245,7 @@ void FStaticMeshStaticLightingTextureMapping::Apply(FQuantizedLightmapData* Quan
 		ELightMapPaddingType PaddingType = GAllowLightmapPadding ? LMPT_NormalPadding : LMPT_NoPadding;
 		const bool bHasNonZeroData = (QuantizedData != NULL && QuantizedData->HasNonZeroData());
 
-		ULevel* StorageLevel = LightingScenario ? LightingScenario : StaticMeshComponent->GetOwner()->GetLevel();
-		UMapBuildDataRegistry* Registry = StorageLevel->GetOrCreateMapBuildData();
+		UMapBuildDataRegistry* Registry = LightingContext->GetOrCreateRegistryForActor(StaticMeshComponent->GetOwner());
 		FMeshMapBuildData& MeshBuildData = Registry->AllocateMeshBuildData(ComponentLODInfo.MapBuildDataId, true);
 
 		// We always create a light map if the surface either has any non-zero lighting data, or if the surface has a shadow map.  The runtime
@@ -306,7 +307,6 @@ void FStaticMeshStaticLightingTextureMapping::Apply(FQuantizedLightmapData* Quan
 	}
 }
 
-#if WITH_EDITOR
 void UStaticMeshComponent::GetStaticLightingInfo(FStaticLightingPrimitiveInfo& OutPrimitiveInfo,const TArray<ULightComponent*>& InRelevantLights,const FLightingBuildOptions& Options)
 {
 	if( HasValidSettingsForStaticLighting(false) )
diff --git a/Engine/Source/Runtime/Engine/Public/LightMap.h b/Engine/Source/Runtime/Engine/Public/LightMap.h
--- a/Engine/Source/Runtime/Engine/Public/LightMap.h
+++ b/Engine/Source/Runtime/Engine/Public/LightMap.h
@@ -304,13 +304,13 @@ public:
 	 * @param	bLightingSuccessful	Whether the lighting build was successful or not.
 	 * @param	bForceCompletion	Force all encoding to be fully completed (they may be asynchronous).
 	 */
-	static ENGINE_API void EncodeTextures( UWorld* InWorld, ULevel* LightingScenario, bool bLightingSuccessful, bool bMultithreadedEncode = false );
+	static ENGINE_API void EncodeTextures(const FStaticLightingBuildContext* LightingContext, bool bLightingSuccessful, bool bMultithreadedEncode = false );
 
 #if WITH_EDITOR
 	/**
 	 * Constructs mip maps for a single shadowmap texture.
 	 */
-	static ENGINE_API int32 EncodeShadowTexture(ULevel* LightingScenario, struct FLightMapPendingTexture& PendingTexture, TArray<TArray<FFourDistanceFieldSamples>>& MipData);
+	static ENGINE_API int32 EncodeShadowTexture(const FStaticLightingBuildContext* LightingContext, struct FLightMapPendingTexture& PendingTexture, TArray<TArray<FFourDistanceFieldSamples>>& MipData);
 #endif // WITH_EDITOR
 
 	/** Call to enable/disable status update of LightMap encoding */
diff --git a/Engine/Source/Runtime/Engine/Public/Model.h b/Engine/Source/Runtime/Engine/Public/Model.h
--- a/Engine/Source/Runtime/Engine/Public/Model.h
+++ b/Engine/Source/Runtime/Engine/Public/Model.h
@@ -584,7 +584,7 @@ public:
 	/**
 	 * Applies all of the finished lighting cached in the NodeGroups 
 	 */
-	void ApplyStaticLighting(ULevel* LightingScenario);
+	void ApplyStaticLighting(const FStaticLightingBuildContext* LightingContext);
 
 	/**
 	 * Apply world origin changes
diff --git a/Engine/Source/Runtime/Engine/Public/ModelLight.h b/Engine/Source/Runtime/Engine/Public/ModelLight.h
--- a/Engine/Source/Runtime/Engine/Public/ModelLight.h
+++ b/Engine/Source/Runtime/Engine/Public/ModelLight.h
@@ -54,7 +54,7 @@ public:
 	}
 
 #if WITH_EDITOR
-	virtual void Apply(FQuantizedLightmapData* QuantizedData, const TMap<ULightComponent*,FShadowMapData2D*>& ShadowMapData, ULevel* LightingScenario) override;
+	virtual void Apply(FQuantizedLightmapData* QuantizedData, const TMap<ULightComponent*,FShadowMapData2D*>& ShadowMapData, const FStaticLightingBuildContext* LightingContext) override;
 	virtual bool DebugThisMapping() const override;
 
 	/** 
diff --git a/Engine/Source/Runtime/Engine/Public/SceneManagement.h b/Engine/Source/Runtime/Engine/Public/SceneManagement.h
--- a/Engine/Source/Runtime/Engine/Public/SceneManagement.h
+++ b/Engine/Source/Runtime/Engine/Public/SceneManagement.h
@@ -50,6 +50,7 @@ class FPrimitiveSceneProxy;
 class FScene;
 class FSceneViewState;
 class FShadowMap;
+class FStaticLightingBuildContext;
 class FStaticMeshRenderData;
 class FTexture;
 class UDecalComponent;
@@ -914,26 +915,26 @@ class FAsyncEncode : public IQueuedWork
 private:
 	TPendingTextureType* PendingTexture;
 	FThreadSafeCounter& Counter;
-	ULevel* LightingScenario;
+	const FStaticLightingBuildContext* LightingContext;
 	class ITextureCompressorModule* Compressor;
 
 public:
 
-	FAsyncEncode(TPendingTextureType* InPendingTexture, ULevel* InLightingScenario, FThreadSafeCounter& InCounter, ITextureCompressorModule* InCompressor) : PendingTexture(nullptr), Counter(InCounter), Compressor(InCompressor)
+	FAsyncEncode(TPendingTextureType* InPendingTexture, const FStaticLightingBuildContext* InLightingContext, FThreadSafeCounter& InCounter, ITextureCompressorModule* InCompressor) : PendingTexture(nullptr), Counter(InCounter), Compressor(InCompressor)
 	{
-		LightingScenario = InLightingScenario;
+		LightingContext = InLightingContext;
 		PendingTexture = InPendingTexture;
 	}
 
 	void Abandon()
 	{
-		PendingTexture->StartEncoding(LightingScenario, Compressor);
+		PendingTexture->StartEncoding(LightingContext, Compressor);
 		Counter.Decrement();
 	}
 
 	void DoThreadedWork()
 	{
-		PendingTexture->StartEncoding(LightingScenario, Compressor);
+		PendingTexture->StartEncoding(LightingContext, Compressor);
 		Counter.Decrement();
 	}
 };
diff --git a/Engine/Source/Runtime/Engine/Public/ShadowMap.h b/Engine/Source/Runtime/Engine/Public/ShadowMap.h
--- a/Engine/Source/Runtime/Engine/Public/ShadowMap.h
+++ b/Engine/Source/Runtime/Engine/Public/ShadowMap.h
@@ -286,15 +286,15 @@ public:
 
 	/**
 	 * Executes all pending shadow-map encoding requests.
-	 * @param	InWorld				World in which the textures exist
+	 * @param	LightingContext				Context in which the textures exist
 	 * @param	bLightingSuccessful	Whether the lighting build was successful or not.
 	 */
-	static ENGINE_API void EncodeTextures(UWorld* InWorld, ULevel* LightingScenario, bool bLightingSuccessful, bool bMultithreadedEncode =false );
+	static ENGINE_API void EncodeTextures(const FStaticLightingBuildContext* LightingContext, bool bLightingSuccessful, bool bMultithreadedEncode =false );
 
 	/**
 	 * Constructs mip maps for a single shadowmap texture.
 	 */
-	static ENGINE_API int32 EncodeSingleTexture(ULevel* LightingScenario, struct FShadowMapPendingTexture& PendingTexture, UShadowMapTexture2D* Texture, TArray< TArray<FFourDistanceFieldSamples>>& MipData);
+	static ENGINE_API int32 EncodeSingleTexture(const FStaticLightingBuildContext* LightingContext, struct FShadowMapPendingTexture& PendingTexture, UShadowMapTexture2D* Texture, TArray< TArray<FFourDistanceFieldSamples>>& MipData);
 
 	static ENGINE_API TRefCountPtr<FShadowMap2D> AllocateShadowMap(
 		UObject* LightMapOuter,
diff --git a/Engine/Source/Runtime/Engine/Public/StaticLighting.h b/Engine/Source/Runtime/Engine/Public/StaticLighting.h
--- a/Engine/Source/Runtime/Engine/Public/StaticLighting.h
+++ b/Engine/Source/Runtime/Engine/Public/StaticLighting.h
@@ -12,6 +12,7 @@
 #include "Components.h"
 
 class FStaticLightingMapping;
+class FStaticLightingBuildContext;
 class FStaticLightingTextureMapping;
 class ULevel;
 class ULightComponent;
@@ -395,7 +396,7 @@ public:
 	 * This function is responsible for deleting ShadowMapData and QuantizedData.
 	 * @param LightMapData - The light-map data which has been computed for the mapping.
 	 */
-	virtual void Apply(struct FQuantizedLightmapData* QuantizedData, const TMap<ULightComponent*,class FShadowMapData2D*>& ShadowMapData, ULevel* LightingScenario) = 0;
+	ENGINE_API virtual void Apply(struct FQuantizedLightmapData* QuantizedData, const TMap<ULightComponent*,class FShadowMapData2D*>& ShadowMapData, const FStaticLightingBuildContext* LightingContext) = 0;
 
 	// FStaticLightingMapping interface.
 	virtual FStaticLightingTextureMapping* GetTextureMapping()
@@ -438,7 +439,7 @@ public:
 	/** Initialization constructor. */
 	ENGINE_API FStaticLightingGlobalVolumeMapping(FStaticLightingMesh* InMesh, UObject* InOwner, int32 InSizeX, int32 InSizeY, int32 InLightmapTextureCoordinateIndex);
 
-	virtual void Apply(struct FQuantizedLightmapData* QuantizedData, const TMap<ULightComponent*,class FShadowMapData2D*>& ShadowMapData, ULevel* LightingScenario) override
+	virtual void Apply(struct FQuantizedLightmapData* QuantizedData, const TMap<ULightComponent*,class FShadowMapData2D*>& ShadowMapData, const FStaticLightingBuildContext* LightingContext) override
 	{
 		// Should never be processed
 		check(false);
diff --git a/Engine/Source/Runtime/Engine/Public/StaticLightingBuildContext.h b/Engine/Source/Runtime/Engine/Public/StaticLightingBuildContext.h
new file mode 100644
--- /dev/null
+++ b/Engine/Source/Runtime/Engine/Public/StaticLightingBuildContext.h
@@ -0,0 +1,52 @@
+// Copyright Epic Games, Inc. All Rights Reserved.
+
+#pragma once
+
+#include "CoreMinimal.h"
+#include "UObject/ObjectMacros.h"
+#include "UObject/Object.h"
+#include "UObject/WeakObjectPtr.h"
+#include "Misc/Guid.h"
+
+class UMapBuildDataRegistry;
+class ULevel;
+class UWorld;
+class AActor;
+struct FVolumetricLightMapGridDesc;
+
+#if WITH_EDITOR
+class FStaticLightingBuildContext
+{	
+	mutable UMapBuildDataRegistry* MapBuildDataRegistry;
+		
+	TMap<FGuid, const TWeakObjectPtr<ULevel>> LevelGuids;	
+
+public:
+
+	UWorld*						World;
+	ULevel*						LightingScenario;
+
+	ENGINE_API FStaticLightingBuildContext(UWorld* InWorld, ULevel* InLightingScenario);
+	ENGINE_API FStaticLightingBuildContext(const FStaticLightingBuildContext& InFrom);
+	
+	ENGINE_API ULevel* GetLightingStorageLevel(ULevel* Level) const;
+
+	ENGINE_API bool ShouldIncludeActor(AActor* Actor) const;
+	ENGINE_API bool ShouldIncludeLevel(ULevel* Level) const;
+
+	ENGINE_API UMapBuildDataRegistry* GetGlobalRegistry() const;
+	ENGINE_API UMapBuildDataRegistry* GetOrCreateGlobalRegistry() const;
+
+	ENGINE_API UMapBuildDataRegistry* GetRegistryForActor(AActor* Actor) const;
+	ENGINE_API UMapBuildDataRegistry* GetOrCreateRegistryForLevelGuid(const FGuid& Guid) const;
+	ENGINE_API FGuid GetLevelGuidForActor(AActor* Actor) const;
+	ENGINE_API FGuid GetPersistentLevelGuid() const;
+	ENGINE_API FGuid GetLevelGuidForLevel(ULevel* Level) const;	
+	ENGINE_API const TWeakObjectPtr<ULevel> GetLevelForGuid(const FGuid& Guid) const;
+	ENGINE_API FGuid GetLevelBuildDataID(const FGuid& LevelGuid) const;
+	
+	ENGINE_API UMapBuildDataRegistry* GetRegistryForLevel(ULevel* Level) const;	
+	ENGINE_API UMapBuildDataRegistry* GetOrCreateRegistryForActor(AActor* Actor) const;
+	ENGINE_API UMapBuildDataRegistry* GetOrCreateRegistryForLevel(ULevel* Level) const;	
+};
+#endif // WITH_EDITOR
\ No newline at end of file
diff --git a/Engine/Source/Runtime/Engine/Public/StaticMeshLight.h b/Engine/Source/Runtime/Engine/Public/StaticMeshLight.h
--- a/Engine/Source/Runtime/Engine/Public/StaticMeshLight.h
+++ b/Engine/Source/Runtime/Engine/Public/StaticMeshLight.h
@@ -108,10 +108,10 @@ public:
 	/** Initialization constructor. */
 	FStaticMeshStaticLightingTextureMapping(UStaticMeshComponent* InPrimitive,int32 InLODIndex,FStaticLightingMesh* InMesh,int32 InSizeX,int32 InSizeY,int32 InTextureCoordinateIndex,bool bPerformFullQualityRebuild);
 
+#if WITH_EDITOR		
 	// FStaticLightingTextureMapping interface
-	virtual void Apply(FQuantizedLightmapData* QuantizedData, const TMap<ULightComponent*,FShadowMapData2D*>& ShadowMapData, ULevel* LightingScenario) override;
+	ENGINE_API virtual void Apply(FQuantizedLightmapData* QuantizedData, const TMap<ULightComponent*,FShadowMapData2D*>& ShadowMapData, const FStaticLightingBuildContext* LightingContext) override;
 
-#if WITH_EDITOR
 	/** 
 	 * Export static lighting mapping instance data to an exporter 
 	 * @param Exporter - export interface to process static lighting data
diff --git a/Engine/Source/Runtime/Landscape/Private/Landscape.cpp b/Engine/Source/Runtime/Landscape/Private/Landscape.cpp
--- a/Engine/Source/Runtime/Landscape/Private/Landscape.cpp
+++ b/Engine/Source/Runtime/Landscape/Private/Landscape.cpp
@@ -1913,18 +1913,8 @@ const FMeshMapBuildData* ULandscapeComponent::GetMeshMapBuildData() const
 
 		if (OwnerLevel && OwnerLevel->OwningWorld)
 		{
-			ULevel* ActiveLightingScenario = OwnerLevel->OwningWorld->GetActiveLightingScenario();
-			UMapBuildDataRegistry* MapBuildData = NULL;
-
-			if (ActiveLightingScenario && ActiveLightingScenario->MapBuildData)
-			{
-				MapBuildData = ActiveLightingScenario->MapBuildData;
-			}
-			else if (OwnerLevel->MapBuildData)
-			{
-				MapBuildData = OwnerLevel->MapBuildData;
-			}
-
+			UMapBuildDataRegistry* MapBuildData = UMapBuildDataRegistry::Get(this);
+			
 			if (MapBuildData)
 			{
 				return MapBuildData->GetMeshBuildData(MapBuildDataId);
diff --git a/Engine/Source/Runtime/Landscape/Private/LandscapeLight.cpp b/Engine/Source/Runtime/Landscape/Private/LandscapeLight.cpp
--- a/Engine/Source/Runtime/Landscape/Private/LandscapeLight.cpp
+++ b/Engine/Source/Runtime/Landscape/Private/LandscapeLight.cpp
@@ -12,6 +12,7 @@ LandscapeLight.cpp: Static lighting for LandscapeComponents
 #include "LightMap.h"
 #include "Engine/Level.h"
 #include "Engine/MapBuildDataRegistry.h"
+#include "StaticLightingBuildContext.h"
 #include "Components/LightComponent.h"
 #include "ShadowMap.h"
 #include "LandscapeComponent.h"
@@ -44,7 +45,7 @@ FStaticLightingTextureMapping(
 {
 }
 
-void FLandscapeStaticLightingTextureMapping::Apply(FQuantizedLightmapData* QuantizedData, const TMap<ULightComponent*,FShadowMapData2D*>& ShadowMapData, ULevel* LightingScenario)
+void FLandscapeStaticLightingTextureMapping::Apply(FQuantizedLightmapData* QuantizedData, const TMap<ULightComponent*,FShadowMapData2D*>& ShadowMapData, const FStaticLightingBuildContext* LightingContext)
 {
 	static const auto CVar = IConsoleManager::Get().FindTConsoleVariableDataInt(TEXT("r.VirtualTexturedLightmaps"));
 	const bool bUseVirtualTextures = (CVar->GetValueOnAnyThread() != 0) && UseVirtualTexturing(GMaxRHIShaderPlatform);
@@ -52,8 +53,7 @@ void FLandscapeStaticLightingTextureMapping::Apply(FQuantizedLightmapData* Quant
 	//ELightMapPaddingType PaddingType = GAllowLightmapPadding ? LMPT_NormalPadding : LMPT_NoPadding;
 	ELightMapPaddingType PaddingType = LMPT_NoPadding;
 
-	ULevel* StorageLevel = LightingScenario ? LightingScenario : LandscapeComponent->GetOwner()->GetLevel();
-	UMapBuildDataRegistry* Registry = StorageLevel->GetOrCreateMapBuildData();
+	UMapBuildDataRegistry* Registry = LightingContext->GetOrCreateRegistryForActor(LandscapeComponent->GetOwner());
 	FMeshMapBuildData& MeshBuildData = Registry->AllocateMeshBuildData(LandscapeComponent->MapBuildDataId, true);
 
 	const bool bHasNonZeroData = QuantizedData != NULL && QuantizedData->HasNonZeroData();
diff --git a/Engine/Source/Runtime/Landscape/Public/LandscapeLight.h b/Engine/Source/Runtime/Landscape/Public/LandscapeLight.h
--- a/Engine/Source/Runtime/Landscape/Public/LandscapeLight.h
+++ b/Engine/Source/Runtime/Landscape/Public/LandscapeLight.h
@@ -26,7 +26,7 @@ public:
 	FLandscapeStaticLightingTextureMapping(ULandscapeComponent* InPrimitive,FStaticLightingMesh* InMesh,int32 InLightMapWidth,int32 InLightMapHeight,bool bPerformFullQualityRebuild);
 
 	// FStaticLightingTextureMapping interface
-	virtual void Apply(FQuantizedLightmapData* QuantizedData, const TMap<ULightComponent*,FShadowMapData2D*>& ShadowMapData, ULevel* LightingScenario);
+	LANDSCAPE_API virtual void Apply(FQuantizedLightmapData* QuantizedData, const TMap<ULightComponent*,FShadowMapData2D*>& ShadowMapData, const FStaticLightingBuildContext* LightingContext);
 
 #if WITH_EDITOR
 	/** 
