MPCook: Add error for invalid cookprocesscount rather than hitting a more obscure error later on in FWorkerId::FromRemoteIndex.
#rnx

diff --git a/Engine/Source/Editor/UnrealEd/Private/CookOnTheFlyServer.cpp b/Engine/Source/Editor/UnrealEd/Private/CookOnTheFlyServer.cpp
--- a/Engine/Source/Editor/UnrealEd/Private/CookOnTheFlyServer.cpp
+++ b/Engine/Source/Editor/UnrealEd/Private/CookOnTheFlyServer.cpp
@@ -6757,6 +6757,13 @@ void UCookOnTheFlyServer::Initialize( ECookMode::Type DesiredCookMode, ECookInit
 			GConfig->GetInt(TEXT("CookSettings"), TEXT("CookProcessCount"), CookProcessCount, GEditorIni);
 		}
 		CookProcessCount = FMath::Max(1, CookProcessCount);
+		if (CookProcessCount > UE::Cook::FWorkerId::GetMaxCookWorkerCount())
+		{
+			// We could clamp it and continue on, but it's not clear what to clamp it to. If they ask for
+			// 1 billion by accidental typo in the ini, what should we set it to?
+			UE_LOG(LogCook, Fatal, TEXT("Invalid CookProcessCount=%d, maximum value is %d."),
+				CookProcessCount, UE::Cook::FWorkerId::GetMaxCookWorkerCount());
+		}
 		if (CookProcessCount > 1)
 		{
 			CookDirector = MakeUnique<UE::Cook::FCookDirector>(*this, CookProcessCount);
diff --git a/Engine/Source/Runtime/CoreUObject/Public/Cooker/MPCollector.h b/Engine/Source/Runtime/CoreUObject/Public/Cooker/MPCollector.h
--- a/Engine/Source/Runtime/CoreUObject/Public/Cooker/MPCollector.h
+++ b/Engine/Source/Runtime/CoreUObject/Public/Cooker/MPCollector.h
@@ -39,6 +39,7 @@ public:
 	constexpr static FWorkerId Local() { return FWorkerId(LocalId); }
 	static FWorkerId FromRemoteIndex(uint8 Index) { check(Index < InvalidId - 1U);  return FWorkerId(Index + 1U); }
 	static FWorkerId FromLocalOrRemoteIndex(uint8 Index) { check(Index < InvalidId);  return FWorkerId(Index); }
+	static int32 GetMaxCookWorkerCount() { return static_cast<int32>(InvalidId - 1); }
 
 	bool IsValid() const { return Id != InvalidId; }
 	bool IsInvalid() const { return Id == InvalidId; }
