Fixed crash when an actor factory returned null from CreateActor

#rnx

diff --git a/Engine/Source/Editor/UnrealEd/Private/AssetSelection.cpp b/Engine/Source/Editor/UnrealEd/Private/AssetSelection.cpp
--- a/Engine/Source/Editor/UnrealEd/Private/AssetSelection.cpp
+++ b/Engine/Source/Editor/UnrealEd/Private/AssetSelection.cpp
@@ -677,11 +677,13 @@ TArray<FTypedElementHandle> PlaceAssetUsingFactory(UObject* Asset, TScriptInterf
 			SpawnParams.ObjectFlags = ObjectFlags;
 			SpawnParams.Name = Name;
 
-			AActor* PlacedActor = ActorFactory->CreateActor(Asset, DesiredLevel, ActorTransform, SpawnParams);
-			FTypedElementHandle Handle = UEngineElementsLibrary::AcquireEditorActorElementHandle(PlacedActor);
-			if (ensure(Handle))
+			if (AActor* PlacedActor = ActorFactory->CreateActor(Asset, DesiredLevel, ActorTransform, SpawnParams))
 			{
-				PlacedItems.Add(Handle);
+				FTypedElementHandle Handle = UEngineElementsLibrary::AcquireEditorActorElementHandle(PlacedActor);
+				if (ensure(Handle))
+				{
+					PlacedItems.Add(Handle);
+				}
 			}
 		}
 	}
