Fixed missed spot in logic to disable FBX level import only

#rnx

diff --git a/Engine/Source/Editor/UnrealEd/Private/FileHelpers.cpp b/Engine/Source/Editor/UnrealEd/Private/FileHelpers.cpp
--- a/Engine/Source/Editor/UnrealEd/Private/FileHelpers.cpp
+++ b/Engine/Source/Editor/UnrealEd/Private/FileHelpers.cpp
@@ -1539,7 +1539,8 @@ void FEditorFileUtils::Import(const FString& InFilename)
 	const FScopedBusyCursor BusyCursor;
 
 	UE::Interchange::FScopedSourceData ScopedSourceData(InFilename);
-	const bool bImportThroughInterchange = UInterchangeManager::GetInterchangeManager().CanTranslateSourceData(ScopedSourceData.GetSourceData());
+	const bool bIsSceneImport = true; // Only scene import is requested from FEditorFileUtils::Import
+	const bool bImportThroughInterchange = UInterchangeManager::GetInterchangeManager().CanTranslateSourceData(ScopedSourceData.GetSourceData(), bIsSceneImport);
 
 	USceneImportFactory* SceneFactory = nullptr;
 
