Fix for ActorPositioning when hitting an object without a primitive and `bAllowNonPrimitiveComponentHits` is set to true

diff --git a/Engine/Source/Editor/UnrealEd/Private/Editor/ObjectPositioning.cpp b/Engine/Source/Editor/UnrealEd/Private/Editor/ObjectPositioning.cpp
--- a/Engine/Source/Editor/UnrealEd/Private/Editor/ObjectPositioning.cpp
+++ b/Engine/Source/Editor/UnrealEd/Private/Editor/ObjectPositioning.cpp
@@ -24,6 +24,8 @@ namespace ObjectPositioningLocals
 	/**
 	 * Prunes list of hit results for object positioning calculations based on conditions that could be tested
 	 * on the game thread and returns a list of primitives for the remaining this.
+	 * @note If a non-primitive based hit is found and bCVarAllowNonPrimitiveComponentHits is true then
+     * an empty weak obj ptr will be added to the result to represent the hit.
 	 */
 	TArray<TWeakObjectPtr<const UPrimitiveComponent>> FilterHitsGameThread(TArray<FHitResult>& InOutHits)
 	{
