Fix thumbnails of transparent textures broken in CL 30324001
SE_BLEND_Translucent does not work with SimpleElementTexture2DPreviewPixelShader, must use SE_BLEND_RGBA_MASK_START instead

diff --git a/Engine/Source/Editor/UnrealEd/Private/ThumbnailRendering/TextureThumbnailRenderer.cpp b/Engine/Source/Editor/UnrealEd/Private/ThumbnailRendering/TextureThumbnailRenderer.cpp
--- a/Engine/Source/Editor/UnrealEd/Private/ThumbnailRendering/TextureThumbnailRenderer.cpp
+++ b/Engine/Source/Editor/UnrealEd/Private/ThumbnailRendering/TextureThumbnailRenderer.cpp
@@ -75,12 +75,16 @@ void UTextureThumbnailRenderer::Draw(UObject* Object, int32 X, int32 Y, uint32 W
 		{
 			CpuCopyTexture = Texture2D->GetCPUCopyTexture();
 		}
+
+		// behavior here should match FTextureEditorViewportClient as much as possible
+
 		if (CpuCopyTexture == nullptr)
 		{
 			// Take the alpha channel into account for textures that have one.
 			// This provides a much better preview than just showing RGB,
 			// Because the RGB content in areas with an alpha of 0 is often garbage that will not be seen in normal conditions.
 			// Non-UI textures often have uncorrelated data in the alpha channel (like a skin mask, specular power, etc) so we only preview UI textures this way.
+			// @todo : this logic depending on LODGroup to decide whether alpha thumbnails should blend is very odd and unexpected for the user.  Consider changing.
 			const bool bUseTranslucentBlend = Texture2D && Texture2D->HasAlphaChannel() && ((Texture2D->LODGroup == TEXTUREGROUP_UI) || (Texture2D->LODGroup == TEXTUREGROUP_Pixels2D));
 
 			UTextureCube* TextureCube = Cast<UTextureCube>(Texture);
@@ -137,6 +141,7 @@ void UTextureThumbnailRenderer::Draw(UObject* Object, int32 X, int32 Y, uint32 W
 			{
 				// BatchedElementParameters is not set
 				// default Canvas will be used, which just shows the texture
+				// (what shader does this use??)
 
 				// some UTexture types can hit this
 				// UTextureRenderTarget
@@ -157,7 +162,17 @@ void UTextureThumbnailRenderer::Draw(UObject* Object, int32 X, int32 Y, uint32 W
 
 			// Use A canvas tile item to draw
 			FCanvasTileItem CanvasTile( FVector2D( X, Y ), Texture->GetResource(), FVector2D( Width,Height ), FLinearColor::White );
-			CanvasTile.BlendMode = bUseTranslucentBlend ? SE_BLEND_Translucent : SE_BLEND_Opaque;
+			if ( BatchedElementParameters == nullptr )
+			{
+				// old style BatchedElementParameters == null wants you to use "Translucent" :
+				CanvasTile.BlendMode = bUseTranslucentBlend ? SE_BLEND_Translucent : SE_BLEND_Opaque;
+			}
+			else
+			{
+				// matches the behavior of FTextureEditorViewportClient::Draw and FTextureEditorToolkit::GetColourChannelBlendMode
+				// I think it may be a bug in BatchedElement that SE_BLEND_Translucent doesn't work here (failing to set up TextureComponentReplicate?)
+				CanvasTile.BlendMode = bUseTranslucentBlend ? (ESimpleElementBlendMode)(SE_BLEND_RGBA_MASK_START+0xF) : SE_BLEND_Opaque;
+			}
 			CanvasTile.BatchedElementParameters = BatchedElementParameters;
 			if (bIsVirtualTexture && Texture->Source.GetNumBlocks() > 1)
 			{
