Fixed attach/detach not always updating properly attached/detached actor folder.

diff --git a/Engine/Source/Editor/SceneOutliner/Private/ActorMode.cpp b/Engine/Source/Editor/SceneOutliner/Private/ActorMode.cpp
--- a/Engine/Source/Editor/SceneOutliner/Private/ActorMode.cpp
+++ b/Engine/Source/Editor/SceneOutliner/Private/ActorMode.cpp
@@ -938,8 +938,6 @@ void FActorMode::OnDrop(ISceneOutlinerTreeItem& DropTarget, const FSceneOutliner
 							if (GEditor->CanParentActors(ParentActor, ChildActor))
 							{
 								GEditor->ParentActors(ParentActor, ChildActor, SocketName);
-
-								ChildActor->SetFolderPath_Recursively(ParentActor->GetFolderPath());
 							}
 						}
 					}
diff --git a/Engine/Source/Editor/UnrealEd/Private/EditorEngine.cpp b/Engine/Source/Editor/UnrealEd/Private/EditorEngine.cpp
--- a/Engine/Source/Editor/UnrealEd/Private/EditorEngine.cpp
+++ b/Engine/Source/Editor/UnrealEd/Private/EditorEngine.cpp
@@ -339,6 +339,22 @@ void DestroySelectionSets()
 
 } // namespace PrivateEditorSelection
 
+namespace ActorAttachmentHelper
+{
+	static const AActor* GetTopMostAttachParentActor(AActor* InActor)
+	{
+		check(InActor);
+		AActor* TopAttachParentActor = InActor;
+		AActor* NextAttachParentActor = TopAttachParentActor->GetAttachParentActor();
+		while (NextAttachParentActor)
+		{
+			TopAttachParentActor = NextAttachParentActor;
+			NextAttachParentActor = TopAttachParentActor->GetAttachParentActor();
+		}
+		return TopAttachParentActor;
+	}
+} // namespace ActorAttachmentHelper
+
 static FAutoConsoleVariable GInvalidateHitProxiesEachSIEFrameCVar(
 	TEXT("r.Editor.Viewport.InvalidateEachSIEFrame"),
 	1,
@@ -4370,6 +4386,10 @@ void UEditorEngine::ParentActors( AActor* ParentActor, AActor* ChildActor, const
 		// Snap to socket if a valid socket name was provided, otherwise attach without changing the relative transform
 		ChildRoot->AttachToComponent(Component ? Component : ParentRoot, FAttachmentTransformRules::KeepWorldTransform, SocketName);
 
+		// Update recursively attached actor folder using top most parent folder
+		const AActor* TopMostAttachParentActor = ActorAttachmentHelper::GetTopMostAttachParentActor(ParentActor);
+		ChildActor->SetFolderPath_Recursively(TopMostAttachParentActor->GetFolderPath());
+
 		// Refresh editor in case child was translated after snapping to socket
 		RedrawLevelEditingViewports();
 	}
@@ -4392,7 +4412,10 @@ bool UEditorEngine::DetachSelectedActors()
 			OldParentActor->Modify(false);
 			RootComp->DetachFromComponent(FDetachmentTransformRules::KeepWorldTransform);
 			bDetachOccurred = true;
-			Actor->SetFolderPath_Recursively(OldParentActor->GetFolderPath());
+
+			// Update recursively detached actor folder using top most parent folder
+			const AActor* TopMostAttachParentActor = ActorAttachmentHelper::GetTopMostAttachParentActor(OldParentActor);
+			Actor->SetFolderPath_Recursively(TopMostAttachParentActor->GetFolderPath());
 		}
 	}
 	return bDetachOccurred;
