Property Coloration: add an option to specify the default color for actors that are not having the current property set.

diff --git a/Engine/Source/Editor/UnrealEd/Private/EditorEngine.cpp b/Engine/Source/Editor/UnrealEd/Private/EditorEngine.cpp
--- a/Engine/Source/Editor/UnrealEd/Private/EditorEngine.cpp
+++ b/Engine/Source/Editor/UnrealEd/Private/EditorEngine.cpp
@@ -550,15 +550,14 @@ UEditorEngine::UEditorEngine(const FObjectInitializer& ObjectInitializer)
 				FInternationalization::ForUseOnlyByLocMacroAndGraphNodeTextLiterals_CreateText(*PropertyColorCustomProperty.Text, TEXT("PropertyColor"), *PropertyColorCustomProperty.Name.ToString()),
 				[this, PropertyColorCustomProperty](const UPrimitiveComponent* InPrimitiveComponent)
 				{
-					FColor PropertyColor(FColor::White);
 					if (AActor* Actor = InPrimitiveComponent->GetOwner())
 					{
 						if (GetPropertyColorationMatch(Actor))
 						{
-							PropertyColor = PropertyColorCustomProperty.PropertyColor;
+							return PropertyColorCustomProperty.PropertyColor;
 						}
 					}
-					return PropertyColor;
+					return PropertyColorCustomProperty.DefaultColor;
 				},
 				[this, PropertyColorCustomProperty]()
 				{
diff --git a/Engine/Source/Editor/UnrealEd/Private/PropertyColorSettings.h b/Engine/Source/Editor/UnrealEd/Private/PropertyColorSettings.h
--- a/Engine/Source/Editor/UnrealEd/Private/PropertyColorSettings.h
+++ b/Engine/Source/Editor/UnrealEd/Private/PropertyColorSettings.h
@@ -26,6 +26,9 @@ struct FPropertyColorCustomProperty
 
 	UPROPERTY(Config)
 	FColor PropertyColor = FColor::Red;
+
+	UPROPERTY(Config)
+	FColor DefaultColor = FColor::White;
 };
 
 /**
