Add a delegate that allows filtering actors out of the selection when copying.

diff --git a/Engine/Source/Editor/UnrealEd/Classes/Editor/EditorEngine.h b/Engine/Source/Editor/UnrealEd/Classes/Editor/EditorEngine.h
--- a/Engine/Source/Editor/UnrealEd/Classes/Editor/EditorEngine.h
+++ b/Engine/Source/Editor/UnrealEd/Classes/Editor/EditorEngine.h
@@ -2357,6 +2357,10 @@ public:
 	FOnCreateLevelEditorDragDropHandler& OnCreateLevelEditorDragDropHandler() { return OnCreateLevelEditorDragDropHandlerDelegate; }
 	UNREALED_API ULevelEditorDragDropHandler* GetLevelEditorDragDropHandler() const;
 
+	DECLARE_MULTICAST_DELEGATE_OneParam(FOnFilterCopiedActors, TArray<AActor*>&);
+	UE_INTERNAL FOnFilterCopiedActors& OnFilterCopiedActors() { return OnFilterCopiedActorsDelegate; }
+	UE_INTERNAL const FOnFilterCopiedActors& OnFilterCopiedActors() const { return OnFilterCopiedActorsDelegate; }
+	
 	/** 
 	 * Gets the interface to manage project references to external content
 	 * @note the returned pointer cannot be null
@@ -2373,6 +2377,7 @@ private:
 
 	FOnMakeAssetReferenceFilter OnMakeAssetReferenceFilterDelegate;
 	FOnCreateLevelEditorDragDropHandler OnCreateLevelEditorDragDropHandlerDelegate;
+	FOnFilterCopiedActors OnFilterCopiedActorsDelegate;
 
 public:
 
diff --git a/Engine/Source/Editor/UnrealEd/Private/EditorActor.cpp b/Engine/Source/Editor/UnrealEd/Private/EditorActor.cpp
--- a/Engine/Source/Editor/UnrealEd/Private/EditorActor.cpp
+++ b/Engine/Source/Editor/UnrealEd/Private/EditorActor.cpp
@@ -158,6 +158,8 @@ void UUnrealEdEngine::CopyActors(const TArray<AActor*>& InActorsToCopy, UWorld*
 			return false;
 		});
 	}
+	// Filter out other actors.
+	OnFilterCopiedActors().Broadcast(ActorsToCopy);
 
 	// Export the actors.
 	FStringOutputDevice Ar;
