Improve FBX Sequencer export:
- Add support of Bool and Enum channels for level sequence tracks
- a User animatable track for IntensityUnits and bUseTemperture
- Export color track as one track with 3 channels instead of 4 different property tracks

diff --git a/Engine/Source/Editor/UnrealEd/Private/Fbx/FbxMainExport.cpp b/Engine/Source/Editor/UnrealEd/Private/Fbx/FbxMainExport.cpp
--- a/Engine/Source/Editor/UnrealEd/Private/Fbx/FbxMainExport.cpp
+++ b/Engine/Source/Editor/UnrealEd/Private/Fbx/FbxMainExport.cpp
@@ -81,11 +81,13 @@
 #include "IMovieScenePlayer.h"
 #include "MovieScene.h"
 #include "Tracks/MovieScene3DTransformTrack.h"
+#include "Tracks/MovieSceneColorTrack.h"
 #include "Tracks/MovieSceneDoubleTrack.h"
 #include "Tracks/MovieSceneFloatTrack.h"
 #include "Tracks/MovieSceneSkeletalAnimationTrack.h"
 #include "Sections/MovieSceneSkeletalAnimationSection.h"
 #include "Sections/MovieScene3DTransformSection.h"
+#include "Sections/MovieSceneColorSection.h"
 #include "Sections/MovieSceneDoubleSection.h"
 #include "Sections/MovieSceneFloatSection.h"
 #include "Evaluation/MovieScenePlayback.h"
@@ -569,6 +571,8 @@ void FFbxExporter::FillFbxLightAttribute(FbxLight* Light, FbxNode* FbxParentNode
 
 	// Add one user property for recording the Brightness animation
 	CreateAnimatableUserProperty(FbxParentNode, BaseLight->Intensity, "UE_Intensity", "UE_Matinee_Light_Intensity");
+	CreateAnimatableUserProperty(FbxParentNode, BaseLight->bUseTemperature, "UE_UseTemperature", "UE_Matinee_Light_UseTemperature", FbxBoolDT);
+	CreateAnimatableUserProperty(FbxParentNode, BaseLight->GetLightUnits(), "UE_IntensityUnits", "UE_Matinee_UE_IntensityUnits", FbxEnumDT);
 
 	// Look for the higher-level light types and determine the lighting method
 	if (BaseLight->IsA(UPointLightComponent::StaticClass()))
@@ -1930,6 +1934,10 @@ bool FFbxExporter::ExportLevelSequenceTracks(UMovieScene* MovieScene, IMovieScen
 		{
 			SkeletalAnimationTrack = Cast<UMovieSceneSkeletalAnimationTrack>(Track);
 		}
+		else if (Track->IsA(UMovieSceneColorTrack::StaticClass()))
+		{
+			ExportLevelSequenceColorTrack(FbxActor, *Cast<UMovieSceneColorTrack>(Track), BoundObject, MovieScene->GetPlaybackRange(), RootToLocalTransform);
+		}
 		else
 		{
 			bool bBakeChannels = false;
@@ -2932,6 +2940,41 @@ void FFbxExporter::ExportConstantChannelToFbxCurve(FbxAnimCurve& InFbxCurve, con
 	InFbxCurve.KeyModifyEnd();
 }
 
+void FFbxExporter::ExportLevelSequenceColorTrack(FbxNode* FbxNode, UMovieSceneColorTrack& ColorTrack, UObject* BoundObject, const TRange<FFrameNumber>& InPlaybackRange, const FMovieSceneSequenceTransform& RootToLocalTransform)
+{
+	UMovieSceneColorSection* ColorSection = ColorTrack.GetAllSections().Num() > 0
+		? Cast<UMovieSceneColorSection>(ColorTrack.GetAllSections()[0])
+		: nullptr;
+
+	if(!ColorSection)
+	{
+		return;
+	}
+
+	if(!FbxNode)
+	{
+		FbxNode = CreateNode(ColorTrack.GetDisplayName().ToString());
+	}
+
+	FbxLight* FbxLight = FbxNode->GetLight();	
+
+	FbxProperty Property;
+	if(FbxLight)
+	{
+		Property = FbxLight->Color;
+
+		FbxAnimLayer* BaseLayer = AnimStack->GetMember<FbxAnimLayer>(0);
+		FbxAnimCurve* CurveRed =  Property.GetCurve(BaseLayer, FBXSDK_CURVENODE_COLOR_RED, true);
+		FbxAnimCurve* CurveGreen =  Property.GetCurve(BaseLayer, FBXSDK_CURVENODE_COLOR_GREEN, true);
+		FbxAnimCurve* CurveBlue =  Property.GetCurve(BaseLayer, FBXSDK_CURVENODE_COLOR_BLUE, true);
+
+		FFrameRate TickResolution = ColorTrack.GetTypedOuter<UMovieScene>()->GetTickResolution();
+		ExportChannelToFbxCurve(*CurveRed, ColorSection->GetRedChannel(), TickResolution, ERichCurveValueMode::Default, false, RootToLocalTransform);
+		ExportChannelToFbxCurve(*CurveGreen, ColorSection->GetGreenChannel(), TickResolution, ERichCurveValueMode::Default, false, RootToLocalTransform);
+		ExportChannelToFbxCurve(*CurveBlue, ColorSection->GetBlueChannel(), TickResolution, ERichCurveValueMode::Default, false, RootToLocalTransform);
+	}
+}
+
 void FFbxExporter::ExportLevelSequence3DTransformTrack(FbxNode* FbxNode, IMovieScenePlayer* MovieScenePlayer, FMovieSceneSequenceIDRef InSequenceID, UMovieScene3DTransformTrack& TransformTrack, UObject* BoundObject, const TRange<FFrameNumber>& InPlaybackRange, const FMovieSceneSequenceTransform& RootToLocalTransform)
 {
 	// TODO: Support more than one section?
@@ -3350,8 +3393,6 @@ void FFbxExporter::ExportLevelSequenceBaked3DTransformTrack(IAnimTrackAdapter& A
 		FbxCurveScaleZ->KeyModifyEnd();
 	}
 }
-
-
 void FFbxExporter::ExportLevelSequenceTrackChannels( FbxNode* FbxNode, UMovieSceneTrack& Track, const TRange<FFrameNumber>& InPlaybackRange, const FMovieSceneSequenceTransform& RootToLocalTransform, bool bBakeBezierCurves)
 {
 	// TODO: Support more than one section?
@@ -3370,7 +3411,9 @@ void FFbxExporter::ExportLevelSequenceTrackChannels( FbxNode* FbxNode, UMovieSce
 	FbxCamera* FbxCamera = FbxNode->GetCamera();
 	FFrameRate TickResolution = Track.GetTypedOuter<UMovieScene>()->GetTickResolution();
 
+	const FName BoolChannelTypeName = FMovieSceneBoolChannel::StaticStruct()->GetFName();
 	const FName DoubleChannelTypeName = FMovieSceneDoubleChannel::StaticStruct()->GetFName();
+	const FName EnumChannelTypeName = FMovieSceneByteChannel::StaticStruct()->GetFName();
 	const FName FloatChannelTypeName = FMovieSceneFloatChannel::StaticStruct()->GetFName();
 	const FName IntegerChannelTypeName = FMovieSceneIntegerChannel::StaticStruct()->GetFName();
 	const FName StringChannelTypeName = FMovieSceneStringChannel::StaticStruct()->GetFName();
@@ -3378,7 +3421,9 @@ void FFbxExporter::ExportLevelSequenceTrackChannels( FbxNode* FbxNode, UMovieSce
 	for (const FMovieSceneChannelEntry& Entry : Section->GetChannelProxy().GetAllEntries())
 	{
 		const FName ChannelTypeName = Entry.GetChannelTypeName();
-		if (ChannelTypeName != DoubleChannelTypeName && 
+		if (ChannelTypeName != BoolChannelTypeName && 
+			ChannelTypeName != DoubleChannelTypeName && 
+			ChannelTypeName != EnumChannelTypeName && 
 			ChannelTypeName != FloatChannelTypeName && 
 			ChannelTypeName != IntegerChannelTypeName && 
 			ChannelTypeName != StringChannelTypeName)
@@ -3393,12 +3438,14 @@ void FFbxExporter::ExportLevelSequenceTrackChannels( FbxNode* FbxNode, UMovieSce
 		{
 			FMovieSceneChannelHandle Channel = ChannelProxy.MakeHandle(ChannelTypeName, Index);
 
+			FMovieSceneBoolChannel* BoolChannel = Entry.GetChannelTypeName() == BoolChannelTypeName ? Channel.Cast<FMovieSceneBoolChannel>().Get() : nullptr;
 			FMovieSceneDoubleChannel* DoubleChannel = Entry.GetChannelTypeName() == DoubleChannelTypeName ? Channel.Cast<FMovieSceneDoubleChannel>().Get() : nullptr;
+			FMovieSceneByteChannel* EnumChannel = Entry.GetChannelTypeName() == EnumChannelTypeName ? Channel.Cast<FMovieSceneByteChannel>().Get() : nullptr;
 			FMovieSceneFloatChannel* FloatChannel = Entry.GetChannelTypeName() == FloatChannelTypeName ? Channel.Cast<FMovieSceneFloatChannel>().Get() : nullptr;
 			FMovieSceneIntegerChannel* IntegerChannel = Entry.GetChannelTypeName() == IntegerChannelTypeName ? Channel.Cast<FMovieSceneIntegerChannel>().Get() : nullptr;
 			FMovieSceneStringChannel* StringChannel = Entry.GetChannelTypeName() == StringChannelTypeName ? Channel.Cast<FMovieSceneStringChannel>().Get() : nullptr;
 
-			if (!DoubleChannel && !FloatChannel && !IntegerChannel && !StringChannel)
+			if (!BoolChannel && !DoubleChannel && !EnumChannel && !FloatChannel && !IntegerChannel && !StringChannel)
 			{
 				continue;
 			}
@@ -3448,13 +3495,29 @@ void FFbxExporter::ExportLevelSequenceTrackChannels( FbxNode* FbxNode, UMovieSce
 			{
 				Property = FbxCamera->FocusDistance;
 			}
+			else if(PropertyName == "bUseTemperature")
+			{
+				Property = FbxNode->FindProperty("UE_UseTemperature", false);
+			}
+			else if(PropertyName == "IntensityUnits")
+			{
+				Property = FbxNode->FindProperty("UE_IntensityUnits", false);
+			}
 
 			if (Property == 0)
 			{
-				if (DoubleChannel)
+				if(BoolChannel)
+				{
+					CreateAnimatableUserProperty(FbxNode, BoolChannel->GetDefault().Get(false), TCHAR_TO_UTF8(*PropertyName), TCHAR_TO_UTF8(*PropertyName), FbxBoolDT);
+				}
+				else if (DoubleChannel)
 				{
 					CreateAnimatableUserProperty(FbxNode, DoubleChannel->GetDefault().Get(MAX_flt), TCHAR_TO_UTF8(*PropertyName), TCHAR_TO_UTF8(*PropertyName));
 				}
+				else if (EnumChannel)
+				{
+					CreateAnimatableUserProperty(FbxNode, EnumChannel->GetDefault().Get(0), TCHAR_TO_UTF8(*PropertyName), TCHAR_TO_UTF8(*PropertyName), FbxEnumDT);
+				}
 				else if (FloatChannel)
 				{
 					CreateAnimatableUserProperty(FbxNode, FloatChannel->GetDefault().Get(MAX_flt), TCHAR_TO_UTF8(*PropertyName), TCHAR_TO_UTF8(*PropertyName));
@@ -3518,11 +3581,25 @@ void FFbxExporter::ExportLevelSequenceTrackChannels( FbxNode* FbxNode, UMovieSce
 			}
 			else if (IntegerChannel)
 			{
-				CurveNode->SetChannelValue<double>(0U, IntegerChannel->GetDefault().Get(0));
+				CurveNode->SetChannelValue<int32>(0U, IntegerChannel->GetDefault().Get(0));
 				CurveNode->ConnectToChannel(AnimCurve, 0U);
 
 				ExportChannelToFbxCurve(*AnimCurve, *IntegerChannel, TickResolution, RootToLocalTransform);
 			}
+			else if (BoolChannel)
+			{
+				CurveNode->SetChannelValue<bool>(0U, BoolChannel->GetDefault().Get(false));
+				CurveNode->ConnectToChannel(AnimCurve, 0U);
+
+				ExportChannelToFbxCurve(*AnimCurve, *BoolChannel, TickResolution, RootToLocalTransform);
+			}
+			else if (EnumChannel)
+			{
+				CurveNode->SetChannelValue<uint8>(0U, EnumChannel->GetDefault().Get(0U));
+				CurveNode->ConnectToChannel(AnimCurve, 0U);
+
+				ExportChannelToFbxCurve(*AnimCurve, *EnumChannel, TickResolution, RootToLocalTransform);
+			}
 		}
 	}
 }
diff --git a/Engine/Source/Editor/UnrealEd/Private/FbxExporter.h b/Engine/Source/Editor/UnrealEd/Private/FbxExporter.h
--- a/Engine/Source/Editor/UnrealEd/Private/FbxExporter.h
+++ b/Engine/Source/Editor/UnrealEd/Private/FbxExporter.h
@@ -30,6 +30,7 @@ class UModel;
 class UMovieScene;
 class UMovieSceneSkeletalAnimationTrack;
 class UMovieScene3DTransformTrack;
+class UMovieSceneColorTrack;
 class UMovieScenePropertyTrack;
 class UMovieSceneTrack;
 class USkeletalMesh;
@@ -437,6 +438,11 @@ private:
 	 */
 	UNREALED_API void ExportLevelSequenceBaked3DTransformTrack(IAnimTrackAdapter& AnimTrackAdapter, FbxNode* FbxActor, IMovieScenePlayer* MovieScenePlayer, FMovieSceneSequenceIDRef InSequenceID, TArray<TWeakObjectPtr<UMovieScene3DTransformTrack> > TransformTracks, UObject* BoundObject, const TRange<FFrameNumber>& InPlaybackRange, const FMovieSceneSequenceTransform& RootToLocalTransform);
 
+	/**
+	 * Exports a level sequence color track into the FBX animation stack. (only works for light)
+	 */
+	UNREALED_API void ExportLevelSequenceColorTrack(FbxNode* FbxActor, UMovieSceneColorTrack& ColorTrack, UObject* BoundObject, const TRange<FFrameNumber>& InPlaybackRange, const FMovieSceneSequenceTransform& RootToLocalTransform);
+
 	/**
 	 * Exports a level sequence property track into the FBX animation stack. 
 	 */
