Fix for bad root motion preview on assets with non-zero root
#rnx

diff --git a/Engine/Source/Editor/UnrealEd/Private/Animation/DebugSkelMeshComponent.cpp b/Engine/Source/Editor/UnrealEd/Private/Animation/DebugSkelMeshComponent.cpp
--- a/Engine/Source/Editor/UnrealEd/Private/Animation/DebugSkelMeshComponent.cpp
+++ b/Engine/Source/Editor/UnrealEd/Private/Animation/DebugSkelMeshComponent.cpp
@@ -29,6 +29,16 @@
 //////////////////////////////////////////////////////////////////////////
 // UDebugSkelMeshComponent
 
+namespace UE::Anim::Private
+{
+	FTransform CalculateInitialTransformFromAssetAndTime(const UDebugSkelMeshComponent& InDebugSkelMeshComponent, const UAnimationAsset* InAnimAsset, const float InTime)
+	{
+		const FTransform InitialRootBoneTransform = InDebugSkelMeshComponent.GetReferenceSkeleton().GetRefBonePose()[0];
+		const FTransform InitialTransform = UE::Anim::ExtractRootTransformFromAnimationAsset(InAnimAsset, InTime);
+		return InitialRootBoneTransform * InitialTransform;
+	}
+}
+
 UDebugSkelMeshComponent::UDebugSkelMeshComponent(const FObjectInitializer& ObjectInitializer)
 	: Super(ObjectInitializer)
 {
@@ -289,7 +299,7 @@ void UDebugSkelMeshComponent::ConsumeRootMotion(const FVector& FloorMin, const F
 			{
 				if (bLooped)
 				{
-					const FTransform InitialTransform = UE::Anim::ExtractRootTransformFromAnimationAsset(PreviewInstance->CurrentAsset, SectionStartPosition);
+					const FTransform InitialTransform = UE::Anim::Private::CalculateInitialTransformFromAssetAndTime(*this, PreviewInstance->CurrentAsset, SectionStartPosition);
 					const FTransform RootMotionDelta = UE::Anim::ExtractRootMotionFromAnimationAsset(PreviewInstance->CurrentAsset, PreviewInstance->GetMirrorDataTable(), SectionStartPosition, CurrentTime);
 					RootMotionTransform = RootMotionDelta * InitialTransform;
 
@@ -433,7 +443,7 @@ void UDebugSkelMeshComponent::SetProcessRootMotionModeInternal(EProcessRootMotio
 			SectionStartPosition = StartTime;
 		}
 	
-		const FTransform InitialTransform = UE::Anim::ExtractRootTransformFromAnimationAsset(PreviewInstance->CurrentAsset, SectionStartPosition);
+		const FTransform InitialTransform = UE::Anim::Private::CalculateInitialTransformFromAssetAndTime(*this, PreviewInstance->CurrentAsset, SectionStartPosition);
 		const FTransform RootMotionDelta = UE::Anim::ExtractRootMotionFromAnimationAsset(PreviewInstance->CurrentAsset, PreviewInstance->GetMirrorDataTable(), SectionStartPosition, CurrentTime);
 		RootMotionTransform = RootMotionDelta * InitialTransform;
 		
@@ -702,7 +712,7 @@ void UDebugSkelMeshComponent::OnMirrorDataTableChanged()
 			SectionStartPosition = StartTime;
 		}
 	
-		const FTransform InitialTransform = UE::Anim::ExtractRootTransformFromAnimationAsset(PreviewInstance->CurrentAsset, SectionStartPosition);
+		const FTransform InitialTransform = UE::Anim::Private::CalculateInitialTransformFromAssetAndTime(*this, PreviewInstance->CurrentAsset, SectionStartPosition);
 		const FTransform RootMotionDelta = UE::Anim::ExtractRootMotionFromAnimationAsset(PreviewInstance->CurrentAsset, PreviewInstance->GetMirrorDataTable(), SectionStartPosition, CurrentTime);
 		RootMotionTransform = RootMotionDelta * InitialTransform;
 
