Slate: Add ticking frame to SlateApplication to help prevent a double RHI BeginFrame. The real fix is to enforce the SlateApplication::OnPreTick and OnPostTick to be more strick and hide it from the user.

diff --git a/Engine/Source/Developer/SourceControl/Private/ScopedSourceControlProgress.cpp b/Engine/Source/Developer/SourceControl/Private/ScopedSourceControlProgress.cpp
--- a/Engine/Source/Developer/SourceControl/Private/ScopedSourceControlProgress.cpp
+++ b/Engine/Source/Developer/SourceControl/Private/ScopedSourceControlProgress.cpp
@@ -236,8 +236,11 @@ void FScopedSourceControlProgress::Tick()
 #if SOURCE_CONTROL_WITH_SLATE
 	if (!(FApp::IsUnattended() || IsRunningCommandlet()) && WindowPtr.IsValid() && FSlateApplication::Get().CanDisplayWindows())
 	{
-		// Mark begin frame
-		if (GIsRHIInitialized)
+		// Testing if we are already ticking the rendering. That is to prevent a double "BeginFrame" in case the user wrongly uses the FSlateApplication::OnPreTick to start a slow task.
+		bool bIsTicking = FSlateApplication::Get().IsTicking();
+
+		// Mark begin frame.
+		if (!bIsTicking && GIsRHIInitialized)
 		{
 			ENQUEUE_RENDER_COMMAND(BeginFrameCmd)([](FRHICommandListImmediate& RHICmdList) { RHICmdList.BeginFrame(); });
 		}
@@ -246,7 +249,7 @@ void FScopedSourceControlProgress::Tick()
 		FSlateApplication::Get().Tick();
 
 		// End frame so frame fence number gets incremented
-		if (GIsRHIInitialized)
+		if (!bIsTicking && GIsRHIInitialized)
 		{
 			ENQUEUE_RENDER_COMMAND(EndFrameCmd)([](FRHICommandListImmediate& RHICmdList) { RHICmdList.EndFrame(); });
 		}
diff --git a/Engine/Source/Editor/UnrealEd/Private/FeedbackContextEditor.cpp b/Engine/Source/Editor/UnrealEd/Private/FeedbackContextEditor.cpp
--- a/Engine/Source/Editor/UnrealEd/Private/FeedbackContextEditor.cpp
+++ b/Engine/Source/Editor/UnrealEd/Private/FeedbackContextEditor.cpp
@@ -421,8 +421,11 @@ static void TickSlate(TSharedPtr<SWindow> SlowTaskWindow)
 	// Avoid re-entrancy by ticking the active modal window again. This can happen if the slow task window is open and a sibling modal window is open as well.  We only tick slate if we are the active modal window or a child of the active modal window
 	if( SlowTaskWindow.IsValid() && ( FSlateApplication::Get().GetActiveModalWindow() == SlowTaskWindow || SlowTaskWindow->IsDescendantOf( FSlateApplication::Get().GetActiveModalWindow() ) ) )
 	{
+		// Testing if we are already ticking the rendering. That is to prevent a double "BeginFrame" in case the user wrongly uses the FSlateApplication::OnPreTick to start a slow task.
+		bool bIsTicking = FSlateApplication::Get().IsTicking();
+
 		// Mark begin frame
-		if (GIsRHIInitialized)
+		if (!bIsTicking && GIsRHIInitialized)
 		{
 			ENQUEUE_RENDER_COMMAND(BeginFrameCmd)([](FRHICommandListImmediate& RHICmdList) { RHICmdList.BeginFrame(); });
 		}
@@ -431,7 +434,7 @@ static void TickSlate(TSharedPtr<SWindow> SlowTaskWindow)
 		FSlateApplication::Get().Tick();
 
 		// End frame so frame fence number gets incremented
-		if (GIsRHIInitialized)
+		if (!bIsTicking && GIsRHIInitialized)
 		{
 			ENQUEUE_RENDER_COMMAND(EndFrameCmd)([](FRHICommandListImmediate& RHICmdList) { RHICmdList.EndFrame(); });
 		}
diff --git a/Engine/Source/Runtime/Slate/Private/Framework/Application/SlateApplication.cpp b/Engine/Source/Runtime/Slate/Private/Framework/Application/SlateApplication.cpp
--- a/Engine/Source/Runtime/Slate/Private/Framework/Application/SlateApplication.cpp
+++ b/Engine/Source/Runtime/Slate/Private/Framework/Application/SlateApplication.cpp
@@ -1510,6 +1510,8 @@ void FSlateApplication::Tick(ESlateTickType TickType)
 
 	FScopeLock SlateTickAccess(&SlateTickCriticalSection);
 
+	TGuardValue<bool> TmpGuard(bIsTicking, true);
+
 #if WITH_EDITOR
 	FScopedPreventDebuggingMode SlatePreventDebugginModeWhileTicking(NSLOCTEXT("EnterDebuggingMode", "WindowTicking", "The window is ticking."));
 #endif
@@ -1551,6 +1553,11 @@ void FSlateApplication::Tick(ESlateTickType TickType)
 	}
 }
 
+bool FSlateApplication::IsTicking() const
+{
+	return bIsTicking;
+}
+
 void FSlateApplication::TickTime()
 {
 	LastTickTime = CurrentTime;
diff --git a/Engine/Source/Runtime/Slate/Public/Framework/Application/SlateApplication.h b/Engine/Source/Runtime/Slate/Public/Framework/Application/SlateApplication.h
--- a/Engine/Source/Runtime/Slate/Public/Framework/Application/SlateApplication.h
+++ b/Engine/Source/Runtime/Slate/Public/Framework/Application/SlateApplication.h
@@ -328,6 +328,9 @@ public:
 	/** Ticks this application */
 	SLATE_API void Tick(ESlateTickType TickType = ESlateTickType::All);
 
+	/** Returns true if we are currently ticking the SlateApplication. */
+	SLATE_API bool IsTicking() const;
+
 	/** Pumps OS messages when a modal window or intra-frame debugging session exists */
 	SLATE_API void PumpMessages();
 
@@ -2022,6 +2025,9 @@ private:
 
 	/** Did we synthesize cursor input this frame? */
 	bool bSynthesizedCursorMove = false;
+	
+	/** Are we ticking the SlateApplication. */
+	bool bIsTicking = false;
 
 	/** Platform mouse movement event count. */
 	uint64 PlatformMouseMovementEvents = 0;
