IterativeCooker: Include configvalues in the targetdomain key for cooked packages.
#rnx

diff --git a/Engine/Source/Editor/UnrealEd/Private/CookOnTheFlyServer.cpp b/Engine/Source/Editor/UnrealEd/Private/CookOnTheFlyServer.cpp
--- a/Engine/Source/Editor/UnrealEd/Private/CookOnTheFlyServer.cpp
+++ b/Engine/Source/Editor/UnrealEd/Private/CookOnTheFlyServer.cpp
@@ -6388,17 +6388,12 @@ void FSaveCookedPackageContext::FinishPlatform()
 	{
 		TOptional<FAssetPackageData> AssetPackageData = COTFS.AssetRegistry->GetAssetPackageDataCopy(Package->GetFName());
 
-		// TODO_BuildDefinitionList: Calculate and store BuildDefinitionList on the PackageData, or collect it here from some other source.
-		TArray<UE::DerivedData::FBuildDefinition> BuildDefinitions;
-		FCbObject BuildDefinitionList = UE::TargetDomain::BuildDefinitionListToObject(BuildDefinitions);
-		FCbObject TargetDomainDependencies;
+		ICookedPackageWriter::FCommitPackageInfo Info;
 		if (COTFS.bHybridIterativeEnabled)
 		{
 			UE_SCOPED_HIERARCHICAL_COOKTIMER(TargetDomainDependencies);
-			TargetDomainDependencies = UE::TargetDomain::CollectDependenciesObject(Package, TargetPlatform, nullptr /* ErrorMessage */);
+			UE::TargetDomain::CollectAndStoreCookAttachments(Package, TargetPlatform, Info.Attachments);
 		}
-
-		ICookedPackageWriter::FCommitPackageInfo Info;
 		if (bSuccessful)
 		{
 			Info.Status = IPackageWriter::ECommitStatus::Success;
@@ -6413,12 +6408,6 @@ void FSaveCookedPackageContext::FinishPlatform()
 		}
 		Info.PackageName = Package->GetFName();
 		Info.PackageHash = AssetPackageData ? AssetPackageData->GetPackageSavedHash() : FIoHash();
-		if (TargetDomainDependencies)
-		{
-			Info.Attachments.Add({ "Dependencies", TargetDomainDependencies });
-		}
-		// TODO: Reenable BuildDefinitionList once FCbPackage support for empty FCbObjects is in
-		//Info.Attachments.Add({ "BuildDefinitionList", BuildDefinitionList });
 		Info.WriteOptions = IPackageWriter::EWriteOptions::None;
 		if (!COTFS.bSkipSave)
 		{
@@ -7714,6 +7703,8 @@ bool UCookOnTheFlyServer::GetCurrentIniVersionStrings( const ITargetPlatform* Ta
 	IniVersionStrings.Reset();
 	return true;
 #else
+	using namespace UE::ConfigAccessTracking;
+
 	// This function should be called after the cook is finished
 	TArray<UE::ConfigAccessTracking::FConfigAccessData> AccessedRecordsArray =
 		UE::ConfigAccessTracking::FCookConfigAccessTracker::Get().GetCookRecords(TargetPlatform);
@@ -7726,7 +7717,7 @@ bool UCookOnTheFlyServer::GetCurrentIniVersionStrings( const ITargetPlatform* Ta
 	while (ConfigFileEndIndex < EndIndex)
 	{
 		int32 ConfigFileStartIndex = ConfigFileEndIndex;
-		const UE::ConfigAccessTracking::FConfigAccessData& FileStartRecord = AccessedRecords[ConfigFileStartIndex];
+		const FConfigAccessData& FileStartRecord = AccessedRecords[ConfigFileStartIndex];
 		++ConfigFileEndIndex;
 		while (ConfigFileEndIndex < EndIndex && AccessedRecords[ConfigFileEndIndex].IsSameConfigFile(FileStartRecord))
 		{
@@ -7734,8 +7725,8 @@ bool UCookOnTheFlyServer::GetCurrentIniVersionStrings( const ITargetPlatform* Ta
 		}
 
 		FConfigFile Temp;
-		FName ConfigPlatform(FileStartRecord.GetConfigPlatform());
 		FName ConfigFileName(FileStartRecord.GetFileName());
+		FConfigAccessData FileRecord = FileStartRecord.GetFileOnlyData();
 		TStringBuilder<64> ConfigFileNameStr(InPlace, ConfigFileName);
 		using UE::String::FindFirst;
 
@@ -7746,19 +7737,11 @@ bool UCookOnTheFlyServer::GetCurrentIniVersionStrings( const ITargetPlatform* Ta
 		{
 			continue;
 		}
-		TStringBuilder<32> ConfigPlatformStr;
-		const TCHAR* SavePlatformStr = TEXT("<Editor>");
-		if (!ConfigPlatform.IsNone())
-		{
-			ConfigPlatformStr << ConfigPlatform;
-			SavePlatformStr = *ConfigPlatformStr;
-		}
-		const TCHAR* LoadTypeStr = LexToString(FileStartRecord.LoadType);
+
 		FullConfigFileNameStr.Reset();
-		FullConfigFileNameStr << LoadTypeStr << TEXT(".") << SavePlatformStr << TEXT(".") << ConfigFileNameStr;
+		FileRecord.AppendFullPath(FullConfigFileNameStr);
 		FName FullConfigFileName(FullConfigFileNameStr);
-		const FConfigFile* ConfigFile = UE::ConfigAccessTracking::FindOrLoadConfigFile(FileStartRecord.LoadType,
-			*ConfigPlatformStr, *ConfigFileNameStr, Temp);
+		const FConfigFile* ConfigFile = UE::ConfigAccessTracking::FindOrLoadConfigFile(FileRecord, Temp);
 		if (!ConfigFile)
 		{
 			// This is logged as Warning; it is unexpected that we were able to load a file from disk that
@@ -7901,79 +7884,6 @@ bool UCookOnTheFlyServer::GetCurrentIniVersionStrings( const ITargetPlatform* Ta
 #endif // !UE_WITH_CONFIG_TRACKING
 }
 
-namespace
-{
-
-void EscapeUsedSettingsToken(FName Token, FStringBuilderBase& Result)
-{
-	Result.Reset();
-	Result << Token;
-	if (Result.ToView().Contains(TEXTVIEW(":")))
-	{
-		FString ReplaceText(Result);
-		ReplaceText.ReplaceInline(TEXT(":"), TEXT("::"));
-		Result.Reset();
-		Result << ReplaceText;
-	}
-}
-
-bool TryTokenizeUsedSettingsString(FStringView Text, TArrayView<FStringBuilderBase*> OutTokens)
-{
-	int32 TextLen = Text.Len();
-	if (TextLen == 0)
-	{
-		return false;
-	}
-	const TCHAR* TextData = Text.GetData();
-
-	int32 NumTokens = OutTokens.Num();
-	check(NumTokens > 0);
-	int32 TokenIndex = 0;
-	FStringBuilderBase* OutToken = OutTokens[TokenIndex++];
-	OutToken->Reset();
-	for (int32 Index = 0; Index < TextLen; ++Index)
-	{
-		TCHAR C = TextData[Index];
-		if (C != ':')
-		{
-			OutToken->AppendChar(C);
-		}
-		else if (Index < TextLen - 1 && TextData[Index + 1] == ':')
-		{
-			++Index;
-			OutToken->AppendChar(':');
-		}
-		else
-		{
-			if (OutToken->Len() == 0)
-			{
-				// Empty token
-				return false;
-			}
-			if (TokenIndex >= NumTokens)
-			{
-				// Too many tokens
-				return false;
-			}
-			OutToken = OutTokens[TokenIndex++];
-			OutToken->Reset();
-		}
-	}
-	if (OutToken->Len() == 0)
-	{
-		// Empty token
-		return false;
-	}
-	if (TokenIndex < NumTokens)
-	{
-		// too few tokens
-		return false;
-	}
-	return true;
-}
-
-} // anonymous namespace
-
 bool UCookOnTheFlyServer::GetCookedIniVersionStrings(const ITargetPlatform* TargetPlatform, UE::Cook::FIniSettingContainer& OutIniSettings, TMap<FString,FString>& OutAdditionalSettings) const
 {
 	const FString EditorIni = GetMetadataDirectory() / TEXT("CookedIniVersion.txt");
@@ -8005,11 +7915,12 @@ bool UCookOnTheFlyServer::GetCookedIniVersionStrings(const ITargetPlatform* Targ
 	FStringBuilderBase* TokenBuffer[] = { &Filename, &SectionName, &ValueName, &ValueIndexStr };
 	TArrayView<FStringBuilderBase*> Tokens = TokenBuffer;
 
+	using namespace UE::ConfigAccessTracking;
 	for (const auto& UsedSetting : *UsedSettings )
 	{
 		KeyStr.Reset();
 		KeyStr << UsedSetting.Key;
-		if (!TryTokenizeUsedSettingsString(KeyStr, Tokens))
+		if (!TryTokenizeConfigTrackingString(KeyStr, Tokens))
 		{
 			UE_LOG(LogCook, Warning, TEXT("Found unparsable ini setting %s for platform %s, invalidating cook."),
 				*KeyStr, *TargetPlatform->PlatformName());
@@ -8198,7 +8109,9 @@ bool UCookOnTheFlyServer::IniSettingsOutOfDate(const ITargetPlatform* TargetPlat
 #if !UE_WITH_CONFIG_TRACKING
 	return false;
 #else
-	UE::ConfigAccessTracking::FIgnoreScope IgnoreScope;
+	using namespace UE::ConfigAccessTracking;
+
+	FIgnoreScope IgnoreScope;
 
 	UE::Cook::FIniSettingContainer OldIniSettings;
 	TMap<FString, FString> OldAdditionalSettings;
@@ -8239,64 +8152,17 @@ bool UCookOnTheFlyServer::IniSettingsOutOfDate(const ITargetPlatform* TargetPlat
 	{
 		OldIniFile.Key.ToString(ConfigNameKeyStr);
 
-		// ConfigSystem.<Editor>.../../../Engine/Config/ConsoleVariables.ini
-		//   -> "ConfigSystem", "<Editor>", "../../../Engine/Config/ConsoleVariables.ini"
-		// 3rd token might have dots, first two cannot.
-		FStringView ConfigNameArray[3];
-		ConfigNameArray[0] = ConfigNameKeyStr;
-		int32 NumTokens = 1;
-		for (NumTokens = 1; NumTokens < UE_ARRAY_COUNT(ConfigNameArray); ++NumTokens)
-		{
-			FStringView& Current = ConfigNameArray[NumTokens-1];
-			FStringView& Next= ConfigNameArray[NumTokens];
-			int32 DotIndex = Current.Find(TEXTVIEW("."));
-			if (DotIndex == INDEX_NONE)
-			{
-				break;
-			}
-			Next = Current.RightChop(DotIndex+1);
-			Current.LeftInline(DotIndex);
-			if (Next.IsEmpty())
-			{
-				// break before incrementing NumTokens so that we do not count the empty Next as a token
-				break; 
-			}
-		}
-		UE::ConfigAccessTracking::ELoadType LoadType = UE::ConfigAccessTracking::ELoadType::Uninitialized;
-		// The input NameKey is of the form 
-		//   LoadType.Platform.ConfigName
-		// Platform is <Editor> if the configfile was an editor config file rather than a platform-specific config file
-		if (NumTokens == 3)
-		{
-			LexFromString(LoadType, ConfigNameArray[0]);
-			if (!UE::ConfigAccessTracking::IsLoadableLoadType(LoadType))
-			{
-				LoadType = UE::ConfigAccessTracking::ELoadType::Uninitialized;
-			}
-			else
-			{
-				PlatformName.Reset();
-				PlatformName << ConfigNameArray[1];
-				Filename.Reset();
-				Filename << ConfigNameArray[2];
-			}
-		}
-		if (LoadType == UE::ConfigAccessTracking::ELoadType::Uninitialized)
+		FConfigAccessData FullFilePathData = FConfigAccessData::Parse(ConfigNameKeyStr);
+		if (!IsLoadableLoadType(FullFilePathData.LoadType))
 		{
 			UE_LOG(LogCook, Warning,
 				TEXT("Invalidating inisettings: Invalid filename key in old ini settings file used by platform %s: key '%s' is invalid."),
 				*TargetPlatform->PlatformName(), *ConfigNameKeyStr);
 			return true;
 		}
-		const TCHAR* PlatformNameToLoad = *PlatformName;
-		if (FCString::Stricmp(PlatformNameToLoad, TEXT("<Editor>")) == 0)
-		{
-			PlatformNameToLoad = TEXT("");
-		}
 
 		FConfigFile Temp;
-		const FConfigFile* ConfigFile = UE::ConfigAccessTracking::FindOrLoadConfigFile(LoadType, PlatformNameToLoad,
-			*Filename, Temp);
+		const FConfigFile* ConfigFile = UE::ConfigAccessTracking::FindOrLoadConfigFile(FullFilePathData, Temp);
 		if (!ConfigFile)
 		{
 			UE_LOG(LogCook, Display,
@@ -8391,6 +8257,7 @@ bool UCookOnTheFlyServer::SaveCurrentIniSettings(const ITargetPlatform* TargetPl
 	const static TCHAR* NAME_UsedSettings =TEXT("UsedSettings");
 	ConfigFile.Remove(NAME_UsedSettings);
 
+	using namespace UE::ConfigAccessTracking;
 	{
 		TStringBuilder<256> NewKey;
 		TStringBuilder<128> FilenameStr;
@@ -8400,15 +8267,15 @@ bool UCookOnTheFlyServer::SaveCurrentIniSettings(const ITargetPlatform* TargetPl
 		for (const auto& CurrentIniFilename : CurrentIniSettings)
 		{
 			FName Filename = CurrentIniFilename.Key;
-			EscapeUsedSettingsToken(Filename, FilenameStr);
+			EscapeConfigTrackingTokenToString(Filename, FilenameStr);
 			for (const auto& CurrentSection : CurrentIniFilename.Value)
 			{
 				FName Section = CurrentSection.Key;
-				EscapeUsedSettingsToken(Section, SectionStr);
+				EscapeConfigTrackingTokenToString(Section, SectionStr);
 				for (const auto& CurrentValue : CurrentSection.Value)
 				{
 					FName ValueName = CurrentValue.Key;
-					EscapeUsedSettingsToken(ValueName, ValueNameStr);
+					EscapeConfigTrackingTokenToString(ValueName, ValueNameStr);
 
 					const TArray<FString>& Values = CurrentValue.Value;
 					for (int Index = 0; Index < Values.Num(); ++Index)
diff --git a/Engine/Source/Editor/UnrealEd/Private/Cooker/CookConfigAccessTracker.cpp b/Engine/Source/Editor/UnrealEd/Private/Cooker/CookConfigAccessTracker.cpp
--- a/Engine/Source/Editor/UnrealEd/Private/Cooker/CookConfigAccessTracker.cpp
+++ b/Engine/Source/Editor/UnrealEd/Private/Cooker/CookConfigAccessTracker.cpp
@@ -17,6 +17,90 @@
 
 #endif
 
+namespace UE::ConfigAccessTracking
+{
+
+void EscapeConfigTrackingTokenToString(FName Token, FStringBuilderBase& Result)
+{
+	Result.Reset();
+	EscapeConfigTrackingTokenAppendString(Token, Result);
+}
+
+void EscapeConfigTrackingTokenAppendString(FName Token, FStringBuilderBase& Result)
+{
+	int32 InitialLength = Result.Len();
+	Result << Token;
+	FStringView AddedView = Result.ToView().RightChop(InitialLength);
+	if (AddedView.Contains(TEXTVIEW(":")))
+	{
+		FString ReplaceText(AddedView);
+		ReplaceText.ReplaceInline(TEXT(":"), TEXT("::"));
+		Result.RemoveSuffix(AddedView.Len());
+		Result << ReplaceText;
+	}
+}
+
+bool TryTokenizeConfigTrackingString(FStringView Text, TArrayView<FStringBuilderBase*> OutTokens)
+{
+	int32 TextLen = Text.Len();
+	if (TextLen == 0)
+	{
+		for (FStringBuilderBase* Token : OutTokens) Token->Reset();
+		return false;
+	}
+	const TCHAR* TextData = Text.GetData();
+
+	int32 NumTokens = OutTokens.Num();
+	check(NumTokens > 0);
+	int32 NextTokenIndex = 0;
+	FStringBuilderBase* OutToken = OutTokens[NextTokenIndex++];
+	OutToken->Reset();
+	for (int32 Index = 0; Index < TextLen; ++Index)
+	{
+		TCHAR C = TextData[Index];
+		if (C != ':')
+		{
+			OutToken->AppendChar(C);
+		}
+		else if (Index < TextLen - 1 && TextData[Index + 1] == ':')
+		{
+			++Index;
+			OutToken->AppendChar(':');
+		}
+		else
+		{
+			if (OutToken->Len() == 0)
+			{
+				// An empty token and therefore the string is invalid; abandon anything that follows it
+				while (NextTokenIndex < NumTokens) OutTokens[NextTokenIndex++]->Reset();
+				return false;
+			}
+			if (NextTokenIndex >= NumTokens)
+			{
+				// Too many tokens
+				return false;
+			}
+			OutToken = OutTokens[NextTokenIndex++];
+			OutToken->Reset();
+		}
+	}
+	if (OutToken->Len() == 0)
+	{
+		// Empty token
+		while (NextTokenIndex < NumTokens) OutTokens[NextTokenIndex++]->Reset();
+		return false;
+	}
+	if (NextTokenIndex < NumTokens)
+	{
+		// too few tokens
+		while (NextTokenIndex < NumTokens) OutTokens[NextTokenIndex++]->Reset();
+		return false;
+	}
+	return true;
+}
+
+} // namespace UE::ConfigAccessTracking
+
 #if UE_WITH_CONFIG_TRACKING
 
 DEFINE_LOG_CATEGORY_STATIC(LogConfigBuildDependencyTracker, Log, All);
@@ -26,22 +110,147 @@ namespace UE::ConfigAccessTracking
 
 FCookConfigAccessTracker FCookConfigAccessTracker::Singleton;
 
-thread_local bool bIgnoreAccess = false;
+FConfigAccessData::FConfigAccessData(ELoadType InLoadType, FNameEntryId InConfigPlatform, FNameEntryId InFileName)
+	: ConfigPlatform(InConfigPlatform)
+	, FileName(InFileName)
+	, LoadType(InLoadType)
+{
+}
 
-TStringBuilder<256> ToString(FNameEntryId FileName, FNameEntryId SectionName, FMinimalName ValueName)
+FConfigAccessData::FConfigAccessData(ELoadType InLoadType, FNameEntryId InConfigPlatform, FNameEntryId InFileName,
+	FNameEntryId InSectionName, FMinimalName InValueName, const ITargetPlatform* InRequestingPlatform)
+	: ConfigPlatform(InConfigPlatform)
+	, FileName(InFileName)
+	, SectionName(InSectionName)
+	, ValueName(InValueName)
+	, RequestingPlatform(InRequestingPlatform)
+	, LoadType(InLoadType)
 {
-	return TStringBuilder<256>(InPlace, FileName, TEXT(":["), SectionName, TEXT("]:"), FName(ValueName));
 }
 
-FIgnoreScope::FIgnoreScope()
-	: bPreviousIgnoreAccess(bIgnoreAccess)
+FConfigAccessData FConfigAccessData::GetFileOnlyData() const
 {
-	bIgnoreAccess = true;
+	return FConfigAccessData(LoadType, ConfigPlatform, FileName);
 }
 
-FIgnoreScope::~FIgnoreScope()
+FConfigAccessData FConfigAccessData::GetPathOnlyData() const
 {
-	bIgnoreAccess = bPreviousIgnoreAccess;
+	return FConfigAccessData(LoadType, ConfigPlatform, FileName, SectionName, ValueName, nullptr);
+}
+
+FString FConfigAccessData::FullPathToString() const
+{
+	TStringBuilder<256> Out;
+	AppendFullPath(Out);
+	return FString(*Out);
+}
+
+void FConfigAccessData::AppendFullPath(FStringBuilderBase& Out) const
+{
+	if (LoadType == ELoadType::Uninitialized || FileName.IsNone())
+	{
+		Out << TEXTVIEW("<Invalid>");
+		return;
+	}
+
+	Out << LexToString(LoadType) << TEXTVIEW(".");
+	if (ConfigPlatform.IsNone())
+	{
+		Out << PlatformAgnosticName;
+	}
+	else
+	{
+		EscapeConfigTrackingTokenAppendString(GetConfigPlatform(), Out);
+	}
+	Out << TEXTVIEW(".");
+	EscapeConfigTrackingTokenAppendString(GetFileName(), Out);
+	if (!SectionName.IsNone())
+	{
+		Out << TEXTVIEW(":[");
+		EscapeConfigTrackingTokenAppendString(GetSectionName(), Out);
+		Out << TEXTVIEW("]");
+		if (!ValueName.IsNone())
+		{
+			Out << TEXTVIEW(":");
+			EscapeConfigTrackingTokenAppendString(GetValueName(), Out);
+		}
+	}
+}
+
+FConfigAccessData FConfigAccessData::Parse(FStringView Text)
+{
+	// ConfigSystem.<Editor>.../../../Engine/Config/ConsoleVariables.ini:Section:ValueName
+	//   -> "ConfigSystem", "<Editor>", "../../../Engine/Config/ConsoleVariables.ini", "Section", "ValueName"
+	// No tokens can contain a single colon, but they might contain double colon, which is the escape code for a single colon.
+	// 3rd token might have dots, first two cannot.
+
+	TStringBuilder<128> FullFilePath;
+	TStringBuilder<64> SectionNameStr;
+	TStringBuilder<64> ValueNameStr;
+	FStringBuilderBase* TokenBuffer[] = { &FullFilePath, &SectionNameStr, &ValueNameStr };
+	TArrayView<FStringBuilderBase*> ColonDelimitedTokens = TokenBuffer;
+	TryTokenizeConfigTrackingString(Text, ColonDelimitedTokens);
+	if (FullFilePath.Len() == 0)
+	{
+		return FConfigAccessData();
+	}
+
+	FStringView FullFilePathTokens[3];
+	FullFilePathTokens[0] = FullFilePath;
+	int32 NumFilePathTokens = 1;
+	for (NumFilePathTokens = 1; NumFilePathTokens < UE_ARRAY_COUNT(FullFilePathTokens); ++NumFilePathTokens)
+	{
+		FStringView& Current = FullFilePathTokens[NumFilePathTokens - 1];
+		FStringView& Next = FullFilePathTokens[NumFilePathTokens];
+		int32 DotIndex = Current.Find(TEXTVIEW("."));
+		if (DotIndex == INDEX_NONE)
+		{
+			break;
+		}
+		Next = Current.RightChop(DotIndex + 1);
+		Current.LeftInline(DotIndex);
+		if (Next.IsEmpty())
+		{
+			// break before incrementing NumFilePathTokens so that we do not count the empty Next as a token
+			break;
+		}
+	}
+	// FullFilePath is of the form 
+	//   LoadType.Platform.ConfigName
+	// Platform is <Editor> if the configfile was an editor config file rather than a platform-specific config file
+	if (NumFilePathTokens < 3)
+	{
+		// Not a valid FConfigAccessData text; we always require all 3 of LoadType,Platform,ConfigName
+		return FConfigAccessData();
+	}
+
+	FConfigAccessData Result;
+	LexFromString(Result.LoadType, FullFilePathTokens[0]);
+	if (Result.LoadType == ELoadType::Uninitialized)
+	{
+		return FConfigAccessData();
+	}
+
+	FName LocalConfigPlatform = FullFilePathTokens[1] == PlatformAgnosticName
+		? NAME_None : FName(FullFilePathTokens[1]);
+	FStringView SectionNameView(SectionNameStr);
+	// We write out sectionnames with [] for readability; remove them if they exist.
+	if (SectionNameView.StartsWith('[')) SectionNameView.RightChopInline(1);
+	if (SectionNameView.EndsWith(']')) SectionNameView.LeftChopInline(1);
+	FName LocalSectionName = SectionNameView.Len() == 0 ? NAME_None : FName(SectionNameView, NAME_NO_NUMBER);
+	FName LocalValueName = ValueNameStr.Len() == 0 ? NAME_None : FName(ValueNameStr);
+
+	Result.ConfigPlatform = LocalConfigPlatform.GetComparisonIndex();
+	Result.FileName = FName(FullFilePathTokens[2]).GetComparisonIndex();
+	Result.SectionName = LocalSectionName.GetComparisonIndex();
+	Result.ValueName = FMinimalName(LocalValueName);
+
+	return Result;
+}
+
+TStringBuilder<256> ToString(FNameEntryId FileName, FNameEntryId SectionName, FMinimalName ValueName)
+{
+	return TStringBuilder<256>(InPlace, FileName, TEXT(":["), SectionName, TEXT("]:"), FName(ValueName));
 }
 
 void FCookConfigAccessTracker::Disable()
@@ -204,13 +413,99 @@ FCookConfigAccessTracker::~FCookConfigAccessTracker()
 	Disable();
 }
 
-void FCookConfigAccessTracker::StaticOnConfigValueRead(UE::ConfigAccessTracking::FSection* Section, FMinimalName ValueName,
-	const FConfigValue& ConfigValue)
+static FConfigFile* FindConfigCacheIniFile(FName ConfigPlatform, FName FileName)
 {
-	if (bIgnoreAccess)
+	// ForPlatform will return GConfig if ConfigPlatform is NAME_None
+	FConfigCacheIni* ConfigSystem = FConfigCacheIni::ForPlatform(ConfigPlatform);
+
+	// The ini files may have been recorded by fullpath or by shortname; search first for a fullpath match using
+	// FindConfigFile and if that fails search for the shortname match by iterating over all files in GConfig
+	FConfigFile* ConfigFile = ConfigSystem->FindConfigFile(FileName.ToString());
+	if (ConfigFile)
 	{
-		return;
+		return ConfigFile;
+	}
+
+	for (const FString& ConfigFilename : ConfigSystem->GetFilenames())
+	{
+		ConfigFile = ConfigSystem->FindConfigFile(ConfigFilename);
+		if (ConfigFile->Name == FileName)
+		{
+			return ConfigFile;
+		}
+	}
+	return nullptr;
+}
+
+FString FCookConfigAccessTracker::GetValue(const FConfigAccessData& AccessData)
+{
+	if (AccessData.SectionName.IsNone() || AccessData.ValueName.IsNone())
+	{
+		return FString();
 	}
+
+	switch (AccessData.LoadType)
+	{
+	case ELoadType::ConfigSystem:
+	{
+		FIgnoreScope IgnoreScope;
+		FConfigFile* ConfigFile = FindConfigCacheIniFile(AccessData.GetConfigPlatform(), AccessData.GetFileName());
+		if (!ConfigFile)
+		{
+			return FString();
+		}
+		const FConfigSection* ConfigSection = ConfigFile->FindSection(AccessData.GetSectionName().ToString());
+		if (!ConfigSection)
+		{
+			return FString();
+		}
+		return MultiValueToString(*ConfigSection, AccessData.GetValueName());
+	}
+	case ELoadType::LocalIniFile:
+	case ELoadType::LocalSingleIniFile:
+	case ELoadType::ExternalIniFile:
+	case ELoadType::ExternalSingleIniFile:
+	{
+		FConfigAccessData PathOnlyData = AccessData.GetPathOnlyData();
+		FConfigAccessData FileOnlyData = PathOnlyData.GetFileOnlyData();
+		uint32 PathOnlyDataHash = GetTypeHash(PathOnlyData);
+		uint32 FileOnlyDataHash = GetTypeHash(FileOnlyData);
+		{
+			UE::TUniqueLock ConfigCacheScopeLock(ConfigCacheLock);
+			if (LoadedConfigFiles.ContainsByHash(FileOnlyDataHash, FileOnlyData))
+			{
+				FString* Result = LoadedValues.FindByHash(PathOnlyDataHash, PathOnlyData);
+				return Result ? *Result : FString();
+			}
+		}
+
+		FIgnoreScope IgnoreScope;
+		FConfigFile Buffer;
+		const FConfigFile* LoadedFile = FindOrLoadConfigFile(FileOnlyData, Buffer);
+		if (!LoadedFile)
+		{
+			return FString();
+		}
+		RecordValuesFromFile(FileOnlyData, *LoadedFile);
+		{
+			UE::TUniqueLock ConfigCacheScopeLock(ConfigCacheLock);
+			FString* Result = LoadedValues.FindByHash(PathOnlyDataHash, PathOnlyData);
+			return Result ? *Result : FString();
+		}
+	}
+	default:
+		return FString();
+	}
+}
+
+FString FCookConfigAccessTracker::GetValue(FStringView AccessDataFullPath)
+{
+	return GetValue(FConfigAccessData::Parse(AccessDataFullPath));
+}
+
+void FCookConfigAccessTracker::StaticOnConfigValueRead(UE::ConfigAccessTracking::FSection* Section, FMinimalName ValueName,
+	const FConfigValue& ConfigValue)
+{
 	if (!Section)
 	{
 		return;
@@ -260,59 +555,98 @@ void FCookConfigAccessTracker::StaticOnConfigValueRead(UE::ConfigAccessTracking:
 #endif
 
 	LLM_SCOPE_BYNAME(TEXTVIEW("ConfigAccessTracking"));
-	FConfigAccessData AccessData{ ConfigPlatform, FileName, SectionName, ValueName, RequestedPlatform, ConfigFile->LoadType };
+	FConfigAccessData AccessData(ConfigFile->LoadType, ConfigPlatform, FileName, SectionName, ValueName, RequestedPlatform);
 	uint32 ReferencerHash = GetTypeHash(Referencer);
 	uint32 AccessDataHash = GetTypeHash(AccessData);
-	UE::TUniqueLock RecordsScopeLock(Singleton.RecordsLock);
-	Singleton.PackageRecords.FindOrAddByHash(ReferencerHash, Referencer).AddByHash(AccessDataHash, MoveTemp(AccessData));
+	FConfigAccessData FileOnlyData = AccessData.GetFileOnlyData();
+	uint32 FileOnlyHash = GetTypeHash(FileOnlyData);
+	{
+		UE::TUniqueLock RecordsScopeLock(Singleton.RecordsLock);
+		Singleton.PackageRecords.FindOrAddByHash(ReferencerHash, Referencer).AddByHash(AccessDataHash, MoveTemp(AccessData));
+	}
+	bool bNeedRecordValuesFromFile = false;
+	{
+		UE::TUniqueLock ConfigCacheScopeLock(Singleton.ConfigCacheLock);
+		bNeedRecordValuesFromFile = !Singleton.LoadedConfigFiles.ContainsByHash(FileOnlyHash, FileOnlyData);
+	}
+	if (bNeedRecordValuesFromFile)
+	{
+		Singleton.RecordValuesFromFile(FileOnlyData, *ConfigFile);
+	}
 }
 
-const FConfigFile* FindOrLoadConfigFile(ELoadType LoadType, const TCHAR* ConfigPlatform, const TCHAR* Filename,
-	FConfigFile& Buffer)
+void FCookConfigAccessTracker::RecordValuesFromFile(const FConfigAccessData& FileOnlyData,
+	const FConfigFile& ConfigFile)
 {
-	bool bHasPlatformName = ConfigPlatform && ConfigPlatform[0] != '\0';
-	FName ConfigPlatformFName;
-	if (bHasPlatformName)
+	FConfigAccessData FullPathData = FileOnlyData.GetFileOnlyData();
+
+	uint32 FullPathHash = GetTypeHash(FullPathData);
+	bool bAlreadyExists = false;
+	UE::TUniqueLock ConfigCacheScopeLock(ConfigCacheLock);
+	LoadedConfigFiles.FindOrAddByHash(FullPathHash, FullPathData, &bAlreadyExists);
+	if (bAlreadyExists)
 	{
-		ConfigPlatformFName = FName(ConfigPlatform);
+		return;
 	}
+	FIgnoreScope IgnoreScope;
 
-	switch (LoadType)
-	{
-	case ELoadType::ConfigSystem:
+	TArray<FName> ValueNames;
+	for (const TPair<FString, FConfigSection>& SectionPair : ConfigFile)
 	{
-		// ForPlatform will return GConfig if ConfigPlatform is NAME_None
-		FConfigCacheIni* ConfigSystem = FConfigCacheIni::ForPlatform(ConfigPlatformFName);
-
-		// The ini files may have been recorded by fullpath or by shortname; search first for a fullpath match using
-		// FindConfigFile and if that fails search for the shortname match by iterating over all files in GConfig
-		FConfigFile* ConfigFile = ConfigSystem->FindConfigFile(FString(Filename));
-		if (ConfigFile)
-		{
-			return ConfigFile;
-		}
-		if (!ConfigFile)
+		FullPathData.SectionName = FName(FStringView(SectionPair.Key), NAME_NO_NUMBER).GetComparisonIndex();
+		const FConfigSection& Section = SectionPair.Value;
+		ValueNames.Reset();
+		SectionPair.Value.GetKeys(ValueNames);
+		for (FName ValueName : ValueNames)
 		{
-			FName FileFName = FName(Filename);
-			for (const FString& ConfigFilename : ConfigSystem->GetFilenames())
-			{
-				ConfigFile = ConfigSystem->FindConfigFile(ConfigFilename);
-				if (ConfigFile->Name == FileFName)
-				{
-					return ConfigFile;
-				}
-			}
+			FullPathData.ValueName = FMinimalName(ValueName);
+			LoadedValues.FindOrAdd(FullPathData) = MultiValueToString(Section, ValueName);
 		}
-		return nullptr;
+	}
+}
+
+FString FCookConfigAccessTracker::MultiValueToString(const FConfigSection& Section, FName ValueName)
+{
+	TArray<const FConfigValue*, TInlineAllocator<8>> Values;
+	Section.MultiFindPointer(ValueName, Values, true /* bMaintainOrder */);
+	if (Values.IsEmpty())
+	{
+		return FString();
+	}
+	if (Values.Num() == 1)
+	{
+		return Values[0]->GetValue();
+	}
+	TStringBuilder<256> ArrayValueStr;
+	ArrayValueStr << Values[0]->GetValue();
+	for (const FConfigValue* Value : TConstArrayView<const FConfigValue*>(Values).RightChop(1))
+	{
+		ArrayValueStr << TEXT("\n") << Value->GetValue();
+	}
+	return FString(*ArrayValueStr);
+}
+
+const FConfigFile* FindOrLoadConfigFile(const FConfigAccessData& AccessData, FConfigFile& Buffer)
+{
+	FName ConfigPlatform(AccessData.GetConfigPlatform());
+	FName FileName(AccessData.GetFileName());
+
+	switch (AccessData.LoadType)
+	{
+	case ELoadType::ConfigSystem:
+	{
+		return FindConfigCacheIniFile(ConfigPlatform, FileName);
 	}
 	case ELoadType::LocalIniFile:
-		if (FConfigCacheIni::LoadLocalIniFile(Buffer, Filename, true /* bIsBaseIniName */, ConfigPlatform))
+		if (FConfigCacheIni::LoadLocalIniFile(Buffer, *WriteToString<128>(FileName), true /* bIsBaseIniName */,
+			*WriteToString<64>(ConfigPlatform)))
 		{
 			return &Buffer;
 		}
 		return nullptr;
 	case ELoadType::LocalSingleIniFile:
-		if (FConfigCacheIni::LoadLocalIniFile(Buffer, Filename, false /* bIsBaseIniName */, ConfigPlatform))
+		if (FConfigCacheIni::LoadLocalIniFile(Buffer, *WriteToString<128>(FileName), false /* bIsBaseIniName */,
+			*WriteToString<64>(ConfigPlatform)))
 		{
 			return &Buffer;
 		}
@@ -321,14 +655,16 @@ const FConfigFile* FindOrLoadConfigFile(ELoadType LoadType, const TCHAR* ConfigP
 		// TODO: LoadExternalIniFile is the same as LoadLocalIniFile, but with possibly redirected
 		// EngineConfigDir and ProjectConfigDir. We can not load them without that extra information.
 		// For now, assume it used the default EngineConfigDir and ProjectConfigDir.
-		if (FConfigCacheIni::LoadLocalIniFile(Buffer, Filename, true /* bIsBaseIniName */, ConfigPlatform))
+		if (FConfigCacheIni::LoadLocalIniFile(Buffer, *WriteToString<128>(FileName), true /* bIsBaseIniName */,
+			*WriteToString<64>(ConfigPlatform)))
 		{
 			return &Buffer;
 		}
 		return nullptr;
 	case ELoadType::ExternalSingleIniFile:
 		// TODO: Same comment as in ExternalIniFile
-		if (FConfigCacheIni::LoadLocalIniFile(Buffer, Filename, false /* bIsBaseIniName */, ConfigPlatform))
+		if (FConfigCacheIni::LoadLocalIniFile(Buffer, *WriteToString<128>(FileName), false /* bIsBaseIniName */,
+			*WriteToString<64>(ConfigPlatform)))
 		{
 			return &Buffer;
 		}
diff --git a/Engine/Source/Editor/UnrealEd/Private/Cooker/CookConfigAccessTracker.h b/Engine/Source/Editor/UnrealEd/Private/Cooker/CookConfigAccessTracker.h
--- a/Engine/Source/Editor/UnrealEd/Private/Cooker/CookConfigAccessTracker.h
+++ b/Engine/Source/Editor/UnrealEd/Private/Cooker/CookConfigAccessTracker.h
@@ -2,25 +2,44 @@
 
 #pragma once
 
+#include "Misc/ConfigAccessTracking.h"
+
+#include "Containers/ArrayView.h"
+#include "Containers/StringView.h"
+#include "Misc/StringBuilder.h"
+#include "UObject/NameTypes.h"
+
+#if UE_WITH_CONFIG_TRACKING
 #include "Async/Mutex.h"
 #include "Cooker/MPCollector.h"
 #include "Containers/Array.h"
 #include "Containers/Map.h"
 #include "Containers/Set.h"
-#include "Containers/StringView.h"
 #include "Interfaces/ITargetPlatform.h"
-#include "Misc/ConfigAccessTracking.h"
 #include "Misc/ConfigCacheIni.h"
 #include "Misc/PackageAccessTracking.h"
 #include "Templates/TypeHash.h"
 #include "Templates/UnrealTemplate.h"
-#include "UObject/NameTypes.h"
+#endif
+
+namespace UE::ConfigAccessTracking
+{
+
+void EscapeConfigTrackingTokenToString(FName Token, FStringBuilderBase& Result);
+void EscapeConfigTrackingTokenAppendString(FName Token, FStringBuilderBase& Result);
+bool TryTokenizeConfigTrackingString(FStringView Text, TArrayView<FStringBuilderBase*> OutTokens);
+
+} // namespace UE::ConfigAccessTracking
 
 #if UE_WITH_CONFIG_TRACKING
 
+namespace UE::ConfigAccessTracking { struct FConfigAccessData; }
+
 namespace UE::ConfigAccessTracking
 {
 
+constexpr FStringView PlatformAgnosticName = TEXTVIEW("<Editor>");
+
 /** Full path of a FConfigValue that was reported read. */
 struct FConfigAccessData
 {
@@ -31,59 +50,32 @@ struct FConfigAccessData
 	const ITargetPlatform* RequestingPlatform = nullptr;
 	ELoadType LoadType = ELoadType::Uninitialized;
 
+	FConfigAccessData() = default;
+	FConfigAccessData(ELoadType InLoadType, FNameEntryId InConfigPlatform, FNameEntryId InFileName,
+		FNameEntryId InSectionName, FMinimalName InValueName, const ITargetPlatform* InRequestingPlatform);
+	FConfigAccessData(ELoadType InLoadType, FNameEntryId InConfigPlatform, FNameEntryId InFileName);
+
 	FName GetConfigPlatform() const { return FName(ConfigPlatform, ConfigPlatform, NAME_NO_NUMBER_INTERNAL); }
 	FName GetFileName() const { return FName(FileName, FileName, NAME_NO_NUMBER_INTERNAL); }
 	FName GetSectionName() const { return FName(SectionName, SectionName, NAME_NO_NUMBER_INTERNAL); }
 	FName GetValueName() const { return FName(ValueName); }
 
-	friend uint32 GetTypeHash(const FConfigAccessData& Data)
-	{
-		uint32 Hash = static_cast<uint32>(Data.LoadType);
-		Hash = HashCombineFast(Hash, Data.ConfigPlatform.ToUnstableInt());
-		Hash = HashCombineFast(Hash, Data.FileName.ToUnstableInt());
-		Hash = HashCombineFast(Hash, Data.SectionName.ToUnstableInt());
-		Hash = HashCombineFast(Hash, GetTypeHash(Data.ValueName));
-		Hash = HashCombineFast(Hash, GetTypeHash(Data.RequestingPlatform));
-		return Hash;
-	}
-	bool IsSameConfigFile(const FConfigAccessData& Other) const
-	{
-		return LoadType == Other.LoadType && ConfigPlatform == Other.ConfigPlatform && FileName == Other.FileName;
-	}
-	bool operator==(const FConfigAccessData& Other) const
-	{
-		return LoadType == Other.LoadType && ConfigPlatform == Other.ConfigPlatform &&
-			FileName == Other.FileName && SectionName == Other.SectionName && ValueName == Other.ValueName &&
-			RequestingPlatform == Other.RequestingPlatform;
-	}
-	bool operator!=(const FConfigAccessData& Other) const
-	{
-		return !(*this == Other);
-	}
-	bool operator<(const FConfigAccessData& Other) const
-	{
-		if (LoadType != Other.LoadType) return static_cast<uint32>(LoadType) < static_cast<uint32>(Other.LoadType);
-		if (ConfigPlatform != Other.ConfigPlatform) return ConfigPlatform.LexicalLess(Other.ConfigPlatform);
-		if (FileName != Other.FileName) return FileName.LexicalLess(Other.FileName);
-		if (SectionName != Other.SectionName) return SectionName.LexicalLess(Other.SectionName);
-		if (ValueName != Other.ValueName) return FName(ValueName).LexicalLess(FName(Other.ValueName));
-		if (RequestingPlatform != Other.RequestingPlatform)
-		{
-			if (RequestingPlatform == nullptr) return true;
-			if (Other.RequestingPlatform == nullptr) return false;
-			return RequestingPlatform->PlatformName() < Other.RequestingPlatform->PlatformName();
-		}
-		return false;
-	}
-};
+	FConfigAccessData GetFileOnlyData() const;
+	FConfigAccessData GetPathOnlyData() const;
+	FString FullPathToString() const;
+	void AppendFullPath(FStringBuilderBase& Out) const;
 
-/** Disables recording of FConfigAccessData reported as accessed on the current thread while in scope. */
-struct FIgnoreScope
-{
-	FIgnoreScope();
-	~FIgnoreScope();
+	/**
+	 * "ConfigSystem.<Editor>.../../../Engine/Config/ConsoleVariables.ini:[Section]:Value"
+	 *   -> "ConfigSystem", "<Editor>", "../../../Engine/Config/ConsoleVariables.ini", "Section", "Value"
+	 */
+	static FConfigAccessData Parse(FStringView Text);
 
-	bool bPreviousIgnoreAccess = false;
+	friend uint32 GetTypeHash(const FConfigAccessData& Data);
+	bool IsSameConfigFile(const FConfigAccessData& Other) const;
+	bool operator==(const FConfigAccessData& Other) const;
+	bool operator!=(const FConfigAccessData& Other) const;
+	bool operator<(const FConfigAccessData& Other) const;
 };
 
 /**
@@ -102,27 +94,33 @@ public:
 	 * Get records requested for the given package and given platform, including RequestingPlatform=nullptr.
 	 * Returned records are SORTED by FConfigAccessData::operator<.
 	 */
-	TArray<UE::ConfigAccessTracking::FConfigAccessData> GetPackageRecords(FName ReferencerPackage,
-		const ITargetPlatform* TargetPlatform) const;
+	TArray<FConfigAccessData> GetPackageRecords(FName ReferencerPackage, const ITargetPlatform* TargetPlatform) const;
 	/**
 	 * Get records for all requesting packages, including records not associated with a package.
 	 * Returned records are SORTED by FConfigAccessData::operator<.
 	 */
-	TArray<UE::ConfigAccessTracking::FConfigAccessData> GetCookRecords() const;
+	TArray<FConfigAccessData> GetCookRecords() const;
 	/**
 	 * Get records requested for all requesting packages, including records not associated with a package,
 	 * but filtered by the given TargetPlatform. Includes records requested with no RequestingPlatform.
 	 * TargetPlatform==nullptr returns only records requested with no RequestingPlatform.
 	 * Returned records are SORTED by FConfigAccessData::operator<.
 	 */
-	TArray<UE::ConfigAccessTracking::FConfigAccessData> GetCookRecords(const ITargetPlatform* TargetPlatform) const;
+	TArray<FConfigAccessData> GetCookRecords(const ITargetPlatform* TargetPlatform) const;
 	/** Add a record as if requested by the given package, or not associated with a package if PackageName.IsNone(). */
-	void AddRecord(FName PackageName, const UE::ConfigAccessTracking::FConfigAccessData& Data);
+	void AddRecord(FName PackageName, const FConfigAccessData& Data);
+
+	/** Lookup in GConfig, LoadConfigFile, or in already-cached values a value indicated by an FConfigAccessData. */
+	FString GetValue(const FConfigAccessData& AccessData);
+	/** Unmarshal a FConfigAccessData.FullPathToString string back to an FConfigAccessData and lookup its value. */
+	FString GetValue(FStringView AccessDataFullPath);
 
 private:
 	FCookConfigAccessTracker();
 	virtual ~FCookConfigAccessTracker();
 
+	void RecordValuesFromFile(const FConfigAccessData& FileOnlyData, const FConfigFile& ConfigFile);
+
 	/** Track object reference reads */
 	static void StaticOnConfigValueRead(UE::ConfigAccessTracking::FSection* Section, FMinimalName ValueName,
 		const FConfigValue& ConfigValue);
@@ -131,11 +129,16 @@ private:
 	static void SortRecordsAndFilterByPlatform(TArray<FConfigAccessData>& Records,
 		const ITargetPlatform* TargetPlatform);
 
+	static FString MultiValueToString(const FConfigSection& Section, FName ValueName);
+
 private:
 	// Use a mutex rather than a critical section for synchronization.  Calls into system libraries, such as windows critical section
 	// functions, are 50 times more expensive on build farm VMs, radically affecting cook times, which this avoids. 
 	mutable UE::FMutex RecordsLock;
+	mutable UE::FMutex ConfigCacheLock;
 	TMap<FName, TSet<UE::ConfigAccessTracking::FConfigAccessData>> PackageRecords;
+	TSet<FConfigAccessData> LoadedConfigFiles;
+	TMap<FConfigAccessData, FString> LoadedValues;
 	UE::ConfigAccessTracking::FConfigValueReadCallbackId OnConfigValueReadCallbackHandle;
 	bool bEnabled = false;
 
@@ -145,15 +148,11 @@ private:
 
 /**
  * Find a ConfigFile by name and ConfigPlatform, either in GConfig or loaded from disk.
- * @param LoadType LoadType describing which ConfigCacheIni method was used to load the ConfigFile
- * @param ConfigPlatform PlatformName used with the ConfigCacheIni method that loaded the ConfigFile, or
- *        nullptr/emptystring for platform-agnostic
- * @param Filename Short ('Engine') or long ('../../../QAGame/Config/DefaultEngine.ini') name of the ConfigFile
+ * @param AccessData Specifies the LoadType, ConfigPlatform, Filename to load.
  * @param Buffer FConfigFile Buffer that will hold the result if LoadConfigFile was called
  * @param The discovered configfile, or nullptr.
  */
-const FConfigFile* FindOrLoadConfigFile(ELoadType LoadType, const TCHAR* ConfigPlatform, const TCHAR* Filename,
-	FConfigFile& Buffer);
+const FConfigFile* FindOrLoadConfigFile(const FConfigAccessData& AccessData, FConfigFile& Buffer);
 
 /** Return whether LoadType is a type that can be loaded by FindOrLoadConfigFile. */
 bool IsLoadableLoadType(ELoadType LoadType);
@@ -183,6 +182,56 @@ public:
 	static FGuid MessageType;
 };
 
+}
+
+namespace UE::ConfigAccessTracking
+{
+
+inline uint32 GetTypeHash(const UE::ConfigAccessTracking::FConfigAccessData& Data)
+{
+	uint32 Hash = static_cast<uint32>(Data.LoadType);
+	Hash = HashCombineFast(Hash, Data.ConfigPlatform.ToUnstableInt());
+	Hash = HashCombineFast(Hash, Data.FileName.ToUnstableInt());
+	Hash = HashCombineFast(Hash, Data.SectionName.ToUnstableInt());
+	Hash = HashCombineFast(Hash, GetTypeHash(Data.ValueName));
+	Hash = HashCombineFast(Hash, GetTypeHash(Data.RequestingPlatform));
+	return Hash;
+}
+
+
+inline bool FConfigAccessData::IsSameConfigFile(const FConfigAccessData& Other) const
+{
+	return LoadType == Other.LoadType && ConfigPlatform == Other.ConfigPlatform && FileName == Other.FileName;
+}
+
+inline bool FConfigAccessData::operator==(const FConfigAccessData& Other) const
+{
+	return LoadType == Other.LoadType && ConfigPlatform == Other.ConfigPlatform &&
+		FileName == Other.FileName && SectionName == Other.SectionName && ValueName == Other.ValueName &&
+		RequestingPlatform == Other.RequestingPlatform;
+}
+
+inline bool FConfigAccessData::operator!=(const FConfigAccessData& Other) const
+{
+	return !(*this == Other);
+}
+
+inline bool FConfigAccessData::operator<(const FConfigAccessData& Other) const
+{
+	if (LoadType != Other.LoadType) return static_cast<uint32>(LoadType) < static_cast<uint32>(Other.LoadType);
+	if (ConfigPlatform != Other.ConfigPlatform) return ConfigPlatform.LexicalLess(Other.ConfigPlatform);
+	if (FileName != Other.FileName) return FileName.LexicalLess(Other.FileName);
+	if (SectionName != Other.SectionName) return SectionName.LexicalLess(Other.SectionName);
+	if (ValueName != Other.ValueName) return FName(ValueName).LexicalLess(FName(Other.ValueName));
+	if (RequestingPlatform != Other.RequestingPlatform)
+	{
+		if (RequestingPlatform == nullptr) return true;
+		if (Other.RequestingPlatform == nullptr) return false;
+		return RequestingPlatform->PlatformName() < Other.RequestingPlatform->PlatformName();
+	}
+	return false;
+}
+
 } // namespace UE::ConfigAccessTracking
 
 #else // !UE_WITH_CONFIG_TRACKING
diff --git a/Engine/Source/Editor/UnrealEd/Private/Cooker/CookRequestCluster.cpp b/Engine/Source/Editor/UnrealEd/Private/Cooker/CookRequestCluster.cpp
--- a/Engine/Source/Editor/UnrealEd/Private/Cooker/CookRequestCluster.cpp
+++ b/Engine/Source/Editor/UnrealEd/Private/Cooker/CookRequestCluster.cpp
@@ -1095,7 +1095,7 @@ void FRequestCluster::FGraphSearch::ExploreVertexEdges(FVertexData& Vertex)
 
 	auto ProcessPlatformAttachments = [this, PackageName, &PackageData, &PlatformAgnosticQueryPlatformData, &AddPlatformDependencyRange]
 		(int32 PlatformIndex, const ITargetPlatform* TargetPlatform, FFetchPlatformData& FetchPlatformData,
-			FPackagePlatformData& PackagePlatformData, const FCookAttachments& PlatformAttachments, bool bExploreDependencies)
+			FPackagePlatformData& PackagePlatformData, FCookAttachments& PlatformAttachments, bool bExploreDependencies)
 	{
 		bool bFoundBuildDefinitions = false;
 		ICookedPackageWriter* PackageWriter = FetchPlatformData.Writer;
@@ -1103,7 +1103,8 @@ void FRequestCluster::FGraphSearch::ExploreVertexEdges(FVertexData& Vertex)
 		if (Cluster.IsIncrementalCook() && PackagePlatformData.IsCookable())
 		{
 			bool bIterativelyUnmodified = false;
-			if (IsCookAttachmentsValid(PackageName, PlatformAttachments))
+			UE::TargetDomain::FCookDependencies& CookDependencies = PlatformAttachments.Dependencies;
+			if (CookDependencies.HasKeyMatch())
 			{
 				if (IsIterativeEnabled(PackageName, Cluster.COTFS.bHybridIterativeAllowAllClasses))
 				{
@@ -1112,11 +1113,11 @@ void FRequestCluster::FGraphSearch::ExploreVertexEdges(FVertexData& Vertex)
 				}
 				if (bExploreDependencies)
 				{
-					AddPlatformDependencyRange(PlatformAttachments.BuildDependencies, PlatformIndex,
+					AddPlatformDependencyRange(CookDependencies.GetPackageDependencies(), PlatformIndex,
 						EInstigator::HardDependency);
 					if (Cluster.bAllowSoftDependencies)
 					{
-						AddPlatformDependencyRange(PlatformAttachments.RuntimeOnlyDependencies, PlatformIndex,
+						AddPlatformDependencyRange(CookDependencies.GetRuntimePackageDependencies(), PlatformIndex,
 							EInstigator::HardDependency);
 					}
 				}
@@ -1125,7 +1126,7 @@ void FRequestCluster::FGraphSearch::ExploreVertexEdges(FVertexData& Vertex)
 				{
 					bFoundBuildDefinitions = true;
 					Cluster.BuildDefinitions.AddBuildDefinitionList(PackageName, TargetPlatform,
-						PlatformAttachments.BuildDefinitionList);
+						PlatformAttachments.BuildDefinitions.Definitions);
 				}
 			}
 			bool bShouldIterativelySkip = bIterativelyUnmodified;
@@ -1147,10 +1148,10 @@ void FRequestCluster::FGraphSearch::ExploreVertexEdges(FVertexData& Vertex)
 		if (Cluster.bPreQueueBuildDefinitions && !bFoundBuildDefinitions)
 		{
 			if (PlatformAgnosticQueryPlatformData.bActive &&
-				IsCookAttachmentsValid(PackageName, PlatformAgnosticQueryPlatformData.CookAttachments))
+				PlatformAgnosticQueryPlatformData.CookAttachments.Dependencies.HasKeyMatch())
 			{
 				Cluster.BuildDefinitions.AddBuildDefinitionList(PackageName, TargetPlatform,
-					PlatformAgnosticQueryPlatformData.CookAttachments.BuildDefinitionList);
+					PlatformAgnosticQueryPlatformData.CookAttachments.BuildDefinitions.Definitions);
 			}
 		}
 	};
@@ -1486,7 +1487,7 @@ void FRequestCluster::FQueryVertexBatch::Send()
 			{
 				RecordCacheResults(PackageName, PlatformIndex, MoveTemp(Attachments));
 			};
-			UE::TargetDomain::FetchCookAttachments(PlatformData.PackageNames, FetchPlatformData.Platform,
+			UE::TargetDomain::FCookAttachments::Fetch(PlatformData.PackageNames, FetchPlatformData.Platform,
 				FetchPlatformData.Writer, MoveTemp(Callback));
 		}
 		else
diff --git a/Engine/Source/Editor/UnrealEd/Private/EditorDomain/EditorDomainUtils.cpp b/Engine/Source/Editor/UnrealEd/Private/EditorDomain/EditorDomainUtils.cpp
--- a/Engine/Source/Editor/UnrealEd/Private/EditorDomain/EditorDomainUtils.cpp
+++ b/Engine/Source/Editor/UnrealEd/Private/EditorDomain/EditorDomainUtils.cpp
@@ -1916,16 +1916,10 @@ bool TrySavePackage(UPackage* Package)
 
 	if (bStorageResultValid)
 	{
-		// TODO_BuildDefinitionList: Calculate and store BuildDefinitionList on the PackageData, or collect it here from some other source.
-		TArray<UE::DerivedData::FBuildDefinition> BuildDefinitions;
-		FCbObject BuildDefinitionList = UE::TargetDomain::BuildDefinitionListToObject(BuildDefinitions);
-		FCbObject TargetDomainDependencies = UE::TargetDomain::CollectDependenciesObject(Package, nullptr, nullptr);
-		if (TargetDomainDependencies)
-		{
-			TArray<IPackageWriter::FCommitAttachmentInfo, TInlineAllocator<2>> Attachments;
-			Attachments.Add({ "Dependencies", TargetDomainDependencies });
-			// TODO: Reenable BuildDefinitionList once FCbPackage support for empty FCbObjects is in
-			//Attachments.Add({ "BuildDefinitionList", BuildDefinitionList });
+		TArray<IPackageWriter::FCommitAttachmentInfo, TInlineAllocator<2>> Attachments;
+		UE::TargetDomain::CollectAndStoreCookAttachments(Package, nullptr, Attachments);
+		if (!Attachments.IsEmpty())
+		{
 			UE::TargetDomain::CommitEditorDomainCookAttachments(Package->GetFName(), Attachments);
 		}
 	}
diff --git a/Engine/Source/Editor/UnrealEd/Private/TargetDomain/TargetDomainUtils.cpp b/Engine/Source/Editor/UnrealEd/Private/TargetDomain/TargetDomainUtils.cpp
--- a/Engine/Source/Editor/UnrealEd/Private/TargetDomain/TargetDomainUtils.cpp
+++ b/Engine/Source/Editor/UnrealEd/Private/TargetDomain/TargetDomainUtils.cpp
@@ -8,6 +8,7 @@
 #include "AssetRegistry/IAssetRegistry.h"
 #include "Containers/Set.h"
 #include "Containers/UnrealString.h"
+#include "Cooker/CookConfigAccessTracker.h"
 #include "Cooker/PackageBuildDependencyTracker.h"
 #include "DerivedDataBuildDefinition.h"
 #include "DerivedDataBuildKey.h"
@@ -22,6 +23,7 @@
 #include "Misc/ScopeRWLock.h"
 #include "Misc/StringBuilder.h"
 #include "Serialization/CompactBinary.h"
+#include "Serialization/CompactBinarySerialization.h"
 #include "Serialization/CompactBinaryWriter.h"
 #include "Serialization/PackageWriter.h"
 #include "ZenStoreHttpClient.h"
@@ -29,6 +31,9 @@
 namespace UE::TargetDomain
 {
 
+constexpr uint32 CookDependenciesVersion = 0x00000001;
+static const FUtf8StringView CookDependenciesAttachmentKey = UTF8TEXTVIEW("CookDependencies");
+static const FUtf8StringView BuildDefinitionsAttachmentKey = UTF8TEXTVIEW("BuildDefinitionsAttachmentKey");
 /**
  * Reads / writes an oplog for EditorDomain BuildDefinitionLists.
  * TODO: Reduce duplication between this class and FZenStoreWriter
@@ -75,9 +80,39 @@ private:
 };
 TUniquePtr<FEditorDomainOplog> GEditorDomainOplog;
 
-bool TryCreateKey(FName PackageName, TConstArrayView<FName> SortedBuildDependencies, FIoHash* OutHash, FString* OutErrorMessage)
+bool FCookDependencies::IsValid() const
+{
+	return bValid;
+}
+
+bool FCookDependencies::HasKeyMatch()
+{
+	if (!bValid)
+	{
+		return false;
+	}
+	if (StoredKey.IsZero())
+	{
+		return false;
+	}
+	if (CurrentKey.IsZero())
+	{
+		if (!TryCalculateCurrentKey())
+		{
+			return false;
+		}
+	}
+	return CurrentKey == StoredKey;
+}
+
+bool FCookDependencies::TryCalculateCurrentKey(FString* OutErrorMessage)
 {
 	IAssetRegistry* AssetRegistry = IAssetRegistry::Get();
+	if (PackageName.IsNone())
+	{
+		if (OutErrorMessage) *OutErrorMessage = TEXT("PackageName is not set.");
+		return false;
+	}
 	if (!AssetRegistry)
 	{
 		if (OutErrorMessage) *OutErrorMessage = TEXT("AssetRegistry is unavailable.");
@@ -96,292 +131,423 @@ bool TryCreateKey(FName PackageName, TConstArrayView<FName> SortedBuildDependenc
 		if (OutErrorMessage) *OutErrorMessage = PackageDigest.GetStatusString();
 		return false;
 	}
+
 	KeyBuilder.Update(&PackageDigest.Hash, sizeof(PackageDigest.Hash));
 
-	for (FName DependencyName : SortedBuildDependencies)
+	for (FName PackageDependency : PackageDependencies)
 	{
-		PackageDigest = EditorDomain->GetPackageDigest(PackageName);
+		PackageDigest = EditorDomain->GetPackageDigest(PackageDependency);
 		if (!PackageDigest.IsSuccessful())
 		{
 			if (OutErrorMessage)
 			{
 				*OutErrorMessage = FString::Printf(TEXT("Could not create PackageDigest for %s: %s"),
-					*DependencyName.ToString(), *PackageDigest.GetStatusString());
+					*PackageDependency.ToString(), *PackageDigest.GetStatusString());
 			}
 			return false;
 		}
 		KeyBuilder.Update(&PackageDigest.Hash, sizeof(PackageDigest.Hash));
 	}
 
-	if (OutHash)
+	if (!ConfigDependencies.IsEmpty())
 	{
-		*OutHash = KeyBuilder.Finalize();
+#if UE_WITH_CONFIG_TRACKING
+		using namespace UE::ConfigAccessTracking;
+		FCookConfigAccessTracker& ConfigTracker = FCookConfigAccessTracker::Get();
+#endif
+		for (const FString& ConfigDependency : ConfigDependencies)
+		{
+			FString Value;
+#if UE_WITH_CONFIG_TRACKING
+			Value = ConfigTracker.GetValue(ConfigDependency);
+#endif
+			uint8 Marker = 0;
+			KeyBuilder.Update(&Marker, sizeof(Marker));
+			if (!Value.IsEmpty())
+			{
+				KeyBuilder.Update(*Value, Value.Len() * sizeof(Value[0]));
+			}
+		}
 	}
+
+	if (OutErrorMessage) OutErrorMessage->Reset();
+	CurrentKey = KeyBuilder.Finalize();
 	return true;
 }
 
-bool TryCollectKeyAndDependencies(UPackage* Package, const ITargetPlatform* TargetPlatform, FIoHash* OutHash, TArray<FName>* OutBuildDependencies,
-	TArray<FName>* OutRuntimeOnlyDependencies, FString* OutErrorMessage)
+void FCookDependencies::Reset()
+{
+	PackageDependencies.Reset();
+	ConfigDependencies.Reset();
+	RuntimePackageDependencies.Reset();
+	PackageName = FName();
+	StoredKey = FIoHash::Zero;
+	CurrentKey = FIoHash::Zero;
+	bValid = false;
+}
+
+void FCookDependencies::Empty()
+{
+	Reset();
+	PackageDependencies.Empty();
+	ConfigDependencies.Empty();
+	RuntimePackageDependencies.Empty();
+}
+
+FCookDependencies FCookDependencies::Collect(UPackage* Package, const ITargetPlatform* TargetPlatform, FString* OutErrorMessage)
 {
 	if (!Package)
 	{
 		if (OutErrorMessage) *OutErrorMessage = TEXT("Invalid null package.");
-		return false;
+		return FCookDependencies();
 	}
 	IAssetRegistry* AssetRegistry = IAssetRegistry::Get();
 	if (!AssetRegistry)
 	{
 		if (OutErrorMessage) *OutErrorMessage = TEXT("AssetRegistry is unavailable.");
-		return false;
+		return FCookDependencies();
 	}
 	FEditorDomain* EditorDomain = FEditorDomain::Get();
 	if (!EditorDomain)
 	{
 		if (OutErrorMessage) *OutErrorMessage = TEXT("EditorDomain is unavailable.");
-		return false;
+		return FCookDependencies();
 	}
 
-	FName PackageName = Package->GetFName();
-	TSet<FName> BuildDependencies;
-	TSet<FName> RuntimeOnlyDependencies;
+	FCookDependencies Result;
+	Result.PackageName = Package->GetFName();
+	TSet<FName> PackageDependenciesSet;
+	TSet<FName> RuntimePackageDependenciesSet;
 
 	TArray<FName> AssetDependencies;
-	AssetRegistry->GetDependencies(PackageName, AssetDependencies, UE::AssetRegistry::EDependencyCategory::Package,
-		UE::AssetRegistry::EDependencyQuery::Game);
+	AssetRegistry->GetDependencies(Result.PackageName, AssetDependencies,
+		UE::AssetRegistry::EDependencyCategory::Package, UE::AssetRegistry::EDependencyQuery::Game);
 
 	FPackageBuildDependencyTracker& Tracker = FPackageBuildDependencyTracker::Get();
 
+#if UE_WITH_PACKAGE_ACCESS_TRACKING
 	if (Tracker.IsEnabled())
 	{
-		TArray<FBuildDependencyAccessData> AccessDatas = Tracker.GetAccessDatas(PackageName);
+		TArray<FBuildDependencyAccessData> AccessDatas = Tracker.GetAccessDatas(Result.PackageName);
 
-		BuildDependencies.Reserve(AccessDatas.Num());
+		PackageDependenciesSet.Reserve(AccessDatas.Num());
 		for (FBuildDependencyAccessData& AccessData : AccessDatas)
 		{
 			if (AccessData.TargetPlatform == TargetPlatform || AccessData.TargetPlatform == nullptr)
 			{
-				BuildDependencies.Add(AccessData.ReferencedPackage);
+				PackageDependenciesSet.Add(AccessData.ReferencedPackage);
 			}
 		}
 
-		RuntimeOnlyDependencies.Reserve(AssetDependencies.Num());
+		RuntimePackageDependenciesSet.Reserve(AssetDependencies.Num());
 		for (FName DependencyName : AssetDependencies)
 		{
-			if (!BuildDependencies.Contains(DependencyName))
+			if (!PackageDependenciesSet.Contains(DependencyName))
 			{
-				RuntimeOnlyDependencies.Add(DependencyName);
+				RuntimePackageDependenciesSet.Add(DependencyName);
 			}
 		}
 	}
 	else
+#endif
 	{
 		// Defensively treat all asset dependencies as build dependencies and have zero runtime only dependencies
-		BuildDependencies.Append(AssetDependencies);
+		PackageDependenciesSet.Append(AssetDependencies);
 	}
 
-	TArray<FName> SortedBuild;
-	SortedBuild = BuildDependencies.Array();
 	TStringBuilder<256> StringBuffer;
 	FName TransientPackageName = GetTransientPackage()->GetFName();
-	auto IsTransientPackageName = [&StringBuffer, TransientPackageName](FName PackageName)
+	auto IsTransientPackageName = [&StringBuffer, TransientPackageName](FName InPackageName)
 	{
-		PackageName.ToString(StringBuffer);
-		return PackageName == TransientPackageName ||
-			FPackageName::IsMemoryPackage(StringBuffer) ||
+		if (InPackageName == TransientPackageName)
+		{
+			return true;
+		}
+		InPackageName.ToString(StringBuffer);
+		return FPackageName::IsMemoryPackage(StringBuffer) ||
 			FPackageName::IsScriptPackage(StringBuffer);
 	};
-	SortedBuild.RemoveAllSwap(IsTransientPackageName, EAllowShrinking::No);
-	SortedBuild.Sort(FNameLexicalLess());
-	TArray<FName> SortedRuntimeOnly;
-	SortedRuntimeOnly = RuntimeOnlyDependencies.Array();
-	SortedRuntimeOnly.RemoveAllSwap(IsTransientPackageName, EAllowShrinking::No);
-	SortedRuntimeOnly.Sort(FNameLexicalLess());
+	Result.PackageDependencies = PackageDependenciesSet.Array();
+	Result.PackageDependencies.RemoveAllSwap(IsTransientPackageName, EAllowShrinking::No);
+	Result.PackageDependencies.Sort(FNameLexicalLess());
+	Result.PackageDependencies.Shrink();
 
-	if (!TryCreateKey(PackageName, SortedBuild, OutHash, OutErrorMessage))
+	Result.RuntimePackageDependencies = RuntimePackageDependenciesSet.Array();
+	Result.RuntimePackageDependencies.RemoveAllSwap(IsTransientPackageName, EAllowShrinking::No);
+	Result.RuntimePackageDependencies.Sort(FNameLexicalLess());
+	Result.RuntimePackageDependencies.Shrink();
+
+#if UE_WITH_CONFIG_TRACKING
 	{
-		return false;
+		using namespace UE::ConfigAccessTracking;
+		FCookConfigAccessTracker& ConfigTracker = FCookConfigAccessTracker::Get();
+		if (ConfigTracker.IsEnabled())
+		{
+			TArray<FConfigAccessData> ConfigKeys = ConfigTracker.GetPackageRecords(Result.PackageName, TargetPlatform);
+			Result.ConfigDependencies.Reserve(ConfigKeys.Num());
+			for (const FConfigAccessData& ConfigKey : ConfigKeys)
+			{
+				Result.ConfigDependencies.Add(ConfigKey.FullPathToString());
+			}
+		}
 	}
-
-	if (OutBuildDependencies)
+#endif
+	if (!Result.TryCalculateCurrentKey(OutErrorMessage))
 	{
-		*OutBuildDependencies = MoveTemp(SortedBuild);
+		return FCookDependencies();
 	}
-	if (OutRuntimeOnlyDependencies)
+	Result.StoredKey = Result.CurrentKey;
+	Result.bValid = true;
+
+	if (OutErrorMessage) OutErrorMessage->Reset();
+	return Result;
+}
+
+}
+
+bool LoadFromCompactBinary(FCbObjectView ObjectView, UE::TargetDomain::FCookDependencies& Dependencies)
+{
+	using namespace UE::TargetDomain;
+
+	Dependencies.Reset();
+	int32 Version = -1;
+	for (FCbFieldView FieldView : ObjectView)
 	{
-		*OutRuntimeOnlyDependencies = MoveTemp(SortedRuntimeOnly);
+		FUtf8StringView FieldName = FieldView.GetName();
+		if (FieldName == "Version")
+		{
+			Version = FieldView.AsInt32();
+			if (FieldView.HasError() || Version != CookDependenciesVersion)
+			{
+				return false;
+			}
+		}
+		else if (FieldName == "StoredKey")
+		{
+			if (!LoadFromCompactBinary(FieldView, Dependencies.StoredKey))
+			{
+				return false;
+			}
+		}
+		else if (FieldName == "PackageDependencies")
+		{
+			if (!LoadFromCompactBinary(FieldView, Dependencies.PackageDependencies))
+			{
+				return false;
+			}
+		}
+		else if (FieldName == "ConfigDependencies")
+		{
+			if (!LoadFromCompactBinary(FieldView, Dependencies.ConfigDependencies))
+			{
+				return false;
+			}
+		}
+		else if (FieldName == "RuntimePackageDependencies")
+		{
+			if (!LoadFromCompactBinary(FieldView, Dependencies.RuntimePackageDependencies))
+			{
+				return false;
+			}
+		}
+		// else ignore
 	}
-	if (OutErrorMessage)
+	if (Version == -1)
 	{
-		OutErrorMessage->Reset();
+		return false;
 	}
-	
+	Dependencies.bValid = true;
 	return true;
 }
 
-FCbObject CollectDependenciesObject(UPackage* Package, const ITargetPlatform* TargetPlatform, FString* ErrorMessage)
+FCbWriter& operator<<(FCbWriter& Writer, const UE::TargetDomain::FCookDependencies& CookDependencies)
 {
-	FIoHash TargetDomainKey;
-	TArray<FName> BuildDependencies;
-	TArray<FName> RuntimeOnlyDependencies;
-	if (!UE::TargetDomain::TryCollectKeyAndDependencies(Package, TargetPlatform, &TargetDomainKey, &BuildDependencies, &RuntimeOnlyDependencies, ErrorMessage))
-	{
-		return FCbObject();
-	}
+	using namespace UE::TargetDomain;
 
-	FCbWriter Writer;
 	Writer.BeginObject();
-	Writer << "targetdomainkey" << TargetDomainKey;
-	TStringBuilder<128> PackageNameBuffer;
-	if (!BuildDependencies.IsEmpty())
+	Writer << "Version" << CookDependenciesVersion;
+	Writer << "StoredKey" << CookDependencies.StoredKey;
+	if (!CookDependencies.PackageDependencies.IsEmpty())
 	{
-		Writer.BeginArray("builddependencies");
-		for (FName DependencyName : BuildDependencies)
-		{
-			DependencyName.ToString(PackageNameBuffer);
-			Writer << PackageNameBuffer;
-		}
-		Writer.EndArray();
+		Writer << "PackageDependencies" << CookDependencies.PackageDependencies;
 	}
-	if (!RuntimeOnlyDependencies.IsEmpty())
+	if (!CookDependencies.ConfigDependencies.IsEmpty())
 	{
-		Writer.BeginArray("runtimeonlydependencies");
-		for (FName DependencyName : RuntimeOnlyDependencies)
-		{
-			DependencyName.ToString(PackageNameBuffer);
-			Writer << PackageNameBuffer;
-		}
-		Writer.EndArray();
+		Writer << "ConfigDependencies" << CookDependencies.ConfigDependencies;
+	}
+	if (!CookDependencies.RuntimePackageDependencies.IsEmpty())
+	{
+		Writer << "RuntimePackageDependencies" << CookDependencies.RuntimePackageDependencies;
 	}
 	Writer.EndObject();
-	return Writer.Save().AsObject();
+	return Writer;
 }
 
-FCbObject BuildDefinitionListToObject(TConstArrayView<UE::DerivedData::FBuildDefinition> BuildDefinitionList)
+namespace UE::TargetDomain
+{
+
+FBuildDefinitionList FBuildDefinitionList::Collect(UPackage* Package, const ITargetPlatform* TargetPlatform,
+	FString* OutErrorMessage)
 {
 	using namespace UE::DerivedData;
 
-	if (BuildDefinitionList.IsEmpty())
-	{
-		return FCbObject();
-	}
+	FBuildDefinitionList Result;
 
-	TArray<const FBuildDefinition*> Sorted;
-	Sorted.Reserve(BuildDefinitionList.Num());
-	for (const FBuildDefinition& BuildDefinition : BuildDefinitionList)
+	// TODO_BuildDefinitionList: Calculate and store BuildDefinitionList on the PackageData, or collect it here from some other source.
+	if (Result.Definitions.IsEmpty())
 	{
-		Sorted.Add(&BuildDefinition);
+		if (OutErrorMessage) *OutErrorMessage = TEXT("Not yet implemented");
+		return FBuildDefinitionList();
 	}
-	Sorted.Sort([](const FBuildDefinition& A, const FBuildDefinition& B)
+
+	TArray<FBuildDefinition>& Defs = Result.Definitions;
+	Algo::Sort(Defs, [](const FBuildDefinition& A, const FBuildDefinition& B)
 		{
 			return A.GetKey().Hash < B.GetKey().Hash;
 		});
 
-	FCbWriter Writer;
+	if (OutErrorMessage) OutErrorMessage->Reset();
+	return Result;
+}
+
+void FBuildDefinitionList::Reset()
+{
+	Definitions.Reset();
+}
+
+void FBuildDefinitionList::Empty()
+{
+	Definitions.Empty();
+}
+
+}
+
+bool LoadFromCompactBinary(FCbObject&& Object, UE::TargetDomain::FBuildDefinitionList& Definitions)
+{
+	using namespace UE::DerivedData;
+
+	FCbField DefinitionsField = Object["BuildDefinitions"];
+	FCbArray DefinitionsArrayField = DefinitionsField.AsArray();
+	if (DefinitionsField.HasError())
+	{
+		return false;
+	}
+	TArray<FBuildDefinition>& Defs = Definitions.Definitions;
+	Defs.Empty(DefinitionsArrayField.Num());
+	for (FCbField& BuildDefinitionObj : DefinitionsArrayField)
+	{
+		FOptionalBuildDefinition BuildDefinition = FBuildDefinition::Load(TEXTVIEW("TargetDomainBuildDefinitionList"),
+			BuildDefinitionObj.AsObject());
+		if (!BuildDefinition)
+		{
+			Defs.Empty();
+			return false;
+		}
+		Defs.Add(MoveTemp(BuildDefinition).Get());
+	}
+
+	return true;
+}
+
+FCbWriter& operator<<(FCbWriter& Writer, const UE::TargetDomain::FBuildDefinitionList& Definitions)
+{
+	using namespace UE::DerivedData;
+
 	Writer.BeginObject();
 	Writer.BeginArray("BuildDefinitions");
-	for (const FBuildDefinition* BuildDefinition : Sorted)
+	for (const FBuildDefinition& BuildDefinition : Definitions.Definitions)
 	{
-		BuildDefinition->Save(Writer);
+		BuildDefinition.Save(Writer);
 	}
 	Writer.EndArray();
-	return Writer.Save().AsObject();
+	return Writer;
+}
+
+namespace UE::TargetDomain
+{
+
+void FCookAttachments::Reset()
+{
+	Dependencies.Reset();
+	BuildDefinitions.Reset();
+}
+
+void FCookAttachments::Empty()
+{
+	Dependencies.Empty();
+	BuildDefinitions.Empty();
+}
+
+bool TryCollectAndStoreCookDependencies(UPackage* Package, const ITargetPlatform* TargetPlatform,
+	IPackageWriter::FCommitAttachmentInfo& OutResult)
+{
+	FCookDependencies CookDependencies = FCookDependencies::Collect(Package, TargetPlatform);
+	if (!CookDependencies.IsValid())
+	{
+		OutResult.Value = FCbObject();
+		return false;
+	}
+
+	FCbWriter Writer;
+	Writer << CookDependencies;
+	OutResult.Key = CookDependenciesAttachmentKey;
+	OutResult.Value = Writer.Save().AsObject();
+	return true;
 }
 
-void FetchCookAttachments(TArrayView<FName> PackageNames, const ITargetPlatform* TargetPlatform, ICookedPackageWriter* PackageWriter,
-	TUniqueFunction<void(FName PackageName, FCookAttachments&& Results)>&& Callback)
+bool TryCollectAndStoreBuildDefinitionList(UPackage* Package, const ITargetPlatform* TargetPlatform,
+	IPackageWriter::FCommitAttachmentInfo& OutResult)
+{
+	FBuildDefinitionList Definitions = FBuildDefinitionList::Collect(Package, TargetPlatform);
+	if (Definitions.Definitions.IsEmpty())
+	{
+		OutResult.Value = FCbObject();
+		return false;
+	}
+
+	FCbWriter Writer;
+	Writer << Definitions;
+	OutResult.Key = BuildDefinitionsAttachmentKey;
+	OutResult.Value = Writer.Save().AsObject();
+	return true;
+}
+
+void FCookAttachments::Fetch(TArrayView<FName> PackageNames, const ITargetPlatform* TargetPlatform,
+	ICookedPackageWriter* PackageWriter,
+	TUniqueFunction<void(FName PackageName, FCookAttachments&& Result)>&& Callback)
 {
 	for (FName PackageName : PackageNames)
 	{
-		FCookAttachments Result;
-		const ANSICHAR* DependenciesKey = "Dependencies";
 		FCbObject DependenciesObj;
+		FCbObject BuildDefinitionsObj;
 		if (TargetPlatform)
 		{
 			check(PackageWriter);
-			DependenciesObj = PackageWriter->GetOplogAttachment(PackageName, DependenciesKey);
+			DependenciesObj = PackageWriter->GetOplogAttachment(PackageName, CookDependenciesAttachmentKey);
+			BuildDefinitionsObj = PackageWriter->GetOplogAttachment(PackageName, CookDependenciesAttachmentKey);
 		}
 		else
 		{
 			if (!GEditorDomainOplog)
 			{
-				Callback(PackageName, MoveTemp(Result));
+				Callback(PackageName, FCookAttachments());
 				continue;
 			}
-			DependenciesObj = GEditorDomainOplog->GetOplogAttachment(PackageName, DependenciesKey);
-		}
-
-		Result.StoredKey = DependenciesObj["targetdomainkey"].AsHash();
-		if (Result.StoredKey.IsZero())
-		{
-			Callback(PackageName, MoveTemp(Result));
-			continue;
-		}
-
-		for (FCbFieldView DepObj : DependenciesObj["builddependencies"])
-		{
-			if (FUtf8StringView DependencyName(DepObj.AsString()); !DependencyName.IsEmpty())
-			{
-				Result.BuildDependencies.Add(FName(DependencyName));
-			}
+			DependenciesObj = GEditorDomainOplog->GetOplogAttachment(PackageName, CookDependenciesAttachmentKey);
+			BuildDefinitionsObj = GEditorDomainOplog->GetOplogAttachment(PackageName, BuildDefinitionsAttachmentKey);
 		}
 
-		for (FCbFieldView DepObj : DependenciesObj["runtimeonlydependencies"])
-		{
-			if (FUtf8StringView DependencyName(DepObj.AsString()); !DependencyName.IsEmpty())
-			{
-				Result.RuntimeOnlyDependencies.Add(FName(DependencyName));
-			}
-		}
-
-		const ANSICHAR* BuildDefinitionListKey = "BuildDefinitionList";
-		FCbObject BuildDefinitionListObj;
-		if (TargetPlatform)
-		{
-			BuildDefinitionListObj = PackageWriter->GetOplogAttachment(PackageName, BuildDefinitionListKey);
-		}
-		else
+		FCookAttachments Result;
+		if (LoadFromCompactBinary(DependenciesObj, Result.Dependencies))
 		{
-			BuildDefinitionListObj = GEditorDomainOplog->GetOplogAttachment(PackageName, BuildDefinitionListKey);
+			Result.Dependencies.PackageName = PackageName;
 		}
+		LoadFromCompactBinary(MoveTemp(BuildDefinitionsObj), Result.BuildDefinitions);
 
-		for (FCbField BuildDefinitionObj : BuildDefinitionListObj)
-		{
-			UE::DerivedData::FOptionalBuildDefinition BuildDefinition =
-				UE::DerivedData::FBuildDefinition::Load(TEXTVIEW("TargetDomainBuildDefinitionList"),
-					BuildDefinitionObj.AsObject());
-			if (!BuildDefinition)
-			{
-				Result.BuildDefinitionList.Empty();
-				break;
-			}
-			Result.BuildDefinitionList.Add(MoveTemp(BuildDefinition).Get());
-		}
-		Result.bValid = true;
 		Callback(PackageName, MoveTemp(Result));
 	}
 }
 
-bool IsCookAttachmentsValid(FName PackageName, const FCookAttachments& CookAttachments)
-{
-	if (!CookAttachments.bValid)
-	{
-		return false;
-	}
-
-	FIoHash CurrentKey;
-	if (!TryCreateKey(PackageName, CookAttachments.BuildDependencies, &CurrentKey, nullptr /* OutErrorMessage */))
-	{
-		return false;
-	}
-
-	if (CookAttachments.StoredKey != CurrentKey)
-	{
-		return false;
-	}
-
-	return true;
-}
-
 bool IsIterativeEnabled(FName PackageName, bool bAllowAllClasses)
 {
 	IAssetRegistry* AssetRegistry = IAssetRegistry::Get();
diff --git a/Engine/Source/Editor/UnrealEd/Private/TargetDomain/TargetDomainUtils.h b/Engine/Source/Editor/UnrealEd/Private/TargetDomain/TargetDomainUtils.h
--- a/Engine/Source/Editor/UnrealEd/Private/TargetDomain/TargetDomainUtils.h
+++ b/Engine/Source/Editor/UnrealEd/Private/TargetDomain/TargetDomainUtils.h
@@ -7,64 +7,151 @@
 #include "Containers/ContainersFwd.h"
 #include "DerivedDataBuildDefinition.h"
 #include "IO/IoHash.h"
+#include "Misc/Optional.h"
 #include "Serialization/PackageWriter.h"
 #include "UObject/NameTypes.h"
 
 class FCbObject;
+class FCbObjectView;
+class FCbWriter;
 class FString;
 class ITargetPlatform;
 class UPackage;
 template <typename FuncType> class TUniqueFunction;
 
+namespace UE::TargetDomain { class FCookDependencies; }
+namespace UE::TargetDomain { struct FBuildDefinitionList; }
+namespace UE::TargetDomain { struct FCookAttachments; }
+
+bool LoadFromCompactBinary(FCbObjectView ObjectView, UE::TargetDomain::FCookDependencies& CookAttachments);
+FCbWriter& operator<<(FCbWriter& Writer, const UE::TargetDomain::FCookDependencies& CookAttachments);
+bool LoadFromCompactBinary(FCbObject&& Object, UE::TargetDomain::FBuildDefinitionList& Definitions);
+FCbWriter& operator<<(FCbWriter& Writer, const UE::TargetDomain::FBuildDefinitionList& Definitions);
+
 namespace UE::TargetDomain
 {
 
+/** Call during Startup to initialize global data used by TargetDomain functions. */
 void CookInitialize();
 
-/** Create the TargetDomainKey based on the EditorDomainKeys of the Package and its dependencies. */
-bool TryCreateKey(FName PackageName, TConstArrayView<FName> SortedBuildDependencies, FIoHash* OutHash, FString* OutErrorMessage);
+/**
+ * Recording of the dependencies of a package discovered during cook, used in incremental cooks. All dependencies
+ * except for those marked Runtime contribute to the packages TargetDomain Key. If HasKeyMatch returns false after
+ * fetching this structure for a package at the beginning of cook, then the package is not iteratively skippable and
+ * needs to be recooked, and this structure needs to be recalculated for the package.
+ * 
+ * Runtime fields on the structure are used to inform the cook of discovered softreferences that need to be added to
+ * the cook when the package is cooked.
+ */
+class FCookDependencies
+{
+public:
+	// Build Dependencies
+	const TArray<FName>& GetPackageDependencies() const { return PackageDependencies; }
+	const TArray<FString>& GetConfigDependencies() const { return ConfigDependencies; }
 
-/** Collect the Package's dependencies and the key based on them. */
-bool TryCollectKeyAndDependencies(UPackage* Package, const ITargetPlatform* TargetPlatform,
-	FIoHash* OutHash, TArray<FName>* OutBuildDependencies, TArray<FName>* OutRuntimeOnlyDependencies, FString* OutErrorMessage);
-/** Collect the Package's dependencies, and create a FCbObject describing them for storage in the OpLog. */
-FCbObject CollectDependenciesObject(UPackage* Package, const ITargetPlatform* TargetPlatform, FString* ErrorMessage);
+	// Runtime Dependencies
+	const TArray<FName>& GetRuntimePackageDependencies() const { return RuntimePackageDependencies; }
 
-/** Marshal the given BuildDefinitionList to FCbObject for storage in DDC. */
-FCbObject BuildDefinitionListToObject(TConstArrayView<UE::DerivedData::FBuildDefinition> BuildDefinitionList);
+	// Data about the Key and Package that are not dependencies
+	/**
+	 * True if the structure has been calculated or fetched and accurately reports dependencies and
+	 * key for the package. False if the stucture is default, has been reset, or was marked invalid.
+	 */
+	bool IsValid() const;
+	FName GetPackageName() const { return PackageName; }
+	const FIoHash& GetStoredKey() const { return StoredKey; }
+	const FIoHash& GetCurrentKey() const { return CurrentKey; }
+	bool HasKeyMatch();
 
-struct FCookAttachments
-{
-	TArray<FName> BuildDependencies;
-	TArray<FName> RuntimeOnlyDependencies;
-	TArray<UE::DerivedData::FBuildDefinition> BuildDefinitionList;
+	// Modifying the structure data
+	bool TryCalculateCurrentKey(FString* OutErrorMessage = nullptr);
+	void Reset();
+	void Empty();
+
+	// Construction functions
+	/**
+	 * Read dependencies for the given targetplatform of the given package out of global dependency trackers
+	 * that have recorded its data during the package's load/other/save operations in the current cook session.
+	 */
+	static FCookDependencies Collect(UPackage* Package, const ITargetPlatform* TargetPlatform,
+		FString* OutErrorMessage = nullptr);
+
+	// Fetch function to load the dependencies from a PackageStore is not yet implemented independently for
+	// this structure. Use FCookAttachments instead. 
+
+private:
+	TArray<FName> PackageDependencies;
+	TArray<FString> ConfigDependencies;
+	TArray<FName> RuntimePackageDependencies;
+	FName PackageName;
 	FIoHash StoredKey;
+	FIoHash CurrentKey;
 	bool bValid = false;
 
-	void Reset()
+	friend bool ::LoadFromCompactBinary(FCbObjectView ObjectView, FCookDependencies& CookAttachments);
+	friend FCbWriter& ::operator<<(FCbWriter& Writer, const FCookDependencies& CookAttachments);
+	friend FCookAttachments;
+};
+
+/** Wrapper around TArray<FBuildDefinition>, used to provide custom functions for compactbinary, collection, and fetch */
+struct FBuildDefinitionList
+{
+public:
+	TArray<UE::DerivedData::FBuildDefinition> Definitions;
+
+	void Reset();
+	void Empty();
+
+	/** Collect DDC BuildDefinitions that were issued from the load/save of the given package and platform. */
+	static FBuildDefinitionList Collect(UPackage* Package, const ITargetPlatform* TargetPlatform,
+		FString* OutErrorMessage = nullptr);
+
+private:
+	friend bool ::LoadFromCompactBinary(FCbObject&& ObjectView, FBuildDefinitionList& Definitions);
+	friend FCbWriter& ::operator<<(FCbWriter& Writer, const FBuildDefinitionList& Definitions);
+};
+
+/**
+ * All of the attachments that we want to read during RequestCluster reference traversal. We read them all at once
+ * to batch up the fetch of the attachments from the PackageStore.
+ */
+struct FCookAttachments
+{
+	FCookDependencies Dependencies;
+	FBuildDefinitionList BuildDefinitions;
+
+	void Reset();
+	void Empty();
+
+	static void Fetch(TArrayView<FName> PackageNames, const ITargetPlatform* TargetPlatform,
+		ICookedPackageWriter* PackageWriter,
+		TUniqueFunction<void(FName PackageName, FCookAttachments&& Result)>&& Callback);
+};
+
+bool TryCollectAndStoreCookDependencies(UPackage* Package, const ITargetPlatform* TargetPlatform,
+	IPackageWriter::FCommitAttachmentInfo& OutResult);
+bool TryCollectAndStoreBuildDefinitionList(UPackage* Package, const ITargetPlatform* TargetPlatform,
+	IPackageWriter::FCommitAttachmentInfo& OutResult);
+
+template <typename ArrayType>
+void CollectAndStoreCookAttachments(UPackage* Package, const ITargetPlatform* TargetPlatform,
+	ArrayType& Output)
+{
+	IPackageWriter::FCommitAttachmentInfo Result;
+	if (TryCollectAndStoreCookDependencies(Package, TargetPlatform, Result))
 	{
-		BuildDependencies.Reset();
-		RuntimeOnlyDependencies.Reset();
-		BuildDefinitionList.Reset();
-		bValid = false;
+		Output.Add(MoveTemp(Result));
 	}
-	void Empty()
+	if (TryCollectAndStoreBuildDefinitionList(Package, TargetPlatform, Result))
 	{
-		BuildDependencies.Empty();
-		RuntimeOnlyDependencies.Empty();
-		BuildDefinitionList.Empty();
-		bValid = false;
+		Output.Add(MoveTemp(Result));
 	}
-};
-void FetchCookAttachments(TArrayView<FName> PackageNames, const ITargetPlatform* TargetPlatform, ICookedPackageWriter* PackageWriter,
-	TUniqueFunction<void (FName PackageName, FCookAttachments&& Result)>&& Callback);
-
-bool IsCookAttachmentsValid(FName PackageName, const FCookAttachments& CookAttachments);
+}
 
 /** Return whether iterative cook is enabled for the given packagename, based on used-class allowlist/blocklist. */
 bool IsIterativeEnabled(FName PackageName, bool bAllowAllClasses);
 
-
 /** Store extra information derived during save and used by the cooker for the given EditorDomain package. */
 void CommitEditorDomainCookAttachments(FName PackageName, TArrayView<IPackageWriter::FCommitAttachmentInfo> Attachments);
 
diff --git a/Engine/Source/Runtime/Core/Private/Misc/ConfigAccessTracking.cpp b/Engine/Source/Runtime/Core/Private/Misc/ConfigAccessTracking.cpp
--- a/Engine/Source/Runtime/Core/Private/Misc/ConfigAccessTracking.cpp
+++ b/Engine/Source/Runtime/Core/Private/Misc/ConfigAccessTracking.cpp
@@ -20,6 +20,8 @@
 namespace UE::ConfigAccessTracking
 {
 
+thread_local bool bIgnoreReads = false;
+
 FFile::FFile(const FConfigFile* InConfigFile)
 	: ConfigFile(InConfigFile)
 	, bSavedHasConfigFile(false)
@@ -153,6 +155,17 @@ FSection::FSection(FFile& InFileAccess, FStringView InSectionName)
 {
 }
 
+FIgnoreScope::FIgnoreScope()
+	: bPreviousIgnoreReads(bIgnoreReads)
+{
+	bIgnoreReads = true;
+}
+
+FIgnoreScope::~FIgnoreScope()
+{
+	bIgnoreReads = bPreviousIgnoreReads;
+}
+
 namespace Private
 {
 
@@ -271,6 +284,13 @@ void OnConfigValueReadInternal(UE::ConfigAccessTracking::FSection* Section, FMin
 	{
 		return;
 	}
+
+	// Implementation of FIgnoreScope
+	if (bIgnoreReads)
+	{
+		return;
+	}
+
 	FConfigReadCallbacks::Get().OnConfigValueRead(Section, ValueName, ConfigValue);
 }
 
diff --git a/Engine/Source/Runtime/Core/Public/Misc/ConfigAccessTracking.h b/Engine/Source/Runtime/Core/Public/Misc/ConfigAccessTracking.h
--- a/Engine/Source/Runtime/Core/Public/Misc/ConfigAccessTracking.h
+++ b/Engine/Source/Runtime/Core/Public/Misc/ConfigAccessTracking.h
@@ -103,6 +103,17 @@ CORE_API FConfigValueReadCallbackId AddConfigValueReadCallback(FConfigValueReadC
 /** Remove a subscriber that was added via FConfigValue. */
 CORE_API void RemoveConfigValueReadCallback(FConfigValueReadCallbackId DelegateHandle);
 
+/** Disables recording of ConfigValues read on the current thread while in scope. */
+struct FIgnoreScope
+{
+public:
+	CORE_API FIgnoreScope();
+	CORE_API ~FIgnoreScope();
+
+private:
+	bool bPreviousIgnoreReads = false;
+};
+
 namespace Private
 {
 
@@ -130,9 +141,19 @@ inline void OnConfigValueRead(UE::ConfigAccessTracking::FSection* Section, FMini
 
 #else  // UE_WITH_CONFIG_TRACKING
 
+namespace UE::ConfigAccessTracking
+{
+
 // Declare FSection even though it is undefined so that we can pass FSection* as an argument to FConfigSection and FConfigValue
 // constructors rather than having to create separate constructors under #ifdef UE_WITH_CONFIG_TRACKING
-namespace UE::ConfigAccessTracking { struct FSection; }
+struct FSection;
+
+// Define FIgnoreScope even when !UE_WITH_CONFIG_TRACKING so we don't need to wrap its use in a macro.
+struct FIgnoreScope
+{
+};
+
+} // namespace UE::ConfigAccessTracking
 
 #endif // else !UE_WITH_CONFIG_TRACKING
 
