When Undoing Apply Instance Changes into a Blueprint class the actor instance now gets replaced with a new/updated version of the BP class that was updated. Fixes an issue with transaction buffer losing reference to the actor and deleting it.

diff --git a/Engine/Source/Editor/UnrealEd/Private/Kismet2/Kismet2.cpp b/Engine/Source/Editor/UnrealEd/Private/Kismet2/Kismet2.cpp
--- a/Engine/Source/Editor/UnrealEd/Private/Kismet2/Kismet2.cpp
+++ b/Engine/Source/Editor/UnrealEd/Private/Kismet2/Kismet2.cpp
@@ -2164,7 +2164,17 @@ int32 FKismetEditorUtilities::ApplyInstanceChangesToBlueprint(AActor* Actor)
 					}
 					if (NumChangedProperties > 0)
 					{
-						Actor = nullptr; // It is unsafe to use Actor after this point as it may have been reinstanced, so set it to null to make this obvious
+						TArray<AActor*> Actors;
+						Actors.Add(Actor);
+
+						FVector Location = Actor->GetActorLocation();
+						FRotator Rotator = Actor->GetActorRotation();
+
+						AActor* NewActor = CreateBlueprintInstanceFromSelection(Blueprint, Actors, Location, Rotator, Actor->GetAttachParentActor());
+						if (NewActor)
+						{
+							NewActor->SetActorScale3D(Actor->GetActorScale3D());
+						}
 					}
 				}
 			}
