Fix GC object leak in AssetManager when unloading plugins that contain loaded primary assets
#rnx

diff --git a/Engine/Source/Editor/UnrealEd/Private/PackageTools.cpp b/Engine/Source/Editor/UnrealEd/Private/PackageTools.cpp
--- a/Engine/Source/Editor/UnrealEd/Private/PackageTools.cpp
+++ b/Engine/Source/Editor/UnrealEd/Private/PackageTools.cpp
@@ -45,6 +45,7 @@
 #include "AssetRegistry/AssetData.h"
 #include "AssetRegistry/AssetRegistryModule.h"
 #include "ComponentReregisterContext.h"
+#include "Engine/AssetManager.h"
 #include "Engine/BlueprintGeneratedClass.h"
 #include "Engine/GameEngine.h"
 #include "Engine/LevelStreaming.h"
@@ -69,6 +70,12 @@
 
 DEFINE_LOG_CATEGORY_STATIC(LogPackageTools, Log, All);
 
+namespace PackageTools_Private
+{
+	static bool bUnloadPackagesUnloadsPrimaryAssets = true;
+	FAutoConsoleVariableRef CVarUnloadPackagesUnloadsPrimaryAssets(TEXT("PackageTools.UnloadPackagesUnloadsPrimaryAssets"), bUnloadPackagesUnloadsPrimaryAssets, TEXT("During unload packages, also unload primary assets"), ECVF_Default);
+}
+
 /** State passed to RestoreStandaloneOnReachableObjects. */
 TSet<UPackage*>* UPackageTools::PackagesBeingUnloaded = nullptr;
 TSet<UObject*> UPackageTools::ObjectsThatHadFlagsCleared;
@@ -484,9 +491,19 @@ UPackageTools::UPackageTools(const FObjectInitializer& ObjectInitializer)
 				}
 
 				GetObjectsWithPackage(PackageBeingUnloaded, ObjectsInPackage, true, RF_Transient, EInternalObjectFlags::Garbage);
-				// Notify any Blueprints that are about to be unloaded, and destroy any leftover worlds.
+				// Notify any Blueprints and other systems that are about to be unloaded, also destroy any leftover worlds.
 				for (UObject* Obj : ObjectsInPackage)
 				{
+					// Asset manager can hold hard references to this object and prevent GC
+					if (PackageTools_Private::bUnloadPackagesUnloadsPrimaryAssets)
+					{
+						const FPrimaryAssetId PrimaryAssetId = UAssetManager::Get().GetPrimaryAssetIdForObject(Obj);
+						if (PrimaryAssetId.IsValid())
+						{
+							UAssetManager::Get().UnloadPrimaryAsset(PrimaryAssetId);
+						}
+					}
+
 					if (UBlueprint* BP = Cast<UBlueprint>(Obj))
 					{
 						BP->ClearEditorReferences();
