Guard InitializeNewlyCreatedInactiveWorld with a GIsLoadingEditorPackage to match the state when the map typically loads in Map_Load. We shouldn't ever dirty the map even if loaded it while inactive.

diff --git a/Engine/Source/Editor/UnrealEd/Private/EditorEngine.cpp b/Engine/Source/Editor/UnrealEd/Private/EditorEngine.cpp
--- a/Engine/Source/Editor/UnrealEd/Private/EditorEngine.cpp
+++ b/Engine/Source/Editor/UnrealEd/Private/EditorEngine.cpp
@@ -7539,6 +7539,9 @@ void UEditorEngine::InitializeNewlyCreatedInactiveWorld(UWorld* World)
 
 	if (!World->bIsWorldInitialized && World->WorldType == EWorldType::Inactive && !World->IsInstanced())
 	{
+		// Guard against dirtying packages while initializing the map
+		TGuardValue<bool> IsEditorLoadingPackageGuard(GIsEditorLoadingPackage, true);
+		// This is probably no longer needed with the EditorLoadingPackage guard but doesn't hurt to keep for safety.
 		const bool bOldDirtyState = World->GetOutermost()->IsDirty();
 
 		// Make sure we have a navigation system if we are cooking the asset.
