Update documentation comment for ShouldDisableCPUThrottlingDelegates to match current behavior

UEditorEngine::ShouldThrottleCPUUsage() was changed to flip the logic but the comment wasn't updated.

diff --git a/Engine/Source/Editor/UnrealEd/Classes/Editor/EditorEngine.h b/Engine/Source/Editor/UnrealEd/Classes/Editor/EditorEngine.h
--- a/Engine/Source/Editor/UnrealEd/Classes/Editor/EditorEngine.h
+++ b/Engine/Source/Editor/UnrealEd/Classes/Editor/EditorEngine.h
@@ -548,7 +548,7 @@ public:
 	DECLARE_MULTICAST_DELEGATE(FPreviewPlatformChanged);
 	FPreviewPlatformChanged PreviewPlatformChanged;
 
-	/** An array of delegates that can force disable throttling cpu usage if any of them return false. */
+	/** An array of delegates that can force disable throttling cpu usage if any of them return true. */
 	DECLARE_DELEGATE_RetVal(bool, FShouldDisableCPUThrottling);
 	TArray<FShouldDisableCPUThrottling> ShouldDisableCPUThrottlingDelegates;
 
