ULevelEditorDragDropHandler now always favors world surrogate referencing object when found (using OnLevelEditorDragDropWorldSurrogateReferencingObjectDelegate).
This ensures that dragging an asset (like a mesh) coming from a gamefeature plugin will favor creating the actor in the right External Data Layer (even if the map domain can see the plugin domain).

diff --git a/Engine/Source/Editor/UnrealEd/Private/LevelEditorDragDropHandler.cpp b/Engine/Source/Editor/UnrealEd/Private/LevelEditorDragDropHandler.cpp
--- a/Engine/Source/Editor/UnrealEd/Private/LevelEditorDragDropHandler.cpp
+++ b/Engine/Source/Editor/UnrealEd/Private/LevelEditorDragDropHandler.cpp
@@ -40,21 +40,20 @@ bool ULevelEditorDragDropHandler::PassesFilter(UWorld* World, const FAssetData&
 
 		FText FailureReason;
 		bool bPassesFilter = AssetPassesFilter(AssetData, ReferencingWorld, &FailureReason);
-		if (!bPassesFilter)
+		
+		// Try to find with referencing asset matching the AssetData
+		if (OnLevelEditorDragDropWorldSurrogateReferencingObjectDelegate.IsBound())
 		{
-			// Try to find with referencing asset matching the AssetData
-			if (OnLevelEditorDragDropWorldSurrogateReferencingObjectDelegate.IsBound())
+			TUniquePtr<FLevelEditorDragDropWorldSurrogateReferencingObject> Object = OnLevelEditorDragDropWorldSurrogateReferencingObjectDelegate.Execute(ReferencingWorld, AssetData.GetSoftObjectPath());
+			const UObject* ReferencingAsset = Object.IsValid() ? Object->GetValue() : nullptr;
+			bool bPassesFilterUsingWorldSurrogate = ReferencingAsset ? AssetPassesFilter(AssetData, ReferencingAsset) : false;
+			if (bPassesFilterUsingWorldSurrogate)
 			{
-				TUniquePtr<FLevelEditorDragDropWorldSurrogateReferencingObject> Object = OnLevelEditorDragDropWorldSurrogateReferencingObjectDelegate.Execute(ReferencingWorld, AssetData.GetSoftObjectPath());
-				const UObject* ReferencingAsset = Object.IsValid() ? Object->GetValue() : nullptr;
-				bPassesFilter = ReferencingAsset ? AssetPassesFilter(AssetData, ReferencingAsset) : false;
-				if (bPassesFilter)
+				if (OutWorldSurrogateReferencingObject)
 				{
-					if (OutWorldSurrogateReferencingObject)
-					{
-						*OutWorldSurrogateReferencingObject = MoveTemp(Object);
-					}
+					*OutWorldSurrogateReferencingObject = MoveTemp(Object);
 				}
+				bPassesFilter = true;
 			}
 		}
 
