UE-207672 - Crash in UEnhancedInputEditorSubsystem when opening Locomotion project

The underlying crash was due to creating a modal window during a post load, which is never going to work correctly. For now, we'll address the issue by simply emulating a commandlet's behavior, which removes the zombie names and issues a warning.

If there's a demand, we can look into adding supporting for this back, but it'll have to be a different implementation.

#rnx

diff --git a/Engine/Source/Editor/UnrealEd/Private/Kismet2/StructureEditorUtils.cpp b/Engine/Source/Editor/UnrealEd/Private/Kismet2/StructureEditorUtils.cpp
--- a/Engine/Source/Editor/UnrealEd/Private/Kismet2/StructureEditorUtils.cpp
+++ b/Engine/Source/Editor/UnrealEd/Private/Kismet2/StructureEditorUtils.cpp
@@ -579,7 +579,7 @@ void FStructureEditorUtils::RemoveInvalidStructureMemberVariableFromBlueprint(UB
 	{
 		const UScriptStruct* FallbackStruct = GetFallbackStruct();
 
-		FString DislpayList;
+		FString DisplayList;
 		TArray<FName> ZombieMemberNames;
 		for (int32 VarIndex = 0; VarIndex < Blueprint->NewVariables.Num(); ++VarIndex)
 		{
@@ -603,40 +603,22 @@ void FStructureEditorUtils::RemoveInvalidStructureMemberVariableFromBlueprint(UB
 			// If this variable is invalid then display a warning
 			if (bIsInvalid)
 			{
-				DislpayList += Var.FriendlyName.IsEmpty() ? Var.VarName.ToString() : Var.FriendlyName;
-				DislpayList += TEXT("\n");
+				DisplayList += Var.FriendlyName.IsEmpty() ? Var.VarName.ToString() : Var.FriendlyName;
+				DisplayList += TEXT("\n");
 				ZombieMemberNames.Add(Var.VarName);
 			}
 		}
 
 		if (ZombieMemberNames.Num())
 		{
-			EAppReturnType::Type Response = EAppReturnType::Ok;
-			if (GIsEditor && !IsRunningCommandlet())
-			{
-				Response = FMessageDialog::Open(
-					EAppMsgType::OkCancel,
-					FText::Format(
-						LOCTEXT("RemoveInvalidStructureMemberVariable_Msg", "The following member variables in blueprint '{0}' have invalid type. Would you like to remove them? \n\n{1}"),
-						FText::FromString(Blueprint->GetFullName()),
-						FText::FromString(DislpayList)
-					));
-			}
-			else
-			{
-				UE_LOG(LogBlueprint, Warning, TEXT("The following member variables in blueprint '%s' have invalid type. Removing them.\n\n%s"), *Blueprint->GetFullName(), *DislpayList);
-			}
-			check((EAppReturnType::Ok == Response) || (EAppReturnType::Cancel == Response));
+			UE_LOG(LogBlueprint, Warning, TEXT("The following member variables in blueprint '%s' have invalid type. Removing them.\n\n%s"), *Blueprint->GetFullName(), *DisplayList);
 
-			if (EAppReturnType::Ok == Response)
-			{				
-				Blueprint->Modify();
+			Blueprint->Modify();
 
-				for (const FName& Name : ZombieMemberNames)
-				{
-					Blueprint->NewVariables.RemoveAll(FFindByNameHelper<FBPVariableDescription>(Name)); //TODO: Add RemoveFirst to TArray
-					FBlueprintEditorUtils::RemoveVariableNodes(Blueprint, Name);
-				}
+			for (const FName& Name : ZombieMemberNames)
+			{
+				Blueprint->NewVariables.RemoveAll(FFindByNameHelper<FBPVariableDescription>(Name)); //TODO: Add RemoveFirst to TArray
+				FBlueprintEditorUtils::RemoveVariableNodes(Blueprint, Name);
 			}
 		}
 	}
