Use FWorldPartitionHelpers::ForEachActorDescInstance when collecting external actor packages for localization

This is more reliable than trying to find external actor packages via their path, and works correctly with content bundles

diff --git a/Engine/Source/Editor/UnrealEd/Classes/Commandlets/GatherTextFromAssetsCommandlet.h b/Engine/Source/Editor/UnrealEd/Classes/Commandlets/GatherTextFromAssetsCommandlet.h
--- a/Engine/Source/Editor/UnrealEd/Classes/Commandlets/GatherTextFromAssetsCommandlet.h
+++ b/Engine/Source/Editor/UnrealEd/Classes/Commandlets/GatherTextFromAssetsCommandlet.h
@@ -63,8 +63,8 @@ private:
 	void ApplyExcludeExactClassesFilter(const FARFilter& InFilter, TArray<FAssetData>& InOutAssetDataArray) const;
 	void FilterAssetsBasedOnIncludeExcludePaths(TArray<FAssetData>& InOutAssetDataArray) const;
 	
-	void DiscoverExternalActors(TArray<FAssetData>& InOutAssetDataArray) const;
-	void RemoveExistingExternalActors(TArray<FAssetData>& InOutAssetDataArray, TArray<FName>& OutExternalActorsSearchPaths) const;
+	void DiscoverExternalActors(TArray<FAssetData>& InOutAssetDataArray);
+	void RemoveExistingExternalActors(TArray<FAssetData>& InOutAssetDataArray, TArray<FName>& OutPartitionedWorldPackageNames) const;
 
 	TSet<FName> GetPackageNamesToGather(const TArray<FAssetData>& InAssetDataArray) const;
 	void PopulatePackagesPendingGather(TSet<FName> PackageNamesToGather);
diff --git a/Engine/Source/Editor/UnrealEd/Private/Commandlets/GatherTextFromAssetsCommandlet.cpp b/Engine/Source/Editor/UnrealEd/Private/Commandlets/GatherTextFromAssetsCommandlet.cpp
--- a/Engine/Source/Editor/UnrealEd/Private/Commandlets/GatherTextFromAssetsCommandlet.cpp
+++ b/Engine/Source/Editor/UnrealEd/Private/Commandlets/GatherTextFromAssetsCommandlet.cpp
@@ -644,10 +644,10 @@ void UGatherTextFromAssetsCommandlet::FilterAssetsBasedOnIncludeExcludePaths(TAr
 	UE_LOG(LogGatherTextFromAssetsCommandlet, Display, TEXT("Filtering assets by include exclude paths took %.2f seconds."), FPlatformTime::Seconds() - FilteringAssetsByIncludeExcludePathsStartTime);
 }
 
-/** Remove any external actors that currently exist in InOutAssetDataArray. OutExternalActorsSearchPaths is populated for the correct paths to search for external actors.*/
-void UGatherTextFromAssetsCommandlet::RemoveExistingExternalActors(TArray<FAssetData>& InOutAssetDataArray, TArray<FName>& OutExternalActorsSearchPaths) const
+/** Remove any external actors that currently exist in InOutAssetDataArray. OutPartitionedWorldPackageNames is populated with the package paths of partitioned worlds.*/
+void UGatherTextFromAssetsCommandlet::RemoveExistingExternalActors(TArray<FAssetData>& InOutAssetDataArray, TArray<FName>& OutPartitionedWorldPackageNames) const
 {
-	InOutAssetDataArray.RemoveAll([&OutExternalActorsSearchPaths](const FAssetData& AssetData)
+	InOutAssetDataArray.RemoveAll([&OutPartitionedWorldPackageNames](const FAssetData& AssetData)
 		{
 			const FNameBuilder PackageNameStr(AssetData.PackageName);
 
@@ -655,10 +655,7 @@ void UGatherTextFromAssetsCommandlet::RemoveExistingExternalActors(TArray<FAsset
 			{
 				if (ULevel::GetIsLevelPartitionedFromAsset(AssetData))
 				{
-					FString ExternalActorsPathForWorld = ULevel::GetExternalActorsPath(*PackageNameStr);
-					// External actors will only be processed if/when their world is gathered.
-					// Adding this as a search path will ensure that all external actors within the world will be added for gather later.
-					OutExternalActorsSearchPaths.Add(*ExternalActorsPathForWorld);
+					OutPartitionedWorldPackageNames.Add(AssetData.PackageName);
 				}
 			}
 			else if (PackageNameStr.ToView().Contains(FPackagePath::GetExternalActorsFolderName()))
@@ -673,17 +670,98 @@ void UGatherTextFromAssetsCommandlet::RemoveExistingExternalActors(TArray<FAsset
 }
 
 /** Appends any external actors that also need to be gathered to the InOutAssetDataArray. */
-void UGatherTextFromAssetsCommandlet::DiscoverExternalActors(TArray<FAssetData>& InOutAssetDataArray) const
+void UGatherTextFromAssetsCommandlet::DiscoverExternalActors(TArray<FAssetData>& InOutAssetDataArray)
 {
 	UE_LOG(LogGatherTextFromAssetsCommandlet, Display, TEXT("Discovering external actors to gather..."));
 	const double DiscoveringExternalActorsStartTime = FPlatformTime::Seconds();
-	TArray<FName> ExternalActorsSearchPaths;
-	RemoveExistingExternalActors(InOutAssetDataArray, ExternalActorsSearchPaths);
 
-	if (ExternalActorsSearchPaths.Num() > 0)
+	TArray<FName> PartitionedWorldPackageNames;
+	RemoveExistingExternalActors(InOutAssetDataArray, PartitionedWorldPackageNames);
+
+	TArray<FName> ExternalActorPackageNames;
 	{
-		IAssetRegistry::GetChecked().GetAssetsByPaths(ExternalActorsSearchPaths, InOutAssetDataArray, /*bRecursive*/true);
+		FLoadPackageLogOutputRedirector LogOutputRedirector;
+
+		int32 NumPackagesProcessed = 0;
+		for (const FName PartitionedWorldPackageName : PartitionedWorldPackageNames)
+		{
+			FNameBuilder PackageNameStr(PartitionedWorldPackageName);
+
+			const int32 CurrentPackageNum = ++NumPackagesProcessed;
+			const float PercentageComplete = static_cast<float>(CurrentPackageNum) / static_cast<float>(PartitionedWorldPackageNames.Num()) * 100.0f;
+			UE_LOG(LogGatherTextFromAssetsCommandlet, Display, TEXT("[%6.2f%%] Loading partitioned world package: '%s'..."), PercentageComplete, *PackageNameStr);
+
+			UPackage* Package = nullptr;
+			{
+				FLoadPackageLogOutputRedirector::FScopedCapture ScopedCapture(&LogOutputRedirector, *PackageNameStr);
+				Package = LoadWorldPackageForEditor(*PackageNameStr, EWorldType::Editor, LOAD_NoWarn | LOAD_Quiet);
+			}
+
+			if (!Package)
+			{
+				UE_LOG(LogGatherTextFromAssetsCommandlet, Warning, TEXT("Failed to load partitioned world package: '%s'."), *PackageNameStr);
+				continue;
+			}
+
+			// Tick background tasks
+			UE::Private::GatherTextFromAssetsCommandlet::TickBackgroundTasks();
+
+			if (UWorld* World = UWorld::FindWorldInPackage(Package))
+			{
+				UWorld::InitializationValues IVS;
+				IVS.InitializeScenes(false);
+				IVS.AllowAudioPlayback(false);
+				IVS.RequiresHitProxies(false);
+				IVS.CreatePhysicsScene(false);
+				IVS.CreateNavigation(false);
+				IVS.CreateAISystem(false);
+				IVS.ShouldSimulatePhysics(false);
+				IVS.EnableTraceCollision(false);
+				IVS.SetTransactional(false);
+				IVS.CreateFXSystem(false);
+				IVS.CreateWorldPartition(true);
+
+				TOptional<FScopedEditorWorld> ScopeEditorWorld;
+				{
+					FLoadPackageLogOutputRedirector::FScopedCapture ScopedCapture(&LogOutputRedirector, *PackageNameStr);
+					ScopeEditorWorld.Emplace(World, IVS); // Initializing FScopedEditorWorld can log warnings, so capture those like we do with loading errors
+				}
+
+				if (UWorldPartition* WorldPartition = World->GetWorldPartition())
+				{
+					FWorldPartitionHelpers::ForEachActorDescInstance(WorldPartition, AActor::StaticClass(), [&ExternalActorPackageNames](const FWorldPartitionActorDescInstance* ActorDescInstance)
+					{
+						ExternalActorPackageNames.Add(ActorDescInstance->GetActorPackage());
+						return true;
+					});
+				}
+			}
+
+			if (HasExceededMemoryLimit(/*bLog*/true))
+			{
+				// First try a minimal purge to only remove things that are no longer referenced or needed by other packages pending gather
+				PurgeGarbage(/*bPurgeReferencedPackages*/false);
+
+				if (HasExceededMemoryLimit(/*bLog*/false))
+				{
+					// If we're still over the memory limit after a minimal purge, then attempt a full purge
+					PurgeGarbage(/*bPurgeReferencedPackages*/true);
+
+					// If we're still over the memory limit after both purges, then log a warning as we may be about to OOM
+					UE_CLOG(HasExceededMemoryLimit(/*bLog*/false), LogGatherTextFromAssetsCommandlet, Warning, TEXT("Flushing failed to reduce process memory to within the requested limits; this process may OOM!"));
+				}
+			}
+		}
 	}
+
+	if (ExternalActorPackageNames.Num() > 0)
+	{
+		FARFilter Filter;
+		Filter.PackageNames = ExternalActorPackageNames;
+		Filter.bIncludeOnlyOnDiskAssets = true;
+		IAssetRegistry::GetChecked().GetAssets(Filter, InOutAssetDataArray);
+	}
+
 	UE_LOG(LogGatherTextFromAssetsCommandlet, Display, TEXT("Discovering external actors took %.2f seconds."), FPlatformTime::Seconds() - DiscoveringExternalActorsStartTime);
 }
 
@@ -1097,7 +1175,11 @@ void UGatherTextFromAssetsCommandlet::LoadAndProcessUncachedPackages(TArray<FNam
 					IVS.CreateFXSystem(false);
 					IVS.CreateWorldPartition(true);
 
-					FScopedEditorWorld ScopeEditorWorld(World, IVS);
+					TOptional<FScopedEditorWorld> ScopeEditorWorld;
+					{
+						FLoadPackageLogOutputRedirector::FScopedCapture ScopedCapture(&LogOutputRedirector, *PackageNameStr);
+						ScopeEditorWorld.Emplace(World, IVS); // Initializing FScopedEditorWorld can log warnings, so capture those like we do with loading errors
+					}
 
 					if (UWorldPartition* WorldPartition = World->GetWorldPartition())
 					{
