[Material Editor] Added the "Automatically set Material usage flags in editor default" project setting to enable/disable making new Materials automatically set usage flags.

diff --git a/Engine/Source/Editor/UnrealEd/Private/Factories/EditorFactories.cpp b/Engine/Source/Editor/UnrealEd/Private/Factories/EditorFactories.cpp
--- a/Engine/Source/Editor/UnrealEd/Private/Factories/EditorFactories.cpp
+++ b/Engine/Source/Editor/UnrealEd/Private/Factories/EditorFactories.cpp
@@ -290,6 +290,12 @@ DEFINE_LOG_CATEGORY(LogEditorFactories);
 
 #define LOCTEXT_NAMESPACE "EditorFactories"
 
+static TAutoConsoleVariable<int32> CVarAutomaticallySetMaterialUsageInEditorDefault(
+	TEXT("r.Material.DefaultAutoMaterialUsage"),
+	1,
+	TEXT("Whether new Materials should automatically set usage flags in the Editor."),
+	ECVF_RenderThreadSafe);
+
 /*------------------------------------------------------------------------------
 	Shared - used by multiple factories
 ------------------------------------------------------------------------------*/
@@ -499,6 +505,8 @@ UObject* UMaterialFactoryNew::FactoryCreateNew(UClass* Class,UObject* InParent,F
 {
 	UMaterial* NewMaterial = NewObject<UMaterial>(InParent, Class, Name, Flags);
 
+	NewMaterial->bAutomaticallySetUsageInEditor = CVarAutomaticallySetMaterialUsageInEditorDefault.GetValueOnAnyThread() != 0;
+
 	if ( InitialTexture != nullptr )
 	{
 		// An initial texture was specified, add it and assign it to the BaseColor
diff --git a/Engine/Source/Runtime/Engine/Classes/Engine/RendererSettings.h b/Engine/Source/Runtime/Engine/Classes/Engine/RendererSettings.h
--- a/Engine/Source/Runtime/Engine/Classes/Engine/RendererSettings.h
+++ b/Engine/Source/Runtime/Engine/Classes/Engine/RendererSettings.h
@@ -1108,6 +1108,15 @@ class URendererSettings : public UDeveloperSettings
 		ConfigRestartRequired = true))
 		uint32 bMaterialEnergyConservation : 1;
 
+	/**
+	"Automatically set Material usage flags in editor default."
+	*/
+	UPROPERTY(config, EditAnywhere, Category = Materials, meta = (
+		ConsoleVariable = "r.Material.DefaultAutoMaterialUsage", DisplayName = "Automatically set Material usage flags in editor default",
+		ToolTip = "Whether new Materials should automatically set usage flags in the Editor.",
+		ConfigRestartRequired = false))
+	uint32 bAutomaticallySetMaterialUsageInEditorDefault : 1;
+
 	/**
 	"Enable Order Independent Transparency (Experimental)."
 	*/
