Make sure we clean up constructed components when using the Actor replacement workflow.

Notes:
- This addresses a long-standing editor regression (first appeared in 4.24) that could otherwise occur as a result of accessing a non-native root component on the new Actor instance after propagating the stale reference from the replaced Actor instance.

diff --git a/Engine/Source/Editor/UnrealEd/Private/EditorEngine.cpp b/Engine/Source/Editor/UnrealEd/Private/EditorEngine.cpp
--- a/Engine/Source/Editor/UnrealEd/Private/EditorEngine.cpp
+++ b/Engine/Source/Editor/UnrealEd/Private/EditorEngine.cpp
@@ -5437,6 +5437,11 @@ void UEditorEngine::ReplaceActors(UActorFactory* Factory, const FAssetData& Asse
 		ULevel* Level = OldActor->GetLevel();
 		AActor* NewActor = NULL;
 
+		// Destroy any non-native constructed components, but make sure we grab the transform first in case it has a
+		// non-native root component. These will be reconstructed as part of the new actor when it's created/instanced.
+		const FTransform OldTransform = OldActor->ActorToWorld();
+		OldActor->DestroyConstructedComponents();
+
 		// Unregister this actors components because we are effectively replacing it with an actor sharing the same ActorGuid.
 		// This allows it to be unregistered before a new actor with the same guid gets registered avoiding conflicts.
 		OldActor->UnregisterAllComponents();
@@ -5455,8 +5460,6 @@ void UEditorEngine::ReplaceActors(UActorFactory* Factory, const FAssetData& Asse
 		// the exact same package, keeping the package name intact.
 		OldActor->UObject::Rename(*OldActorReplacedNamed.ToString(), OldActor->GetOuter(), REN_DoNotDirty | REN_DontCreateRedirectors | REN_ForceNoResetLoaders);
 
-		const FTransform OldTransform = OldActor->ActorToWorld();
-
 		// create the actor
 		NewActor = Factory->CreateActor(Asset, Level, OldTransform, SpawnParams);
 		// For blueprints, try to copy over properties
