[Mode Manager] Fix an issue when using a custom selection set where we assumed the level editor is the owner of the mode manager

Simply check against the global level editor mode manager pointer before we go grabbing selection from the global level editor viewport client

diff --git a/Engine/Source/Editor/UnrealEd/Private/EditorModeManager.cpp b/Engine/Source/Editor/UnrealEd/Private/EditorModeManager.cpp
--- a/Engine/Source/Editor/UnrealEd/Private/EditorModeManager.cpp
+++ b/Engine/Source/Editor/UnrealEd/Private/EditorModeManager.cpp
@@ -1057,10 +1057,10 @@ FMatrix FEditorModeTools::GetCustomCoordinateSystem(TUniqueFunction<void(const T
 	if (!CustomCoordinateSystemProvided)
 	{
 		TTypedElement<ITypedElementWorldInterface> LastSelected;
-		if (GCurrentLevelEditingViewportClient)
+		if ((this == &GLevelEditorModeTools()) && GCurrentLevelEditingViewportClient)
 		{
 			// Use the cache from the viewport when available
-			 LastSelected = GCurrentLevelEditingViewportClient->GetElementsToManipulate()->GetBottomElement<ITypedElementWorldInterface>();
+			LastSelected = GCurrentLevelEditingViewportClient->GetElementsToManipulate()->GetBottomElement<ITypedElementWorldInterface>();
 		}
 		else
 		{
