Apply min dimensions to ClientSize to fix auto centering of dialog.

diff --git a/Engine/Source/Editor/UnrealEd/Private/Dialogs/Dialogs.cpp b/Engine/Source/Editor/UnrealEd/Private/Dialogs/Dialogs.cpp
--- a/Engine/Source/Editor/UnrealEd/Private/Dialogs/Dialogs.cpp
+++ b/Engine/Source/Editor/UnrealEd/Private/Dialogs/Dialogs.cpp
@@ -1117,11 +1117,13 @@ TSharedRef<SWindow> UE::Private::CreateModalDialogWindow(
 	ESizingRule Sizing,
 	FVector2D MinDimensions)
 {
+	// clang-format off
 	return SNew(SWindow)
 		.Title(InTitle)
 		.SizingRule(Sizing)
 		.MinWidth(MinDimensions.X)
 		.MinHeight(MinDimensions.Y)
+		.ClientSize(MinDimensions)
 		.SupportsMaximize(false)
 		.SupportsMinimize(false)
 		.HasCloseButton(false)
@@ -1133,6 +1135,7 @@ TSharedRef<SWindow> UE::Private::CreateModalDialogWindow(
 				MoveTemp(Contents)
 			]
 		];
+	// clang-format on
 }
 
 void UE::Private::ShowModalDialogWindow(TSharedRef<SWindow> Window)
