Removing PCX importing/exporting.

diff --git a/Engine/Plugins/Interchange/Runtime/Source/Import/Private/InterchangeImportModule.cpp b/Engine/Plugins/Interchange/Runtime/Source/Import/Private/InterchangeImportModule.cpp
--- a/Engine/Plugins/Interchange/Runtime/Source/Import/Private/InterchangeImportModule.cpp
+++ b/Engine/Plugins/Interchange/Runtime/Source/Import/Private/InterchangeImportModule.cpp
@@ -33,7 +33,6 @@
 #include "Texture/InterchangeIESTranslator.h"
 #include "Texture/InterchangeJPGTranslator.h"
 #include "Texture/InterchangeUEJPEGTranslator.h"
-#include "Texture/InterchangePCXTranslator.h"
 #include "Texture/InterchangePSDTranslator.h"
 #include "Texture/InterchangeTextureFactory.h"
 
@@ -92,7 +91,6 @@ void FInterchangeImportModule::StartupModule()
 		InterchangeManager.RegisterTranslator(UInterchangeDDSTranslator::StaticClass());
 		InterchangeManager.RegisterTranslator(UInterchangeUEJPEGTranslator::StaticClass());
 		InterchangeManager.RegisterTranslator(UInterchangeJPGTranslator::StaticClass());
-		InterchangeManager.RegisterTranslator(UInterchangePCXTranslator::StaticClass());
 		InterchangeManager.RegisterTranslator(UInterchangePSDTranslator::StaticClass());
 		InterchangeManager.RegisterTranslator(UInterchangeIESTranslator::StaticClass());
 
diff --git a/Engine/Plugins/Interchange/Runtime/Source/Import/Private/Texture/InterchangePCXTranslator.cpp b/Engine/Plugins/Interchange/Runtime/Source/Import/Private/Texture/InterchangePCXTranslator.cpp
deleted file mode 100644
diff --git a/Engine/Plugins/Interchange/Runtime/Source/Import/Public/Texture/InterchangePCXTranslator.h b/Engine/Plugins/Interchange/Runtime/Source/Import/Public/Texture/InterchangePCXTranslator.h
deleted file mode 100644
diff --git a/Engine/Source/Editor/UnrealEd/Classes/Exporters/TextureExporterPCX.h b/Engine/Source/Editor/UnrealEd/Classes/Exporters/TextureExporterPCX.h
deleted file mode 100644
diff --git a/Engine/Source/Editor/UnrealEd/Private/Factories/EditorFactories.cpp b/Engine/Source/Editor/UnrealEd/Private/Factories/EditorFactories.cpp
--- a/Engine/Source/Editor/UnrealEd/Private/Factories/EditorFactories.cpp
+++ b/Engine/Source/Editor/UnrealEd/Private/Factories/EditorFactories.cpp
@@ -137,7 +137,6 @@
 #include "Exporters/TextureExporterBMP.h"
 #include "Exporters/TextureExporterHDR.h"
 #include "Exporters/RenderTargetExporterHDR.h"
-#include "Exporters/TextureExporterPCX.h"
 #include "Exporters/TextureExporterTGA.h"
 #include "Exporters/TextureExporterPNG.h"
 #include "Exporters/TextureExporterEXR.h"
@@ -2405,40 +2404,6 @@ UObject* UTextureRenderTargetVolumeFactoryNew::FactoryCreateNew(UClass* Class, U
 
 // .PCX file header.
 #pragma pack(push,1)
-class FPCXFileHeader
-{
-public:
-	uint8	Manufacturer;		// Always 10.
-	uint8	Version;			// PCX file version.
-	uint8	Encoding;			// 1=run-length, 0=none.
-	uint8	BitsPerPixel;		// 1,2,4, or 8.
-	uint16	XMin;				// Dimensions of the image.
-	uint16	YMin;				// Dimensions of the image.
-	uint16	XMax;				// Dimensions of the image.
-	uint16	YMax;				// Dimensions of the image.
-	uint16	XDotsPerInch;		// Horizontal printer resolution.
-	uint16	YDotsPerInch;		// Vertical printer resolution.
-	uint8	OldColorMap[48];	// Old colormap info data.
-	uint8	Reserved1;			// Must be 0.
-	uint8	NumPlanes;			// Number of color planes (1, 3, 4, etc).
-	uint16	BytesPerLine;		// Number of bytes per scanline.
-	uint16	PaletteType;		// How to interpret palette: 1=color, 2=gray.
-	uint16	HScreenSize;		// Horizontal monitor size.
-	uint16	VScreenSize;		// Vertical monitor size.
-	uint8	Reserved2[54];		// Must be 0.
-	friend FArchive& operator<<( FArchive& Ar, FPCXFileHeader& H )
-	{
-		Ar << H.Manufacturer << H.Version << H.Encoding << H.BitsPerPixel;
-		Ar << H.XMin << H.YMin << H.XMax << H.YMax << H.XDotsPerInch << H.YDotsPerInch;
-		for( int32 i=0; i<UE_ARRAY_COUNT(H.OldColorMap); i++ )
-			Ar << H.OldColorMap[i];
-		Ar << H.Reserved1 << H.NumPlanes;
-		Ar << H.BytesPerLine << H.PaletteType << H.HScreenSize << H.VScreenSize;
-		for( int32 i=0; i<UE_ARRAY_COUNT(H.Reserved2); i++ )
-			Ar << H.Reserved2[i];
-		return Ar;
-	}
-};
 
 struct FTGAFileFooter
 {
@@ -3098,122 +3063,6 @@ bool UTextureFactory::ImportImage(const uint8* Buffer, int64 Length, FFeedbackCo
 	}
 
 
-	//
-	// PCX
-	//
-	const FPCXFileHeader*    PCX = (FPCXFileHeader *)Buffer;
-	if (Length >= sizeof(FPCXFileHeader) && PCX->Manufacturer == 10)
-	{
-		int32 NewU = PCX->XMax + 1 - PCX->XMin;
-		int32 NewV = PCX->YMax + 1 - PCX->YMin;
-
-		// Check the resolution of the imported texture to ensure validity
-		if (!IsImportResolutionValid(NewU, NewV, bAllowNonPowerOfTwo, Warn))
-		{
-			return false;
-		}
-		else if (PCX->NumPlanes == 1 && PCX->BitsPerPixel == 8)
-		{
-
-			// Set texture properties.
-			OutImage.Init2DWithOneMip(
-				NewU,
-				NewV,
-				TSF_BGRA8
-			);
-			FColor* DestPtr = (FColor*)OutImage.RawData.GetData();
-
-			// Import the palette.
-			uint8* PCXPalette = (uint8 *)(Buffer + Length - 256 * 3);
-			TArray<FColor>	Palette;
-			for (uint32 i = 0; i < 256; i++)
-			{
-				Palette.Add(FColor(PCXPalette[i * 3 + 0], PCXPalette[i * 3 + 1], PCXPalette[i * 3 + 2], i == 0 ? 0 : 255));
-			}
-
-			// Import it.
-			FColor* DestEnd = DestPtr + NewU * NewV;
-			Buffer += 128;
-			while (DestPtr < DestEnd)
-			{
-				uint8 Color = *Buffer++;
-				if ((Color & 0xc0) == 0xc0)
-				{
-					uint32 RunLength = Color & 0x3f;
-					Color = *Buffer++;
-
-					for (uint32 Index = 0; Index < RunLength; Index++)
-					{
-						*DestPtr++ = Palette[Color];
-					}
-				}
-				else *DestPtr++ = Palette[Color];
-			}
-		}
-		else if (PCX->NumPlanes == 3 && PCX->BitsPerPixel == 8)
-		{
-			// Set texture properties.
-			OutImage.Init2DWithOneMip(
-				NewU,
-				NewV,
-				TSF_BGRA8
-			);
-
-			uint8* Dest = OutImage.RawData.GetData();
-
-			// Doing a memset to make sure the alpha channel is set to 0xff since we only have 3 color planes.
-			FMemory::Memset(Dest, 0xff, NewU * NewV * FTextureSource::GetBytesPerPixel(OutImage.Format));
-
-			// Copy upside-down scanlines.
-			Buffer += 128;
-			int32 CountU = FMath::Min<int32>(PCX->BytesPerLine, NewU);
-			for (int32 i = 0; i < NewV; i++)
-			{
-				// We need to decode image one line per time building RGB image color plane by color plane.
-				int32 RunLength, Overflow = 0;
-				uint8 Color = 0;
-				for (int32 ColorPlane = 2; ColorPlane >= 0; ColorPlane--)
-				{
-					for (int32 j = 0; j < CountU; j++)
-					{
-						if (!Overflow)
-						{
-							Color = *Buffer++;
-							if ((Color & 0xc0) == 0xc0)
-							{
-								RunLength = FMath::Min((Color & 0x3f), CountU - j);
-								Overflow = (Color & 0x3f) - RunLength;
-								Color = *Buffer++;
-							}
-							else
-								RunLength = 1;
-						}
-						else
-						{
-							RunLength = FMath::Min(Overflow, CountU - j);
-							Overflow = Overflow - RunLength;
-						}
-
-						//checkf(((i*NewU + RunLength) * 4 + ColorPlane) < (Texture->Source.CalcMipSize(0)),
-						//	TEXT("RLE going off the end of buffer"));
-						for (int32 k = j; k < j + RunLength; k++)
-						{
-							Dest[(i*NewU + k) * 4 + ColorPlane] = Color;
-						}
-						j += RunLength - 1;
-					}
-				}
-			}
-		}
-		else
-		{
-			Warn->Logf(ELogVerbosity::Error, TEXT("PCX uses an unsupported format (%i/%i)"), PCX->NumPlanes, PCX->BitsPerPixel);
-			return false;
-		}
-
-		return true;
-	}
-
 	//
 	// PSD File
 	//
@@ -4613,96 +4462,6 @@ void UTextureFactory::ParseFromJson(TSharedRef<class FJsonObject> ImportSettings
 	FJsonObjectConverter::JsonObjectToUStruct(ImportSettingsJson, GetClass(), this, 0, CPF_InstancedReference);
 }
 
-/*------------------------------------------------------------------------------
-	UTextureExporterPCX implementation.
-	UTextureExporterPCX does not use TextureExporterGeneric because there's no PCX ImageWrapper
-		therefore does not support UDIM and other niceties like TextureExporterGeneric
-------------------------------------------------------------------------------*/
-UTextureExporterPCX::UTextureExporterPCX(const FObjectInitializer& ObjectInitializer)
-	: Super(ObjectInitializer)
-{
-	SupportedClass = UTexture2D::StaticClass();
-	PreferredFormatIndex = 0;
-	FormatExtension.Add(TEXT("PCX"));
-	FormatDescription.Add(TEXT("PCX File"));
-}
-
-bool UTextureExporterPCX::SupportsObject(UObject* Object) const
-{
-	bool bSupportsObject = false;
-	if (Super::SupportsObject(Object))
-	{
-		UTexture2D* Texture = Cast<UTexture2D>(Object);
-
-		if (Texture)
-		{
-			bSupportsObject = Texture->Source.GetFormat() == TSF_BGRA8;
-			
-			if (Texture->Source.GetNumBlocks() > 1 )
-			{
-				// does not support UDIM
-				bSupportsObject = false;
-			}		
-		}
-	}
-	return bSupportsObject;
-}
-
-bool UTextureExporterPCX::ExportBinary( UObject* Object, const TCHAR* Type, FArchive& Ar, FFeedbackContext* Warn, int32 FileIndex, uint32 PortFlags )
-{
-	UTexture2D* Texture = CastChecked<UTexture2D>( Object );
-
-	if( !Texture->Source.IsValid() || Texture->Source.GetFormat() != TSF_BGRA8 )
-	{
-		return false;
-	}
-
-	uint16 SizeX = IntCastChecked<uint16>(Texture->Source.GetSizeX());
-	uint16 SizeY = IntCastChecked<uint16>(Texture->Source.GetSizeY());
-	TArray64<uint8> RawData;
-	verify( Texture->Source.GetMipData(RawData, 0) );
-
-	// Set all PCX file header properties.
-	FPCXFileHeader PCX;
-	FMemory::Memzero( &PCX, sizeof(PCX) );
-	PCX.Manufacturer	= 10;
-	PCX.Version			= 05;
-	PCX.Encoding		= 1;
-	PCX.BitsPerPixel	= 8;
-	PCX.XMin			= 0;
-	PCX.YMin			= 0;
-	PCX.XMax			= SizeX-1;
-	PCX.YMax			= SizeY-1;
-	PCX.XDotsPerInch	= SizeX;
-	PCX.YDotsPerInch	= SizeY;
-	PCX.BytesPerLine	= SizeX;
-	PCX.PaletteType		= 0;
-	PCX.HScreenSize		= 0;
-	PCX.VScreenSize		= 0;
-
-	// Copy all RLE bytes.
-	uint8 RleCode=0xc1;
-
-	PCX.NumPlanes = 3;
-	Ar << PCX;
-	for( int32 Line=0; Line<SizeY; Line++ )
-	{
-		for( int32 ColorPlane = 2; ColorPlane >= 0; ColorPlane-- )
-		{
-			uint8* ScreenPtr = RawData.GetData() + (Line * SizeX * 4) + ColorPlane;
-			for( int32 Row=0; Row<SizeX; Row++ )
-			{
-				if( (*ScreenPtr&0xc0)==0xc0 )
-					Ar << RleCode;
-				Ar << *ScreenPtr;
-				ScreenPtr += 4;
-			}
-		}
-	}
-
-	return true;
-}
-
 
 /*------------------------------------------------------------------------------
 	UTextureExporterDDS implementation.
diff --git a/Engine/Source/Editor/UnrealEd/Public/UnrealEd.h b/Engine/Source/Editor/UnrealEd/Public/UnrealEd.h
--- a/Engine/Source/Editor/UnrealEd/Public/UnrealEd.h
+++ b/Engine/Source/Editor/UnrealEd/Public/UnrealEd.h
@@ -246,7 +246,6 @@ MONOLITHIC_HEADER_BOILERPLATE()
 #include "Exporters/TextureExporterBMP.h"
 #include "Exporters/TextureExporterHDR.h"
 #include "Exporters/RenderTargetExporterHDR.h"
-#include "Exporters/TextureExporterPCX.h"
 #include "Exporters/TextureExporterTGA.h"
 #include "ThumbnailRendering/ThumbnailRenderer.h"
 #include "ThumbnailRendering/DefaultSizedThumbnailRenderer.h"
