Fix subobject of a CDO being moved outside of thier original package by the BP reinstantiator!

diff --git a/Engine/Source/Editor/UnrealEd/Private/Kismet2/BlueprintEditorUtils.cpp b/Engine/Source/Editor/UnrealEd/Private/Kismet2/BlueprintEditorUtils.cpp
--- a/Engine/Source/Editor/UnrealEd/Private/Kismet2/BlueprintEditorUtils.cpp
+++ b/Engine/Source/Editor/UnrealEd/Private/Kismet2/BlueprintEditorUtils.cpp
@@ -1417,6 +1417,8 @@ void FBlueprintEditorUtils::PatchCDOSubobjectsIntoExport(UObject* PreviousCDO, U
 			{
 				TArray<UObject*> OldSubObjects;
 				GetObjectsWithOuter(OldObj, OldSubObjects, /*bIncludeNestedSubObjects =*/false);
+				UPackage* Package= OldObj->GetPackage();
+				checkf(Package, TEXT("Expecting a package for this object"));
 
 				// Exit now if we don't have any subobjects to process.
 				if (OldSubObjects.Num() == 0)
@@ -1440,7 +1442,8 @@ void FBlueprintEditorUtils::PatchCDOSubobjectsIntoExport(UObject* PreviousCDO, U
 					// Resolve new instances through the new object and patch them into the linker's export table.
 					for (const FInstancedSubObjRef& OldInstancedSubObjRef : OldInstancedSubObjRefs)
 					{
-						if (UObject* OldSubObj = OldInstancedSubObjRef.SubObjInstance)
+						// For non null subobject, make sure they are in the same package, otherwise it will move them and this is not what we want.
+						if (UObject* OldSubObj = OldInstancedSubObjRef.SubObjInstance; OldSubObj && OldSubObj->IsInPackage(Package))
 						{
 							if (UObject* NewSubObj = OldInstancedSubObjRef.PropertyPath.Resolve(NewObj))
 							{
