Cooker: DiffOnly: Add option to skip reporting indeterminism errors in classes in the TargetDomain IterativeClassDenyList.
#rnx

diff --git a/Engine/Source/Editor/UnrealEd/Private/Cooker/DiffPackageWriter.cpp b/Engine/Source/Editor/UnrealEd/Private/Cooker/DiffPackageWriter.cpp
--- a/Engine/Source/Editor/UnrealEd/Private/Cooker/DiffPackageWriter.cpp
+++ b/Engine/Source/Editor/UnrealEd/Private/Cooker/DiffPackageWriter.cpp
@@ -6,6 +6,7 @@
 #include "Containers/Map.h"
 #include "CookOnTheSide/CookLog.h"
 #include "CoreGlobals.h"
+#include "EditorDomain/EditorDomainUtils.h"
 #include "Engine/Engine.h"
 #include "Logging/LogMacros.h"
 #include "Misc/CommandLine.h"
@@ -63,7 +64,9 @@ void FDiffPackageWriter::ParseCmds()
 	const TCHAR* DumpObjectsParam = TEXT("dumpobjects");
 
 	FString CmdsText;
-	if (FParse::Value(FCommandLine::Get(), TEXT("-diffcmds="), CmdsText, false))
+	const TCHAR* const CommandLine = FCommandLine::Get();
+
+	if (FParse::Value(CommandLine, TEXT("-diffcmds="), CmdsText, false))
 	{
 		CmdsText = CmdsText.TrimQuotes();
 		TArray<FString> CmdsList;
@@ -82,6 +85,11 @@ void FDiffPackageWriter::ParseCmds()
 			}
 		}
 	}
+
+	if (FParse::Param(CommandLine, TEXT("DiffDenyList")))
+	{
+		CompareDenyListClasses = UE::EditorDomain::ConstructTargetIterativeClassBlockList();
+	}
 }
 
 void FDiffPackageWriter::ParseDumpObjList(FString InParams)
@@ -130,6 +138,7 @@ void FDiffPackageWriter::BeginPackage(const FBeginPackageInfo& Info)
 	Accumulators[1].SafeRelease();
 
 	BeginInfo = Info;
+	Package = FindObjectFast<UPackage>(nullptr, BeginInfo.PackageName);
 	ConditionallyDumpObjList();
 	ConditionallyDumpObjects();
 	Inner->BeginPackage(Info);
@@ -147,7 +156,8 @@ void FDiffPackageWriter::CommitPackage(FCommitPackageInfo&& Info)
 	{
 		EnumRemoveFlags(Info.WriteOptions, EWriteOptions::Write);
 	}
-	return Inner->CommitPackage(MoveTemp(Info));
+	Inner->CommitPackage(MoveTemp(Info));
+	Package = nullptr;
 }
 
 void FDiffPackageWriter::WritePackageData(const FPackageInfo& Info, FLargeMemoryWriter& ExportsArchive,
@@ -177,6 +187,10 @@ void FDiffPackageWriter::WritePackageData(const FPackageInfo& Info, FLargeMemory
 		Accumulator.OnFirstSaveComplete(LocalInfo.LooseFilePath, LocalInfo.HeaderSize, Info.HeaderSize,
 			MoveTemp(PreviousInnerData));
 		bIsDifferent = Accumulator.HasDifferences();
+		if (bIsDifferent && !IsPackageDiffAllowed())
+		{
+			bIsDifferent = false;
+		}
 	}
 	else
 	{
@@ -197,6 +211,28 @@ void FDiffPackageWriter::WritePackageData(const FPackageInfo& Info, FLargeMemory
 	Inner->WritePackageData(LocalInfo, ExportsArchive, FileRegions);
 }
 
+bool FDiffPackageWriter::IsPackageDiffAllowed() const
+{
+	if (!CompareDenyListClasses.IsEmpty() && Package != nullptr)
+	{
+		bool bHasDenyClass = false;
+		ForEachObjectWithPackage(Package, [this, &bHasDenyClass](UObject* Object)
+			{
+				FTopLevelAssetPath ClassPath = Object->GetClass()->GetClassPathName();
+				if (CompareDenyListClasses.Contains(ClassPath))
+				{
+					bHasDenyClass = true;
+					return false; // Stop iterating
+				}
+				return true; // Keep iterating
+			});
+		if (bHasDenyClass)
+		{
+			return false;
+		}
+	}
+	return true;
+}
 UE::DiffWriter::FMessageCallback FDiffPackageWriter::GetDiffWriterMessageCallback()
 {
 	return UE::DiffWriter::FMessageCallback([this](ELogVerbosity::Type Verbosity, FStringView Message)
diff --git a/Engine/Source/Editor/UnrealEd/Private/Cooker/DiffPackageWriter.h b/Engine/Source/Editor/UnrealEd/Private/Cooker/DiffPackageWriter.h
--- a/Engine/Source/Editor/UnrealEd/Private/Cooker/DiffPackageWriter.h
+++ b/Engine/Source/Editor/UnrealEd/Private/Cooker/DiffPackageWriter.h
@@ -130,11 +130,15 @@ protected:
 	virtual void OnDiffWriterMessage(ELogVerbosity::Type Verbosity, FStringView Message);
 	FString ResolveText(FStringView Message);
 	UE::DiffWriter::FAccumulator& ConstructAccumulator(FName PackageName, UObject* Asset, uint16 MultiOutputIndex);
+	bool IsPackageDiffAllowed() const;
 
 	TRefCountPtr<UE::DiffWriter::FAccumulator> Accumulators[2];
 	FBeginPackageInfo BeginInfo;
+	TSet<FTopLevelAssetPath> CompareDenyListClasses;
 	TUniquePtr<ICookedPackageWriter> Inner;
 	TUniquePtr<UE::DiffWriter::FAccumulatorGlobals> AccumulatorGlobals;
+	/** Only non-null between BeginPackage and CommitPackage. */
+	UPackage* Package = nullptr;
 	const TCHAR* Indent = nullptr;
 	const TCHAR* NewLine = nullptr;
 	FString DumpObjListParams;
diff --git a/Engine/Source/Editor/UnrealEd/Private/EditorDomain/EditorDomainUtils.h b/Engine/Source/Editor/UnrealEd/Private/EditorDomain/EditorDomainUtils.h
--- a/Engine/Source/Editor/UnrealEd/Private/EditorDomain/EditorDomainUtils.h
+++ b/Engine/Source/Editor/UnrealEd/Private/EditorDomain/EditorDomainUtils.h
@@ -126,6 +126,9 @@ void UtilsCookInitialize();
  */
 FStringView RemoveConfigComment(FStringView Line);
 
+/** Parse Editor.ini:[TargetDomain]:IterativeClassDenyList */
+TSet<FTopLevelAssetPath> ConstructTargetIterativeClassBlockList();
+
 #if ENABLE_COOK_STATS
 namespace CookStats
 {
@@ -133,4 +136,4 @@ namespace CookStats
 }
 #endif
 
-}
\ No newline at end of file
+} // namespace UE::EditorDomain
\ No newline at end of file
