[Asset Editor] Allow child toolkits to override the viewport tab spawn function

diff --git a/Engine/Source/Editor/UnrealEd/Public/Tools/BaseAssetToolkit.h b/Engine/Source/Editor/UnrealEd/Public/Tools/BaseAssetToolkit.h
--- a/Engine/Source/Editor/UnrealEd/Public/Tools/BaseAssetToolkit.h
+++ b/Engine/Source/Editor/UnrealEd/Public/Tools/BaseAssetToolkit.h
@@ -57,7 +57,7 @@ protected:
 	UNREALED_API virtual void RegisterToolbar();
 	UNREALED_API virtual AssetEditorViewportFactoryFunction GetViewportDelegate();
 	UNREALED_API virtual TSharedPtr<FEditorViewportClient> CreateEditorViewportClient() const;
-	UNREALED_API TSharedRef<SDockTab> SpawnTab_Viewport(const FSpawnTabArgs& Args);
+	UNREALED_API virtual TSharedRef<SDockTab> SpawnTab_Viewport(const FSpawnTabArgs& Args);
 	UNREALED_API virtual TSharedRef<SDockTab> SpawnTab_Details(const FSpawnTabArgs& Args);
 
 protected:
