PR #11564: Updating FUnrealEdMisc::OnMessageTokenActivated to get outer actor

diff --git a/Engine/Source/Editor/UnrealEd/Private/UnrealEdMisc.cpp b/Engine/Source/Editor/UnrealEd/Private/UnrealEdMisc.cpp
--- a/Engine/Source/Editor/UnrealEd/Private/UnrealEdMisc.cpp
+++ b/Engine/Source/Editor/UnrealEd/Private/UnrealEdMisc.cpp
@@ -1402,15 +1402,9 @@ void FUnrealEdMisc::OnMessageTokenActivated(const TSharedRef<IMessageToken>& Tok
 		else
 		{
 			AActor* Actor = Cast<AActor>(Object);
-			UPrimitiveComponent* Component = Cast<UPrimitiveComponent>(Object);
-
-			if (Component)
+			if( !Actor )
 			{
-				check( !Actor);
-				if( Component->GetOwner())
-				{
-					Actor = Component->GetOwner();
-				}		
+				Actor = Object->GetTypedOuter<AActor>();
 			}
 
 			if (Actor && Actor->GetLevel() != nullptr)
