[BugFix] CopyFile throwing ensure when merge is committed

diff --git a/Engine/Source/Editor/UnrealEd/Private/MergeUtils.cpp b/Engine/Source/Editor/UnrealEd/Private/MergeUtils.cpp
--- a/Engine/Source/Editor/UnrealEd/Private/MergeUtils.cpp
+++ b/Engine/Source/Editor/UnrealEd/Private/MergeUtils.cpp
@@ -121,8 +121,11 @@ void UUndoableResolveHandler::SetManagedObject(UObject* Object)
 
 	// save package and copy the package to a temp file so it can be reverted
 	const FString BaseFilename = FPaths::GetBaseFilename(Filepath);
-	BackupFilepath = FPaths::CreateTempFilename(*(FPaths::ProjectSavedDir()/TEXT("Temp")), *BaseFilename.Left(32));
-	ensure(FPlatformFileManager::Get().GetPlatformFile().CopyFile(*BackupFilepath, *Filepath));
+	const FString Directory = FPaths::ProjectSavedDir()/TEXT("Temp");
+	IPlatformFile& PlatformFile = FPlatformFileManager::Get().GetPlatformFile();
+	BackupFilepath = FPaths::CreateTempFilename(*Directory, *BaseFilename.Left(32));
+	ensure(PlatformFile.CreateDirectoryTree(*Directory));
+	ensure(PlatformFile.CopyFile(*BackupFilepath, *Filepath));
 }
 
 void UUndoableResolveHandler::MarkResolved()
