Cooker: Fix typo in bSkipResolveChunkDependencyGraph config value that prevented us from reading it if a licensee added the properly named value in their config.
#rnx

diff --git a/Engine/Source/Editor/UnrealEd/Private/Commandlets/AssetRegistryGenerator.cpp b/Engine/Source/Editor/UnrealEd/Private/Commandlets/AssetRegistryGenerator.cpp
--- a/Engine/Source/Editor/UnrealEd/Private/Commandlets/AssetRegistryGenerator.cpp
+++ b/Engine/Source/Editor/UnrealEd/Private/Commandlets/AssetRegistryGenerator.cpp
@@ -2365,7 +2365,7 @@ void FAssetRegistryGenerator::FixupPackageDependenciesForChunks(UE::Cook::FCookS
 	FConfigFile PlatformIniFile;
 	FConfigCacheIni::LoadLocalIniFile(PlatformIniFile, TEXT("Engine"), true, *TargetPlatform->IniPlatformName());
 	bool bSkipResolveChunkDependencyGraph = false;
-	PlatformIniFile.GetBool(TEXT("Script/UnrealEd.ChunkDependencyInfo"), TEXT("bSkipResolveChunkDependencyGraph"), bSkipResolveChunkDependencyGraph);
+	PlatformIniFile.GetBool(TEXT("/Script/UnrealEd.ChunkDependencyInfo"), TEXT("bSkipResolveChunkDependencyGraph"), bSkipResolveChunkDependencyGraph);
 
 	const FChunkDependencyTreeNode* ChunkDepGraph = DependencyInfo.GetOrBuildChunkDependencyGraph(!bSkipResolveChunkDependencyGraph ? HighestChunkId : 0);
 
