Prevent the legacy fbx importer from creating an asset that dont have the same name has the outer package.

#rnx

diff --git a/Engine/Source/Editor/UnrealEd/Private/Fbx/FbxFactory.cpp b/Engine/Source/Editor/UnrealEd/Private/Fbx/FbxFactory.cpp
--- a/Engine/Source/Editor/UnrealEd/Private/Fbx/FbxFactory.cpp
+++ b/Engine/Source/Editor/UnrealEd/Private/Fbx/FbxFactory.cpp
@@ -709,11 +709,20 @@ UObject* UFbxFactory::FactoryCreateFile
 							bool bMapMorphTargetToTimeZero = false;
 							if (LODIndex == 0 && SkelMeshNodeArray.Num() != 0)
 							{
-								FName OutputName = NAME_None;
+								FName OutputName = FbxImporter->MakeNameForMesh((Package == nullptr) ? TEXT("None") : Name.ToString(), SkelMeshNodeArray[0]);
+
+								if (Package)
+								{
+									if (!OutputName.ToString().Equals(FPaths::GetCleanFilename(Package->GetName()), ESearchCase::IgnoreCase))
+									{
+										//We need to create a new package
+										Package = nullptr;
+									}
+								}
+
 								if (Package == nullptr)
 								{
 									FString NewPackageName;
-									OutputName = FbxImporter->MakeNameForMesh(TEXT("None"), SkelMeshNodeArray[0]);
 									if (InParent != nullptr && InParent->GetOutermost() != nullptr)
 									{
 										NewPackageName = FPackageName::GetLongPackagePath(InParent->GetOutermost()->GetName()) + TEXT("/") + OutputName.ToString();
@@ -727,10 +736,6 @@ UObject* UFbxFactory::FactoryCreateFile
 									Package = CreatePackage(*NewPackageName);
 									Package->FullyLoad();
 								}
-								else
-								{
-									OutputName = FbxImporter->MakeNameForMesh(Name.ToString(), SkelMeshNodeArray[0]);
-								}
 
 								UnFbx::FFbxImporter::FImportSkeletalMeshArgs ImportSkeletalMeshArgs;
 								ImportSkeletalMeshArgs.InParent = Package;
