Fix for material preview showing WorldGridMaterial when using Nanite custom mesh

diff --git a/Engine/Source/Editor/UnrealEd/Private/PreviewMaterial.cpp b/Engine/Source/Editor/UnrealEd/Private/PreviewMaterial.cpp
--- a/Engine/Source/Editor/UnrealEd/Private/PreviewMaterial.cpp
+++ b/Engine/Source/Editor/UnrealEd/Private/PreviewMaterial.cpp
@@ -90,7 +90,8 @@ public:
 			}
 
 			// we only need local vertex factory for the preview static mesh
-			if (VertexFactoryType != FindVertexFactoryType(FName(TEXT("FLocalVertexFactory"), FNAME_Find)))
+			if (VertexFactoryType != FindVertexFactoryType(FName(TEXT("FLocalVertexFactory"), FNAME_Find)) &&
+				VertexFactoryType != FindVertexFactoryType(FName(TEXT("FNaniteVertexFactory"), FNAME_Find)))
 			{
 				//cache for gpu skinned vertex factory if the material allows it
 				//this way we can have a preview skeletal mesh
@@ -155,6 +156,10 @@ public:
 			{
 				bShaderTypeMatches = true;
 			}
+			else if (FCString::Stristr(ShaderType->GetName(), TEXT("TBasePassCSFNoLightMapPolicy")))
+			{
+				bShaderTypeMatches = true;
+			}
 			else if (FCString::Stristr(ShaderType->GetName(), TEXT("CachedPointIndirectLightingPolicy")))
 			{
 				bShaderTypeMatches = true;
