MPCook and CookPackageSplitter: Rename FGeneratorPackage to FGenerationHelper to remove confusion between the helper object and the parent generator packagedata.
#rnx

diff --git a/Engine/Source/Editor/UnrealEd/Classes/CookOnTheSide/CookOnTheFlyServer.h b/Engine/Source/Editor/UnrealEd/Classes/CookOnTheSide/CookOnTheFlyServer.h
--- a/Engine/Source/Editor/UnrealEd/Classes/CookOnTheSide/CookOnTheFlyServer.h
+++ b/Engine/Source/Editor/UnrealEd/Classes/CookOnTheSide/CookOnTheFlyServer.h
@@ -186,7 +186,7 @@ namespace UE::Cook
 	struct FCookGenerationInfo;
 	struct FCookSavePackageContext;
 	struct FDiscoveredPlatformSet;
-	struct FGeneratorPackage;
+	struct FGenerationHelper;
 	struct FInitializeConfigSettings;
 	struct FPackageData;
 	struct FPackageDatas;
@@ -1162,7 +1162,7 @@ private:
 		bool bPrecaching);
 	UE::Cook::EPollStatus PrepareSaveInternal(UE::Cook::FPackageData& PackageData, UE::Cook::FCookerTimer& Timer,
 		bool bPrecaching);
-	/** Call BeginCacheForCookedPlatformData on all objects in a PackageData or a GeneratorPackage's current round. */
+	/** Call BeginCacheForCookedPlatformData on all objects currently found in the PackageData. */
 	UE::Cook::EPollStatus CallBeginCacheOnObjects(UE::Cook::FPackageData& PackageData, UPackage* Package,
 		TArray<UE::Cook::FCachedObjectInOuter>& Objects, int32& NextIndex, UE::Cook::FCookerTimer& Timer);
 
@@ -1344,31 +1344,39 @@ private:
 
 	UE::Cook::EPollStatus ConditionalCreateGeneratorPackage(UE::Cook::FPackageData& PackageData, bool bPrecaching);
 
-	/** Generate the list of cook-time-created packages created by the Generator. */
-	UE::Cook::EPollStatus QueueGeneratedPackages(UE::Cook::FGeneratorPackage& Generator,
+	/** Generate the list of cook-time-created packages created by the generator package. */
+	UE::Cook::EPollStatus QueueGeneratedPackages(UE::Cook::FGenerationHelper& GenerationHelper,
 		UE::Cook::FPackageData& PackageData);
-	/** Run additional steps in PrepareSave that are required when the package has a Generator. */
-	UE::Cook::EPollStatus PrepareSaveGeneratedPackage(UE::Cook::FGeneratorPackage& Generator,
+	/** Run additional steps in PrepareSave that are required when the package has a GenerationHelper. */
+	UE::Cook::EPollStatus PrepareSaveGeneratedPackage(UE::Cook::FGenerationHelper& GenerationHelper,
 		UE::Cook::FPackageData& PackageData, UE::Cook::FCookerTimer& Timer, bool bPrecaching);
-	/** Call BeginCacheForCookedPlatformData on the objects the Generator plans to move into its main UPackage. */
-	UE::Cook::EPollStatus BeginCacheObjectsToMove(UE::Cook::FGeneratorPackage& Generator,
+	/**
+	 * Call BeginCacheForCookedPlatformData on objects the CookPackageSplitter plans to move
+	 * into its main UPackage.
+	 */
+	UE::Cook::EPollStatus BeginCacheObjectsToMove(UE::Cook::FGenerationHelper& GenerationHelper,
 		UE::Cook::FCookGenerationInfo& Info, UE::Cook::FCookerTimer& Timer,
 		TArray<ICookPackageSplitter::FGeneratedPackageForPreSave>& GeneratedPackagesForPresave);
-	/** Call the Generator's PreSaveGeneratorPackage to create/move objects into its main UPackage. */
+	/** Call the CookPackageSplitter's PreSaveGeneratorPackage to create/move objects into its main UPackage. */
 	UE::Cook::EPollStatus PreSaveGeneratorPackage(UE::Cook::FPackageData& PackageData,
-		UE::Cook::FGeneratorPackage& Generator, UE::Cook::FCookGenerationInfo& Info,
+		UE::Cook::FGenerationHelper& GenerationHelper, UE::Cook::FCookGenerationInfo& Info,
 		TArray<ICookPackageSplitter::FGeneratedPackageForPreSave>& GeneratedPackagesForPresave);
-	/** Construct the list of generated packages that is required for some of the CookPackageSplitter interface calls. */
-	void ConstructGeneratedPackagesForPresave(UE::Cook::FPackageData& PackageData, UE::Cook::FGeneratorPackage& Generator,
+	/** Construct the list of generated packages that is required for some CookPackageSplitter interface calls. */
+	void ConstructGeneratedPackagesForPresave(UE::Cook::FPackageData& PackageData,
+		UE::Cook::FGenerationHelper& GenerationHelper,
 		TArray<ICookPackageSplitter::FGeneratedPackageForPreSave>& GeneratedPackagesForPresave);
-	/** Call BeginCacheForCookedPlatformData on any undeclared objects in the Generator's main UPackage after the move. */
-	UE::Cook::EPollStatus BeginCachePostMove(UE::Cook::FGeneratorPackage& Generator,
+	/**
+	 * Call BeginCacheForCookedPlatformData on any undeclared objects in the CookPackageSplitter's main UPackage
+	 * after the move.
+	 */
+	UE::Cook::EPollStatus BeginCachePostMove(UE::Cook::FGenerationHelper& GenerationHelper,
 		UE::Cook::FCookGenerationInfo& Info, UE::Cook::FCookerTimer& Timer);
 
 	/** Try creating (or finding from earlier creation) the generated package for later population */
-	UPackage* TryCreateGeneratedPackage(UE::Cook::FGeneratorPackage& Generator, UE::Cook::FCookGenerationInfo& GeneratedInfo);
-	/** Try calling the splitter's populate to create the package */
-	UE::Cook::EPollStatus TryPopulateGeneratedPackage(UE::Cook::FGeneratorPackage& Generator,
+	UPackage* TryCreateGeneratedPackage(UE::Cook::FGenerationHelper& GenerationHelper,
+		UE::Cook::FCookGenerationInfo& GeneratedInfo);
+	/** Try calling the CookPackageSplitter's populate to create the package */
+	UE::Cook::EPollStatus TryPopulateGeneratedPackage(UE::Cook::FGenerationHelper& GenerationHelper,
 		UE::Cook::FCookGenerationInfo& GeneratedInfo);
 
 	ICookedPackageWriter& FindOrCreatePackageWriter(const ITargetPlatform* TargetPlatform);
@@ -1592,7 +1600,7 @@ private:
 	friend UE::Cook::FCookWorkerClient;
 	friend UE::Cook::FCookWorkerServer;
 	friend UE::Cook::FDiagnostics;
-	friend UE::Cook::FGeneratorPackage;
+	friend UE::Cook::FGenerationHelper;
 	friend UE::Cook::FInitializeConfigSettings;
 	friend UE::Cook::FPackageData;
 	friend UE::Cook::FPackageDatas;
diff --git a/Engine/Source/Editor/UnrealEd/Private/CookOnTheFlyServer.cpp b/Engine/Source/Editor/UnrealEd/Private/CookOnTheFlyServer.cpp
--- a/Engine/Source/Editor/UnrealEd/Private/CookOnTheFlyServer.cpp
+++ b/Engine/Source/Editor/UnrealEd/Private/CookOnTheFlyServer.cpp
@@ -2881,23 +2881,27 @@ void UCookOnTheFlyServer::LoadPackageInQueue(UE::Cook::FPackageData& PackageData
 	}
 	else
 	{
-		FGeneratorPackage* Generator = PackageData.GetGeneratedOwner();
-		if (!Generator)
+		FGenerationHelper* GenerationHelper = PackageData.GetGeneratedOwner();
+		if (!GenerationHelper)
 		{
-			UE_LOG(LogCook, Error, TEXT("Package %s is an out-of-date generated package with a no-longer-available generator. It can not be loaded."), *PackageFileName.ToString());
+			UE_LOG(LogCook, Error,
+				TEXT("Package %s is an out-of-date generated package with a no-longer-available generator. It can not be loaded."),
+				*PackageFileName.ToString());
 			RejectPackageToLoad(PackageData, TEXT("is an orphaned generated package"), ESuppressCookReason::OrphanedGenerated);
 			return;
 		}
-		FCookGenerationInfo* Info = Generator->FindInfo(PackageData);
+		FCookGenerationInfo* Info = GenerationHelper->FindInfo(PackageData);
 		if (!Info)
 		{
-			UE_LOG(LogCook, Error, TEXT("Package %s is a generated package but its generator no longer has a record of it. It can not be loaded."), *PackageFileName.ToString());
+			UE_LOG(LogCook, Error,
+				TEXT("Package %s is a generated package but its generator no longer has a record of it. It can not be loaded."),
+				*PackageFileName.ToString());
 			RejectPackageToLoad(PackageData, TEXT("is an orphaned generated package"), ESuppressCookReason::OrphanedGenerated);
 			return;
 		}
 
-		FPackageData& OwnerPackageData = Generator->GetOwner();
-		UPackage* OwnerPackage = Generator->GetOwnerPackage();
+		FPackageData& OwnerPackageData = GenerationHelper->GetOwner();
+		UPackage* OwnerPackage = GenerationHelper->GetOwnerPackage();
 		if (!OwnerPackage)
 		{
 			OwnerPackage = FindObject<UPackage>(nullptr, *OwnerPackageData.GetPackageName().ToString());
@@ -2908,26 +2912,29 @@ void UCookOnTheFlyServer::LoadPackageInQueue(UE::Cook::FPackageData& PackageData
 			UObject* SplitterDataObject = nullptr;
 			if (bLoadFullySuccessful)
 			{
-				SplitterDataObject = Generator->FindSplitDataObject();
+				SplitterDataObject = GenerationHelper->FindSplitDataObject();
 			}
 			if (!SplitterDataObject)
 			{
 				ResultFlags |= COSR_ErrorLoadingPackage;
-				UE_LOG(LogCook, Error, TEXT("Package %s is a generated package and we could not load its generator package %s. It can not be loaded."),
+				UE_LOG(LogCook, Error,
+					TEXT("Package %s is a generated package and we could not load its generator package %s. It can not be loaded."),
 					*PackageFileName.ToString(), *OwnerPackageData.GetFileName().ToString());
-				RejectPackageToLoad(PackageData, TEXT("is a generated package which could not load its generator"), ESuppressCookReason::LoadError);
+				RejectPackageToLoad(PackageData, TEXT("is a generated package which could not load its generator"),
+					ESuppressCookReason::LoadError);
 				return;
 			}
 			FScopedActivePackage ScopedActivePackage(*this, OwnerPackageData.GetPackageName(),
 				PackageAccessTrackingOps::NAME_CookerBuildObject);
-			Generator->GetCookPackageSplitterInstance()->OnOwnerReloaded(OwnerPackage, SplitterDataObject);
-			Generator->SetOwnerPackage(OwnerPackage);
+			GenerationHelper->GetCookPackageSplitterInstance()->OnOwnerReloaded(OwnerPackage, SplitterDataObject);
+			GenerationHelper->SetOwnerPackage(OwnerPackage);
 		}
 
-		LoadedPackage = TryCreateGeneratedPackage(*Generator, *Info);
+		LoadedPackage = TryCreateGeneratedPackage(*GenerationHelper, *Info);
 		if (!LoadedPackage)
 		{
-			RejectPackageToLoad(PackageData, TEXT("is a generated package which could not be populated"), ESuppressCookReason::LoadError);
+			RejectPackageToLoad(PackageData, TEXT("is a generated package which could not be populated"),
+				ESuppressCookReason::LoadError);
 			return;
 		}
 	}
@@ -3275,8 +3282,8 @@ UE::Cook::EPollStatus UCookOnTheFlyServer::ConditionalCreateGeneratorPackage(UE:
 	{
 		if (!bGeneratorExists && !bIncomplete)
 		{
-			// Destroy any old GeneratorPackage if we no longer find we need one
-			PackageData.DestroyGeneratorPackage();
+			// Destroy any old GenerationHelper if we no longer find we need one
+			PackageData.DestroyGenerationHelper();
 		}
 	};
 
@@ -3347,40 +3354,41 @@ UE::Cook::EPollStatus UCookOnTheFlyServer::ConditionalCreateGeneratorPackage(UE:
 		return EPollStatus::Error;
 	}
 
-	// Create a FGeneratorPackage helper object using this CookPackageSplitter instance
+	// Create an FGenerationHelper struct using this CookPackageSplitter instance
 	bGeneratorExists = true;
-	PackageData.CreateGeneratorPackage(SplitDataObject, SplitterInstance);
+	PackageData.CreateGenerationHelper(SplitDataObject, SplitterInstance);
 	return EPollStatus::Success;
 }
 
-UE::Cook::EPollStatus UCookOnTheFlyServer::QueueGeneratedPackages(UE::Cook::FGeneratorPackage& Generator,
+UE::Cook::EPollStatus UCookOnTheFlyServer::QueueGeneratedPackages(UE::Cook::FGenerationHelper& GenerationHelper,
 	UE::Cook::FPackageData& PackageData)
 {
 	using namespace UE::Cook;
 
-	ICookPackageSplitter* Splitter = Generator.GetCookPackageSplitterInstance();
-	UObject* SplitObject = Generator.FindSplitDataObject();
-	FCookGenerationInfo& Info = Generator.GetOwnerInfo();
+	ICookPackageSplitter* Splitter = GenerationHelper.GetCookPackageSplitterInstance();
+	UObject* SplitObject = GenerationHelper.FindSplitDataObject();
+	FCookGenerationInfo& Info = GenerationHelper.GetOwnerInfo();
 	if (!SplitObject)
 	{
-		UE_LOG(LogCook, Error, TEXT("Could not find SplitDataObject %s"), *Generator.GetSplitDataObjectName().ToString());
+		UE_LOG(LogCook, Error, TEXT("Could not find SplitDataObject %s"),
+			*GenerationHelper.GetSplitDataObjectName().ToString());
 		return EPollStatus::Error;
 	}
 
 	if (Info.GetSaveState() <= FCookGenerationInfo::ESaveState::GenerateList)
 	{
 		// Call the splitter to generate the list
-		if (!Generator.TryGenerateList(SplitObject, *PackageDatas))
+		if (!GenerationHelper.TryGenerateList(SplitObject, *PackageDatas))
 		{
 			return EPollStatus::Error;
 		}
-		Generator.SetOwnerPackage(PackageData.GetPackage());
+		GenerationHelper.SetOwnerPackage(PackageData.GetPackage());
 		Info.SetSaveStateComplete(FCookGenerationInfo::ESaveState::GenerateList);
 	}
 
 	if (Info.GetSaveState() <= FCookGenerationInfo::ESaveState::ClearOldPackagesLastAttempt)
 	{
-		for (const FCookGenerationInfo& ChildInfo: Generator.GetPackagesToGenerate())
+		for (const FCookGenerationInfo& ChildInfo: GenerationHelper.GetPackagesToGenerate())
 		{
 			const FString GeneratedPackageName = ChildInfo.PackageData->GetPackageName().ToString();
 			if (FindObject<UPackage>(nullptr, *GeneratedPackageName))
@@ -3393,8 +3401,9 @@ UE::Cook::EPollStatus UCookOnTheFlyServer::QueueGeneratedPackages(UE::Cook::FGen
 				}
 				else
 				{
-					UE_LOG(LogCook, Error, TEXT("PackageSplitter was unable to construct new generated packages because an old version of the package is already in memory and GC did not remove it. Splitter=%s, Generated=%s."),
-						*Generator.GetSplitDataObjectName().ToString(), *ChildInfo.RelativePath);
+					UE_LOG(LogCook, Error,
+						TEXT("PackageSplitter was unable to construct new generated packages because an old version of the package is already in memory and GC did not remove it. Splitter=%s, Generated=%s."),
+						*GenerationHelper.GetSplitDataObjectName().ToString(), *ChildInfo.RelativePath);
 					return EPollStatus::Error;
 				}
 			}
@@ -3409,11 +3418,11 @@ UE::Cook::EPollStatus UCookOnTheFlyServer::QueueGeneratedPackages(UE::Cook::FGen
 	{
 		TArray<const ITargetPlatform*, TInlineAllocator<ExpectedMaxNumPlatforms>> ReachablePlatforms;
 		PackageData.GetReachablePlatforms(ReachablePlatforms);
-		for (const FCookGenerationInfo& ChildInfo: Generator.GetPackagesToGenerate())
+		for (const FCookGenerationInfo& ChildInfo: GenerationHelper.GetPackagesToGenerate())
 		{
 			FPackageData* ChildPackageData = ChildInfo.PackageData;
 			// Set the Instigator now rather than delaying it until the discovery queue is processed.
-			ChildPackageData->SetInstigator(Generator, FInstigator(EInstigator::GeneratedPackage, OwnerName));
+			ChildPackageData->SetInstigator(GenerationHelper, FInstigator(EInstigator::GeneratedPackage, OwnerName));
 			// When running -cookfirst, generated packages should also be cooked first, so set the urgency of the
 			// generated packages to match the urgency of the generator
 			bool bUrgent = PackageData.GetIsUrgent();
@@ -3427,12 +3436,12 @@ UE::Cook::EPollStatus UCookOnTheFlyServer::QueueGeneratedPackages(UE::Cook::FGen
 	return EPollStatus::Success;
 }
 
-UE::Cook::EPollStatus UCookOnTheFlyServer::PrepareSaveGeneratedPackage(UE::Cook::FGeneratorPackage& Generator,
+UE::Cook::EPollStatus UCookOnTheFlyServer::PrepareSaveGeneratedPackage(UE::Cook::FGenerationHelper& GenerationHelper,
 	UE::Cook::FPackageData& PackageData, UE::Cook::FCookerTimer& Timer, bool bPrecaching)
 {
 	using namespace UE::Cook;
 
-	FCookGenerationInfo* InfoPtr = Generator.FindInfo(PackageData);
+	FCookGenerationInfo* InfoPtr = GenerationHelper.FindInfo(PackageData);
 	if (!InfoPtr)
 	{
 		UE_LOG(LogCook, Error, TEXT("Generated package %s is missing its generation data and cannot be saved."),
@@ -3456,7 +3465,7 @@ UE::Cook::EPollStatus UCookOnTheFlyServer::PrepareSaveGeneratedPackage(UE::Cook:
 	{
 		if (Info.GetSaveState() <= FCookGenerationInfo::ESaveState::BeginCacheObjectsToMove)
 		{
-			EPollStatus Result = BeginCacheObjectsToMove(Generator, Info, Timer, GeneratedPackagesForPresave);
+			EPollStatus Result = BeginCacheObjectsToMove(GenerationHelper, Info, Timer, GeneratedPackagesForPresave);
 			if (Result != EPollStatus::Success)
 			{
 				return Result;
@@ -3469,7 +3478,7 @@ UE::Cook::EPollStatus UCookOnTheFlyServer::PrepareSaveGeneratedPackage(UE::Cook:
 			return EPollStatus::Incomplete;
 		}
 		bool bFoundNewObjects;
-		EPollStatus Result = Info.RefreshPackageObjects(Generator, PackageData.GetPackage(), bFoundNewObjects,
+		EPollStatus Result = Info.RefreshPackageObjects(GenerationHelper, PackageData.GetPackage(), bFoundNewObjects,
 			FCookGenerationInfo::ESaveState::BeginCacheObjectsToMove);
 		if (Result != EPollStatus::Success)
 		{
@@ -3479,7 +3488,7 @@ UE::Cook::EPollStatus UCookOnTheFlyServer::PrepareSaveGeneratedPackage(UE::Cook:
 		{
 			// Call this function recursively to reexecute CallBeginCacheOnObjects in BeginCacheObjectsToMove.
 			// Note that RefreshPackageObjects checked for too many recursive calls and ErrorExited if so.
-			return PrepareSaveGeneratedPackage(Generator, PackageData, Timer, bPrecaching);
+			return PrepareSaveGeneratedPackage(GenerationHelper, PackageData, Timer, bPrecaching);
 		}
 		Info.SetSaveStateComplete(FCookGenerationInfo::ESaveState::FinishCacheObjectsToMove);
 	}
@@ -3497,11 +3506,11 @@ UE::Cook::EPollStatus UCookOnTheFlyServer::PrepareSaveGeneratedPackage(UE::Cook:
 		EPollStatus Result;
 		if (Info.IsGenerator())
 		{
-			Result = PreSaveGeneratorPackage(PackageData, Generator, Info, GeneratedPackagesForPresave);
+			Result = PreSaveGeneratorPackage(PackageData, GenerationHelper, Info, GeneratedPackagesForPresave);
 		}
 		else
 		{
-			Result = TryPopulateGeneratedPackage(Generator, Info);
+			Result = TryPopulateGeneratedPackage(GenerationHelper, Info);
 		}
 		if (Result != EPollStatus::Success)
 		{
@@ -3514,7 +3523,7 @@ UE::Cook::EPollStatus UCookOnTheFlyServer::PrepareSaveGeneratedPackage(UE::Cook:
 	{
 		if (Info.GetSaveState() <= FCookGenerationInfo::ESaveState::BeginCachePostMove)
 		{
-			EPollStatus Result = BeginCachePostMove(Generator, Info, Timer);
+			EPollStatus Result = BeginCachePostMove(GenerationHelper, Info, Timer);
 			if (Result != EPollStatus::Success)
 			{
 				return Result;
@@ -3527,7 +3536,7 @@ UE::Cook::EPollStatus UCookOnTheFlyServer::PrepareSaveGeneratedPackage(UE::Cook:
 			return EPollStatus::Incomplete;
 		}
 		bool bFoundNewObjects;
-		EPollStatus Result = Info.RefreshPackageObjects(Generator, PackageData.GetPackage(), bFoundNewObjects,
+		EPollStatus Result = Info.RefreshPackageObjects(GenerationHelper, PackageData.GetPackage(), bFoundNewObjects,
 			FCookGenerationInfo::ESaveState::BeginCachePostMove);
 		if (Result != EPollStatus::Success)
 		{
@@ -3537,7 +3546,7 @@ UE::Cook::EPollStatus UCookOnTheFlyServer::PrepareSaveGeneratedPackage(UE::Cook:
 		{
 			// Call this function recursively to reexecute CallBeginCacheOnObjects in BeginCachePostMove
 			// Note that RefreshPackageObjects checked for too many recursive calls and ErrorExited if so.
-			return PrepareSaveGeneratedPackage(Generator, PackageData, Timer, bPrecaching);
+			return PrepareSaveGeneratedPackage(GenerationHelper, PackageData, Timer, bPrecaching);
 		}
 
 		Info.SetSaveStateComplete(FCookGenerationInfo::ESaveState::FinishCachePostMove);
@@ -3547,7 +3556,7 @@ UE::Cook::EPollStatus UCookOnTheFlyServer::PrepareSaveGeneratedPackage(UE::Cook:
 	return EPollStatus::Success;
 }
 
-UE::Cook::EPollStatus UCookOnTheFlyServer::BeginCacheObjectsToMove(UE::Cook::FGeneratorPackage& Generator,
+UE::Cook::EPollStatus UCookOnTheFlyServer::BeginCacheObjectsToMove(UE::Cook::FGenerationHelper& GenerationHelper,
 	UE::Cook::FCookGenerationInfo& Info, UE::Cook::FCookerTimer& Timer,
 	TArray<ICookPackageSplitter::FGeneratedPackageForPreSave>& GeneratedPackagesForPresave)
 {
@@ -3556,11 +3565,12 @@ UE::Cook::EPollStatus UCookOnTheFlyServer::BeginCacheObjectsToMove(UE::Cook::FGe
 	check(Info.PackageData); // Caller validated this
 	FPackageData& PackageData(*Info.PackageData);
 	UPackage* Package = PackageData.GetPackage();
-	ICookPackageSplitter* Splitter = Generator.GetCookPackageSplitterInstance();
-	UObject* SplitDataObject = Generator.FindSplitDataObject();
+	ICookPackageSplitter* Splitter = GenerationHelper.GetCookPackageSplitterInstance();
+	UObject* SplitDataObject = GenerationHelper.FindSplitDataObject();
 	if (!Package || !Splitter || !SplitDataObject)
 	{
-		UE_LOG(LogCook, Error, TEXT("CookPackageSplitter is missing %s during BeginCacheObjectsToMove. PackageName: %s."),
+		UE_LOG(LogCook, Error,
+			TEXT("CookPackageSplitter is missing %s during BeginCacheObjectsToMove. PackageName: %s."),
 			(!Package ? TEXT("Package") : (!Splitter ? TEXT("Splitter") : TEXT("SplitDataObject"))),
 			*PackageData.GetPackageName().ToString());
 		return EPollStatus::Error;
@@ -3573,16 +3583,17 @@ UE::Cook::EPollStatus UCookOnTheFlyServer::BeginCacheObjectsToMove(UE::Cook::FGe
 		TArray<UPackage*> KeepReferencedPackages;
 		if (Info.IsGenerator())
 		{
-			ConstructGeneratedPackagesForPresave(PackageData, Generator, GeneratedPackagesForPresave);
-			FScopedActivePackage ScopedActivePackage(*this, Generator.GetOwner().GetPackageName(),
+			ConstructGeneratedPackagesForPresave(PackageData, GenerationHelper, GeneratedPackagesForPresave);
+			FScopedActivePackage ScopedActivePackage(*this, GenerationHelper.GetOwner().GetPackageName(),
 				PackageAccessTrackingOps::NAME_CookerBuildObject);
-			bPopulateSucceeded = Splitter->PopulateGeneratorPackage(Package, SplitDataObject, GeneratedPackagesForPresave,
-				ObjectsToMove, KeepReferencedPackages);
+			bPopulateSucceeded = Splitter->PopulateGeneratorPackage(Package, SplitDataObject,
+				GeneratedPackagesForPresave, ObjectsToMove, KeepReferencedPackages);
 		}
 		else
 		{
-			ICookPackageSplitter::FGeneratedPackageForPopulate SplitterInfo{ Info.RelativePath, Info.GeneratedRootPath, Package, Info.IsCreateAsMap() };
-			FScopedActivePackage ScopedActivePackage(*this, Generator.GetOwner().GetPackageName(),
+			ICookPackageSplitter::FGeneratedPackageForPopulate SplitterInfo{ Info.RelativePath,
+				Info.GeneratedRootPath, Package, Info.IsCreateAsMap() };
+			FScopedActivePackage ScopedActivePackage(*this, GenerationHelper.GetOwner().GetPackageName(),
 				PackageAccessTrackingOps::NAME_CookerBuildObject);
 			bPopulateSucceeded = Splitter->PopulateGeneratedPackage(Package, SplitDataObject, SplitterInfo,
 				ObjectsToMove, KeepReferencedPackages);
@@ -3592,13 +3603,14 @@ UE::Cook::EPollStatus UCookOnTheFlyServer::BeginCacheObjectsToMove(UE::Cook::FGe
 		{
 			UE_LOG(LogCook, Error, TEXT("CookPackageSplitter returned false from %s. Splitter=%s%s"),
 				Info.IsGenerator() ? TEXT("PopulateGeneratorPackage") : TEXT("PopulateGeneratedPackage"),
-				*Generator.GetSplitDataObjectName().ToString(), 
-				Info.IsGenerator() ? TEXT("") : *FString::Printf(TEXT("\nGeneratedPackage: %s"), *PackageData.GetPackageName().ToString()));
+				*GenerationHelper.GetSplitDataObjectName().ToString(), 
+				Info.IsGenerator() ? TEXT("") : *FString::Printf(TEXT("\nGeneratedPackage: %s"),
+					*PackageData.GetPackageName().ToString()));
 			return EPollStatus::Error;
 		}
 
 		Info.AddKeepReferencedPackages(KeepReferencedPackages);
-		Info.TakeOverCachedObjectsAndAddMoved(Generator, PackageData.GetCachedObjectsInOuter(), ObjectsToMove);
+		Info.TakeOverCachedObjectsAndAddMoved(GenerationHelper, PackageData.GetCachedObjectsInOuter(), ObjectsToMove);
 		Info.SetSaveStateComplete(FCookGenerationInfo::ESaveState::CallObjectsToMove);
 	}
 
@@ -3612,14 +3624,14 @@ UE::Cook::EPollStatus UCookOnTheFlyServer::BeginCacheObjectsToMove(UE::Cook::FGe
 }
 
 UE::Cook::EPollStatus UCookOnTheFlyServer::PreSaveGeneratorPackage(UE::Cook::FPackageData& PackageData,
-	UE::Cook::FGeneratorPackage& Generator, UE::Cook::FCookGenerationInfo& Info,
+	UE::Cook::FGenerationHelper& GenerationHelper, UE::Cook::FCookGenerationInfo& Info,
 	TArray<ICookPackageSplitter::FGeneratedPackageForPreSave>& GeneratedPackagesForPresave)
 {
 	using namespace UE::Cook;
 
 	UPackage* Package = PackageData.GetPackage();
-	ICookPackageSplitter* Splitter = Generator.GetCookPackageSplitterInstance();
-	UObject* SplitDataObject = Generator.FindSplitDataObject();
+	ICookPackageSplitter* Splitter = GenerationHelper.GetCookPackageSplitterInstance();
+	UObject* SplitDataObject = GenerationHelper.FindSplitDataObject();
 	if (!Package || !Splitter || !SplitDataObject)
 	{
 		UE_LOG(LogCook, Error, TEXT("CookPackageSplitter is missing %s during PreSaveGeneratorPackage. PackageName: %s."),
@@ -3629,14 +3641,14 @@ UE::Cook::EPollStatus UCookOnTheFlyServer::PreSaveGeneratorPackage(UE::Cook::FPa
 	}
 
 	TArray<UPackage*> KeepReferencedPackages;
-	ConstructGeneratedPackagesForPresave(PackageData, Generator, GeneratedPackagesForPresave);
+	ConstructGeneratedPackagesForPresave(PackageData, GenerationHelper, GeneratedPackagesForPresave);
 	{
-		FScopedActivePackage ScopedActivePackage(*this, Generator.GetOwner().GetPackageName(),
+		FScopedActivePackage ScopedActivePackage(*this, GenerationHelper.GetOwner().GetPackageName(),
 			PackageAccessTrackingOps::NAME_CookerBuildObject);
 		if (!Splitter->PreSaveGeneratorPackage(Package, SplitDataObject, GeneratedPackagesForPresave, KeepReferencedPackages))
 		{
 			UE_LOG(LogCook, Error, TEXT("PackageSplitter returned false from PreSaveGeneratorPackage. Splitter=%s"),
-				*Generator.GetSplitDataObjectName().ToString());
+				*GenerationHelper.GetSplitDataObjectName().ToString());
 			return EPollStatus::Error;
 		}
 	}
@@ -3645,7 +3657,7 @@ UE::Cook::EPollStatus UCookOnTheFlyServer::PreSaveGeneratorPackage(UE::Cook::FPa
 	return EPollStatus::Success;
 }
 
-void UCookOnTheFlyServer::ConstructGeneratedPackagesForPresave(UE::Cook::FPackageData& PackageData, UE::Cook::FGeneratorPackage& Generator,
+void UCookOnTheFlyServer::ConstructGeneratedPackagesForPresave(UE::Cook::FPackageData& PackageData, UE::Cook::FGenerationHelper& GenerationHelper,
 	TArray<ICookPackageSplitter::FGeneratedPackageForPreSave>& GeneratedPackagesForPresave)
 {
 	using namespace UE::Cook;
@@ -3660,7 +3672,7 @@ void UCookOnTheFlyServer::ConstructGeneratedPackagesForPresave(UE::Cook::FPackag
 
 	// We need to find or (create empty stub packages for) each of the PackagesToGenerate so that PreSaveGeneratorPackage
 	// can refer to them to create hardlinks in the cooked Generator package
-	TArrayView<FCookGenerationInfo> PackagesToGenerate = Generator.GetPackagesToGenerate();
+	TArrayView<FCookGenerationInfo> PackagesToGenerate = GenerationHelper.GetPackagesToGenerate();
 	TArray<ICookPackageSplitter::FGeneratedPackageForPreSave> SplitterDatas;
 	SplitterDatas.Reserve(PackagesToGenerate.Num());
 	for (FCookGenerationInfo& Info : PackagesToGenerate)
@@ -3674,12 +3686,12 @@ void UCookOnTheFlyServer::ConstructGeneratedPackagesForPresave(UE::Cook::FPackag
 		SplitterData.Package = FindObject<UPackage>(nullptr, *GeneratedPackageName);
 		if (!SplitterData.Package)
 		{
-			SplitterData.Package = Generator.CreateGeneratedUPackage(Info, Package, *GeneratedPackageName);
+			SplitterData.Package = GenerationHelper.CreateGeneratedUPackage(Info, Package, *GeneratedPackageName);
 		}
 	}
 }
 
-UE::Cook::EPollStatus UCookOnTheFlyServer::BeginCachePostMove(UE::Cook::FGeneratorPackage& Generator,
+UE::Cook::EPollStatus UCookOnTheFlyServer::BeginCachePostMove(UE::Cook::FGenerationHelper& GenerationHelper,
 	UE::Cook::FCookGenerationInfo& Info, UE::Cook::FCookerTimer& Timer)
 {
 	using namespace UE::Cook;
@@ -3687,8 +3699,8 @@ UE::Cook::EPollStatus UCookOnTheFlyServer::BeginCachePostMove(UE::Cook::FGenerat
 	check(Info.PackageData); // Caller has validated
 	UE::Cook::FPackageData& PackageData(*Info.PackageData);
 	UPackage* Package = PackageData.GetPackage();
-	ICookPackageSplitter* Splitter = Generator.GetCookPackageSplitterInstance();
-	UObject* SplitDataObject = Generator.FindSplitDataObject();
+	ICookPackageSplitter* Splitter = GenerationHelper.GetCookPackageSplitterInstance();
+	UObject* SplitDataObject = GenerationHelper.FindSplitDataObject();
 	if (!Package || !Splitter || !SplitDataObject)
 	{
 		UE_LOG(LogCook, Error, TEXT("CookPackageSplitter is missing %s during BeginCachePostMove. PackageName: %s."),
@@ -3700,7 +3712,7 @@ UE::Cook::EPollStatus UCookOnTheFlyServer::BeginCachePostMove(UE::Cook::FGenerat
 	if (Info.GetSaveState() <= FCookGenerationInfo::ESaveState::CallGetPostMoveObjects)
 	{
 		bool bFoundNewObjects;
-		EPollStatus Result = Info.RefreshPackageObjects(Generator, Package, bFoundNewObjects,
+		EPollStatus Result = Info.RefreshPackageObjects(GenerationHelper, Package, bFoundNewObjects,
 			FCookGenerationInfo::ESaveState::Last);
 		if (Result != EPollStatus::Success)
 		{
@@ -3712,7 +3724,7 @@ UE::Cook::EPollStatus UCookOnTheFlyServer::BeginCachePostMove(UE::Cook::FGenerat
 	EPollStatus Result = CallBeginCacheOnObjects(PackageData, Package, PackageData.GetCachedObjectsInOuter(),
 		PackageData.GetCookedPlatformDataNextIndex(), Timer);
 	if (PackageData.GetNumPendingCookedPlatformData() > 0 &&
-		!Generator.GetCookPackageSplitterInstance()->UseInternalReferenceToAvoidGarbageCollect() &&
+		!GenerationHelper.GetCookPackageSplitterInstance()->UseInternalReferenceToAvoidGarbageCollect() &&
 		!Info.HasIssuedUndeclaredMovedObjectsWarning())
 	{
 		UObject* FirstPendingObject = nullptr;
@@ -3750,13 +3762,13 @@ UE::Cook::EPollStatus UCookOnTheFlyServer::BeginCachePostMove(UE::Cook::FGenerat
 	return EPollStatus::Success;
 }
 
-UPackage* UCookOnTheFlyServer::TryCreateGeneratedPackage(UE::Cook::FGeneratorPackage& Generator, UE::Cook::FCookGenerationInfo& Info)
+UPackage* UCookOnTheFlyServer::TryCreateGeneratedPackage(UE::Cook::FGenerationHelper& GenerationHelper, UE::Cook::FCookGenerationInfo& Info)
 {
 	using namespace UE::Cook;
 	// Caller is responsible for validating OwnerPackage and Generated PackageData
 	check(Info.PackageData); // Caller is responsible for validating
 	UE::Cook::FPackageData& GeneratedPackageData = *Info.PackageData;
-	UPackage* OwnerPackage = Generator.GetOwnerPackage();
+	UPackage* OwnerPackage = GenerationHelper.GetOwnerPackage();
 	check(OwnerPackage); // Caller is responsible for validating
 
 	const FString GeneratedPackageName = GeneratedPackageData.GetPackageName().ToString();
@@ -3769,7 +3781,7 @@ UPackage* UCookOnTheFlyServer::TryCreateGeneratedPackage(UE::Cook::FGeneratorPac
 			UE_LOG(LogCook, Error, TEXT("PackageSplitter found an existing copy of a package it was trying to populate;")
 				TEXT("this is unexpected since garbage has been collected and the package should have been unreferenced so it should have been collected.")
 				TEXT("Splitter=%s, Generated=%s."),
-				*Generator.GetSplitDataObjectName().ToString(), *GeneratedPackageName);
+				*GenerationHelper.GetSplitDataObjectName().ToString(), *GeneratedPackageName);
 			EReferenceChainSearchMode SearchMode = EReferenceChainSearchMode::Shortest
 				| EReferenceChainSearchMode::PrintAllResults
 				| EReferenceChainSearchMode::FullChain;
@@ -3782,32 +3794,32 @@ UPackage* UCookOnTheFlyServer::TryCreateGeneratedPackage(UE::Cook::FGeneratorPac
 	}
 	else
 	{
-		GeneratedPackage = Generator.CreateGeneratedUPackage(Info, OwnerPackage, *GeneratedPackageName);
+		GeneratedPackage = GenerationHelper.CreateGeneratedUPackage(Info, OwnerPackage, *GeneratedPackageName);
 	}
 	return GeneratedPackage;
 }
 
-UE::Cook::EPollStatus UCookOnTheFlyServer::TryPopulateGeneratedPackage(UE::Cook::FGeneratorPackage& Generator,
+UE::Cook::EPollStatus UCookOnTheFlyServer::TryPopulateGeneratedPackage(UE::Cook::FGenerationHelper& GenerationHelper,
 	UE::Cook::FCookGenerationInfo& GeneratedInfo)
 {
 	using namespace UE::Cook;
 
-	UPackage* OwnerPackage = Generator.GetOwnerPackage();
+	UPackage* OwnerPackage = GenerationHelper.GetOwnerPackage();
 	check(GeneratedInfo.PackageData); // Caller already checked this
 	UE::Cook::FPackageData& GeneratedPackageData = *GeneratedInfo.PackageData;
 	const FString GeneratedPackageName = GeneratedPackageData.GetPackageName().ToString();
 	UPackage* GeneratedPackage = GeneratedPackageData.GetPackage();
 	check(GeneratedPackage); // We would have been kicked out of save if the package were gone
 
-	UObject* OwnerObject = Generator.FindSplitDataObject();
+	UObject* OwnerObject = GenerationHelper.FindSplitDataObject();
 	if (!OwnerObject)
 	{
 		UE_LOG(LogCook, Error, TEXT("PopulateGeneratedPacakge could not find the original splitting object. Generated package can not be created. Splitter=%s, Generated=%s."),
-			*Generator.GetSplitDataObjectName().ToString(), *GeneratedPackageName);
+			*GenerationHelper.GetSplitDataObjectName().ToString(), *GeneratedPackageName);
 		return EPollStatus::Error;
 	}
 
-	ICookPackageSplitter* Splitter = Generator.GetCookPackageSplitterInstance();
+	ICookPackageSplitter* Splitter = GenerationHelper.GetCookPackageSplitterInstance();
 
 	// Populate package using CookPackageSplitterInstance and pass GeneratedPackage's cooked name for it to
 	// properly setup any internal reference to this package (SoftObjectPaths or others)
@@ -3818,12 +3830,12 @@ UE::Cook::EPollStatus UCookOnTheFlyServer::TryPopulateGeneratedPackage(UE::Cook:
 	PopulateData.bCreatedAsMap = GeneratedInfo.IsCreateAsMap();
 	TArray<UPackage*> KeepReferencedPackages;
 	{
-		FScopedActivePackage ScopedActivePackage(*this, Generator.GetOwner().GetPackageName(),
+		FScopedActivePackage ScopedActivePackage(*this, GenerationHelper.GetOwner().GetPackageName(),
 			PackageAccessTrackingOps::NAME_CookerBuildObject);
 		if (!Splitter->PreSaveGeneratedPackage(OwnerPackage, OwnerObject, PopulateData, KeepReferencedPackages))
 		{
 			UE_LOG(LogCook, Error, TEXT("PackageSplitter returned false from PreSaveGeneratedPackage. Splitter=%s, Generated=%s."),
-				*Generator.GetSplitDataObjectName().ToString(), *GeneratedPackageName);
+				*GenerationHelper.GetSplitDataObjectName().ToString(), *GeneratedPackageName);
 			return EPollStatus::Error;
 		}
 	}
@@ -3833,7 +3845,7 @@ UE::Cook::EPollStatus UCookOnTheFlyServer::TryPopulateGeneratedPackage(UE::Cook:
 	{
 		UE_LOG(LogCook, Error, TEXT("PackageSplitter specified generated package is %s in GetGenerateList results, but then in PreSaveGeneratedPackage created it as %s. Splitter=%s, Generated=%s."),
 			(GeneratedInfo.IsCreateAsMap() ? TEXT("map") : TEXT("uasset")), (bPackageIsMap ? TEXT("map") : TEXT("uasset")),
-			*Generator.GetSplitDataObjectName().ToString(), *GeneratedPackageName);
+			*GenerationHelper.GetSplitDataObjectName().ToString(), *GeneratedPackageName);
 		return EPollStatus::Error;
 	}
 
@@ -3893,7 +3905,7 @@ UE::Cook::EPollStatus UCookOnTheFlyServer::PrepareSaveInternal(UE::Cook::FPackag
 	UPackage* Package = PackageData.GetPackage();
 	check(Package && Package->IsFullyLoaded());
 	check(PackageData.GetState() == EPackageState::Save);
-	FGeneratorPackage* Generator = nullptr;
+	FGenerationHelper* GenerationHelper = nullptr;
 
 	if (!PackageData.GetCookedPlatformDataCalled())
 	{
@@ -3960,10 +3972,10 @@ UE::Cook::EPollStatus UCookOnTheFlyServer::PrepareSaveInternal(UE::Cook::FPackag
 			}
 			PackageData.SetInitializedGeneratorSave(true);
 		}
-		Generator = PackageData.GetGeneratorPackage();
-		if (Generator)
+		GenerationHelper = PackageData.GetGenerationHelper();
+		if (GenerationHelper)
 		{
-			Result = QueueGeneratedPackages(*Generator, PackageData);
+			Result = QueueGeneratedPackages(*GenerationHelper, PackageData);
 			if (Result != EPollStatus::Success)
 			{
 				return Result;
@@ -3974,12 +3986,12 @@ UE::Cook::EPollStatus UCookOnTheFlyServer::PrepareSaveInternal(UE::Cook::FPackag
 	}
 	else
 	{
-		Generator = PackageData.GetGeneratorPackage();
+		GenerationHelper = PackageData.GetGenerationHelper();
 	}
 
-	if (Generator)
+	if (GenerationHelper)
 	{
-		EPollStatus Result = PrepareSaveGeneratedPackage(*Generator, PackageData, Timer, bPrecaching);
+		EPollStatus Result = PrepareSaveGeneratedPackage(*GenerationHelper, PackageData, Timer, bPrecaching);
 		if (Result != EPollStatus::Success)
 		{
 			return Result;
@@ -3987,10 +3999,10 @@ UE::Cook::EPollStatus UCookOnTheFlyServer::PrepareSaveInternal(UE::Cook::FPackag
 	}
 	else if (PackageData.IsGenerated())
 	{
-		FGeneratorPackage* ParentGenerator = PackageData.GetGeneratedOwner();
+		FGenerationHelper* ParentGenerator = PackageData.GetGeneratedOwner();
 		if (!ParentGenerator)
 		{
-			UE_LOG(LogCook, Error, TEXT("Generated package %s is missing its Parent GeneratorPackage and cannot be saved."),
+			UE_LOG(LogCook, Error, TEXT("Generated package %s is missing its Parent generator package and cannot be saved."),
 				*PackageData.GetPackageName().ToString());
 			return EPollStatus::Error;
 		}
@@ -4124,12 +4136,12 @@ void UCookOnTheFlyServer::ReleaseCookedPlatformData(UE::Cook::FPackageData& Pack
 		return;
 	}
 
-	FGeneratorPackage* Generator = PackageData.GetGeneratorPackage();
-	if (!Generator)
+	FGenerationHelper* GenerationHelper = PackageData.GetGenerationHelper();
+	if (!GenerationHelper)
 	{
-		Generator = PackageData.GetGeneratedOwner();
+		GenerationHelper = PackageData.GetGeneratedOwner();
 	}
-	FCookGenerationInfo* GenerationInfo = Generator ? Generator->FindInfo(PackageData) : nullptr;
+	FCookGenerationInfo* GenerationInfo = GenerationHelper ? GenerationHelper->FindInfo(PackageData) : nullptr;
 
 	// For every BeginCacheForCookedPlatformData call we made we need to call ClearAllCachedCookedPlatformData
 	if (ReleaseSaveReason == EStateChangeReason::Completed)
@@ -4210,7 +4222,7 @@ void UCookOnTheFlyServer::ReleaseCookedPlatformData(UE::Cook::FPackageData& Pack
 
 	if (GenerationInfo)
 	{
-		Generator->ResetSaveState(*GenerationInfo, PackageData.GetPackage(), ReleaseSaveReason);
+		GenerationHelper->ResetSaveState(*GenerationInfo, PackageData.GetPackage(), ReleaseSaveReason);
 		if (GenerationInfo->IsGenerator())
 		{
 			PackageData.SetInitializedGeneratorSave(false);
@@ -4218,19 +4230,19 @@ void UCookOnTheFlyServer::ReleaseCookedPlatformData(UE::Cook::FPackageData& Pack
 
 		if (ReleaseSaveReason == EStateChangeReason::Completed)
 		{
-			Generator->SetPackageSaved(*GenerationInfo, PackageData);
-			if (Generator->IsComplete())
+			GenerationHelper->SetPackageSaved(*GenerationInfo, PackageData);
+			if (GenerationHelper->IsComplete())
 			{
 				if (GenerationInfo->IsGenerator())
 				{
-					PackageData.DestroyGeneratorPackage();
+					PackageData.DestroyGenerationHelper();
 				}
 				else
 				{
-					Generator->GetOwner().DestroyGeneratorPackage();
+					GenerationHelper->GetOwner().DestroyGenerationHelper();
 				}
 				// Clear now-dangling pointers
-				Generator = nullptr;
+				GenerationHelper = nullptr;
 				GenerationInfo = nullptr;
 			}
 		}
@@ -5300,16 +5312,17 @@ void UCookOnTheFlyServer::PreGarbageCollect()
 	// or prevent them from being garbage collected if the splitter wants to keep them referenced
 	for (FPackageData* PackageData : PackageDatas->GetSaveQueue())
 	{
-		FGeneratorPackage* Generator = PackageData->GetGeneratorPackage();
-		if (!Generator)
+		FGenerationHelper* GenerationHelper = PackageData->GetGenerationHelper();
+		if (!GenerationHelper)
 		{
-			Generator = PackageData->GetGeneratedOwner();
+			GenerationHelper = PackageData->GetGeneratedOwner();
 		}
-		FCookGenerationInfo* Info = Generator ? Generator->FindInfo(*PackageData) : nullptr;
+		FCookGenerationInfo* Info = GenerationHelper ? GenerationHelper->FindInfo(*PackageData) : nullptr;
 		if (Info)
 		{
 			bool bShouldDemote;
-			Generator->PreGarbageCollect(*Info, GCKeepObjects, GCKeepPackages, GCKeepPackageDatas, bShouldDemote);
+			GenerationHelper->PreGarbageCollect(*Info, GCKeepObjects, GCKeepPackages,
+				GCKeepPackageDatas, bShouldDemote);
 			if (bShouldDemote)
 			{
 				ReleaseCookedPlatformData(*PackageData, UE::Cook::EStateChangeReason::GeneratorPreGarbageCollected);
@@ -5551,9 +5564,9 @@ void UCookOnTheFlyServer::PostGarbageCollect()
 
 	PackageDatas->LockAndEnumeratePackageDatas([](FPackageData* PackageData)
 	{
-		if (FGeneratorPackage* GeneratorPackage = PackageData->GetGeneratorPackage())
+		if (FGenerationHelper* GenerationHelper = PackageData->GetGenerationHelper())
 		{
-			GeneratorPackage->PostGarbageCollect();
+			GenerationHelper->PostGarbageCollect();
 		}
 	});
 	PackageDatas->LockAndEnumeratePackageDatas([](FPackageData* PackageData)
@@ -6300,10 +6313,11 @@ void FSaveCookedPackageContext::SetupPlatform(const ITargetPlatform* InTargetPla
 	{
 		if (TSet<FName>* NeverCookPackages = COTFS.PackageTracker->PlatformSpecificNeverCookPackages.Find(TargetPlatform))
 		{
-			FGeneratorPackage* Generator = PackageData.IsGenerated() ? PackageData.GetGeneratedOwner() : nullptr;
+			FGenerationHelper* GenerationHelper =
+				PackageData.IsGenerated() ? PackageData.GetGeneratedOwner() : nullptr;
 
 			if (NeverCookPackages->Find(Package->GetFName()) ||
-				(Generator && NeverCookPackages->Find(Generator->GetOwner().GetPackageName())))
+				(GenerationHelper && NeverCookPackages->Find(GenerationHelper->GetOwner().GetPackageName())))
 			{
 				SavePackageResult = ESavePackageResult::ContainsEditorOnlyData;
 				UE_LOG(LogCook, Display, TEXT("Excluding %s"), *PackageName);
@@ -6432,8 +6446,8 @@ void FSaveCookedPackageContext::FinishPlatform()
 		// during user saves for non-generated packages.
 		TOptional<TArray<FAssetDependency>> OverridePackageDependencies;
 		TOptional<FAssetPackageData> OverrideAssetPackageData;
-		FGeneratorPackage* GeneratorPackage;
-		if (GeneratorPackage = PackageData.GetGeneratorPackage(); GeneratorPackage)
+		FGenerationHelper* GenerationHelper;
+		if (GenerationHelper = PackageData.GetGenerationHelper(); GenerationHelper)
 		{
 			OverridePackageDependencies.Emplace();
 
@@ -6441,8 +6455,8 @@ void FSaveCookedPackageContext::FinishPlatform()
 			// each generated package.
 			COTFS.AssetRegistry->GetDependencies(Package->GetFName(), *OverridePackageDependencies,
 				UE::AssetRegistry::EDependencyCategory::Package);
-			OverridePackageDependencies->Reserve(GeneratorPackage->GetPackagesToGenerate().Num());
-			for (FCookGenerationInfo& GeneratedInfo : GeneratorPackage->GetPackagesToGenerate())
+			OverridePackageDependencies->Reserve(GenerationHelper->GetPackagesToGenerate().Num());
+			for (FCookGenerationInfo& GeneratedInfo : GenerationHelper->GetPackagesToGenerate())
 			{
 				FAssetDependency& Dependency = OverridePackageDependencies->Emplace_GetRef();
 				Dependency.AssetId = FAssetIdentifier(GeneratedInfo.PackageData->GetPackageName());
@@ -6452,13 +6466,13 @@ void FSaveCookedPackageContext::FinishPlatform()
 
 			if (!bHasFirstPlatformResults)
 			{
-				GeneratorPackage->FetchExternalActorDependencies();
-				COTFS.RecordExternalActorDependencies(GeneratorPackage->GetExternalActorDependencies());
+				GenerationHelper->FetchExternalActorDependencies();
+				COTFS.RecordExternalActorDependencies(GenerationHelper->GetExternalActorDependencies());
 			}
 		}
-		else if (GeneratorPackage = PackageData.GetGeneratedOwner(); GeneratorPackage)
+		else if (GenerationHelper = PackageData.GetGeneratedOwner(); GenerationHelper)
 		{
-			FCookGenerationInfo* GeneratedInfo = GeneratorPackage->FindInfo(PackageData);
+			FCookGenerationInfo* GeneratedInfo = GenerationHelper->FindInfo(PackageData);
 			if (!GeneratedInfo)
 			{
 				UE_LOG(LogCook, Error, TEXT("GeneratedInfo missing for package %s."), *PackageData.GetPackageName().ToString());
@@ -8539,7 +8553,7 @@ void UCookOnTheFlyServer::PopulateCookedPackages(TArrayView<const ITargetPlatfor
 					*Pair.Key.ToString());
 				continue;
 			}
-			Generator->CreateGeneratorPackage(nullptr, nullptr).SetPreviousGeneratedPackages(MoveTemp(Pair.Value.Generated));
+			Generator->CreateGenerationHelper(nullptr, nullptr).SetPreviousGeneratedPackages(MoveTemp(Pair.Value.Generated));
 		}
 
 		PlatformAssetRegistry.SetPreviousAssetRegistry(MoveTemp(PreviousAssetRegistry));
@@ -10359,7 +10373,7 @@ void UCookOnTheFlyServer::ShutdownCookSession()
 
 	PackageDatas->LockAndEnumeratePackageDatas([](UE::Cook::FPackageData* PackageData)
 	{
-		PackageData->DestroyGeneratorPackage();
+		PackageData->DestroyGenerationHelper();
 	});
 
 	if (IsCookByTheBookMode())
@@ -11522,7 +11536,7 @@ void UCookOnTheFlyServer::GetPackagesToRetract(int32 NumToRetract, TArray<FName>
 		}
 
 		if (PackageData->GetWorkerAssignmentConstraint().IsValid() || PackageData->IsGenerated() ||
-			PackageData->GetGeneratorPackage())
+			PackageData->GetGenerationHelper())
 		{
 			// Don't send back Packages that are constrained to this worker. Doing so will just
 			// cause the CookDirector to send it back to us, and this can cause the cooker to crash
diff --git a/Engine/Source/Editor/UnrealEd/Private/Cooker/CookDiagnostics.cpp b/Engine/Source/Editor/UnrealEd/Private/Cooker/CookDiagnostics.cpp
--- a/Engine/Source/Editor/UnrealEd/Private/Cooker/CookDiagnostics.cpp
+++ b/Engine/Source/Editor/UnrealEd/Private/Cooker/CookDiagnostics.cpp
@@ -164,7 +164,7 @@ void FDiagnostics::AnalyzeHiddenDependencies(UCookOnTheFlyServer& COTFS, FPackag
 			continue;
 		}
 
-		if (PackageData.GetGeneratorPackage())
+		if (PackageData.GetGenerationHelper())
 		{
 			// TODO: Collect SaveReferences for all generated packages and collect Unsolicited from all externalactor packages
 			// and run the missing solicited test when the GeneratorPackage finishes saving all generated
diff --git a/Engine/Source/Editor/UnrealEd/Private/Cooker/CookPackageData.cpp b/Engine/Source/Editor/UnrealEd/Private/Cooker/CookPackageData.cpp
--- a/Engine/Source/Editor/UnrealEd/Private/Cooker/CookPackageData.cpp
+++ b/Engine/Source/Editor/UnrealEd/Private/Cooker/CookPackageData.cpp
@@ -113,7 +113,7 @@ FPackageData::~FPackageData()
 
 void FPackageData::ClearReferences()
 {
-	DestroyGeneratorPackage();
+	DestroyGenerationHelper();
 }
 
 const FName& FPackageData::GetPackageName() const
@@ -292,12 +292,12 @@ void FPackageData::SetInstigator(FRequestCluster& Cluster, FInstigator&& InInsti
 	SetInstigatorInternal(MoveTemp(InInstigator));
 }
 
-void FPackageData::SetInstigator(FCookWorkerClient& Cluster, FInstigator&& InInstigator)
+void FPackageData::SetInstigator(FCookWorkerClient& Client, FInstigator&& InInstigator)
 {
 	SetInstigatorInternal(MoveTemp(InInstigator));
 }
 
-void FPackageData::SetInstigator(FGeneratorPackage& Cluster, FInstigator&& InInstigator)
+void FPackageData::SetInstigator(FGenerationHelper& InGenerationHelper, FInstigator&& InInstigator)
 {
 	SetInstigatorInternal(MoveTemp(InInstigator));
 }
@@ -1335,8 +1335,8 @@ void FPackageData::CheckCookedPlatformDataEmpty() const
 	check(!GetCookedPlatformDataStarted());
 	check(!GetCookedPlatformDataCalled());
 	check(!GetCookedPlatformDataComplete());
-	check(!GetGeneratorPackage() ||
-		GetGeneratorPackage()->GetOwnerInfo().GetSaveState() <= FCookGenerationInfo::ESaveState::StartPopulate);
+	check(!GetGenerationHelper() ||
+		GetGenerationHelper()->GetOwnerInfo().GetSaveState() <= FCookGenerationInfo::ESaveState::StartPopulate);
 	if (GetGeneratedOwner())
 	{
 		FCookGenerationInfo* Info = GetGeneratedOwner()->FindInfo(*this);
@@ -1427,9 +1427,9 @@ void FPackageData::UpdateSaveAfterGarbageCollect(bool& bOutDemote)
 		}
 	}
 
-	if (GeneratorPackage)
+	if (GenerationHelper)
 	{
-		GeneratorPackage->UpdateSaveAfterGarbageCollect(*this, bOutDemote);
+		GenerationHelper->UpdateSaveAfterGarbageCollect(*this, bOutDemote);
 	}
 	else if (IsGenerated())
 	{
@@ -1444,16 +1444,16 @@ void FPackageData::UpdateSaveAfterGarbageCollect(bool& bOutDemote)
 	}
 }
 
-void FPackageData::SetGeneratedOwner(FGeneratorPackage* InGeneratedOwner)
+void FPackageData::SetGeneratedOwner(FGenerationHelper* InGeneratedOwner)
 {
 	check(IsGenerated());
 	check(!(GeneratedOwner && InGeneratedOwner));
 	GeneratedOwner = InGeneratedOwner;
 }
 
-UE::Cook::FGeneratorPackage* FPackageData::GetGeneratorPackage() const
+UE::Cook::FGenerationHelper* FPackageData::GetGenerationHelper() const
 {
-	UE::Cook::FGeneratorPackage* Result = GeneratorPackage.Get();
+	UE::Cook::FGenerationHelper* Result = GenerationHelper.Get();
 	if (Result && Result->IsInitialized())
 	{
 		return Result;
@@ -1461,24 +1461,24 @@ UE::Cook::FGeneratorPackage* FPackageData::GetGeneratorPackage() const
 	return nullptr;
 }
 
-UE::Cook::FGeneratorPackage& FPackageData::CreateGeneratorPackage(const UObject* InSplitDataObject,
+UE::Cook::FGenerationHelper& FPackageData::CreateGenerationHelper(const UObject* InSplitDataObject,
 	ICookPackageSplitter* InCookPackageSplitterInstance)
 {
-	if (!GeneratorPackage)
+	if (!GenerationHelper)
 	{
-		GeneratorPackage.Reset(new UE::Cook::FGeneratorPackage(*this, InSplitDataObject,
+		GenerationHelper.Reset(new UE::Cook::FGenerationHelper(*this, InSplitDataObject,
 			InCookPackageSplitterInstance));
 	}
 	else
 	{
-		GeneratorPackage->InitializeSave(InSplitDataObject, InCookPackageSplitterInstance);
-		if (GeneratorPackage->IsInitialized())
+		GenerationHelper->InitializeSave(InSplitDataObject, InCookPackageSplitterInstance);
+		if (GenerationHelper->IsInitialized())
 		{
 			// The earlier exit from SaveState should have reset the progress back to StartGeneratorSave or earlier
-			check(GeneratorPackage->GetOwnerInfo().GetSaveState() <= FCookGenerationInfo::ESaveState::StartPopulate);
+			check(GenerationHelper->GetOwnerInfo().GetSaveState() <= FCookGenerationInfo::ESaveState::StartPopulate);
 		}
 	}
-	FGeneratorPackage* Result = GeneratorPackage.Get();
+	FGenerationHelper* Result = GenerationHelper.Get();
 	check(Result);
 	return *Result;
 }
@@ -1538,16 +1538,16 @@ namespace UE::Cook
 {
 
 //////////////////////////////////////////////////////////////////////////
-// FGeneratorPackage
+// FGenerationHelper
 
-FGeneratorPackage::FGeneratorPackage(UE::Cook::FPackageData& InOwner, const UObject* InSplitDataObject,
+FGenerationHelper::FGenerationHelper(UE::Cook::FPackageData& InOwner, const UObject* InSplitDataObject,
 	ICookPackageSplitter* InCookPackageSplitterInstance)
 : OwnerInfo(InOwner, true /* bInGenerated */)
 {
 	InitializeSave(InSplitDataObject, InCookPackageSplitterInstance);
 }
 
-void FGeneratorPackage::InitializeSave(const UObject* InSplitDataObject,
+void FGenerationHelper::InitializeSave(const UObject* InSplitDataObject,
 	ICookPackageSplitter* InCookPackageSplitterInstance)
 {
 	if (InCookPackageSplitterInstance)
@@ -1572,7 +1572,7 @@ void FGeneratorPackage::InitializeSave(const UObject* InSplitDataObject,
 	}
 }
 
-FGeneratorPackage::~FGeneratorPackage()
+FGenerationHelper::~FGenerationHelper()
 {
 	if (!IsInitialized())
 	{
@@ -1583,7 +1583,7 @@ FGeneratorPackage::~FGeneratorPackage()
 	ClearGeneratedPackages();
 }
 
-void FGeneratorPackage::ConditionalNotifyCompletion(ICookPackageSplitter::ETeardown Status)
+void FGenerationHelper::ConditionalNotifyCompletion(ICookPackageSplitter::ETeardown Status)
 {
 	if (!bNotifiedCompletion)
 	{
@@ -1593,7 +1593,7 @@ void FGeneratorPackage::ConditionalNotifyCompletion(ICookPackageSplitter::ETeard
 	}
 }
 
-void FGeneratorPackage::ClearGeneratedPackages()
+void FGenerationHelper::ClearGeneratedPackages()
 {
 	check(IsInitialized());
 	for (FCookGenerationInfo& Info: PackagesToGenerate)
@@ -1607,7 +1607,7 @@ void FGeneratorPackage::ClearGeneratedPackages()
 	}
 }
 
-bool FGeneratorPackage::TryGenerateList(UObject* OwnerObject, FPackageDatas& PackageDatas)
+bool FGenerationHelper::TryGenerateList(UObject* OwnerObject, FPackageDatas& PackageDatas)
 {
 	check(IsInitialized());
 
@@ -1636,13 +1636,15 @@ bool FGeneratorPackage::TryGenerateList(UObject* OwnerObject, FPackageDatas& Pac
 	{
 		if (!SplitterData.GetCreateAsMap().IsSet())
 		{
-			UE_LOG(LogCook, Error, TEXT("PackageSplitter did not specify whether CreateAsMap is true for generated package. Splitter=%s, Generated=%s."),
+			UE_LOG(LogCook, Error,
+				TEXT("PackageSplitter did not specify whether CreateAsMap is true for generated package. Splitter=%s, Generated=%s."),
 				*this->GetSplitDataObjectName().ToString(), *OwnerPackageData.GetPackageName().ToString());
 			return false;
 		}
 		bool bCreateAsMap = *SplitterData.GetCreateAsMap();
 
-		FString PackageRoot = SplitterData.GeneratedRootPath.IsEmpty() ? OwnerPackageData.GetPackageName().ToString() : SplitterData.GeneratedRootPath;
+		FString PackageRoot = SplitterData.GeneratedRootPath.IsEmpty() ?
+			OwnerPackageData.GetPackageName().ToString() : SplitterData.GeneratedRootPath;
 
 		FString PackageName = FPaths::RemoveDuplicateSlashes(FString::Printf(TEXT("/%s/%s/%s"),
 			*PackageRoot, GeneratedPackageSubPath, *SplitterData.RelativePath));
@@ -1651,7 +1653,8 @@ bool FGeneratorPackage::TryGenerateList(UObject* OwnerObject, FPackageDatas& Pac
 			false /* bRequireExists */, bCreateAsMap);
 		if (!PackageData)
 		{
-			UE_LOG(LogCook, Error, TEXT("PackageSplitter could not find mounted filename for generated packagepath. Splitter=%s, Generated=%s."),
+			UE_LOG(LogCook, Error,
+				TEXT("PackageSplitter could not find mounted filename for generated packagepath. Splitter=%s, Generated=%s."),
 				*this->GetSplitDataObjectName().ToString(), *PackageName);
 			return false;
 		}
@@ -1662,7 +1665,8 @@ bool FGeneratorPackage::TryGenerateList(UObject* OwnerObject, FPackageDatas& Pac
 		check(PackageData->GetGeneratedOwner() == nullptr);
 		if (IFileManager::Get().FileExists(*PackageData->GetFileName().ToString()))
 		{
-			UE_LOG(LogCook, Warning, TEXT("PackageSplitter specified a generated package that already exists in the workspace domain. Splitter=%s, Generated=%s."),
+			UE_LOG(LogCook, Warning,
+				TEXT("PackageSplitter specified a generated package that already exists in the workspace domain. Splitter=%s, Generated=%s."),
 				*this->GetSplitDataObjectName().ToString(), *PackageName);
 			return false;
 		}
@@ -1723,7 +1727,7 @@ bool FGeneratorPackage::TryGenerateList(UObject* OwnerObject, FPackageDatas& Pac
 	return true;
 }
 
-void FGeneratorPackage::FetchExternalActorDependencies()
+void FGenerationHelper::FetchExternalActorDependencies()
 {
 	check(IsInitialized());
 
@@ -1762,7 +1766,7 @@ void FGeneratorPackage::FetchExternalActorDependencies()
 	ExternalActorDependencies.Shrink();
 }
 
-FCookGenerationInfo* FGeneratorPackage::FindInfo(const FPackageData& PackageData)
+FCookGenerationInfo* FGenerationHelper::FindInfo(const FPackageData& PackageData)
 {
 	check(IsInitialized());
 
@@ -1779,18 +1783,19 @@ FCookGenerationInfo* FGeneratorPackage::FindInfo(const FPackageData& PackageData
 	}
 	return nullptr;
 }
-const FCookGenerationInfo* FGeneratorPackage::FindInfo(const FPackageData& PackageData) const
+const FCookGenerationInfo* FGenerationHelper::FindInfo(const FPackageData& PackageData) const
 {
-	return const_cast<FGeneratorPackage*>(this)->FindInfo(PackageData);
+	return const_cast<FGenerationHelper*>(this)->FindInfo(PackageData);
 }
 
-ICookPackageSplitter* FGeneratorPackage::GetCookPackageSplitterInstance() const
+ICookPackageSplitter* FGenerationHelper::GetCookPackageSplitterInstance() const
 {
-	checkf(!bNotifiedCompletion, TEXT("It is illegal for the cooker to try to access the CookPackageSplitterInstance after calling Teardown on it."));
+	checkf(!bNotifiedCompletion,
+		TEXT("It is illegal for the cooker to try to access the CookPackageSplitterInstance after calling Teardown on it."));
 	return CookPackageSplitterInstance.Get();
 }
 
-UObject* FGeneratorPackage::FindSplitDataObject() const
+UObject* FGenerationHelper::FindSplitDataObject() const
 {
 	check(IsInitialized());
 	FString ObjectPath = GetSplitDataObjectName().ToString();
@@ -1805,7 +1810,7 @@ UObject* FGeneratorPackage::FindSplitDataObject() const
 	return FindObject<UObject>(nullptr, *ObjectPath);
 }
 
-void FGeneratorPackage::PreGarbageCollect(FCookGenerationInfo& Info, TArray<TObjectPtr<UObject>>& GCKeepObjects,
+void FGenerationHelper::PreGarbageCollect(FCookGenerationInfo& Info, TArray<TObjectPtr<UObject>>& GCKeepObjects,
 	TArray<UPackage*>& GCKeepPackages, TArray<FPackageData*>& GCKeepPackageDatas, bool& bOutShouldDemote)
 {
 	if (!IsInitialized())
@@ -1857,7 +1862,7 @@ void FGeneratorPackage::PreGarbageCollect(FCookGenerationInfo& Info, TArray<TObj
 	}
 }
 
-void FGeneratorPackage::PostGarbageCollect()
+void FGenerationHelper::PostGarbageCollect()
 {
 	if (!IsInitialized())
 	{
@@ -1871,7 +1876,8 @@ void FGeneratorPackage::PostGarbageCollect()
 		// objects, so it should still be loaded
 		if (!Owner.GetPackage() || !FindSplitDataObject())
 		{
-			UE_LOG(LogCook, Error, TEXT("PackageSplitter object was deleted by garbage collection while generation was still ongoing. This will break the generation.")
+			UE_LOG(LogCook, Error,
+				TEXT("PackageSplitter object was deleted by garbage collection while generation was still ongoing. This will break the generation.")
 				TEXT("\n\tSplitter=%s."), *GetSplitDataObjectName().ToString());
 		}
 	}
@@ -1888,7 +1894,8 @@ void FGeneratorPackage::PostGarbageCollect()
 				!Owner.IsKeepReferencedDuringGC() &&
 				!IsUseInternalReferenceToAvoidGarbageCollect())
 			{
-				UE_LOG(LogCook, Error, TEXT("PackageSplitter found the Generator package still in memory after it should have been deleted by GC.")
+				UE_LOG(LogCook, Error,
+					TEXT("PackageSplitter found the Generator package still in memory after it should have been deleted by GC.")
 					TEXT("\n\tThis is unexpected since garbage has been collected and the package should have been unreferenced so it should have been collected, and will break population of Generated packages.")
 					TEXT("\n\tSplitter=%s"), *GetSplitDataObjectName().ToString());
 				EReferenceChainSearchMode SearchMode = EReferenceChainSearchMode::Shortest
@@ -1906,14 +1913,17 @@ void FGeneratorPackage::PostGarbageCollect()
 		{
 			if (!Info.PackageData->IsKeepReferencedDuringGC() && !Info.HasSaved() && !bHasIssuedWarning)
 			{
-				UE_LOG(LogCook, Warning, TEXT("PackageSplitter found a package it generated that was not removed from memory during garbage collection. This will cause errors later during population.")
-					TEXT("\n\tSplitter=%s, Generated=%s."), *GetSplitDataObjectName().ToString(), *Info.PackageData->GetPackageName().ToString());
+				UE_LOG(LogCook, Warning,
+					TEXT("PackageSplitter found a package it generated that was not removed from memory during garbage collection. This will cause errors later during population.")
+					TEXT("\n\tSplitter=%s, Generated=%s."), *GetSplitDataObjectName().ToString(),
+					*Info.PackageData->GetPackageName().ToString());
 				
 				{
 					// Compute UCookOnTheFlyServer's references so they are gathered by OBJ REFS below 
 					UCookOnTheFlyServer::FScopeFindCookReferences(Owner.GetPackageDatas().GetCookOnTheFlyServer());
 
-					StaticExec(nullptr, *FString::Printf(TEXT("OBJ REFS NAME=%s"), *Info.PackageData->GetPackageName().ToString()));
+					StaticExec(nullptr, *FString::Printf(TEXT("OBJ REFS NAME=%s"),
+						*Info.PackageData->GetPackageName().ToString()));
 				}
 				
 				bHasIssuedWarning = true; // Only issue the warning once per GC
@@ -1926,7 +1936,7 @@ void FGeneratorPackage::PostGarbageCollect()
 	}
 }
 
-UPackage* FGeneratorPackage::CreateGeneratedUPackage(FCookGenerationInfo& GeneratedInfo,
+UPackage* FGenerationHelper::CreateGeneratedUPackage(FCookGenerationInfo& GeneratedInfo,
 	const UPackage* InOwnerPackage, const TCHAR* GeneratedPackageName)
 {
 	check(IsInitialized());
@@ -1946,7 +1956,7 @@ UPackage* FGeneratorPackage::CreateGeneratedUPackage(FCookGenerationInfo& Genera
 	return GeneratedPackage;
 }
 
-void FGeneratorPackage::SetPackageSaved(FCookGenerationInfo& Info, FPackageData& PackageData)
+void FGenerationHelper::SetPackageSaved(FCookGenerationInfo& Info, FPackageData& PackageData)
 {
 	check(IsInitialized());
 	if (Info.HasSaved())
@@ -1962,13 +1972,13 @@ void FGeneratorPackage::SetPackageSaved(FCookGenerationInfo& Info, FPackageData&
 	}
 }
 
-bool FGeneratorPackage::IsComplete() const
+bool FGenerationHelper::IsComplete() const
 {
 	check(IsInitialized());
 	return RemainingToPopulate == 0;
 }
 
-void FGeneratorPackage::ResetSaveState(FCookGenerationInfo& Info, UPackage* Package, EStateChangeReason ReleaseSaveReason)
+void FGenerationHelper::ResetSaveState(FCookGenerationInfo& Info, UPackage* Package, EStateChangeReason ReleaseSaveReason)
 {
 	check(IsInitialized());
 	if (Info.GetSaveState() > FCookGenerationInfo::ESaveState::CallPopulate)
@@ -1976,7 +1986,8 @@ void FGeneratorPackage::ResetSaveState(FCookGenerationInfo& Info, UPackage* Pack
 		UObject* SplitObject = FindSplitDataObject();
 		if (!SplitObject || !Package)
 		{
-			UE_LOG(LogCook, Warning, TEXT("PackageSplitter: %s on %s was GarbageCollected before we finished saving it. This prevents us from calling PostSave and may corrupt other packages that it altered during Populate. Splitter=%s."),
+			UE_LOG(LogCook, Warning,
+				TEXT("PackageSplitter: %s on %s was GarbageCollected before we finished saving it. This prevents us from calling PostSave and may corrupt other packages that it altered during Populate. Splitter=%s."),
 				(!Package ? TEXT("UPackage") : TEXT("SplitDataObject")),
 				Info.PackageData ? *Info.PackageData->GetPackageName().ToString() : *Info.RelativePath,
 				*GetSplitDataObjectName().ToString());
@@ -1997,7 +2008,8 @@ void FGeneratorPackage::ResetSaveState(FCookGenerationInfo& Info, UPackage* Pack
 				PopulateInfo.GeneratedRootPath = Info.GeneratedRootPath;
 				PopulateInfo.bCreatedAsMap = Info.IsCreateAsMap();
 				PopulateInfo.Package = Package;
-				GetCookPackageSplitterInstance()->PostSaveGeneratedPackage(GetOwnerPackage(), SplitObject, PopulateInfo);
+				GetCookPackageSplitterInstance()->PostSaveGeneratedPackage(GetOwnerPackage(), SplitObject,
+					PopulateInfo);
 			}
 		}
 	}
@@ -2034,7 +2046,8 @@ void FGeneratorPackage::ResetSaveState(FCookGenerationInfo& Info, UPackage* Pack
 			(ReleaseSaveReason != EStateChangeReason::Completed && ReleaseSaveReason != EStateChangeReason::DoneForNow &&
 			 ReleaseSaveReason != EStateChangeReason::SaveError && ReleaseSaveReason != EStateChangeReason::CookerShutdown))
 		{
-			UE_LOG(LogCook, Error, TEXT("CookPackageSplitter failure: We are demoting a %s package from save and removing our references that keep its objects loaded.\n")
+			UE_LOG(LogCook, Error,
+				TEXT("CookPackageSplitter failure: We are demoting a %s package from save and removing our references that keep its objects loaded.\n")
 				TEXT("This will allow the objects to be garbage collected and cause failures in the splitter which expects them to remain loaded.\n")
 				TEXT("Package=%s, Splitter=%s, ReleaseSaveReason=%s"),
 				Info.IsGenerator() ? TEXT("generator") : TEXT("generated"),
@@ -2047,7 +2060,7 @@ void FGeneratorPackage::ResetSaveState(FCookGenerationInfo& Info, UPackage* Pack
 	Info.KeepReferencedPackages.Reset();
 }
 
-void FGeneratorPackage::UpdateSaveAfterGarbageCollect(const FPackageData& PackageData, bool& bInOutDemote)
+void FGenerationHelper::UpdateSaveAfterGarbageCollect(const FPackageData& PackageData, bool& bInOutDemote)
 {
 	if (!IsInitialized())
 	{
@@ -2108,7 +2121,8 @@ void FGeneratorPackage::UpdateSaveAfterGarbageCollect(const FPackageData& Packag
 			{
 				if (!PackageData.GetPackage()->IsFullyLoaded())
 				{
-					DeletedObject = FString::Printf(TEXT("UPackage %s is no longer FullyLoaded"), *PackageData.GetPackageName().ToString());
+					DeletedObject = FString::Printf(TEXT("UPackage %s is no longer FullyLoaded"),
+						*PackageData.GetPackageName().ToString());
 				}
 				else
 				{
@@ -2116,13 +2130,15 @@ void FGeneratorPackage::UpdateSaveAfterGarbageCollect(const FPackageData& Packag
 				}
 			}
 		}
-		UE_LOG(LogCook, Error, TEXT("A %s package had some of its UObjects deleted during garbage collection after it started generating. This will cause errors during save of the package.")
+		UE_LOG(LogCook, Error,
+			TEXT("A %s package had some of its UObjects deleted during garbage collection after it started generating. This will cause errors during save of the package.")
 			TEXT("\n\tDeleted object: %s")
 			TEXT("\n\tSplitter=%s%s"),
 			Info->IsGenerator() ? TEXT("Generator") : TEXT("Generated"),
 			*DeletedObject,
 			*GetSplitDataObjectName().ToString(),
-			Info->IsGenerator() ? TEXT(".") : *FString::Printf(TEXT(", Generated=%s."), *Info->PackageData->GetPackageName().ToString()));
+			Info->IsGenerator() ? TEXT(".") : *FString::Printf(TEXT(", Generated=%s."),
+				*Info->PackageData->GetPackageName().ToString()));
 	}
 
 	// Remove raw pointers from RootMovedObjects if they no longer exist in the weakpointers in CachedObjectsInOuter
@@ -2178,7 +2194,7 @@ void FCachedObjectInOuterGeneratorInfo::Initialize(UObject* Object)
 	bInitialized = true;
 }
 
-void FCookGenerationInfo::TakeOverCachedObjectsAndAddMoved(FGeneratorPackage& Generator,
+void FCookGenerationInfo::TakeOverCachedObjectsAndAddMoved(FGenerationHelper& GenerationHelper,
 	TArray<FCachedObjectInOuter>& CachedObjectsInOuter, TArray<UObject*>& MovedObjects)
 {
 	CachedObjectsInOuterInfo.Reset();
@@ -2197,11 +2213,13 @@ void FCookGenerationInfo::TakeOverCachedObjectsAndAddMoved(FGeneratorPackage& Ge
 	{
 		if (!IsValid(Object))
 		{
-			UE_LOG(LogCook, Warning, TEXT("CookPackageSplitter found non-valid object %s returned from %s on Splitter %s%s. Ignoring it."),
+			UE_LOG(LogCook, Warning,
+				TEXT("CookPackageSplitter found non-valid object %s returned from %s on Splitter %s%s. Ignoring it."),
 				Object ? *Object->GetFullName() : TEXT("<null>"),
 				IsGenerator() ? TEXT("PopulateGeneratorPackage") : TEXT("PopulateGeneratedPackage"),
-				*Generator.GetSplitDataObjectName().ToString(),
-				IsGenerator() ? TEXT("") : *FString::Printf(TEXT(", Package %s"), *PackageData->GetPackageName().ToString()));
+				*GenerationHelper.GetSplitDataObjectName().ToString(),
+				IsGenerator() ? TEXT("") : *FString::Printf(TEXT(", Package %s"),
+					*PackageData->GetPackageName().ToString()));
 			continue;
 		}
 		FCachedObjectInOuterGeneratorInfo& Info = CachedObjectsInOuterInfo.FindOrAdd(Object);
@@ -2211,7 +2229,8 @@ void FCookGenerationInfo::TakeOverCachedObjectsAndAddMoved(FGeneratorPackage& Ge
 			Info.bMoved = true;
 			Info.bMovedRoot = true;
 			CachedObjectsInOuter.Emplace(Object);
-			GetObjectsWithOuter(Object, ChildrenOfMovedObjects, true /* bIncludeNestedObjects */, RF_NoFlags, EInternalObjectFlags::Garbage);
+			GetObjectsWithOuter(Object, ChildrenOfMovedObjects, true /* bIncludeNestedObjects */, RF_NoFlags,
+				EInternalObjectFlags::Garbage);
 		}
 	}
 
@@ -2230,12 +2249,13 @@ void FCookGenerationInfo::TakeOverCachedObjectsAndAddMoved(FGeneratorPackage& Ge
 	SetHasTakenOverCachedCookedPlatformData(true);
 }
 
-EPollStatus FCookGenerationInfo::RefreshPackageObjects(FGeneratorPackage& Generator, UPackage* Package,
+EPollStatus FCookGenerationInfo::RefreshPackageObjects(FGenerationHelper& GenerationHelper, UPackage* Package,
 	bool& bOutFoundNewObjects, ESaveState DemotionState)
 {
 	bOutFoundNewObjects = false;
 	TArray<UObject*> CurrentObjectsInOuter;
-	GetObjectsWithOuter(Package, CurrentObjectsInOuter, true /* bIncludeNestedObjects */, RF_NoFlags, EInternalObjectFlags::Garbage);
+	GetObjectsWithOuter(Package, CurrentObjectsInOuter, true /* bIncludeNestedObjects */, RF_NoFlags,
+		EInternalObjectFlags::Garbage);
 
 	check(PackageData); // RefreshPackageObjects is only called when there is a PackageData
 	TArray<FCachedObjectInOuter>& CachedObjectsInOuter = PackageData->GetCachedObjectsInOuter();
@@ -2260,11 +2280,13 @@ EPollStatus FCookGenerationInfo::RefreshPackageObjects(FGeneratorPackage& Genera
 		SetSaveState(DemotionState);
 		if (++PackageData->GetNumRetriesBeginCacheOnObjects() > FPackageData::GetMaxNumRetriesBeginCacheOnObjects())
 		{
-			UE_LOG(LogCook, Error, TEXT("Cooker has repeatedly tried to call BeginCacheForCookedPlatformData on all objects in a generated package, but keeps finding new objects.\n")
+			UE_LOG(LogCook, Error,
+				TEXT("Cooker has repeatedly tried to call BeginCacheForCookedPlatformData on all objects in a generated package, but keeps finding new objects.\n")
 				TEXT("Aborting the save of the package; programmer needs to debug why objects keep getting added to the package.\n")
 				TEXT("Splitter: %s%s. Most recent created object: %s."),
-				*Generator.GetSplitDataObjectName().ToString(),
-				IsGenerator() ? TEXT("") : *FString::Printf(TEXT(", Package: %s"), *PackageData->GetPackageName().ToString()),
+				*GenerationHelper.GetSplitDataObjectName().ToString(),
+				IsGenerator() ? TEXT("") : *FString::Printf(TEXT(", Package: %s"),
+					*PackageData->GetPackageName().ToString()),
 				*FirstNewObject->GetFullName());
 			return EPollStatus::Error;
 		}
@@ -2300,11 +2322,11 @@ void FCookGenerationInfo::CreatePackageHash()
 	PRAGMA_ENABLE_DEPRECATION_WARNINGS;
 }
 
-void FCookGenerationInfo::IterativeCookValidateOrClear(FGeneratorPackage& Generator,
+void FCookGenerationInfo::IterativeCookValidateOrClear(FGenerationHelper& GenerationHelper,
 	TConstArrayView<const ITargetPlatform*> RequestedPlatforms, const FIoHash& PreviousPackageHash,
 	bool& bOutIterativelyUnmodified)
 {
-	UCookOnTheFlyServer& COTFS = Generator.GetOwner().GetPackageDatas().GetCookOnTheFlyServer();
+	UCookOnTheFlyServer& COTFS = GenerationHelper.GetOwner().GetPackageDatas().GetCookOnTheFlyServer();
 	bOutIterativelyUnmodified = PreviousPackageHash == this->PackageHash;
 	if (bOutIterativelyUnmodified)
 	{
@@ -3147,8 +3169,8 @@ void FPackageDatas::UpdateThreadsafePackageData(const FPackageData& PackageData)
 		[&PackageData](FThreadsafePackageData& ThreadsafeData, bool bNew)
 		{
 			ThreadsafeData.Instigator = PackageData.GetInstigator();
-			FGeneratorPackage* Generator = PackageData.GetGeneratedOwner();
-			ThreadsafeData.Generator = Generator ? Generator->GetOwner().GetPackageName() : NAME_None;
+			FGenerationHelper* GenerationHelper = PackageData.GetGeneratedOwner();
+			ThreadsafeData.Generator = GenerationHelper ? GenerationHelper->GetOwner().GetPackageName() : NAME_None;
 		});
 }
 
diff --git a/Engine/Source/Editor/UnrealEd/Private/Cooker/CookPackageData.h b/Engine/Source/Editor/UnrealEd/Private/Cooker/CookPackageData.h
--- a/Engine/Source/Editor/UnrealEd/Private/Cooker/CookPackageData.h
+++ b/Engine/Source/Editor/UnrealEd/Private/Cooker/CookPackageData.h
@@ -50,7 +50,7 @@ namespace UE::Cook
 
 class FPackageDataQueue;
 class FRequestCluster;
-struct FGeneratorPackage;
+struct FGenerationHelper;
 struct FPackageData;
 struct FPackageDataMonitor;
 struct FPendingCookedPlatformDataCancelManager;
@@ -121,8 +121,8 @@ struct FCachedObjectInOuter
 };
 
 /**
- * Extra information about CachedObjectsInOuter that a generator needs to know for diagnostics.
- * The generator constructs an associated array (aka TMap) of these structures when it takes
+ * Extra information about CachedObjectsInOuter that a GenerationHelper needs to know for diagnostics.
+ * The GenerationHelper constructs an associated array (aka TMap) of these structures when it takes
  * over the CachedObjectsInOuter.
  */
 struct FCachedObjectInOuterGeneratorInfo
@@ -220,8 +220,8 @@ struct FPackagePlatformData
 	/** All flags modified by reachability calculations are returned to default. */
 	void ResetReachable();
 	/**
-	 * Mark the platform as ExplorableOverride=true and reset all data necessary to reexplore it, including reachability.
-	 * Caller is responsbile for marking it again as reachable.
+	 * Mark platform as ExplorableOverride=true and reset all data necessary to reexplore it, including reachability.
+	 * Caller is responsible for marking it again as reachable.
 	 */
 	void MarkAsExplorable();
 
@@ -402,9 +402,12 @@ public:
 	 * FindOrAdd each TargetPlatform and set its flags: CookAttempted=true, Succeeded=<given>.
 	 * In version that takes two arrays, TargetPlatforms and Succeeded must be the same length.
 	 */
-	void SetPlatformsCooked(const TConstArrayView<const ITargetPlatform*> TargetPlatforms, const TConstArrayView<ECookResult> Succeeded, bool bInWasCookedThisSession = true);
-	void SetPlatformsCooked(const TConstArrayView<const ITargetPlatform*> TargetPlatforms, ECookResult Result, bool bInWasCookedThisSession = true);
-	void SetPlatformCooked(const ITargetPlatform* TargetPlatform, ECookResult Result, bool bInWasCookedThisSession = true);
+	void SetPlatformsCooked(const TConstArrayView<const ITargetPlatform*> TargetPlatforms,
+		const TConstArrayView<ECookResult> Succeeded, bool bInWasCookedThisSession = true);
+	void SetPlatformsCooked(const TConstArrayView<const ITargetPlatform*> TargetPlatforms,
+		ECookResult Result, bool bInWasCookedThisSession = true);
+	void SetPlatformCooked(const ITargetPlatform* TargetPlatform, ECookResult Result,
+		bool bInWasCookedThisSession = true);
 	/**
 	 * FindOrAdd each TargetPlatform and set its flags: CookAttempted=false.
 	 * In Version that takes no TargetPlatform, CookAttempted is cleared from all existing platforms.
@@ -417,7 +420,8 @@ public:
 
 	/** Access the information about platforms interacted with by *this. */
 	const TSortedMap<const ITargetPlatform*, FPackagePlatformData, TInlineAllocator<1>>& GetPlatformDatas() const;
-	TSortedMap<const ITargetPlatform*, FPackagePlatformData, TInlineAllocator<1>>& GetPlatformDatasConstKeysMutableValues();
+	TSortedMap<const ITargetPlatform*, FPackagePlatformData, TInlineAllocator<1>>&
+		GetPlatformDatasConstKeysMutableValues();
 
 	/** Add a platform if not already existing and return a writable pointer to its flags. */
 	FPackagePlatformData& FindOrAddPlatformData(const ITargetPlatform* TargetPlatform);
@@ -486,7 +490,7 @@ public:
 	bool IsInStateProperty(EPackageStateProperty Property) const;
 
 	/*
-	 * CompletionCallback - A callback that will be called when this PackageData next transitions from InProgress to not
+	 * CompletionCallback - A callback that is called when this PackageData next transitions from InProgress to not
 	 * InProgress because of cook success, failure, skip, or cancel.
 	 */
 	/** Get a reference to the currently set callback, to e.g. move it into a local variable during execution. */
@@ -594,7 +598,7 @@ public:
 	 */
 	void ClearCookedPlatformData();
 
-	/** Get/Set the PackageDataMonitor flag that counts whether this PackageData has finished cooking and with what result. */
+	/** Get/Set the Monitor's flag that counts whether this PackageData has finished cooking and with what result. */
 	ECookResult GetMonitorCookResult() const { return (ECookResult)MonitorCookResult; }
 	void SetMonitorCookResult(ECookResult Value) { MonitorCookResult = (uint8) Value; }
 
@@ -607,13 +611,13 @@ public:
 	/** Swap all ITargetPlatform* stored on this instance according to the mapping in @param Remap. */
 	void RemapTargetPlatforms(const TMap<ITargetPlatform*, ITargetPlatform*>& Remap);
 
-	/** Create and return a GeneratorPackage helper object for a given CookPackageSplitter. */
-	FGeneratorPackage& CreateGeneratorPackage(const UObject* InSplitDataObject,
+	/** Create a GenerationHelper struct for this PackageData and the given CookPackageSplitter and return it. */
+	FGenerationHelper& CreateGenerationHelper(const UObject* InSplitDataObject,
 		ICookPackageSplitter* InCookPackageSplitterInstance);
-	/** Destroy GeneratorPackage helper object. */
-	void DestroyGeneratorPackage() { GeneratorPackage.Reset(); }
-	/** Get GeneratorPackage helper object. */
-	FGeneratorPackage* GetGeneratorPackage() const;
+	/** Destroy this PackageData's GenerationHelper struct. */
+	void DestroyGenerationHelper() { GenerationHelper.Reset(); }
+	/** Get the helper object containing data and functions for packages that need to generate other packages. */
+	FGenerationHelper* GetGenerationHelper() const;
 	/** Get whether the PackageData has done any necessary Generator steps and is ready for BeginCache calls. */
 	bool HasCompletedGeneration() const { return static_cast<bool>(bCompletedGeneration); }
 	/** Set whether the PackageData has done any necessary Generator steps and is ready for BeginCache calls. */
@@ -627,9 +631,9 @@ public:
 	/** Set whether the PackageData is a generated package created by a owning generator PackageData. */
 	void SetGenerated(bool Value) { bGenerated = Value != 0; }
 	/** Set the owning generator PackageData. Only valid to call if SetGenerated(true) has been called. */
-	void SetGeneratedOwner(FGeneratorPackage* InGeneratedOwner);
-	/** Return the owning generator PackageData. Will be null if not a generated package or if orphaned. */
-	FGeneratorPackage* GetGeneratedOwner() const { return GeneratedOwner; }
+	void SetGeneratedOwner(FGenerationHelper* InGeneratedOwner);
+	/** Return the parent generator's GenerationHelper. Will be null if not a generated package or if orphaned. */
+	FGenerationHelper* GetGeneratedOwner() const { return GeneratedOwner; }
 
 	/**
 	 * Return the instigator for this package. The Instigator is the first code location or
@@ -639,8 +643,8 @@ public:
 	bool HasInstigator() const { return Instigator.Category != EInstigator::NotYetRequested; }
 	/** Setting the instigator is mostly private - it should only be done during clustering. */
 	void SetInstigator(FRequestCluster& Cluster, FInstigator&& InInstigator);
-	void SetInstigator(FCookWorkerClient& Cluster, FInstigator&& InInstigator);
-	void SetInstigator(FGeneratorPackage& Cluster, FInstigator&& InInstigator);
+	void SetInstigator(FCookWorkerClient& Client, FInstigator&& InInstigator);
+	void SetInstigator(FGenerationHelper& InGenerationHelper, FInstigator&& InInstigator);
 
 	/** Get whether COTFS is keeping this package referenced referenced during GC. */
 	bool IsKeepReferencedDuringGC() const { return static_cast<bool>(bKeepReferencedDuringGC); }
@@ -662,7 +666,7 @@ public:
 	/** Set the workerid that is the only worker allowed to cook this package; default is InvalidId; */
 	void SetWorkerAssignmentConstraint(FWorkerId InWorkerAssignment) { WorkerAssignmentConstraint = InWorkerAssignment; }
 
-	/** Marshall this PackageData to a ConstructData that can be used later or on a remote machine to reconstruct it. */
+	/** Marshall this PackageData to a ConstructData that is used later or on a remote machine to reconstruct it. */
 	FConstructPackageData CreateConstructData();
 
 	/** Only used When IsDebugRecordUnsolicited: storage for the list of unsolicited packages for this package */
@@ -683,8 +687,8 @@ private:
 
 	void SetIsUrgent(bool Value);
 	void SetInstigatorInternal(FInstigator&& InInstigator);
-	static void AddReachablePlatformsInternal(FPackageData& PackageData, TConstArrayView<const ITargetPlatform*> Platforms,
-		FInstigator&& InInstigator);
+	static void AddReachablePlatformsInternal(FPackageData& PackageData,
+		TConstArrayView<const ITargetPlatform*> Platforms, FInstigator&& InInstigator);
 	static void QueueAsDiscoveredInternal(FPackageData& PackageData, FInstigator&& InInstigator,
 		FDiscoveredPlatformSet&& ReachablePlatforms, bool bUrgent);
 
@@ -748,8 +752,8 @@ private:
 
 	void OnPackageDataFirstMarkedReachable(FInstigator&& InInstigator);
 
-	TUniquePtr<FGeneratorPackage> GeneratorPackage;
-	FGeneratorPackage* GeneratedOwner;
+	TUniquePtr<FGenerationHelper> GenerationHelper;
+	FGenerationHelper* GeneratedOwner;
 	/** Data for each platform that has been interacted with by *this. */
 	TSortedMap<const ITargetPlatform*, FPackagePlatformData, TInlineAllocator<1>> PlatformDatas;
 
@@ -811,7 +815,7 @@ private:
 };
 
 /**
- * Struct uses with ICookPackageSplitter, which contains the state information for the save of either the
+ * Struct used with ICookPackageSplitter, which contains the state information for the save of either the
  * Generator package or one of its Generated packages.
  */
 struct FCookGenerationInfo
@@ -862,17 +866,19 @@ public:
 
 	/**
 	 * Steal the list of cached objects to call BeginCacheForCookedPlatformData on from the PackageData,
-	 * and add them to this. Also add the list of NewObjects reported by the splitter that will be moved into the package.
+	 * and add them to this. Also add NewObjects reported by the splitter that will be moved into the package.
 	 */
-	void TakeOverCachedObjectsAndAddMoved(FGeneratorPackage& Generator, 
+	void TakeOverCachedObjectsAndAddMoved(FGenerationHelper& GenerationHelper,
 		TArray<FCachedObjectInOuter>& CachedObjectsInOuter, TArray<UObject*>& MovedObjects);
 	/**
-	 * Fetch all the objects currently in the package and add them the list of objects that need BeginCacheForCookedPlatformData.
-	 * Reports whether new objects were found. If DemotionState is not ESaveState::Last, will SetState back to DemotionState
+	 * Fetch all the objects currently in the package and add them to the list of objects that need
+	 * BeginCacheForCookedPlatformData.
+	 * Reports whether new objects were found. If DemotionState is not ESaveState::Last,
+	 * will SetState back to DemotionState
 	 * if new objects were found, and will error exit if this demotion has happened too many times.
 	 */
-	EPollStatus RefreshPackageObjects(FGeneratorPackage& Generator, UPackage* Package, bool& bOutFoundNewObjects,
-		ESaveState DemotionState);
+	EPollStatus RefreshPackageObjects(FGenerationHelper& GenerationHelper, UPackage* Package,
+		bool& bOutFoundNewObjects, ESaveState DemotionState);
 
 	void AddKeepReferencedPackages(TArray<UPackage*>& InKeepReferencedPackages)
 	{
@@ -887,7 +893,7 @@ public:
 	 * test whether they need to be invalidated now and clear the results if so. Only called in legacy iterative;
 	 * incremental cooks handle invalidation by querying the TargetDomainDigest during the RequestCluster.
 	 */
-	void IterativeCookValidateOrClear(FGeneratorPackage& Generator,
+	void IterativeCookValidateOrClear(FGenerationHelper& GeneratorHelper,
 		TConstArrayView<const ITargetPlatform*> RequestedPlatforms, const FIoHash& PreviousHash,
 		bool& bOutIterativelyUnmodified);
 
@@ -916,13 +922,13 @@ private:
  * Helper that wraps a ICookPackageSplitter, gets/caches packages to generate and provide
  * the necessary to iterate over all of them iteratively.
  */
-struct FGeneratorPackage
+struct FGenerationHelper
 {
 public:
 	/** Store the provided CookPackageSplitter and prepare the packages to generate. */
-	FGeneratorPackage(UE::Cook::FPackageData& InOwner, const UObject* InSplitDataObject,
+	FGenerationHelper(UE::Cook::FPackageData& InOwner, const UObject* InSplitDataObject,
 		ICookPackageSplitter* InCookPackageSplitterInstance);
-	~FGeneratorPackage();
+	~FGenerationHelper();
 	void InitializeSave(const UObject* InSplitDataObject, ICookPackageSplitter* InCookPackageSplitterInstance);
 	bool IsInitialized() const { return bInitialized; }
 	/** Clear references to owned generated packages, and mark those packages as orphaned */
@@ -930,12 +936,15 @@ public:
 
 	/** Call the Splitter's GetGenerateList and create the PackageDatas */
 	bool TryGenerateList(UObject* OwnerObject, FPackageDatas& PackageDatas);
-	/** Record all dependencies from the Generator that are ExternalActor dependencies and store them on this. */
+	/**
+	 * Record all dependencies from the generator package that are ExternalActor dependencies and
+	 * store them on this.
+	 */
 	void FetchExternalActorDependencies();
 
 	/** Accessor for the packages to generate */
 	TArrayView<UE::Cook::FCookGenerationInfo> GetPackagesToGenerate() { check(IsInitialized()); return PackagesToGenerate; }
-	/** Return the GenerationInfo used to save the Generator's UPackage */
+	/** Return the GenerationInfo used to save the generator package's UPackage */
 	UE::Cook::FCookGenerationInfo& GetOwnerInfo() { check(IsInitialized()); return OwnerInfo; }
 	/** Return owner FPackageData. */
 	UE::Cook::FPackageData& GetOwner() { return *OwnerInfo.PackageData; }
@@ -968,7 +977,10 @@ public:
 	/** Call CreatePackage and set PackageData for deterministic generated packages, and update status. */
 	UPackage* CreateGeneratedUPackage(FCookGenerationInfo& GenerationInfo,
 		const UPackage* OwnerPackage, const TCHAR* GeneratedPackageName);
-	/** Mark that the generator or a generated package has saved, to keep track of when *this is no longer needed. */
+	/**
+	 * Mark that the generator package or a generated package has saved, to keep track of when *this
+	 * is no longer needed.
+	 */
 	void SetPackageSaved(FCookGenerationInfo& Info, FPackageData& PackageData);
 	/** Return whether list has been generated and all generated packages have been populated */
 	bool IsComplete() const;
@@ -1119,9 +1131,15 @@ public:
 	/** Callback called from FPackageData when it transitions to or from inprogress. */
 	void OnInProgressChanged(FPackageData& PackageData, bool bInProgress);
 	void OnPreloadAllocatedChanged(FPackageData& PackageData, bool bPreloadAllocated);
-	/** Callback called from FPackageData when it has set a platform to CookAttempted=true and it does not have any others cooked. */
+	/**
+	 * Callback called from FPackageData when it has set a platform to CookAttempted=true and it does not have
+	 * any others cooked.
+	 */
 	void OnFirstCookedPlatformAdded(FPackageData& PackageData, ECookResult CookResult);
-	/** Callback called from FPackageData when it has set a platform to CookAttempted=false and it does not have any others cooked. */
+	/**
+	 * Callback called from FPackageData when it has set a platform to CookAttempted=false and it does not have
+	 * any others cooked.
+	 */
 	void OnLastCookedPlatformRemoved(FPackageData& PackageData);
 	/** Callback called from FPackageData when it has changed its urgency. */
 	void OnUrgencyChanged(FPackageData& PackageData);
@@ -1184,8 +1202,8 @@ public:
 	FPackageDataSet& GetRestartedRequests() { return RestartedRequests; }
 	/**
 	 * Unlike other containers on PackageData, GetDiscoveryQueue is not an ownership container.
-	 * PackageDatas in the DiscoveryQueue are PackageDatas we need to look at - in the right order during PumpRequests -
-	 * they can be in any state and are owned by another container (or are in the idle state).
+	 * PackageDatas in the DiscoveryQueue are PackageDatas we need to look at - in the right order during
+	 * PumpRequests - they can be in any state and are owned by another container (or are in the idle state).
 	 */
 	TRingBuffer<FDiscoveryQueueElement>& GetDiscoveryQueue() { return DiscoveryQueue; }
 	TRingBuffer<FRequestCluster>& GetRequestClusters() { return RequestClusters; }
diff --git a/Engine/Source/Editor/UnrealEd/Private/Cooker/CookWorkerClient.cpp b/Engine/Source/Editor/UnrealEd/Private/Cooker/CookWorkerClient.cpp
--- a/Engine/Source/Editor/UnrealEd/Private/Cooker/CookWorkerClient.cpp
+++ b/Engine/Source/Editor/UnrealEd/Private/Cooker/CookWorkerClient.cpp
@@ -193,9 +193,9 @@ void FCookWorkerClient::ReportPromoteToSaveComplete(FPackageData& PackageData)
 	Result->SetPackageName(PackageName);
 	Result->SetSuppressCookReason(ESuppressCookReason::NotSuppressed);
 	Result->SetPlatforms(OrderedSessionPlatforms);
-	if (FGeneratorPackage* Generator = PackageData.GetGeneratorPackage(); Generator)
+	if (FGenerationHelper* GenerationHelper = PackageData.GetGenerationHelper(); GenerationHelper)
 	{
-		Result->SetExternalActorDependencies(Generator->ReleaseExternalActorDependencies());
+		Result->SetExternalActorDependencies(GenerationHelper->ReleaseExternalActorDependencies());
 	}
 
 	int32 NumPlatforms = OrderedSessionPlatforms.Num();
