Prevent the PreCreateSubObjectsForReinstantiation to recreate DefaultSubObject if they are already created. This came up while using asyncloading and it was working without it by relying on UObject recylcing which was not intended.

diff --git a/Engine/Source/Editor/UnrealEd/Private/Kismet2/KismetReinstanceUtilities.cpp b/Engine/Source/Editor/UnrealEd/Private/Kismet2/KismetReinstanceUtilities.cpp
--- a/Engine/Source/Editor/UnrealEd/Private/Kismet2/KismetReinstanceUtilities.cpp
+++ b/Engine/Source/Editor/UnrealEd/Private/Kismet2/KismetReinstanceUtilities.cpp
@@ -3239,14 +3239,22 @@ void FBlueprintCompileReinstancer::PreCreateSubObjectsForReinstantiation_Inner(c
 				// Only pre-create object where the class does not have newer version of the it
 				if(!SubObjectClass->HasAnyClassFlags(CLASS_NewerVersionExists))
 				{
-					UObject* NewSubObject = NewObject<UObject>(NewUObject, SubObjectClass, SubObjectName, SubObjectFlags);
+					UObject* NewSubObject = nullptr;
+					if (UObject** ExistingNewSubObject = ContainedNewSubObjects.FindByPredicate([SubObjectName](UObject* SubObject) { return SubObject && SubObject->GetFName() == SubObjectName; }))
+					{
+						NewSubObject = *ExistingNewSubObject;
+					}
+					else
+					{
+						NewSubObject = NewObject<UObject>(NewUObject, SubObjectClass, SubObjectName, SubObjectFlags);
+					}
 					CreatedInstanceMap.Add(OldSubObject, NewSubObject);
 					PreCreateSubObjectsForReinstantiation_Inner(OldInstancedSubObjects, OldToNewClassMap, OldSubObject, NewSubObject, CreatedInstanceMap, OldToNewInstanceMap);
 				}
 			}
 		}
 		// There might be new subobjects attached to the sub object that are particular to this instance, let's traverse it to find them out.
-		else if (UObject** NewSubObject = ContainedNewSubObjects.FindByPredicate([SubObjectName](UObject* SubObject) { return SubObject && SubObject->GetName() == SubObjectName; }))
+		else if (UObject** NewSubObject = ContainedNewSubObjects.FindByPredicate([SubObjectName](UObject* SubObject) { return SubObject && SubObject->GetFName() == SubObjectName; }))
 		{
 			PreCreateSubObjectsForReinstantiation_Inner(OldInstancedSubObjects, OldToNewClassMap, OldSubObject, *NewSubObject, CreatedInstanceMap, OldToNewInstanceMap);
 		}
