Fixes for subobject when copying prefab from the content bowser

diff --git a/Engine/Source/Editor/UnrealEd/Private/Kismet2/BlueprintEditorUtils.cpp b/Engine/Source/Editor/UnrealEd/Private/Kismet2/BlueprintEditorUtils.cpp
--- a/Engine/Source/Editor/UnrealEd/Private/Kismet2/BlueprintEditorUtils.cpp
+++ b/Engine/Source/Editor/UnrealEd/Private/Kismet2/BlueprintEditorUtils.cpp
@@ -1,6 +1,7 @@
 // Copyright Epic Games, Inc. All Rights Reserved.
 
 #include "Kismet2/BlueprintEditorUtils.h"
+#include "Kismet2/KismetReinstanceUtilities.h"
 #include "Algo/AnyOf.h"
 #include "Algo/Copy.h"
 #include "Algo/RemoveIf.h"
@@ -1779,7 +1780,21 @@ void FBlueprintEditorUtils::PostDuplicateBlueprint(UBlueprint* Blueprint, bool b
 
 			UObject* NewCDO = Blueprint->GeneratedClass->GetDefaultObject();
 			check(NewCDO != nullptr);
-			UEditorEngine::CopyPropertiesForUnrelatedObjects(OldCDO, NewCDO);
+
+			TMap<UClass*, UClass*> InOutOldToNewClassMap;
+			TMap<UObject*, UObject*> CreatedInstanceMap;
+			FBlueprintCompileReinstancer::PreCreateSubObjectsForReinstantiation(InOutOldToNewClassMap, OldCDO, NewCDO, CreatedInstanceMap);
+
+			// We only need to copy properties of the pre-created instances, the rest of the default sub object is done inside the UEditorEngine::CopyPropertiesForUnrelatedObjects
+			TMap<UObject*, UObject*> OldToNewInstanceMap(CreatedInstanceMap);
+			UEngine::FCopyPropertiesForUnrelatedObjectsParams Params;
+			Params.OptionalReplacementMappings = &OldToNewInstanceMap;
+			Params.bOnlyHandleDirectSubObjects = true;
+			Params.bReplaceInternalReferenceUponRead = true;
+			for (const auto& Pair : CreatedInstanceMap)
+			{
+				UEditorEngine::CopyPropertiesForUnrelatedObjects(Pair.Key, Pair.Value, Params);
+			}
 
 			// copy sparse data over to the new class sparse data, if any:
 			const TObjectPtr<UScriptStruct> SparseData = Blueprint->GeneratedClass->GetSparseClassDataStruct();
