Cooker: bErrorOnEngineContentUse: Allow references to uncooked objectredirectors outside of the dlc; we do not mark those as "cooked" by the base game so our typical method of deciding "has been cooked by the base game" does not work for them.
#rnx

diff --git a/Engine/Source/Editor/UnrealEd/Private/Cooker/CookRequestCluster.cpp b/Engine/Source/Editor/UnrealEd/Private/Cooker/CookRequestCluster.cpp
--- a/Engine/Source/Editor/UnrealEd/Private/Cooker/CookRequestCluster.cpp
+++ b/Engine/Source/Editor/UnrealEd/Private/Cooker/CookRequestCluster.cpp
@@ -2,6 +2,7 @@
 
 #include "Cooker/CookRequestCluster.h"
 
+#include "Algo/AllOf.h"
 #include "Algo/BinarySearch.h"
 #include "Algo/Sort.h"
 #include "Algo/TopologicalSort.h"
@@ -1619,19 +1620,36 @@ void FRequestCluster::IsRequestCookable(const ITargetPlatform* Platform, FName P
 		FileName.ToString(NameBuffer);
 		if (!FStringView(NameBuffer).StartsWith(InDLCPath))
 		{
-			if (!PackageData.HasCookedPlatform(Platform, true /* bIncludeFailed */))
+			// Editoronly content that was not cooked by the base game is allowed to be "cooked"; if it references
+			// something not editoronly then we will exclude and give a warning on that followup asset. We need to
+			// handle editoronly objects being referenced because the base game will not have marked them as cooked so
+			// we will think we still need to "cook" them.
+			// The only case where this comes up is in ObjectRedirectors, so we only test for those for performance.
+			TArray<FAssetData> Assets;
+			IAssetRegistry::GetChecked().GetAssetsByPackageName(PackageName, Assets,
+				true /* bIncludeOnlyOnDiskAssets */);
+			bool bEditorOnly = !Assets.IsEmpty() &&
+				Algo::AllOf(Assets, [](const FAssetData& Asset)
+					{
+						return Asset.IsRedirector();
+					});
+
+			if (!bEditorOnly)
 			{
-				// AllowUncookedAssetReferences should only be used when the DLC plugin to cook is going to be mounted where uncooked packages are available.
-				// This will allow a DLC plugin to be recooked continually and mounted in an uncooked editor which is useful for CI.
-				if (!InCOTFS.CookByTheBookOptions->bAllowUncookedAssetReferences)
+				if (!PackageData.HasCookedPlatform(Platform, true /* bIncludeFailed */))
 				{
-					UE_LOG(LogCook, Error, TEXT("Uncooked Engine or Game content %s is being referenced by DLC!"), *NameBuffer);
+					// AllowUncookedAssetReferences should only be used when the DLC plugin to cook is going to be mounted where uncooked packages are available.
+					// This will allow a DLC plugin to be recooked continually and mounted in an uncooked editor which is useful for CI.
+					if (!InCOTFS.CookByTheBookOptions->bAllowUncookedAssetReferences)
+					{
+						UE_LOG(LogCook, Error, TEXT("Uncooked Engine or Game content %s is being referenced by DLC!"), *NameBuffer);
+					}
 				}
+				OutReason = ESuppressCookReason::NotInCurrentPlugin;
+				bOutCookable = false;
+				bOutExplorable = false;
+				return;
 			}
-			OutReason = ESuppressCookReason::NotInCurrentPlugin;
-			bOutCookable = false;
-			bOutExplorable = false;
-			return;
 		}
 	}
 
