Fix reinstancer to correctly restore individual actor initialization state instead of relying on the global world state.
This fixes several issues with reinstancing actors when compiling blueprints inside partially loaded streaming levels.

diff --git a/Engine/Source/Editor/UnrealEd/Private/Kismet2/KismetReinstanceUtilities.cpp b/Engine/Source/Editor/UnrealEd/Private/Kismet2/KismetReinstanceUtilities.cpp
--- a/Engine/Source/Editor/UnrealEd/Private/Kismet2/KismetReinstanceUtilities.cpp
+++ b/Engine/Source/Editor/UnrealEd/Private/Kismet2/KismetReinstanceUtilities.cpp
@@ -745,6 +745,9 @@ public:
 				UWorld* World = Actor->GetWorld();
 				if (World)
 				{
+					// NOTE: This function does not handle gameplay edge cases correctly!
+					// FActorReplacementHelper has a better implementation of this code
+
 					// Remove any pending latent actions, as the compiled script code may have changed, and thus the
 					// cached LinkInfo data may now be invalid. This could happen in the fast path, since the original
 					// Actor instance will not be replaced in that case, and thus might still have latent actions pending.
@@ -1438,6 +1441,12 @@ struct FActorReplacementHelper
 		TArray<AActor*> AttachedActors;
 		OldActor->GetAttachedActors(AttachedActors);
 
+		// Cache the actor initialization status
+		bHasRegisteredAllComponents = OldActor->HasActorRegisteredAllComponents();
+		bHasInitialized = OldActor->IsActorInitialized();
+		bHasBegunPlay = OldActor->HasActorBegunPlay();
+		bWasHiddenEdLevel = OldActor->bHiddenEdLevel;
+
 		// if there are attached objects detach them and store the socket names
 		for (AActor* AttachedActor : AttachedActors)
 		{
@@ -1497,6 +1506,10 @@ private:
 	AActor*          NewActor;
 	FTransform       TargetWorldTransform;
 	FActorAttachmentData AttachmentData;
+	bool bHasRegisteredAllComponents = false;
+	bool bHasInitialized = false;
+	bool bHasBegunPlay = false;
+	bool bWasHiddenEdLevel = false;
 
 	/** Holds actor component data, etc. that we use to apply */
 	TSharedPtr<FActorTransactionAnnotation> CachedActorData;
@@ -1515,6 +1528,8 @@ void FActorReplacementHelper::Finalize(const TMap<UObject*, UObject*>& OldToNewI
 	FEditorScriptExecutionGuard ScriptGuard;
 
 	// run the construction script, which will use the properties we just copied over
+	// @TODO: This code is similar to AActor::RerunConstructionScripts and ideally could use shared code for restoring state
+
 	bool bCanReRun = UBlueprint::IsBlueprintHierarchyErrorFree(NewActor->GetClass());
 	if (NewActor->CurrentTransactionAnnotation.IsValid() && bCanReRun)
 	{
@@ -1537,28 +1552,31 @@ void FActorReplacementHelper::Finalize(const TMap<UObject*, UObject*>& OldToNewI
 		NewActor->ExecuteConstruction(TargetWorldTransform, nullptr, &DummyComponentData);
 	}	
 
-	// The reinstancing case doesn't ever explicitly call Actor->FinishSpawning, we've handled the construction script
-	// portion above but still need the PostActorConstruction() case so BeginPlay gets routed correctly while in a BegunPlay world
+	// Try to restore gameplay initialization state
 	if (UWorld* World = NewActor->GetWorld())
 	{
-		if (World->HasBegunPlay())
+		// This is unsafe to call from a loading stack but that should never happen for an actor that was fully initialized
+		// @TODO: If there is a need for this case, it must be deferred until later in the frame
+		if (World->IsGameWorld() && bHasInitialized && ensure(!FUObjectThreadContext::Get().IsRoutingPostLoad))
 		{
-			//GAllowActorScriptExecutionInEditor must be false when we call BeginPlay
+			// GAllowActorScriptExecutionInEditor must be false when we call events from initialization
 			TGuardValue AutoRestore(GAllowActorScriptExecutionInEditor, false);
-			NewActor->PostActorConstruction();
-		}
-	}
 
-	// make sure that the actor is properly hidden if it's in a hidden sublevel:
-	bool bIsInHiddenLevel = false;
-	if (ULevel* Level = NewActor->GetLevel())
-	{
-		// consider non visible level as hidden in editor only if this is post world initialization
-		UWorld* World = Level->GetWorld();
-		bIsInHiddenLevel = !Level->bIsVisible && World && World->bIsWorldInitialized;
+			// Restore initialization state
+			NewActor->PreInitializeComponents();
+			NewActor->InitializeComponents();
+			NewActor->PostInitializeComponents();
+
+			// Also call begin play if necessary
+			if (bHasBegunPlay)
+			{
+				NewActor->DispatchBeginPlay(false);
+			}
+		}
 	}
 
-	if (bIsInHiddenLevel)
+	// Restore editor visibility
+	if (bWasHiddenEdLevel)
 	{
 		NewActor->bHiddenEdLevel = true;
 		NewActor->MarkComponentsRenderStateDirty();
@@ -1580,7 +1598,6 @@ void FActorReplacementHelper::Finalize(const TMap<UObject*, UObject*>& OldToNewI
 		GEditor->NotifyToolsOfObjectReplacement(ConstructedComponentReplacementMap);
 	}
 
-	// Destroy actor and clear references.
 	NewActor->Modify();
 	if (GEditor)
 	{
@@ -2451,6 +2468,7 @@ static void ReplaceActorHelper(AActor* OldActor, UClass* OldClass, UObject*& New
 	// running the NewActor's construction-script is saved for that 
 	// second pass (because the construction-script may reference 
 	// another instance that hasn't been replaced yet).
+	bool bHadRegisteredComponents = OldActor->HasActorRegisteredAllComponents();
 	FActorAttachmentData& CurrentAttachmentData = ActorAttachmentData.FindChecked(OldActor);
 	ReplacementActors.Add(FActorReplacementHelper(NewActor, OldActor, MoveTemp(CurrentAttachmentData)));
 	ActorAttachmentData.Remove(OldActor);
@@ -2482,10 +2500,9 @@ static void ReplaceActorHelper(AActor* OldActor, UClass* OldClass, UObject*& New
 	// reset properties/streams
 	NewActor->ResetPropertiesForConstruction();
 
-	// Only register the components if the world is already initialized
-	if (World->bIsWorldInitialized)
+	// Only register the native components if the actor had already registered them
+	if (bHadRegisteredComponents)
 	{
-		// register native components
 		NewActor->RegisterAllComponents();
 	}
 	
