New TRS Gizmos: fixed crash with Editor Utility Widgets
    - as the InteractiveToolsContext is replaced within the FWidgetModeManager constructor, we have to make sure to unregister / register the transform gizmo context object when replacing it

#rnx

diff --git a/Engine/Source/Editor/UnrealEd/Private/WidgetModeManager.cpp b/Engine/Source/Editor/UnrealEd/Private/WidgetModeManager.cpp
--- a/Engine/Source/Editor/UnrealEd/Private/WidgetModeManager.cpp
+++ b/Engine/Source/Editor/UnrealEd/Private/WidgetModeManager.cpp
@@ -2,6 +2,7 @@
 
 #include "WidgetModeManager.h"
 
+#include "EditorGizmos/EditorTransformGizmoUtil.h"
 #include "ToolContexts/WidgetToolsContext.h"
 
 //////////////////////////////////////////////////////////////////////////
@@ -13,6 +14,7 @@ FWidgetModeManager::FWidgetModeManager()
 	// Since we can't call virtuals in constructor, need free and re-create InteractiveToolsContext resource
 	if (UObjectInitialized())
 	{
+		UE::EditorTransformGizmoUtil::UnregisterTransformGizmoContextObject(this);
 		InteractiveToolsContext->ShutdownContext();
 		InteractiveToolsContext = nullptr;
 	}
@@ -20,6 +22,7 @@ FWidgetModeManager::FWidgetModeManager()
 	CachedWidgetToolContext = NewObject<UWidgetToolsContext>(GetTransientPackage(), UWidgetToolsContext::StaticClass(), NAME_None, RF_Transient);
 	InteractiveToolsContext = CachedWidgetToolContext;
 	InteractiveToolsContext->InitializeContextWithEditorModeManager(this);
+	UE::EditorTransformGizmoUtil::RegisterTransformGizmoContextObject(this);
 }
 
 bool FWidgetModeManager::OnKeyChar(const FGeometry& MyGeometry, const FCharacterEvent& InCharacterEvent)
