PR #10360: [SourceControl] Fix DiffAgainstShelvedFile() to catch the date of the revision

diff --git a/Engine/Source/Editor/SourceControlWindows/Private/SSourceControlHistory.cpp b/Engine/Source/Editor/SourceControlWindows/Private/SSourceControlHistory.cpp
--- a/Engine/Source/Editor/SourceControlWindows/Private/SSourceControlHistory.cpp
+++ b/Engine/Source/Editor/SourceControlWindows/Private/SSourceControlHistory.cpp
@@ -1823,6 +1823,7 @@ bool FSourceControlWindows::DiffAgainstShelvedFile(const FSourceControlStateRef&
 					FRevisionInfo ShelvedRevision;
 					ShelvedRevision.Changelist = Revision->GetCheckInIdentifier();
 					ShelvedRevision.Revision = TEXT("Shelved");
+					ShelvedRevision.Date = Revision->GetDate();
 
 					FRevisionInfo NewRevision;
 					NewRevision.Revision = TEXT("");
diff --git a/Engine/Source/Editor/UnrealEd/Private/DiffUtils.cpp b/Engine/Source/Editor/UnrealEd/Private/DiffUtils.cpp
--- a/Engine/Source/Editor/UnrealEd/Private/DiffUtils.cpp
+++ b/Engine/Source/Editor/UnrealEd/Private/DiffUtils.cpp
@@ -1495,13 +1495,13 @@ FText DiffViewUtils::GetPanelLabel(const UObject* Asset, const FRevisionInfo& Re
 			RevisionData = FText::Format(NSLOCTEXT("DiffViewUtils", "RevisionData", "Revision {0} - CL {1} - {2}")
 				, FText::FromString(Revision.Revision)
 				, FText::AsNumber(Revision.Changelist, &FNumberFormattingOptions::DefaultNoGrouping())
-				, FText::FromString(Revision.Date.ToString(TEXT("%m/%d/%Y"))));
+				, FText::AsDateTime(Revision.Date));
 		}
 		else
 		{
 			RevisionData = FText::Format(NSLOCTEXT("DiffViewUtils", "RevisionDataNoChangelist", "Revision {0} - {1}")
 				, FText::FromString(Revision.Revision)
-				, FText::FromString(Revision.Date.ToString(TEXT("%m/%d/%Y"))));		
+				, FText::AsDateTime(Revision.Date));
 		}
 
 		if (Label.IsEmpty())
