Downgrade absence of memory pressure stats from OS during cook to be a warning instead of an error.

diff --git a/Engine/Source/Editor/UnrealEd/Private/CookOnTheFlyServer.cpp b/Engine/Source/Editor/UnrealEd/Private/CookOnTheFlyServer.cpp
--- a/Engine/Source/Editor/UnrealEd/Private/CookOnTheFlyServer.cpp
+++ b/Engine/Source/Editor/UnrealEd/Private/CookOnTheFlyServer.cpp
@@ -5023,7 +5023,7 @@ bool UCookOnTheFlyServer::PumpHasExceededMaxMemory(uint32& OutResultFlags)
 		if (PressureStatus == FPlatformMemoryStats::EMemoryPressureStatus::Unknown)
 		{
 			UE_CALL_ONCE([&]() {
-				UE_LOG(LogCook, Error,
+				UE_LOG(LogCook, Warning,
 				TEXT("MemoryPressureStatus is not available from the operating system. We may run out of memory due to lack of knowledge of when to collect garbage."));
 				});
 		}
