Fix ensure when trying to update 'object being debugged' and we were previously debugging an object that is no longer valid, but still exists

diff --git a/Engine/Source/Editor/UnrealEd/Private/Kismet2/KismetDebugUtilities.cpp b/Engine/Source/Editor/UnrealEd/Private/Kismet2/KismetDebugUtilities.cpp
--- a/Engine/Source/Editor/UnrealEd/Private/Kismet2/KismetDebugUtilities.cpp
+++ b/Engine/Source/Editor/UnrealEd/Private/Kismet2/KismetDebugUtilities.cpp
@@ -281,7 +281,7 @@ void FKismetDebugUtilities::OnScriptException(const UObject* ActiveObject, const
 		{
 			// Check if we need to update the object being debugged
 			UObject* ObjectToDebug = FindObjectSafe<UObject>(nullptr, *PathToDebug);
-			if (ObjectToDebug)
+			if (IsValid(ObjectToDebug))
 			{
 				// If the path to debug matches a newly-spawned object, set the hard reference now
 				ObjectBeingDebugged = ObjectToDebug;
