Changed rootmotion movement in DebugSkelmeshComponent back to using AddLocalTransform()
- Fixes issue with additive animation playback getting stuck on first loop
- Fixes simulation preview in physics asset editor

diff --git a/Engine/Source/Editor/UnrealEd/Classes/Animation/DebugSkelMeshComponent.h b/Engine/Source/Editor/UnrealEd/Classes/Animation/DebugSkelMeshComponent.h
--- a/Engine/Source/Editor/UnrealEd/Classes/Animation/DebugSkelMeshComponent.h
+++ b/Engine/Source/Editor/UnrealEd/Classes/Animation/DebugSkelMeshComponent.h
@@ -311,9 +311,6 @@ class UDebugSkelMeshComponent : public USkeletalMeshComponent
 	/** Storage of Source Animation Pose for when bDisplaySourceAnimation == true, as they have to be calculated */
 	TArray<FTransform> SourceAnimationPoses;
 
-	/** Accumulated root motion. */
-	FTransform RootMotionTransform;
-
 	/** Transform representing the actor transform at the beginning of the animation sequence. */
 	FTransform RootMotionReferenceTransform;
 
@@ -410,18 +407,6 @@ class UDebugSkelMeshComponent : public USkeletalMeshComponent
 	UNREALED_API virtual void ShowReferencePose(bool bRefPose);
 	UNREALED_API virtual bool IsReferencePoseShown() const;
 
-	/** @return Accumulated root motion with turntable animation. */
-	FTransform GetRootMotionTransform() const
-	{
-		return RootMotionTransform * TurntableTransform;
-	}
-
-	/** @return Transform representing the actor transform at the beginning of the animation with turntable animation. */
-	FTransform GetRootMotionReferenceTransform() const
-	{
-		return RootMotionReferenceTransform * TurntableTransform;
-	}
-
 	/** Called when mirror data table changes on anim instance. */
 	UNREALED_API void OnMirrorDataTableChanged();
 
@@ -635,8 +620,6 @@ public:
 	EPersonaTurnTableMode::Type TurnTableMode;
 	/** Current turn table speed scaling */
 	float TurnTableSpeedScaling;
-	/** Current turn table transform */
-	FTransform TurntableTransform;
 
 	UNREALED_API virtual void TickComponent(float DeltaTime, enum ELevelTick TickType, FActorComponentTickFunction *ThisTickFunction) override;
 
diff --git a/Engine/Source/Editor/UnrealEd/Private/Animation/DebugSkelMeshComponent.cpp b/Engine/Source/Editor/UnrealEd/Private/Animation/DebugSkelMeshComponent.cpp
--- a/Engine/Source/Editor/UnrealEd/Private/Animation/DebugSkelMeshComponent.cpp
+++ b/Engine/Source/Editor/UnrealEd/Private/Animation/DebugSkelMeshComponent.cpp
@@ -34,7 +34,13 @@ namespace UE::Anim::Private
 	FTransform CalculateInitialTransformFromAssetAndTime(const UDebugSkelMeshComponent& InDebugSkelMeshComponent, const UAnimationAsset* InAnimAsset, const float InTime)
 	{
 		const FTransform InitialRootBoneTransform = InDebugSkelMeshComponent.GetReferenceSkeleton().GetRefBonePose()[0];
-		const FTransform InitialTransform = UE::Anim::ExtractRootTransformFromAnimationAsset(InAnimAsset, InTime);
+		FTransform InitialTransform = UE::Anim::ExtractRootTransformFromAnimationAsset(InAnimAsset, InTime);
+		if (InAnimAsset->IsValidAdditive())
+		{
+			// Additive animations have zero scale as "no scaling" value.
+			// This function is used to set the initial transform, we explicitly set the scale to start at 1.
+			InitialTransform.SetScale3D(FVector::OneVector);
+		}
 		return InitialRootBoneTransform * InitialTransform;
 	}
 }
@@ -50,14 +56,12 @@ UDebugSkelMeshComponent::UDebugSkelMeshComponent(const FObjectInitializer& Objec
 	bMeshSocketsVisible = true;
 	bSkeletonSocketsVisible = true;
 
-	bShowNotificationVisualizations = false;
+	bShowNotificationVisualizations = true;
 	bShowRootMotionVisualizations = false;
 
 	TurnTableSpeedScaling = 1.f;
 	TurnTableMode = EPersonaTurnTableMode::Stopped;
-	TurntableTransform = FTransform::Identity;
 
-	RootMotionTransform = FTransform::Identity;
 	RootMotionReferenceTransform = FTransform::Identity;
 
 	bPauseClothingSimulationWithAnim = false;
@@ -197,6 +201,9 @@ double WrapInRange(double StartVal, double MinVal, double MaxVal)
 
 void UDebugSkelMeshComponent::ConsumeRootMotion(const FVector& FloorMin, const FVector& FloorMax)
 {
+	// Note: this method is called from FAnimationEditorPreviewScene() after world tick,
+	// so care must be taken how the transforms are adjusted. Other features,
+	// such as physic simulation or turn table rotation may change the transform. 
 	if (PreviewInstance == nullptr)
 	{
 		return;
@@ -226,105 +233,121 @@ void UDebugSkelMeshComponent::ConsumeRootMotion(const FVector& FloorMin, const F
 	const float PreviousTime = ConsumeRootMotionPreviousPlaybackTime;
 	ConsumeRootMotionPreviousPlaybackTime = CurrentTime;
 
-	bool bLooped = false;
-
-	if (ProcessRootMotionMode != EProcessRootMotionMode::Ignore && !ShouldBlendPhysicsBones())
+	// Apply the root motion, including resetting the actors location in case the process root motion mode requires it.
+	// If ShouldBlendPhysicsBones() is true, do not alter the transform, because the readback from the physics bodies into
+	// the bone transforms will already have been done during the tick based on the component transform.
+	// If we change the component transform here, then the bone transforms will not match the physics simulation.
+	if (!ShouldBlendPhysicsBones())
 	{
-		if (PreviewInstance->IsPlaying())
+		if (ProcessRootMotionMode != EProcessRootMotionMode::Ignore)
 		{
-			// Figure out if the animation has looped, and the start end position of the root motion extraction.
-			float SectionStartPosition = 0.0f;
-
-			// We have to deal with montage explicitly because we can have multiple sections and we want to reset the position when the section loops
-			// and depending on the composition, CurrentTime < PreviousTime (or CurrentTime > PreviousTime when playing in reverse) is not enough
-			if (const UAnimMontage* Montage = Cast<UAnimMontage>(PreviewInstance->CurrentAsset))
+			if (PreviewInstance->IsPlaying())
 			{
-				const int32 PreviewStartSectionIdx = Montage->CompositeSections.IsValidIndex(PreviewInstance->MontagePreviewStartSectionIdx) ? PreviewInstance->MontagePreviewStartSectionIdx : Montage->GetSectionIndexFromPosition(CurrentTime);
-				const int32 FirstSectionIdx = PreviewInstance->MontagePreview_FindFirstSectionAsInMontage(PreviewStartSectionIdx);
-				const int32 LastSectionIdx = PreviewInstance->MontagePreview_FindLastSection(FirstSectionIdx);
+				// Figure out if the animation has looped, and the start end position of the root motion extraction.
+				float SectionStartPosition = 0.0f;
+				bool bLooped = false;
 
-				// If FirstSection == LastSection we are previewing a single section
-				// In this case to know if we have looped we just need to check if CurrentTime < PreviousTime (or the oposite if we are playing the montage in reverse)
-				if (FirstSectionIdx == LastSectionIdx)
+				// We have to deal with montage explicitly because we can have multiple sections and we want to reset the position when the section loops
+				// and depending on the composition, CurrentTime < PreviousTime (or CurrentTime > PreviousTime when playing in reverse) is not enough
+				if (const UAnimMontage* Montage = Cast<UAnimMontage>(PreviewInstance->CurrentAsset))
 				{
-					bLooped = PreviewInstance->IsReverse() ? (CurrentTime > PreviousTime) : (CurrentTime < PreviousTime);
-				}
-				// Otherwise, we are previewing a montage with multiple section. In this case we check if section at CurrentTime is the FirstSection and the section at PreviewTime is the LastSection (or the opposite if we are playing the montage in reverse)
-				else
-				{
-					const int32 SectionIndexPrevTime = Montage->GetSectionIndexFromPosition(PreviousTime);
-					const int32 SectionIndexCurrentTime = Montage->GetSectionIndexFromPosition(CurrentTime);
-					bLooped = PreviewInstance->IsReverse() ? (SectionIndexPrevTime == FirstSectionIdx && SectionIndexCurrentTime == LastSectionIdx) : (SectionIndexPrevTime == LastSectionIdx && SectionIndexCurrentTime == FirstSectionIdx);
-				}
+					const int32 PreviewStartSectionIdx = Montage->CompositeSections.IsValidIndex(PreviewInstance->MontagePreviewStartSectionIdx) ? PreviewInstance->MontagePreviewStartSectionIdx : Montage->GetSectionIndexFromPosition(CurrentTime);
+					const int32 FirstSectionIdx = PreviewInstance->MontagePreview_FindFirstSectionAsInMontage(PreviewStartSectionIdx);
+					const int32 LastSectionIdx = PreviewInstance->MontagePreview_FindLastSection(FirstSectionIdx);
 
-				// If we have looped...
-				if (bLooped)
+					// If FirstSection == LastSection we are previewing a single section
+					// In this case to know if we have looped we just need to check if CurrentTime < PreviousTime (or the oposite if we are playing the montage in reverse)
+					if (FirstSectionIdx == LastSectionIdx)
+					{
+						bLooped = PreviewInstance->IsReverse() ? (CurrentTime > PreviousTime) : (CurrentTime < PreviousTime);
+					}
+					// Otherwise, we are previewing a montage with multiple section. In this case we check if section at CurrentTime is the FirstSection and the section at PreviewTime is the LastSection (or the opposite if we are playing the montage in reverse)
+					else
+					{
+						const int32 SectionIndexPrevTime = Montage->GetSectionIndexFromPosition(PreviousTime);
+						const int32 SectionIndexCurrentTime = Montage->GetSectionIndexFromPosition(CurrentTime);
+						bLooped = PreviewInstance->IsReverse() ? (SectionIndexPrevTime == FirstSectionIdx && SectionIndexCurrentTime == LastSectionIdx) : (SectionIndexPrevTime == LastSectionIdx && SectionIndexCurrentTime == FirstSectionIdx);
+					}
+
+					// If we have looped...
+					if (bLooped)
+					{
+						float StartTime = 0.0f, EndTime = 0.0f;
+						Montage->GetSectionStartAndEndTime(LastSectionIdx, StartTime, EndTime);
+						SectionStartPosition = StartTime;
+					}
+				}
+				else // CurrentAsset is not a Montage
 				{
-					float StartTime = 0.0f, EndTime = 0.0f;
-					Montage->GetSectionStartAndEndTime(LastSectionIdx, StartTime, EndTime);
-					SectionStartPosition = StartTime;
+					bLooped = PreviewInstance->IsReverse() ? (CurrentTime > PreviousTime) : (CurrentTime < PreviousTime);
 				}
-			}
-			else // CurrentAsset is not a Montage
-			{
-				bLooped = PreviewInstance->IsReverse() ? (CurrentTime > PreviousTime) : (CurrentTime < PreviousTime);
-			}
-			
-			// Loop Mode: Preview mesh will consume root motion continually
-			if (ProcessRootMotionMode == EProcessRootMotionMode::Loop)
-			{
-				const FTransform RootMotion = ExtractedRootMotion.GetRootMotionTransform();
-				RootMotionTransform = RootMotion * RootMotionTransform;
-
-				//Handle moving component so that it stays within the editor floor
-				FVector Trans = RootMotionTransform.GetLocation();
-				Trans.X = WrapInRange(Trans.X, FloorMin.X, FloorMax.X);
-				Trans.Y = WrapInRange(Trans.Y, FloorMin.Y, FloorMax.Y);
-				const FVector WrapDelta = Trans - RootMotionTransform.GetTranslation();
-				RootMotionTransform.SetTranslation(Trans);
 				
-				// If the location wraps, move the root motion reference too.
-				if (!WrapDelta.IsNearlyZero())
+				// Loop Mode: Preview mesh will consume root motion continually
+				if (ProcessRootMotionMode == EProcessRootMotionMode::Loop)
 				{
-					RootMotionReferenceTransform.AddToTranslation(WrapDelta);
+					FTransform RootMotionTransform = GetRelativeTransform();
+					const FTransform RootMotionDelta = ExtractedRootMotion.GetRootMotionTransform();
+					RootMotionTransform = RootMotionDelta * RootMotionTransform;
+
+					//Handle moving component so that it stays within the editor floor
+					FVector Trans = RootMotionTransform.GetLocation();
+					Trans.X = WrapInRange(Trans.X, FloorMin.X, FloorMax.X);
+					Trans.Y = WrapInRange(Trans.Y, FloorMin.Y, FloorMax.Y);
+					const FVector WrapDelta = Trans - RootMotionTransform.GetTranslation();
+					RootMotionTransform.SetTranslation(Trans);
+
+					if (!WrapDelta.IsNearlyZero())
+					{
+						SetRelativeTransform(RootMotionTransform);
+					}
+					else
+					{
+						AddLocalTransform(RootMotionDelta);
+					}
+					
+					// Looping resets the root motion reference transform.
+					if (bLooped)
+					{
+						RootMotionReferenceTransform = RootMotionTransform;
+					}
 				}
-
-				if (bLooped)
+				// Loop and Reset Mode: Preview mesh will consume root motion resetting the position back to the origin every time the animation loops
+				else if (ProcessRootMotionMode == EProcessRootMotionMode::LoopAndReset)
 				{
-					RootMotionReferenceTransform = RootMotionTransform;
+					if (bLooped)
+					{
+						const FTransform InitialTransform = UE::Anim::Private::CalculateInitialTransformFromAssetAndTime(*this, PreviewInstance->CurrentAsset, SectionStartPosition);
+						const FTransform RootMotionDelta = UE::Anim::ExtractRootMotionFromAnimationAsset(PreviewInstance->CurrentAsset, PreviewInstance->GetMirrorDataTable(), SectionStartPosition, CurrentTime);
+						const FTransform RootMotionTransform = RootMotionDelta * InitialTransform;
+
+						// Reference transform is always relative to the beginning of the animation sequence. 
+						RootMotionReferenceTransform = UE::Anim::ExtractRootTransformFromAnimationAsset(PreviewInstance->CurrentAsset, 0.0f);
+						
+						SetRelativeTransform(RootMotionTransform);
+					}
+					else
+					{
+						const FTransform RootMotionDelta = ExtractedRootMotion.GetRootMotionTransform();
+						AddLocalTransform(RootMotionDelta);
+					}
 				}
 			}
-			// Loop and Reset Mode: Preview mesh will consume root motion resetting the position back to the origin every time the animation loops
-			else if (ProcessRootMotionMode == EProcessRootMotionMode::LoopAndReset)
+			else // Not Playing. When not playing user can still scrub the time line but animation is not ticking so we have to extract and apply root motion manually
 			{
-				if (bLooped)
-				{
-					const FTransform InitialTransform = UE::Anim::Private::CalculateInitialTransformFromAssetAndTime(*this, PreviewInstance->CurrentAsset, SectionStartPosition);
-					const FTransform RootMotionDelta = UE::Anim::ExtractRootMotionFromAnimationAsset(PreviewInstance->CurrentAsset, PreviewInstance->GetMirrorDataTable(), SectionStartPosition, CurrentTime);
-					RootMotionTransform = RootMotionDelta * InitialTransform;
-
-					// Reference transform is always relative to the beginning of the animation sequence. 
-					RootMotionReferenceTransform = UE::Anim::ExtractRootTransformFromAnimationAsset(PreviewInstance->CurrentAsset, 0.0f);
-				}
-				else
-				{
-					const FTransform RootMotion = ExtractedRootMotion.GetRootMotionTransform();
-					RootMotionTransform = RootMotion * RootMotionTransform;
-				}
+				const FTransform RootMotionDelta = UE::Anim::ExtractRootMotionFromAnimationAsset(PreviewInstance->CurrentAsset, PreviewInstance->GetMirrorDataTable(), PreviousTime, CurrentTime);
+				AddLocalTransform(RootMotionDelta);
 			}
 		}
-		else // Not Playing. When not playing user can still scrub the time line but animation is not ticking so we have to extract and apply root motion manually
+		else
 		{
-			const FTransform RootMotion = UE::Anim::ExtractRootMotionFromAnimationAsset(PreviewInstance->CurrentAsset, PreviewInstance->GetMirrorDataTable(), PreviousTime, CurrentTime);
-			RootMotionTransform = RootMotion * RootMotionTransform;
+			// No root motion, reset to identity for consistency.
+			const FTransform RootMotionTransform = GetRelativeTransform();
+			if (!RootMotionTransform.Equals(FTransform::Identity))
+			{
+				SetRelativeTransform(FTransform::Identity);
+			}
 		}
 	}
-	else
-	{
-		RootMotionTransform = FTransform::Identity;
-	}
-
-	SetRelativeTransform(RootMotionTransform * TurntableTransform, /*bSweep*/false, /*OutSweepResult*/nullptr, bLooped ? ETeleportType::ResetPhysics : ETeleportType::None);
 }
 
 bool UDebugSkelMeshComponent::IsProcessingRootMotion() const 
@@ -428,6 +451,13 @@ void UDebugSkelMeshComponent::SetProcessRootMotionModeInternal(EProcessRootMotio
 {
 	ProcessRootMotionMode = Mode;
 
+	if (!DoesCurrentAssetHaveRootMotion())
+	{
+		return;
+	}
+
+	FTransform RootMotionTransform = FTransform::Identity;
+	
 	if (ProcessRootMotionMode == EProcessRootMotionMode::LoopAndReset || ProcessRootMotionMode == EProcessRootMotionMode::Loop)
 	{
 		// Reset transform
@@ -456,7 +486,7 @@ void UDebugSkelMeshComponent::SetProcessRootMotionModeInternal(EProcessRootMotio
 		RootMotionReferenceTransform = FTransform::Identity;
 	}
 
-	SetRelativeTransform(RootMotionTransform * TurntableTransform, /*bSweep*/false, /*OutSweepResult*/nullptr, ETeleportType::ResetPhysics);
+	SetRelativeTransform(RootMotionTransform);
 }
 
 bool UDebugSkelMeshComponent::IsTrackingAttachedLOD() const
@@ -697,6 +727,11 @@ void UDebugSkelMeshComponent::PostInitMeshObject(FSkeletalMeshObject* InMeshObje
 
 void UDebugSkelMeshComponent::OnMirrorDataTableChanged()
 {
+	if (!DoesCurrentAssetHaveRootMotion())
+	{
+		return;
+	}
+
 	if (ProcessRootMotionMode == EProcessRootMotionMode::LoopAndReset)
 	{
 		// Reset transform
@@ -711,15 +746,15 @@ void UDebugSkelMeshComponent::OnMirrorDataTableChanged()
 			Montage->GetSectionStartAndEndTime(LastSectionIdx, StartTime, EndTime);
 			SectionStartPosition = StartTime;
 		}
-	
+
 		const FTransform InitialTransform = UE::Anim::Private::CalculateInitialTransformFromAssetAndTime(*this, PreviewInstance->CurrentAsset, SectionStartPosition);
 		const FTransform RootMotionDelta = UE::Anim::ExtractRootMotionFromAnimationAsset(PreviewInstance->CurrentAsset, PreviewInstance->GetMirrorDataTable(), SectionStartPosition, CurrentTime);
-		RootMotionTransform = RootMotionDelta * InitialTransform;
+		const FTransform RootMotionTransform = RootMotionDelta * InitialTransform;
 
 		// Reference transform is always relative to the beginning of the animation sequence. 
-		RootMotionReferenceTransform = UE::Anim::ExtractRootTransformFromAnimationAsset(PreviewInstance->CurrentAsset, 0.0f);;
+		RootMotionReferenceTransform = UE::Anim::ExtractRootTransformFromAnimationAsset(PreviewInstance->CurrentAsset, 0.0f);
 
-		SetRelativeTransform(RootMotionTransform * TurntableTransform, /*bSweep*/false, /*OutSweepResult*/nullptr, ETeleportType::ResetPhysics);
+		SetRelativeTransform(RootMotionTransform);
 	}
 }
 
@@ -1177,10 +1212,9 @@ void UDebugSkelMeshComponent::TickComponent(float DeltaTime, enum ELevelTick Tic
 		return;
 	}
 
-	// Update turn table transform. The transform is applied in ConsumeRootMotion().
 	if (TurnTableMode == EPersonaTurnTableMode::Playing)
 	{
-		FRotator Rotation = TurntableTransform.Rotator();
+		FRotator Rotation = GetRelativeTransform().Rotator();
 		// Take into account time dilation, so it doesn't affect turn table turn rate.
 		float CurrentTimeDilation = 1.0f;
 		if (UWorld* MyWorld = GetWorld())
@@ -1188,13 +1222,9 @@ void UDebugSkelMeshComponent::TickComponent(float DeltaTime, enum ELevelTick Tic
 			CurrentTimeDilation = MyWorld->GetWorldSettings()->GetEffectiveTimeDilation();
 		}
 		Rotation.Yaw += 36.f * TurnTableSpeedScaling * DeltaTime / FMath::Max(CurrentTimeDilation, KINDA_SMALL_NUMBER);
-		TurntableTransform.SetRotation(Rotation.Quaternion());
+		SetRelativeRotation(Rotation);
 	}
-	else if (TurnTableMode == EPersonaTurnTableMode::Stopped)
-	{
-		TurntableTransform = FTransform::Identity;
-	}
-
+	
     // Brute force approach to ensure that when materials are changed the names are cached parameter names are updated 
 	bCachedMaterialParameterIndicesAreDirty = true;
 	
