Fix ordering of copy when reinstantiating, going from the leaves to the root. That way CPFUO will not reset subobject still pointing to the old object chain before the subobject itself had the chance to fix them.

diff --git a/Engine/Source/Editor/Kismet/Private/BlueprintCompilationManager.cpp b/Engine/Source/Editor/Kismet/Private/BlueprintCompilationManager.cpp
--- a/Engine/Source/Editor/Kismet/Private/BlueprintCompilationManager.cpp
+++ b/Engine/Source/Editor/Kismet/Private/BlueprintCompilationManager.cpp
@@ -2555,7 +2555,8 @@ void FBlueprintCompilationManagerImpl::ReinstanceBatch(TArray<FReinstancingJob>&
 				}
 
 				TMap<UObject*, UObject*> CreatedInstanceMap;
-				FBlueprintCompileReinstancer::PreCreateSubObjectsForReinstantiation(InOutOldToNewClassMap, OldCDO, NewCDO, CreatedInstanceMap);
+				TArray< TTuple<UObject*, UObject*>> OrderedListOfObjectToCopy;
+				FBlueprintCompileReinstancer::PreCreateSubObjectsForReinstantiation(InOutOldToNewClassMap, OldCDO, NewCDO, CreatedInstanceMap, nullptr, &OrderedListOfObjectToCopy);
 
 				// We only need to copy properties of the pre-created instances, the rest of the default sub object is done inside the UEditorEngine::CopyPropertiesForUnrelatedObjects
 				TMap<UObject*, UObject*> OldToNewInstanceMap(CreatedInstanceMap);
@@ -2563,7 +2564,7 @@ void FBlueprintCompilationManagerImpl::ReinstanceBatch(TArray<FReinstancingJob>&
 				{
 					OldToNewInstanceMap.Append(*OldToNewTemplateMap);
 				}
-				for(const auto& Pair : CreatedInstanceMap)
+				for(const auto& Pair : OrderedListOfObjectToCopy)
 				{
 					FBlueprintCompileReinstancer::CopyPropertiesForUnrelatedObjects(Pair.Key, Pair.Value, /*bClearExternalReferences*/true, bUseDeltaSerialization, /*bOnlyHandleDirectSubObjects*/true, &OldToNewInstanceMap, &InOutOldToNewClassMap);
 				}
@@ -2808,7 +2809,8 @@ void FBlueprintCompilationManagerImpl::ReinstanceBatch(TArray<FReinstancingJob>&
 					const bool bUseDeltaSerialization = ReinstancingJob.Reinstancer.IsValid() ? ReinstancingJob.Reinstancer->bUseDeltaSerializationToCopyProperties : false;
 
 					TMap<UObject*, UObject*> CreatedInstanceMap;
-					FBlueprintCompileReinstancer::PreCreateSubObjectsForReinstantiation(InOutOldToNewClassMap, OldInstance, NewArchetype, CreatedInstanceMap);
+					TArray< TTuple<UObject*, UObject*>> OrderedListOfObjectToCopy;
+					FBlueprintCompileReinstancer::PreCreateSubObjectsForReinstantiation(InOutOldToNewClassMap, OldInstance, NewArchetype, CreatedInstanceMap, nullptr, &OrderedListOfObjectToCopy);
 
 					// We only need to copy properties of the pre-created instances, the rest of the default sub object is done inside the UEditorEngine::CopyPropertiesForUnrelatedObjects
 					TMap<UObject*, UObject*> OldToNewInstanceMap(CreatedInstanceMap);
@@ -2816,7 +2818,7 @@ void FBlueprintCompilationManagerImpl::ReinstanceBatch(TArray<FReinstancingJob>&
 					{
 						OldToNewInstanceMap.Append(*OldToNewTemplateMap);
 					}
-					for (const auto& Pair : CreatedInstanceMap)
+					for (const auto& Pair : OrderedListOfObjectToCopy)
 					{
 						FBlueprintCompileReinstancer::CopyPropertiesForUnrelatedObjects(Pair.Key, Pair.Value, /*bClearExternalReferences*/true, bUseDeltaSerialization, /*bOnlyHandleDirectSubObjects*/true, &OldToNewInstanceMap, &InOutOldToNewClassMap);
 					}
diff --git a/Engine/Source/Editor/UnrealEd/Private/Kismet2/BlueprintEditorUtils.cpp b/Engine/Source/Editor/UnrealEd/Private/Kismet2/BlueprintEditorUtils.cpp
--- a/Engine/Source/Editor/UnrealEd/Private/Kismet2/BlueprintEditorUtils.cpp
+++ b/Engine/Source/Editor/UnrealEd/Private/Kismet2/BlueprintEditorUtils.cpp
@@ -1783,7 +1783,8 @@ void FBlueprintEditorUtils::PostDuplicateBlueprint(UBlueprint* Blueprint, bool b
 
 			TMap<UClass*, UClass*> InOutOldToNewClassMap;
 			TMap<UObject*, UObject*> CreatedInstanceMap;
-			FBlueprintCompileReinstancer::PreCreateSubObjectsForReinstantiation(InOutOldToNewClassMap, OldCDO, NewCDO, CreatedInstanceMap);
+			TArray< TTuple<UObject*, UObject*>> OrderedListOfObjectToCopy;
+			FBlueprintCompileReinstancer::PreCreateSubObjectsForReinstantiation(InOutOldToNewClassMap, OldCDO, NewCDO, CreatedInstanceMap, nullptr, &OrderedListOfObjectToCopy);
 
 			// We only need to copy properties of the pre-created instances, the rest of the default sub object is done inside the UEditorEngine::CopyPropertiesForUnrelatedObjects
 			TMap<UObject*, UObject*> OldToNewInstanceMap(CreatedInstanceMap);
@@ -1791,7 +1792,7 @@ void FBlueprintEditorUtils::PostDuplicateBlueprint(UBlueprint* Blueprint, bool b
 			Params.OptionalReplacementMappings = &OldToNewInstanceMap;
 			Params.bOnlyHandleDirectSubObjects = true;
 			Params.bReplaceInternalReferenceUponRead = true;
-			for (const auto& Pair : CreatedInstanceMap)
+			for (const auto& Pair : OrderedListOfObjectToCopy)
 			{
 				UEditorEngine::CopyPropertiesForUnrelatedObjects(Pair.Key, Pair.Value, Params);
 			}
diff --git a/Engine/Source/Editor/UnrealEd/Private/Kismet2/KismetReinstanceUtilities.cpp b/Engine/Source/Editor/UnrealEd/Private/Kismet2/KismetReinstanceUtilities.cpp
--- a/Engine/Source/Editor/UnrealEd/Private/Kismet2/KismetReinstanceUtilities.cpp
+++ b/Engine/Source/Editor/UnrealEd/Private/Kismet2/KismetReinstanceUtilities.cpp
@@ -2271,7 +2271,8 @@ static void ReplaceObjectHelper(UObject*& OldObject, UClass* OldClass, UObject*&
 	NewUObject->SetFlags(OldFlags & UE::ReinstanceUtils::FlagMask);
 
 	TMap<UObject*, UObject*> CreatedInstanceMap;
-	FBlueprintCompileReinstancer::PreCreateSubObjectsForReinstantiation(OldToNewClassMap, OldObject, NewUObject, CreatedInstanceMap, &OldToNewInstanceMap);
+	TArray< TTuple<UObject*, UObject*>> OrderedListOfObjectToCopy;
+	FBlueprintCompileReinstancer::PreCreateSubObjectsForReinstantiation(OldToNewClassMap, OldObject, NewUObject, CreatedInstanceMap, &OldToNewInstanceMap, &OrderedListOfObjectToCopy);
 	OldToNewInstanceMap.Append(CreatedInstanceMap);
 
 	// Copy property values
@@ -2293,7 +2294,7 @@ static void ReplaceObjectHelper(UObject*& OldObject, UClass* OldClass, UObject*&
 		Options.bDoDelta = false;
 	}
 	// We only need to copy properties of the pre-created instances, the rest of the default sub object is done inside the UEditorEngine::CopyPropertiesForUnrelatedObjects
-	for (const auto& Pair : CreatedInstanceMap)
+	for (const auto& Pair : OrderedListOfObjectToCopy)
 	{
 		UEditorEngine::CopyPropertiesForUnrelatedObjects(Pair.Key, Pair.Value, Options);
 	}
@@ -3187,17 +3188,33 @@ void FBlueprintCompileReinstancer::CopyPropertiesForUnrelatedObjects(UObject* Ol
 	UEngine::CopyPropertiesForUnrelatedObjects(OldObject, NewObject, Params);
 }
 
-void FBlueprintCompileReinstancer::PreCreateSubObjectsForReinstantiation(const TMap<UClass*, UClass*>& OldToNewClassMap, UObject* OldObject, UObject* NewUObject, TMap<UObject*, UObject*>& CreatedInstanceMap, const TMap<UObject*, UObject*>* OldToNewInstanceMap/* = nullptr*/)
+void FBlueprintCompileReinstancer::PreCreateSubObjectsForReinstantiation(
+	const TMap<UClass*, UClass*>& OldToNewClassMap, 
+	UObject* OldObject, UObject* NewUObject, 
+	TMap<UObject*, UObject*>& CreatedInstanceMap, 
+	const TMap<UObject*, UObject*>* OldToNewInstanceMap/* = nullptr*/, 
+	TArray< TTuple<UObject*, UObject*>>* OrderedListOfObjectToCopy /*= nullptr*/)
 {
 	TSet<UObject*> OldInstancedSubObjects;
 	FReplaceReferenceHelper::GetOwnedSubobjectsRecursive(OldObject, OldInstancedSubObjects);
 
 	// Add the mapping from the old to the new object exists...
 	CreatedInstanceMap.Add(OldObject, NewUObject);
-	PreCreateSubObjectsForReinstantiation_Inner(OldInstancedSubObjects, OldToNewClassMap, OldObject, NewUObject, CreatedInstanceMap, OldToNewInstanceMap);
+	PreCreateSubObjectsForReinstantiation_Inner(OldInstancedSubObjects, OldToNewClassMap, OldObject, NewUObject, CreatedInstanceMap, OldToNewInstanceMap, OrderedListOfObjectToCopy);
+	if (OrderedListOfObjectToCopy)
+	{
+		// Post add for deep first order
+		OrderedListOfObjectToCopy->Add({OldObject, NewUObject});
+	}
 }
 
-void FBlueprintCompileReinstancer::PreCreateSubObjectsForReinstantiation_Inner(const TSet<UObject*>& OldInstancedSubObjects, const TMap<UClass*, UClass*>& OldToNewClassMap, UObject* OldObject, UObject* NewUObject, TMap<UObject*, UObject*>& CreatedInstanceMap, const TMap<UObject*, UObject*>* OldToNewInstanceMap/* = nullptr*/)
+void FBlueprintCompileReinstancer::PreCreateSubObjectsForReinstantiation_Inner(
+	const TSet<UObject*>& OldInstancedSubObjects, 
+	const TMap<UClass*, UClass*>& OldToNewClassMap, 
+	UObject* OldObject, UObject* NewUObject, 
+	TMap<UObject*, UObject*>& CreatedInstanceMap, 
+	const TMap<UObject*, UObject*>* OldToNewInstanceMap,
+	TArray< TTuple<UObject*, UObject*>>* OrderedListOfObjectToCopy)
 {
 	// Gather subobjects on old object and pre-create them if needed
 	TArray<UObject*> ContainedOldSubObjects;
@@ -3266,14 +3283,19 @@ void FBlueprintCompileReinstancer::PreCreateSubObjectsForReinstantiation_Inner(c
 						NewSubObject = NewObject<UObject>(NewUObject, SubObjectClass, SubObjectName, SubObjectFlags);
 					}
 					CreatedInstanceMap.Add(OldSubObject, NewSubObject);
-					PreCreateSubObjectsForReinstantiation_Inner(OldInstancedSubObjects, OldToNewClassMap, OldSubObject, NewSubObject, CreatedInstanceMap, OldToNewInstanceMap);
+					PreCreateSubObjectsForReinstantiation_Inner(OldInstancedSubObjects, OldToNewClassMap, OldSubObject, NewSubObject, CreatedInstanceMap, OldToNewInstanceMap, OrderedListOfObjectToCopy);
+					if (OrderedListOfObjectToCopy)
+					{
+						// Post add for deep first order
+						OrderedListOfObjectToCopy->Add({OldSubObject, NewSubObject});
+					}
 				}
 			}
 		}
 		// There might be new subobjects attached to the sub object that are particular to this instance, let's traverse it to find them out.
 		else if (UObject** NewSubObject = ContainedNewSubObjects.FindByPredicate([SubObjectName](UObject* SubObject) { return SubObject && SubObject->GetFName() == SubObjectName; }))
 		{
-			PreCreateSubObjectsForReinstantiation_Inner(OldInstancedSubObjects, OldToNewClassMap, OldSubObject, *NewSubObject, CreatedInstanceMap, OldToNewInstanceMap);
+			PreCreateSubObjectsForReinstantiation_Inner(OldInstancedSubObjects, OldToNewClassMap, OldSubObject, *NewSubObject, CreatedInstanceMap, OldToNewInstanceMap, OrderedListOfObjectToCopy);
 		}
 	}
 }
diff --git a/Engine/Source/Editor/UnrealEd/Public/Kismet2/KismetReinstanceUtilities.h b/Engine/Source/Editor/UnrealEd/Public/Kismet2/KismetReinstanceUtilities.h
--- a/Engine/Source/Editor/UnrealEd/Public/Kismet2/KismetReinstanceUtilities.h
+++ b/Engine/Source/Editor/UnrealEd/Public/Kismet2/KismetReinstanceUtilities.h
@@ -318,12 +318,13 @@ public:
 	 * @param NewUObject where to store those pre-created sub objects
 	 * @param CreatedInstanceMap in/out of the result of all of the pre-created objects
 	 * @param OldToNewInstanceMap optional parameter of the possible re-instanced sub objects if any
+	 * @param OrderedListOfObjectToCopy List of object to copy in order
 	 */
-	static UNREALED_API void PreCreateSubObjectsForReinstantiation(const TMap<UClass*, UClass*>& OldToNewClassMap, UObject* OldObject, UObject* NewUObject, TMap<UObject*, UObject*>& CreatedInstanceMap, const TMap<UObject*, UObject*>* OldToNewInstanceMap = nullptr);
+	static UNREALED_API void PreCreateSubObjectsForReinstantiation(const TMap<UClass*, UClass*>& OldToNewClassMap, UObject* OldObject, UObject* NewUObject, TMap<UObject*, UObject*>& CreatedInstanceMap, const TMap<UObject*, UObject*>* OldToNewInstanceMap = nullptr, TArray< TTuple<UObject*, UObject*>>* OrderedListOfObjectToCopy = nullptr);
 
 private:
 	/** Handles the sub object pre-creation recursively */
-	static UNREALED_API void PreCreateSubObjectsForReinstantiation_Inner(const TSet<UObject*>& OldInstancedSubObjects, const TMap<UClass*, UClass*>& OldToNewClassMap, UObject* OldObject, UObject* NewUObject, TMap<UObject*, UObject*>& CreatedInstanceMap, const TMap<UObject*, UObject*>* OldToNewInstanceMap);
+	static UNREALED_API void PreCreateSubObjectsForReinstantiation_Inner(const TSet<UObject*>& OldInstancedSubObjects, const TMap<UClass*, UClass*>& OldToNewClassMap, UObject* OldObject, UObject* NewUObject, TMap<UObject*, UObject*>& CreatedInstanceMap, const TMap<UObject*, UObject*>* OldToNewInstanceMap, TArray< TTuple<UObject*, UObject*>>* OrderedListOfObjectToCopy);
 
 	/** Handles the work of ReplaceInstancesOfClass, handling both normal replacement of instances and batch */
 	static UNREALED_API void ReplaceInstancesOfClass_Inner(const TMap<UClass*, UClass*>& InOldToNewClassMap, const FReplaceInstancesOfClassParameters& Params);
