Fix CIS build

#rnx

diff --git a/Engine/Source/Editor/UnrealEd/Private/FbxMeshUtils.cpp b/Engine/Source/Editor/UnrealEd/Private/FbxMeshUtils.cpp
--- a/Engine/Source/Editor/UnrealEd/Private/FbxMeshUtils.cpp
+++ b/Engine/Source/Editor/UnrealEd/Private/FbxMeshUtils.cpp
@@ -508,7 +508,7 @@ namespace FbxMeshUtils
 							if (bResult)
 							{
 								//If we use alternate skinweight, we must re-import all profile for this LOD
-								if(!SelectedSkelMesh->GetSkinWeightProfiles().IsEmpty())
+								if(SelectedSkelMesh && !SelectedSkelMesh->GetSkinWeightProfiles().IsEmpty())
 								{
 									//Enqueue the re-import alternate skinning
 									TSharedPtr<FInterchangeSkeletalMeshAlternateSkinWeightPostImportTask> SkeletalMeshPostImportTask = MakeShared<FInterchangeSkeletalMeshAlternateSkinWeightPostImportTask>(SelectedSkelMesh);
@@ -996,7 +996,7 @@ namespace FbxMeshUtils
 					if (bResult)
 					{
 						//If we use alternate skinweight, we must re-import all profile for this LOD
-						if (!SkeletalMesh->GetSkinWeightProfiles().IsEmpty())
+						if (SkeletalMesh && !SkeletalMesh->GetSkinWeightProfiles().IsEmpty())
 						{
 							//Enqueue the re-import alternate skinning
 							TSharedPtr<FInterchangeSkeletalMeshAlternateSkinWeightPostImportTask> SkeletalMeshPostImportTask = MakeShared<FInterchangeSkeletalMeshAlternateSkinWeightPostImportTask>(SkeletalMesh);
