Performance optimizations for FBX animation import.

diff --git a/Engine/Source/Editor/UnrealEd/Private/SkeletalMeshEdit.cpp b/Engine/Source/Editor/UnrealEd/Private/SkeletalMeshEdit.cpp
--- a/Engine/Source/Editor/UnrealEd/Private/SkeletalMeshEdit.cpp
+++ b/Engine/Source/Editor/UnrealEd/Private/SkeletalMeshEdit.cpp
@@ -1275,94 +1275,6 @@ bool UnFbx::FFbxImporter::ImportCurve(const FbxAnimCurve* FbxCurve, FRichCurve&
 	return false;
 }
 
-/** This is to debug FBX importing animation. It saves source data and compare with what we use internally, so that it does detect earlier to find out there is transform issue
- *	We don't support skew(shearing), so if you have animation that has shearing(skew), this won't be preserved. Instead it will try convert to our format, which will visually look wrong. 
- *	If you have shearing(skew), please use "Preserve Local Transform" option, but it won't preserve its original animated transform */
-namespace AnimationTransformDebug
-{
-	// Data sturctutre to debug bone transform of animation issues
-	struct FAnimationTransformDebugData
-	{
-		int32 BoneIndex;
-		FName BoneName;
-		TArray<FTransform>	RecalculatedLocalTransform;
-		// this is used to calculate for intermediate result, not the source parent global transform
-		TArray<FTransform>	RecalculatedParentTransform;
-
-		//source data to convert from
-		TArray<FTransform>	SourceGlobalTransform;
-		TArray<FTransform>	SourceParentGlobalTransform;
-
-		FAnimationTransformDebugData()
-			: BoneIndex(INDEX_NONE), BoneName(NAME_None)
-		{}
-
-		void SetTrackData( int32 InBoneIndex, FName InBoneName)
-		{
-			BoneIndex = InBoneIndex;
-			BoneName = InBoneName;
-		}
-	};
-
-	void OutputAnimationTransformDebugData(TArray<AnimationTransformDebug::FAnimationTransformDebugData> &TransformDebugData, int32 TotalNumKeys, const FReferenceSkeleton& RefSkeleton)
-	{
-		bool bShouldOutputToMessageLog = true;
-
-		for(int32 Key=0; Key<TotalNumKeys; ++Key)
-		{
-			// go through all bones and find 
-			for(int32 BoneIndex=0; BoneIndex<TransformDebugData.Num(); ++BoneIndex)
-			{
-				FAnimationTransformDebugData& Data = TransformDebugData[BoneIndex];
-				int32 ParentIndex = RefSkeleton.GetParentIndex(Data.BoneIndex);
-				int32 ParentTransformDebugDataIndex = 0;
-
-				check(Data.RecalculatedLocalTransform.Num() == TotalNumKeys);
-				check(Data.SourceGlobalTransform.Num() == TotalNumKeys);
-				check(Data.SourceParentGlobalTransform.Num() == TotalNumKeys);
-
-				for(; ParentTransformDebugDataIndex<BoneIndex; ++ParentTransformDebugDataIndex)
-				{
-					if(ParentIndex == TransformDebugData[ParentTransformDebugDataIndex].BoneIndex)
-					{
-						FTransform ParentTransform = TransformDebugData[ParentTransformDebugDataIndex].RecalculatedLocalTransform[Key] * TransformDebugData[ParentTransformDebugDataIndex].RecalculatedParentTransform[Key];
-						Data.RecalculatedParentTransform.Add(ParentTransform);
-						break;
-					}
-				}
-
-				// did not find Parent
-				if(ParentTransformDebugDataIndex == BoneIndex)
-				{
-					Data.RecalculatedParentTransform.Add(FTransform::Identity);
-				}
-
-				check(Data.RecalculatedParentTransform.Num() == Key+1);
-
-				FTransform GlobalTransform = Data.RecalculatedLocalTransform[Key] * Data.RecalculatedParentTransform[Key];
-				// makes more generous on the threshold. 
-				if(GlobalTransform.Equals(Data.SourceGlobalTransform[Key], 0.1f) == false)
-				{
-					// so that we don't spawm with this message
-					if(bShouldOutputToMessageLog)
-					{
-						UnFbx::FFbxImporter* FFbxImporter = UnFbx::FFbxImporter::GetInstance();
-						// now print information - it doesn't match well, find out what it is
-						FFbxImporter->AddTokenizedErrorMessage(FTokenizedMessage::Create(EMessageSeverity::Warning, FText::Format(LOCTEXT("FBXImport_TransformError", "Imported bone transform is different from original. Please check Output Log to see detail of error. "),
-							FText::FromName(Data.BoneName), FText::AsNumber(Data.BoneIndex), FText::FromString(Data.SourceGlobalTransform[Key].ToString()), FText::FromString(GlobalTransform.ToString()))), FFbxErrors::Animation_TransformError);
-						
-						// now print information - it doesn't match well, find out what it is
-						UE_LOG(LogFbx, Warning, TEXT("IMPORT TRANSFORM ERROR : Bone (%s:%d) \r\nSource Global Transform (%s), \r\nConverted Global Transform (%s)"),
-							*Data.BoneName.ToString(), Data.BoneIndex, *Data.SourceGlobalTransform[Key].ToString(), *GlobalTransform.ToString());
-
-						bShouldOutputToMessageLog = false;
-					}
-				}
-			}
-		}
-	}
-}
-
 bool UnFbx::FFbxImporter::ImportCurveToAnimSequence(class UAnimSequence * TargetSequence, const FString& CurveName, const FbxAnimCurve* FbxCurve, int32 CurveFlags,const FbxTimeSpan& AnimTimeSpan, const bool bReimport, float ValueScale/*=1.f*/) const
 {
 	if (TargetSequence && FbxCurve)
@@ -1843,7 +1755,6 @@ bool UnFbx::FFbxImporter::ImportAnimation(USkeleton* Skeleton, UAnimSequence * D
 	}
 	// importing custom attribute END
 	
-	TArray<AnimationTransformDebug::FAnimationTransformDebugData> TransformDebugData;
 	int32 TotalNumKeys = 0;
 	const FReferenceSkeleton& RefSkeleton = Skeleton->GetReferenceSkeleton();
 
@@ -1851,9 +1762,7 @@ bool UnFbx::FFbxImporter::ImportAnimation(USkeleton* Skeleton, UAnimSequence * D
 	if (ImportOptions->bImportBoneTracks)
 	{
 		FbxNode* SkeletalMeshRootNode = NodeArray.Num() > 0 ? NodeArray[0] : nullptr;
-		ImportBoneTracks(Skeleton, AnimImportSettings, SkeletalMeshRootNode, ResampleRate, TransformDebugData, TotalNumKeys, bReimport);
-
-		AnimationTransformDebug::OutputAnimationTransformDebugData(TransformDebugData, TotalNumKeys, RefSkeleton);
+		ImportBoneTracks(Skeleton, AnimImportSettings, SkeletalMeshRootNode, ResampleRate, TotalNumKeys, bReimport);
 	}
 	else if (CurveAttributeKeyCount > 0)
 	{
@@ -2135,11 +2044,8 @@ void UnFbx::FFbxImporter::ImportAnimationCustomAttribute(FAnimCurveImportSetting
 	}
 }
 
-void UnFbx::FFbxImporter::ImportBoneTracks(USkeleton* Skeleton, FAnimCurveImportSettings& AnimImportSettings, FbxNode* SkeletalMeshRootNode, const int32 ResampleRate, TArray<AnimationTransformDebug::FAnimationTransformDebugData>& TransformDebugData, int32& OutTotalNumKeys, const bool bReimport)
+void UnFbx::FFbxImporter::ImportBoneTracks(USkeleton* Skeleton, FAnimCurveImportSettings& AnimImportSettings, FbxNode* SkeletalMeshRootNode, const int32 ResampleRate, int32& OutTotalNumKeys, const bool bReimport)
 {
-	FScopedSlowTask SlowTask(AnimImportSettings.FbxRawBoneNames.Num(), LOCTEXT("BeginImportAnimation", "Importing Animation"), true);
-	SlowTask.MakeDialog();
-
 	OutTotalNumKeys = 0;
 	UnFbx::FFbxImporter* FbxImporter = UnFbx::FFbxImporter::GetInstance();
 	const bool bPreserveLocalTransform = FbxImporter->GetImportOptions()->bPreserveLocalTransform;
@@ -2153,6 +2059,7 @@ void UnFbx::FFbxImporter::ImportBoneTracks(USkeleton* Skeleton, FAnimCurveImport
 	FbxAMatrix FbxAddedMatrix;
 	BuildFbxMatrixForImportTransform(FbxAddedMatrix, TemplateData);
 	FMatrix AddedMatrix = Converter.ConvertMatrix(FbxAddedMatrix);
+	FTransform AddedTransform(AddedMatrix);
 
 	bool bIsRigidMeshAnimation = false;
 	if (ImportOptions->bImportScene && AnimImportSettings.SortedLinks.Num() > 0)
@@ -2168,224 +2075,255 @@ void UnFbx::FFbxImporter::ImportBoneTracks(USkeleton* Skeleton, FAnimCurveImport
 		}
 	}
 
-	const int32 NumSamplingFrame = FMath::RoundToInt((AnimTimeSpan.GetDuration().GetSecondDouble() * ResampleRate));
-	//Set the time increment from the re-sample rate
+	// Prepare per-track data
+
+	const FReferenceSkeleton& RefSkeleton = Skeleton->GetReferenceSkeleton();
+
+	int32 SourceTrackNum = AnimImportSettings.FbxRawBoneNames.Num();
+	int32 KeyTotal = FMath::RoundToInt((AnimTimeSpan.GetDuration().GetSecondDouble() * ResampleRate)) + 1; // Expected Number of Keys
+
+	TArray<int32> BoneTreeIndices;
+	TArray<bool> SourceTrackValid;
+	TArray<FRawAnimSequenceTrack> RawTracks;
+	TArray<TArray<float>> TimeKeys;
+
+	BoneTreeIndices.SetNum(SourceTrackNum);
+	SourceTrackValid.SetNum(SourceTrackNum);
+	RawTracks.SetNum(SourceTrackNum);
+	TimeKeys.SetNum(SourceTrackNum);
+
+	for (int32 SourceTrackIdx = 0; SourceTrackIdx < SourceTrackNum; ++SourceTrackIdx)
+	{
+		BoneTreeIndices[SourceTrackIdx] = RefSkeleton.FindBoneIndex(AnimImportSettings.FbxRawBoneNames[SourceTrackIdx]);
+		SourceTrackValid[SourceTrackIdx] = BoneTreeIndices[SourceTrackIdx] != INDEX_NONE || IsUnrealTransformAttribute(AnimImportSettings.SortedLinks[SourceTrackIdx]);
+		RawTracks[SourceTrackIdx].PosKeys.Reserve(KeyTotal);
+		RawTracks[SourceTrackIdx].RotKeys.Reserve(KeyTotal);
+		RawTracks[SourceTrackIdx].ScaleKeys.Reserve(KeyTotal);
+		TimeKeys[SourceTrackIdx].Reserve(KeyTotal);
+	}
+
+	// Set the time increment from the re-sample rate
+
 	FbxTime TimeIncrement = 0;
 	TimeIncrement.SetSecondDouble(1.0 / ((double)(ResampleRate)));
 
-	//Add a threshold when we compare if we have reach the end of the animation
+	// Add a threshold when we compare if we have reach the end of the animation
 	const FbxTime TimeComparisonThreshold = (KINDA_SMALL_NUMBER * static_cast<float>(FBXSDK_TC_SECOND));
-	const FReferenceSkeleton& RefSkeleton = Skeleton->GetReferenceSkeleton();
 
-	IAnimationDataModel* DataModel = DestSeq->GetDataModel();
-	IAnimationDataController& Controller = DestSeq->GetController();	
-	Controller.SetFrameRate(FFrameRate(ResampleRate, 1), bShouldTransact);
+	// Loop Over Key Times
 
-	for (int32 SourceTrackIdx = 0; SourceTrackIdx < AnimImportSettings.FbxRawBoneNames.Num(); ++SourceTrackIdx)
-	{
-		int32 NumKeysForTrack = 0;
+	FScopedSlowTask SlowTask(KeyTotal, LOCTEXT("BeginImportAnimation", "Importing Animation"), true);
+	SlowTask.MakeDialog();
 
-		// see if it's found in Skeleton
-		FName BoneName = AnimImportSettings.FbxRawBoneNames[SourceTrackIdx];
-		int32 BoneTreeIndex = RefSkeleton.FindBoneIndex(BoneName);
+	int32 KeyNum = 0;
+	for (FbxTime CurTime = AnimTimeSpan.GetStart(); CurTime < (AnimTimeSpan.GetStop() + TimeComparisonThreshold); CurTime += TimeIncrement)
+	{
+		// Update Status
 
-		// update status
 		FFormatNamedArguments Args;
-		Args.Add(TEXT("TrackName"), FText::FromName(BoneName));
-		Args.Add(TEXT("TotalKey"), FText::AsNumber(NumSamplingFrame + 1)); //Key number is Frame + 1
-		Args.Add(TEXT("TrackIndex"), FText::AsNumber(SourceTrackIdx + 1));
-		Args.Add(TEXT("TotalTracks"), FText::AsNumber(AnimImportSettings.FbxRawBoneNames.Num()));
-		const FText StatusUpate = FText::Format(LOCTEXT("ImportingAnimTrackDetail", "Importing Animation Track [{TrackName}] ({TrackIndex}/{TotalTracks}) - TotalKey {TotalKey}"), Args);
-		SlowTask.EnterProgressFrame(1, StatusUpate);
+		Args.Add(TEXT("KeyNum"), FText::AsNumber(KeyNum + 1));
+		Args.Add(TEXT("KeyTotal"), FText::AsNumber(KeyTotal));
+		const FText StatusUpate = FText::Format(LOCTEXT("ImportingAnimTrackDetail", "Importing Animation Keys {KeyNum} of {KeyTotal}"), Args);
+		SlowTask.EnterProgressFrame(1.0f, StatusUpate);
 
-		if (BoneTreeIndex != INDEX_NONE || IsUnrealTransformAttribute(AnimImportSettings.SortedLinks[SourceTrackIdx]))
-		{
-			bool bSuccess = true;
+		// Loop Over Tracks
 
-			FRawAnimSequenceTrack RawTrack;
-			RawTrack.PosKeys.Empty();
-			RawTrack.RotKeys.Empty();
-			RawTrack.ScaleKeys.Empty();
-
-			TArray<float> TimeKeys;
+		for (int32 SourceTrackIdx = 0; SourceTrackIdx < SourceTrackNum; ++SourceTrackIdx)
+		{
+			// Skip Invalid Tracks
+			if (!SourceTrackValid[SourceTrackIdx])
+			{
+				continue;
+			}
 
-			AnimationTransformDebug::FAnimationTransformDebugData NewDebugData;
+			// Get Bone name and index.
+			FName BoneName = AnimImportSettings.FbxRawBoneNames[SourceTrackIdx];
+			int32 BoneTreeIndex = BoneTreeIndices[SourceTrackIdx];
 
+			// Get Links
 			FbxNode* Link = AnimImportSettings.SortedLinks[SourceTrackIdx];
 			FbxNode* LinkParent = Link->GetParent();
-			for (FbxTime CurTime = AnimTimeSpan.GetStart(); CurTime < (AnimTimeSpan.GetStop() + TimeComparisonThreshold); CurTime += TimeIncrement)
+
+			FTransform LocalTransform;
+			if (!bPreserveLocalTransform && LinkParent)
 			{
-				// save global trasnform
+				// get global transform
 				FbxAMatrix GlobalMatrix = Link->EvaluateGlobalTransform(CurTime) * FFbxDataConverter::GetJointPostConversionMatrix();
+
 				// we'd like to verify this before going to Transform. 
 				// currently transform has tons of NaN check, so it will crash there
 				FMatrix GlobalUEMatrix = Converter.ConvertMatrix(GlobalMatrix);
 				if (GlobalUEMatrix.ContainsNaN())
 				{
-					bSuccess = false;
+					SourceTrackValid[SourceTrackIdx] = false;
 					AddTokenizedErrorMessage(FTokenizedMessage::Create(EMessageSeverity::Error, FText::Format(LOCTEXT("Error_InvalidTransform", "Track {0} contains invalid transform. Could not import the track."), FText::FromName(BoneName))), FFbxErrors::Animation_TransformError);
-					break;
+					continue;
 				}
 
 				FTransform GlobalTransform = Converter.ConvertTransform(GlobalMatrix);
 				if (GlobalTransform.ContainsNaN())
 				{
-					bSuccess = false;
+					SourceTrackValid[SourceTrackIdx] = false;
 					AddTokenizedErrorMessage(FTokenizedMessage::Create(EMessageSeverity::Error, FText::Format(LOCTEXT("Error_InvalidUnrealTransform", "Track {0} has invalid transform(NaN). Zero scale transform can cause this issue."), FText::FromName(BoneName))), FFbxErrors::Animation_TransformError);
-					break;
+					continue;
 				}
 
-				// debug data, including import transformation
-				FTransform AddedTransform(AddedMatrix);
-				NewDebugData.SourceGlobalTransform.Add(GlobalTransform * AddedTransform);
-
-				FTransform LocalTransform;
-				if (!bPreserveLocalTransform && LinkParent)
+				// I can't rely on LocalMatrix. I need to recalculate quaternion/scale based on global transform if Parent exists
+				FbxAMatrix ParentGlobalMatrix = Link->GetParent()->EvaluateGlobalTransform(CurTime);
+				if (BoneTreeIndex != 0)
 				{
-					// I can't rely on LocalMatrix. I need to recalculate quaternion/scale based on global transform if Parent exists
-					FbxAMatrix ParentGlobalMatrix = Link->GetParent()->EvaluateGlobalTransform(CurTime);
-					if (BoneTreeIndex != 0)
-					{
-						ParentGlobalMatrix = ParentGlobalMatrix * FFbxDataConverter::GetJointPostConversionMatrix();
-					}
-					FTransform ParentGlobalTransform = Converter.ConvertTransform(ParentGlobalMatrix);
-					//In case we do a scene import we need to add the skeletal mesh root node matrix to the parent link.
-					if (ImportOptions->bImportScene && !ImportOptions->bTransformVertexToAbsolute && BoneTreeIndex == 0 && SkeletalMeshRootNode != nullptr)
-					{
-						//In the case of a rigidmesh animation we have to use the skeletalMeshRootNode position at zero since the mesh can be animate.
-						FbxAMatrix GlobalSkeletalNodeFbx = bIsRigidMeshAnimation ? SkeletalMeshRootNode->EvaluateGlobalTransform(0) : SkeletalMeshRootNode->EvaluateGlobalTransform(CurTime);
-						FTransform GlobalSkeletalNode = Converter.ConvertTransform(GlobalSkeletalNodeFbx);
-						ParentGlobalTransform = ParentGlobalTransform * GlobalSkeletalNode;
-					}
-
-					LocalTransform = GlobalTransform.GetRelativeTransform(ParentGlobalTransform);
-					NewDebugData.SourceParentGlobalTransform.Add(ParentGlobalTransform);
+					ParentGlobalMatrix = ParentGlobalMatrix * FFbxDataConverter::GetJointPostConversionMatrix();
 				}
-				else
+				FTransform ParentGlobalTransform = Converter.ConvertTransform(ParentGlobalMatrix);
+				//In case we do a scene import we need to add the skeletal mesh root node matrix to the parent link.
+				if (ImportOptions->bImportScene && !ImportOptions->bTransformVertexToAbsolute && BoneTreeIndex == 0 && SkeletalMeshRootNode != nullptr)
 				{
-					FbxAMatrix& LocalMatrix = Link->EvaluateLocalTransform(CurTime);
-					FbxVector4 NewLocalT = LocalMatrix.GetT();
-					FbxVector4 NewLocalS = LocalMatrix.GetS();
-					FbxQuaternion NewLocalQ = LocalMatrix.GetQ();
+					//In the case of a rigidmesh animation we have to use the skeletalMeshRootNode position at zero since the mesh can be animate.
+					FbxAMatrix GlobalSkeletalNodeFbx = bIsRigidMeshAnimation ? SkeletalMeshRootNode->EvaluateGlobalTransform(0) : SkeletalMeshRootNode->EvaluateGlobalTransform(CurTime);
+					FTransform GlobalSkeletalNode = Converter.ConvertTransform(GlobalSkeletalNodeFbx);
+					ParentGlobalTransform = ParentGlobalTransform * GlobalSkeletalNode;
+				}
 
-					LocalTransform.SetTranslation(Converter.ConvertPos(NewLocalT));
-					LocalTransform.SetScale3D(Converter.ConvertScale(NewLocalS));
-					LocalTransform.SetRotation(Converter.ConvertRotToQuat(NewLocalQ));
+				LocalTransform = GlobalTransform.GetRelativeTransform(ParentGlobalTransform);
+			}
+			else
+			{
+				FbxAMatrix& LocalMatrix = Link->EvaluateLocalTransform(CurTime);
+				FbxVector4 NewLocalT = LocalMatrix.GetT();
+				FbxVector4 NewLocalS = LocalMatrix.GetS();
+				FbxQuaternion NewLocalQ = LocalMatrix.GetQ();
+
+				LocalTransform.SetTranslation(Converter.ConvertPos(NewLocalT));
+				LocalTransform.SetScale3D(Converter.ConvertScale(NewLocalS));
+				LocalTransform.SetRotation(Converter.ConvertRotToQuat(NewLocalQ));
+			}
 
-					NewDebugData.SourceParentGlobalTransform.Add(FTransform::Identity);
-				}
+			if (TemplateData && BoneTreeIndex == 0)
+			{
+				// If we found template data earlier, apply the import transform matrix to
+				// the root track.
+				LocalTransform.SetFromMatrix(LocalTransform.ToMatrixWithScale() * AddedMatrix);
+			}
 
-				if (TemplateData && BoneTreeIndex == 0)
-				{
-					// If we found template data earlier, apply the import transform matrix to
-					// the root track.
-					LocalTransform.SetFromMatrix(LocalTransform.ToMatrixWithScale() * AddedMatrix);
-				}
+			if (LocalTransform.ContainsNaN())
+			{
+				SourceTrackValid[SourceTrackIdx] = false;
+				AddTokenizedErrorMessage(FTokenizedMessage::Create(EMessageSeverity::Error, FText::Format(LOCTEXT("Error_InvalidUnrealLocalTransform", "Track {0} has invalid transform(NaN). If you have zero scale transform, that can cause this."), FText::FromName(BoneName))), FFbxErrors::Animation_TransformError);
+				continue;
+			}
 
-				if (LocalTransform.ContainsNaN())
-				{
-					bSuccess = false;
-					AddTokenizedErrorMessage(FTokenizedMessage::Create(EMessageSeverity::Error, FText::Format(LOCTEXT("Error_InvalidUnrealLocalTransform", "Track {0} has invalid transform(NaN). If you have zero scale transform, that can cause this."), FText::FromName(BoneName))), FFbxErrors::Animation_TransformError);
-					break;
-				}
+			// Add data to RawTracks, TimeKeys, NewDebugDatas
 
-				RawTrack.ScaleKeys.Add(FVector3f(LocalTransform.GetScale3D()));
-				RawTrack.PosKeys.Add(FVector3f(LocalTransform.GetTranslation()));
-				RawTrack.RotKeys.Add(FQuat4f(LocalTransform.GetRotation()));
+			RawTracks[SourceTrackIdx].ScaleKeys.Add(FVector3f(LocalTransform.GetScale3D()));
+			RawTracks[SourceTrackIdx].PosKeys.Add(FVector3f(LocalTransform.GetTranslation()));
+			RawTracks[SourceTrackIdx].RotKeys.Add(FQuat4f(LocalTransform.GetRotation()));
+			TimeKeys[SourceTrackIdx].Add((CurTime - AnimTimeSpan.GetStart()).GetSecondDouble());
+		}
 
-				TimeKeys.Add((CurTime - AnimTimeSpan.GetStart()).GetSecondDouble());
+		// Update Key Num
 
-				NewDebugData.RecalculatedLocalTransform.Add(LocalTransform);
-				++NumKeysForTrack;
-			}
+		KeyNum++;
+	}
 
-			if (bSuccess)
-			{
-				check(RawTrack.ScaleKeys.Num() == NumKeysForTrack);
-				check(RawTrack.PosKeys.Num() == NumKeysForTrack);
-				check(RawTrack.RotKeys.Num() == NumKeysForTrack);
-				check(TimeKeys.Num() == NumKeysForTrack);
+	// Compute Total Number of Keys
 
-				if (BoneTreeIndex != INDEX_NONE)
-				{
-					//add new track
-					if (BoneName.GetStringLength() > 92)
-					{
-						//The bone name exceed the maximum length supported by the animation system
-						//The animation system is adding _CONTROL to the bone name to name the animation controller and
-						//the maximum total length is cap at 100, so user should not import bone name longer then 92 characters
-						AddTokenizedErrorMessage(FTokenizedMessage::Create(EMessageSeverity::Warning, FText::Format(LOCTEXT("Error_BoneNameExceed92Characters", "Bone with animation cannot have a name exceeding 92 characters: {0}"), FText::FromName(BoneName))), FFbxErrors::Animation_InvalidData);
-						continue;
-					}
-					Controller.AddBoneCurve(BoneName, bShouldTransact);
-					Controller.SetBoneTrackKeys(BoneName, RawTrack.PosKeys, RawTrack.RotKeys, RawTrack.ScaleKeys, bShouldTransact);
-					NewDebugData.SetTrackData(BoneTreeIndex, BoneName);
+	for (int32 SourceTrackIdx = 0; SourceTrackIdx < SourceTrackNum; ++SourceTrackIdx)
+	{
+		OutTotalNumKeys = FMath::Max(OutTotalNumKeys, TimeKeys[SourceTrackIdx].Num());
+	}
 
-					// add mapping to skeleton bone track
-					TransformDebugData.Add(NewDebugData);
-				}
-				else if (NumKeysForTrack > 0) // add transform attribute
-				{
-					FbxNode* TargetBoneLink = LinkParent;
-					while (TargetBoneLink != nullptr && !IsUnrealBone(TargetBoneLink))
-					{
-						TargetBoneLink = TargetBoneLink->GetParent();
-					}
+	// Put raw track data into animation
+	
+	IAnimationDataController& Controller = DestSeq->GetController();
+	Controller.SetFrameRate(FFrameRate(ResampleRate, 1), bShouldTransact);
+	
+	for (int32 SourceTrackIdx = 0; SourceTrackIdx < SourceTrackNum; ++SourceTrackIdx)
+	{
+		if (!SourceTrackValid[SourceTrackIdx])
+		{
+			continue;
+		}
 
-					if (TargetBoneLink)
+		FName BoneName = AnimImportSettings.FbxRawBoneNames[SourceTrackIdx];
+		int32 BoneTreeIndex = BoneTreeIndices[SourceTrackIdx];
+
+		FbxNode* Link = AnimImportSettings.SortedLinks[SourceTrackIdx];
+		FbxNode* LinkParent = Link->GetParent();
+
+		const int32 SourceTrackKeyNum = TimeKeys[SourceTrackIdx].Num();
+
+		if (BoneTreeIndex != INDEX_NONE)
+		{
+			//add new track
+			if (BoneName.GetStringLength() > 92)
+			{
+				//The bone name exceed the maximum length supported by the animation system
+				//The animation system is adding _CONTROL to the bone name to name the animation controller and
+				//the maximum total length is cap at 100, so user should not import bone name longer then 92 characters
+				AddTokenizedErrorMessage(FTokenizedMessage::Create(EMessageSeverity::Warning, FText::Format(LOCTEXT("Error_BoneNameExceed92Characters", "Bone with animation cannot have a name exceeding 92 characters: {0}"), FText::FromName(BoneName))), FFbxErrors::Animation_InvalidData);
+				continue;
+			}
+			Controller.AddBoneCurve(BoneName, bShouldTransact);
+			Controller.SetBoneTrackKeys(BoneName, RawTracks[SourceTrackIdx].PosKeys, RawTracks[SourceTrackIdx].RotKeys, RawTracks[SourceTrackIdx].ScaleKeys, bShouldTransact);
+		}
+		else if (SourceTrackKeyNum > 0) // add transform attribute
+		{
+			FbxNode* TargetBoneLink = LinkParent;
+			while (TargetBoneLink != nullptr && !IsUnrealBone(TargetBoneLink))
+			{
+				TargetBoneLink = TargetBoneLink->GetParent();
+			}
+
+			if (TargetBoneLink)
+			{
+				int32 TargetBoneTrackIndex = AnimImportSettings.SortedLinks.Find(TargetBoneLink);
+				if (TargetBoneTrackIndex != INDEX_NONE)
+				{
+					FName TargetBoneName = AnimImportSettings.FbxRawBoneNames[TargetBoneTrackIndex];
+					if (RefSkeleton.FindBoneIndex(TargetBoneName) != INDEX_NONE)
 					{
-						int32 TargetBoneTrackIndex = AnimImportSettings.SortedLinks.Find(TargetBoneLink);
-						if (TargetBoneTrackIndex != INDEX_NONE)
+						FAnimationAttributeIdentifier AttributeIdentifier = UAnimationAttributeIdentifierExtensions::CreateAttributeIdentifier(DestSeq, FName(BoneName), TargetBoneName, FTransformAnimationAttribute::StaticStruct());
+						if (AttributeIdentifier.IsValid())
 						{
-							FName TargetBoneName = AnimImportSettings.FbxRawBoneNames[TargetBoneTrackIndex];
-							if (RefSkeleton.FindBoneIndex(TargetBoneName) != INDEX_NONE)
+							// remove any existing attribute with the same identifier
+							if (const IAnimationDataModel* Model = Controller.GetModel())
 							{
-								FAnimationAttributeIdentifier AttributeIdentifier = UAnimationAttributeIdentifierExtensions::CreateAttributeIdentifier(DestSeq, FName(BoneName), TargetBoneName, FTransformAnimationAttribute::StaticStruct());
-								if (AttributeIdentifier.IsValid())
+								if (Model->FindAttribute(AttributeIdentifier))
 								{
-									// remove any existing attribute with the same identifier
-									if (const IAnimationDataModel* Model = Controller.GetModel())
-									{
-										if (Model->FindAttribute(AttributeIdentifier))
-										{
-											Controller.RemoveAttribute(AttributeIdentifier, bShouldTransact);
-										}
-									}
-
-									// pack the separate rot/pos/scale key arrays into a single array
-									TArray<FTransform> TransformValues;
-									TransformValues.Reserve(NumKeysForTrack);
-									for (int32 KeyIndex = 0; KeyIndex < NumKeysForTrack; ++KeyIndex)
-									{
-										const FQuat Q(RawTrack.RotKeys[KeyIndex]);
-										const FVector T(RawTrack.PosKeys[KeyIndex]);
-										const FVector S(RawTrack.ScaleKeys[KeyIndex]);
-										TransformValues.Add(FTransform(Q, T, S));
-									}
+									Controller.RemoveAttribute(AttributeIdentifier, bShouldTransact);
+								}
+							}
 
-									// reduce keys for the common case where all of the keys have the same values
-									bool bReduceKeys = true;
-									for (int32 KeyIndex = 1; KeyIndex < NumKeysForTrack; ++KeyIndex)
-									{
-										if (!TransformValues[KeyIndex].Equals(TransformValues[0]))
-										{
-											bReduceKeys = false;
-											break;
-										}
-									}
+							// pack the separate rot/pos/scale key arrays into a single array
+							TArray<FTransform> TransformValues;
+							TransformValues.Reserve(SourceTrackKeyNum);
+							for (int32 KeyIndex = 0; KeyIndex < SourceTrackKeyNum; ++KeyIndex)
+							{
+								const FQuat Q(RawTracks[SourceTrackIdx].RotKeys[KeyIndex]);
+								const FVector T(RawTracks[SourceTrackIdx].PosKeys[KeyIndex]);
+								const FVector S(RawTracks[SourceTrackIdx].ScaleKeys[KeyIndex]);
+								TransformValues.Add(FTransform(Q, T, S));
+							}
 
-									// create the attribute and add the transform keys
-									const int32 NumAttributeKeys = (bReduceKeys) ? 1 : NumKeysForTrack;
-									UE::Anim::AddTypedCustomAttribute<FTransformAnimationAttribute, FTransform>(FName(BoneName), TargetBoneName, DestSeq, MakeArrayView(TimeKeys.GetData(), NumAttributeKeys), MakeArrayView(TransformValues.GetData(), NumAttributeKeys), bShouldTransact);
+							// reduce keys for the common case where all of the keys have the same values
+							bool bReduceKeys = true;
+							for (int32 KeyIndex = 1; KeyIndex < SourceTrackKeyNum; ++KeyIndex)
+							{
+								if (!TransformValues[KeyIndex].Equals(TransformValues[0]))
+								{
+									bReduceKeys = false;
+									break;
 								}
 							}
+
+							// create the attribute and add the transform keys
+							const int32 NumAttributeKeys = (bReduceKeys) ? 1 : SourceTrackKeyNum;
+							UE::Anim::AddTypedCustomAttribute<FTransformAnimationAttribute, FTransform>(FName(BoneName), TargetBoneName, DestSeq, MakeArrayView(TimeKeys[SourceTrackIdx].GetData(), NumAttributeKeys), MakeArrayView(TransformValues.GetData(), NumAttributeKeys), bShouldTransact);
 						}
 					}
 				}
 			}
 		}
-
-		OutTotalNumKeys = FMath::Max(OutTotalNumKeys, NumKeysForTrack);
 	}
 }
 
diff --git a/Engine/Source/Editor/UnrealEd/Public/FbxImporter.h b/Engine/Source/Editor/UnrealEd/Public/FbxImporter.h
--- a/Engine/Source/Editor/UnrealEd/Public/FbxImporter.h
+++ b/Engine/Source/Editor/UnrealEd/Public/FbxImporter.h
@@ -36,10 +36,6 @@ struct FRichCurve;
 struct FStaticMaterial;
 struct FSkeletalMaterial;
 
-namespace AnimationTransformDebug
-{
-	struct FAnimationTransformDebugData;
-}
 // Temporarily disable a few warnings due to virtual function abuse in FBX source files
 #pragma warning( push )
 
@@ -1305,11 +1301,10 @@ private:
 	 * @param AnimImportSettings	Common settings to import animation.
 	 * @param SkeletalMeshRootNode	The fbx root node of the skeletalmesh.
 	 * @param ResampleRate			The rate at which the animations are resampled.
-	 * @param TransformDebugData	Out parameter data for internal debugging.
 	 * @param OutTotalNumKeys		Out parameter returning the number of keys imported, used to set the number of keys in the sequencer.
  	 * @param bReimport				Flag indicating whether or not this operation is part of reimporting.
 	 */
-	void ImportBoneTracks(USkeleton* Skeleton, FAnimCurveImportSettings& AnimImportSettings, FbxNode* SkeletalMeshRootNode, const int32 ResampleRate, TArray<AnimationTransformDebug::FAnimationTransformDebugData>& TransformDebugData, int32& OutTotalNumKeys, const bool bReimport);
+	void ImportBoneTracks(USkeleton* Skeleton, FAnimCurveImportSettings& AnimImportSettings, FbxNode* SkeletalMeshRootNode, const int32 ResampleRate, int32& OutTotalNumKeys, const bool bReimport);
 
 public:
 	// current Fbx scene we are importing. Make sure to release it after import
