FScopedEditorWorld: Support redirectors by using UAssetRegistryHelpers::FixupRedirectedAssetPath

#rnx

diff --git a/Engine/Source/Editor/UnrealEd/Private/EditorWorldUtils.cpp b/Engine/Source/Editor/UnrealEd/Private/EditorWorldUtils.cpp
--- a/Engine/Source/Editor/UnrealEd/Private/EditorWorldUtils.cpp
+++ b/Engine/Source/Editor/UnrealEd/Private/EditorWorldUtils.cpp
@@ -2,6 +2,7 @@
 
 #include "EditorWorldUtils.h"
 #include "Editor.h"
+#include "AssetRegistry/AssetRegistryHelpers.h"
 
 FScopedEditorWorld::FScopedEditorWorld()
 	: World(nullptr)
@@ -16,9 +17,23 @@ FScopedEditorWorld::FScopedEditorWorld(UWorld* InWorld, const UWorld::Initializa
 }
 
 FScopedEditorWorld::FScopedEditorWorld(const FStringView InLongPackageName, const UWorld::InitializationValues& InInitializationValues, EWorldType::Type InWorldType)
+	: FScopedEditorWorld(TSoftObjectPtr<UWorld>(FSoftObjectPath(InLongPackageName)), InInitializationValues, InWorldType)
+{
+}
+
+FScopedEditorWorld::FScopedEditorWorld(const TSoftObjectPtr<UWorld>& InSoftWorld, const UWorld::InitializationValues& InInitializationValues, EWorldType::Type InWorldType)
 	: FScopedEditorWorld()
 {
-	if (UPackage* WorldPackage = LoadWorldPackageForEditor(InLongPackageName))
+	FSoftObjectPath WorldObjectPath(InSoftWorld.ToSoftObjectPath());
+	// Handle cases where the SoftObjectPath doesn't have a valid AssetName
+	if (WorldObjectPath.GetAssetFName().IsNone())
+	{
+		FString WorldObjectPathStr = WorldObjectPath.GetLongPackageName() + TEXT(".") + FPackageName::GetShortName(WorldObjectPath.GetLongPackageName());
+		WorldObjectPath = FSoftObjectPath(WorldObjectPathStr);
+	}
+	UAssetRegistryHelpers::FixupRedirectedAssetPath(WorldObjectPath);
+	
+	if (UPackage* WorldPackage = LoadWorldPackageForEditor(WorldObjectPath.GetLongPackageName()))
 	{
 		if (UWorld* RuntimeWorld = UWorld::FindWorldInPackage(WorldPackage))
 		{
@@ -27,11 +42,6 @@ FScopedEditorWorld::FScopedEditorWorld(const FStringView InLongPackageName, cons
 	}
 }
 
-FScopedEditorWorld::FScopedEditorWorld(const TSoftObjectPtr<UWorld>& InSoftWorld, const UWorld::InitializationValues& InInitializationValues, EWorldType::Type InWorldType)
-	: FScopedEditorWorld(InSoftWorld.ToSoftObjectPath().GetLongPackageName(), InInitializationValues, InWorldType)
-{
-}
-
 void FScopedEditorWorld::Init(UWorld* InWorld, const UWorld::InitializationValues& InInitializationValues, EWorldType::Type InWorldType)
 {
 	check(InWorld);
