#UE Remove unreachable code

diff --git a/Engine/Source/Editor/UnrealEd/Private/EditorSelectUtils.cpp b/Engine/Source/Editor/UnrealEd/Private/EditorSelectUtils.cpp
--- a/Engine/Source/Editor/UnrealEd/Private/EditorSelectUtils.cpp
+++ b/Engine/Source/Editor/UnrealEd/Private/EditorSelectUtils.cpp
@@ -72,13 +72,15 @@ void UUnrealEdEngine::NoteActorMovement()
 		GLevelEditorModeTools().Snapping=0;
 		
 		AActor* SelectedActor = NULL;
-		for ( FSelectionIterator It( GetSelectedActorIterator() ) ; It ; ++It )
 		{
-			AActor* Actor = static_cast<AActor*>( *It );
-			checkSlow( Actor->IsA(AActor::StaticClass()) );
+			FSelectionIterator It(GetSelectedActorIterator());
+			if (It)
+			{
+				AActor* Actor = static_cast<AActor*>( *It );
+				checkSlow( Actor->IsA(AActor::StaticClass()) );
 
-			SelectedActor = Actor;
-			break;
+				SelectedActor = Actor;
+			}
 		}
 
 		if( SelectedActor == NULL )
@@ -89,13 +91,15 @@ void UUnrealEdEngine::NoteActorMovement()
 		}
 
 		// Look for an actor that requires snapping.
-		for ( FSelectionIterator It( GetSelectedActorIterator() ) ; It ; ++It )
 		{
-			AActor* Actor = static_cast<AActor*>( *It );
-			checkSlow( Actor->IsA(AActor::StaticClass()) );
+			FSelectionIterator It(GetSelectedActorIterator());
+			if (It)
+			{
+				AActor* Actor = static_cast<AActor*>( *It );
+				checkSlow( Actor->IsA(AActor::StaticClass()) );
 
-			GLevelEditorModeTools().Snapping = 1;
-			break;
+				GLevelEditorModeTools().Snapping = 1;
+			}
 		}
 
 		TSet<AGroupActor*> GroupActors;
