[CrashFix] Crash when attempting to change parent class of a blueprint which original parent class was deleted

diff --git a/Engine/Source/Editor/UnrealEd/Private/Kismet2/BlueprintEditorUtils.cpp b/Engine/Source/Editor/UnrealEd/Private/Kismet2/BlueprintEditorUtils.cpp
--- a/Engine/Source/Editor/UnrealEd/Private/Kismet2/BlueprintEditorUtils.cpp
+++ b/Engine/Source/Editor/UnrealEd/Private/Kismet2/BlueprintEditorUtils.cpp
@@ -8691,7 +8691,10 @@ TSharedRef<SWidget> FBlueprintEditorUtils::ConstructBlueprintParentClassPicker(
 		// for example, reparenting AActor to AEditorUtilityActor is illegal because AActor uses UBlueprint while
 		// AEditorUtilityActor uses UEditorUtilityBlueprint
 		const IKismetCompilerInterface& KismetCompilerModule = FModuleManager::LoadModuleChecked<IKismetCompilerInterface>("KismetCompiler");
-		KismetCompilerModule.GetSubclassesWithDifferingBlueprintTypes(Blueprint->ParentClass, Filter->DisallowedChildrenOfClasses);
+		if (Blueprint->ParentClass) // ParentClass is rarely null. Allow all children in that case.
+		{
+			KismetCompilerModule.GetSubclassesWithDifferingBlueprintTypes(Blueprint->ParentClass, Filter->DisallowedChildrenOfClasses);
+		}
 	}
 
 	return FModuleManager::LoadModuleChecked<FClassViewerModule>("ClassViewer").CreateClassViewer(Options, OnPicked);
