Fix grammar in property tooltip

diff --git a/Engine/Source/Editor/UnrealEd/Classes/Settings/EditorLoadingSavingSettings.h b/Engine/Source/Editor/UnrealEd/Classes/Settings/EditorLoadingSavingSettings.h
--- a/Engine/Source/Editor/UnrealEd/Classes/Settings/EditorLoadingSavingSettings.h
+++ b/Engine/Source/Editor/UnrealEd/Classes/Settings/EditorLoadingSavingSettings.h
@@ -125,7 +125,7 @@ private:
 public:
 
 	/**Automatically reimports textures when a change to source content is detected */
-	UPROPERTY(EditAnywhere, config, Category=AutoReimport, meta=(DisplayName="Monitor Content Directories", ToolTip="When enabled, changes to made to source content files inside the content directories will automatically be reflected in the content browser.\nNote that source content files must reside in one of the monitored directories to be eligible for auto-reimport.\nAdvanced setup options are available below."))
+	UPROPERTY(EditAnywhere, config, Category=AutoReimport, meta=(DisplayName="Monitor Content Directories", ToolTip="When enabled, changes to source content files inside the content directories will automatically be reflected in the content browser.\nNote that source content files must reside in one of the monitored directories to be eligible for auto-reimport.\nAdvanced setup options are available below."))
 	bool bMonitorContentDirectories;
 
 	UPROPERTY(config)
