Cooker: Whitespace: add linebreaks to long lines.
#rnx

diff --git a/Engine/Source/Editor/UnrealEd/Private/Cooker/CookDirector.cpp b/Engine/Source/Editor/UnrealEd/Private/Cooker/CookDirector.cpp
--- a/Engine/Source/Editor/UnrealEd/Private/Cooker/CookDirector.cpp
+++ b/Engine/Source/Editor/UnrealEd/Private/Cooker/CookDirector.cpp
@@ -129,21 +129,24 @@ FCookDirector::FCookDirector(UCookOnTheFlyServer& InCOTFS, int32 CookProcessCoun
 	ParseConfig(CookProcessCount, bConfigValid);
 	if (!bConfigValid)
 	{
-		UE_LOG(LogCook, Error, TEXT("CookDirector initialization failure: config settings are invalid for multiprocess. CookMultiprocess is disabled and the cooker is running as a single process."));
+		UE_LOG(LogCook, Error,
+			TEXT("CookDirector initialization failure: config settings are invalid for multiprocess. CookMultiprocess is disabled and the cooker is running as a single process."));
 		bMultiprocessAvailable = false;
 		return;
 	}
 	ISocketSubsystem* SocketSubsystem = ISocketSubsystem::Get();
 	if (!SocketSubsystem)
 	{
-		UE_LOG(LogCook, Error, TEXT("CookDirector initialization failure: platform does not support network sockets. CookMultiprocess is disabled and the cooker is running as a single process."));
+		UE_LOG(LogCook, Error,
+			TEXT("CookDirector initialization failure: platform does not support network sockets. CookMultiprocess is disabled and the cooker is running as a single process."));
 		bMultiprocessAvailable = false;
 		return;
 	}
 	bMultiprocessAvailable = true;
 
 	UE_LOG(LogCook, Display, TEXT("CookProcessCount=%d. CookMultiprocess is enabled with 1 CookDirector and %d %s."),
-		RequestedCookWorkerCount+1, RequestedCookWorkerCount, RequestedCookWorkerCount > 1 ? TEXT("CookWorkers") : TEXT("CookWorker"));
+		RequestedCookWorkerCount+1, RequestedCookWorkerCount,
+		RequestedCookWorkerCount > 1 ? TEXT("CookWorkers") : TEXT("CookWorker"));
 
 	Register(new FLogMessagesMessageHandler());
 	Register(new TMPCollectorServerMessageCallback<FRetractionResultsMessage>([this]
@@ -225,7 +228,8 @@ void FCookDirector::ParseConfig(int32 CookProcessCount, bool& bOutValid)
 	if (MultiprocessId != 0)
 	{
 		bOutValid = false;
-		UE_LOG(LogCook, Error, TEXT("CookMultiprocess is incompatible with -MultiprocessId on the CookDirector's commandline. The CookDirector needs to be able to specify all MultiprocessIds."));
+		UE_LOG(LogCook, Error,
+			TEXT("CookMultiprocess is incompatible with -MultiprocessId on the CookDirector's commandline. The CookDirector needs to be able to specify all MultiprocessIds."));
 	}
 }
 
@@ -243,9 +247,11 @@ FCookDirector::~FCookDirector()
 	}
 	for (FPackageData* PackageData : AbortedAssignments)
 	{
-		check(PackageData->IsInProgress()); // Packages that were assigned to workers should be in the AssignedToWorker state
+		// Packages that were assigned to workers should be in the AssignedToWorker state
+		check(PackageData->IsInProgress());
 		PackageData->SetWorkerAssignment(FWorkerId::Invalid(), ESendFlags::QueueNone);
-		PackageData->SendToState(UE::Cook::EPackageState::Request, ESendFlags::QueueAddAndRemove, EStateChangeReason::CookerShutdown);
+		PackageData->SendToState(UE::Cook::EPackageState::Request, ESendFlags::QueueAddAndRemove,
+			EStateChangeReason::CookerShutdown);
 	}
 	RemoteWorkers.Empty();
 	RemoteWorkerProfileDatas.Empty();
@@ -337,8 +343,10 @@ void FCookDirector::AssignRequests(TArrayView<FPackageData*> Requests, TArray<FW
 	AssignRequests(MoveTemp(WorkerIds), LocalRemoteWorkers, Requests, OutAssignments, MoveTemp(RequestGraph));
 }
 
-void FCookDirector::AssignRequests(TArray<FWorkerId>&& InWorkers, TArray<TRefCountPtr<FCookWorkerServer>>& InRemoteWorkers,
-	TArrayView<FPackageData*> Requests, TArray<FWorkerId>& OutAssignments, TMap<FPackageData*, TArray<FPackageData*>>&& RequestGraph)
+void FCookDirector::AssignRequests(TArray<FWorkerId>&& InWorkers,
+	TArray<TRefCountPtr<FCookWorkerServer>>& InRemoteWorkers,
+	TArrayView<FPackageData*> Requests, TArray<FWorkerId>& OutAssignments, TMap<FPackageData*,
+	TArray<FPackageData*>>&& RequestGraph)
 {
 	check(InWorkers.Num() > 0);
 	if (InWorkers.Num() <= 1)
@@ -353,7 +361,8 @@ void FCookDirector::AssignRequests(TArray<FWorkerId>&& InWorkers, TArray<TRefCou
 			FWorkerId WorkerIdConstraint = Request->GetWorkerAssignmentConstraint();
 			if (WorkerIdConstraint.IsValid() && WorkerIdConstraint != WorkerId)
 			{
-				UE_LOG(LogCook, Warning, TEXT("Package %s can only be cooked by a now-disconnected CookWorker. The package can not be cooked."),
+				UE_LOG(LogCook, Warning,
+					TEXT("Package %s can only be cooked by a now-disconnected CookWorker. The package can not be cooked."),
 					*Request->GetPackageName().ToString());
 				Assignment = FWorkerId::Invalid();
 				RemovedRequests.Add(Request);
@@ -431,7 +440,8 @@ void FCookDirector::AssignRequests(TArray<FWorkerId>&& InWorkers, TArray<TRefCou
 			uint8 RemoteIndex = WorkerId.GetRemoteIndex();
 			if (RemoteIndex >= RemoteBatches.Num() || !RemoteIndexIsValid[RemoteIndex])
 			{
-				UE_LOG(LogCook, Error, TEXT("Package %s can only be cooked by a now-disconnected CookWorker. The package can not be cooked."),
+				UE_LOG(LogCook, Error,
+					TEXT("Package %s can only be cooked by a now-disconnected CookWorker. The package can not be cooked."),
 					*Requests[RequestIndex]->GetPackageName().ToString());
 				WorkerId = FWorkerId::Invalid();
 				continue;
@@ -454,7 +464,8 @@ void FCookDirector::AssignRequests(TArray<FWorkerId>&& InWorkers, TArray<TRefCou
 			TRefCountPtr<FCookWorkerServer>* RemoteWorker = InRemoteWorkers.FindByPredicate(
 				[&WorkerId](const TRefCountPtr<FCookWorkerServer>& X) { return X->GetWorkerId() == WorkerId; });
 			check(RemoteWorker);
-			(*RemoteWorker)->AppendAssignments(RemoteBatches[WorkerId.GetRemoteIndex()], ECookDirectorThread::SchedulerThread);
+			(*RemoteWorker)->AppendAssignments(RemoteBatches[WorkerId.GetRemoteIndex()],
+				ECookDirectorThread::SchedulerThread);
 		}
 	}
 
@@ -556,7 +567,8 @@ void FCookDirector::TickFromSchedulerThread()
 		RetractionHandler->TickFromSchedulerThread(bAllWorkersConnected, bAnyIdle, BusiestNumAssignments);
 	}
 
-	bool bIsStalled = bLocalWorkerIdle && !COTFS.PackageDatas->GetAssignedToWorkerSet().IsEmpty() && WorkersWithMessage.IsEmpty();
+	bool bIsStalled = bLocalWorkerIdle && !COTFS.PackageDatas->GetAssignedToWorkerSet().IsEmpty()
+		&& WorkersWithMessage.IsEmpty();
 	for (TRefCountPtr<FCookWorkerServer>& Worker : WorkersWithMessage)
 	{
 		Worker->HandleReceiveMessages(ECookDirectorThread::SchedulerThread);
@@ -773,7 +785,8 @@ void FCookDirector::ShutdownCookSession()
 		FScopeLock CommunicationScopeLock(&CommunicationLock);
 		for (TPair<FCookWorkerServer*, TRefCountPtr<FCookWorkerServer>>& Pair : ShuttingDownWorkers)
 		{
-			check(Pair.Value); // The Value was set by AbortWorker for any new entries, and all old entries guarantee the value is set
+			// The Value was set by AbortWorker for any new entries, and all old entries guarantee the value is set
+			check(Pair.Value);
 			if (Pair.Key->IsFlushingBeforeShutdown())
 			{
 				WorkersNeedingAbort.Add(Pair.Value);
@@ -819,7 +832,8 @@ void FCookDirector::Register(IMPCollector* Collector)
 	TRefCountPtr<IMPCollector>& Existing = Collectors.FindOrAdd(Collector->GetMessageType());
 	if (Existing)
 	{
-		UE_LOG(LogCook, Error, TEXT("Duplicate IMPCollectors registered. Guid: %s, Existing: %s, Registering: %s. Keeping the Existing."),
+		UE_LOG(LogCook, Error,
+			TEXT("Duplicate IMPCollectors registered. Guid: %s, Existing: %s, Registering: %s. Keeping the Existing."),
 			*Collector->GetMessageType().ToString(), Existing->GetDebugName(), Collector->GetDebugName());
 		return;
 	}
@@ -832,7 +846,8 @@ void FCookDirector::Unregister(IMPCollector* Collector)
 	Collectors.RemoveAndCopyValue(Collector->GetMessageType(), Existing);
 	if (Existing && Existing.GetReference() != Collector)
 	{
-		UE_LOG(LogCook, Error, TEXT("Duplicate IMPCollector during Unregister. Guid: %s, Existing: %s, Unregistering: %s. Ignoring the Unregister."),
+		UE_LOG(LogCook, Error,
+			TEXT("Duplicate IMPCollector during Unregister. Guid: %s, Existing: %s, Unregistering: %s. Ignoring the Unregister."),
 			*Collector->GetMessageType().ToString(), Existing->GetDebugName(), Collector->GetDebugName());
 		Collectors.Add(Collector->GetMessageType(), MoveTemp(Existing));
 	}
@@ -860,7 +875,8 @@ void FCookDirector::SetWorkersStalled(bool bInWorkersStalled)
 		const double CurrentTime = FPlatformTime::Seconds();
 		if (CurrentTime >= WorkersStalledWarnTimeSeconds)
 		{
-			UE_LOG(LogCook, Display, TEXT("Cooker has been blocked with no results from remote CookWorkers for %.0f seconds."),
+			UE_LOG(LogCook, Display,
+				TEXT("Cooker has been blocked with no results from remote CookWorkers for %.0f seconds."),
 				(float)(CurrentTime - WorkersStalledStartTimeSeconds));
 			WorkersStalledWarnTimeSeconds = CurrentTime + GCookProgressWarnBusyTime;
 		}
@@ -964,7 +980,8 @@ bool FCookDirector::TryCreateWorkerConnectSocket()
 		TEXT("FCookDirector-WorkerConnect"), ErrorReason);
 	if (!WorkerConnectSocket)
 	{
-		UE_LOG(LogCook, Error, TEXT("CookDirector could not create listen socket, CookWorkers will be disabled. Reason: %s."),
+		UE_LOG(LogCook, Error,
+			TEXT("CookDirector could not create listen socket, CookWorkers will be disabled. Reason: %s."),
 			*ErrorReason);
 		return false;
 	}
@@ -1001,7 +1018,8 @@ void FCookDirector::InitializeWorkers()
 	{
 		int32 ProfileId = RemoteWorkerProfileDatas.Num();
 		RemoteWorkerProfileDatas.Emplace();
-		RemoteWorkers.Add(RemoteIndex, new FCookWorkerServer(*this, ProfileId, FWorkerId::FromRemoteIndex(RemoteIndex)));
+		RemoteWorkers.Add(RemoteIndex, new FCookWorkerServer(*this, ProfileId,
+			FWorkerId::FromRemoteIndex(RemoteIndex)));
 	}
 	bWorkersActive = true;
 
@@ -1050,7 +1068,8 @@ void FCookDirector::RecreateWorkers()
 		}
 		int32 ProfileId = RemoteWorkerProfileDatas.Num();
 		RemoteWorkerProfileDatas.Emplace();
-		RemoteWorkers.Add(RemoteIndex, new FCookWorkerServer(*this, ProfileId, FWorkerId::FromRemoteIndex(RemoteIndex)));
+		RemoteWorkers.Add(RemoteIndex, new FCookWorkerServer(*this, ProfileId,
+			FWorkerId::FromRemoteIndex(RemoteIndex)));
 		bWorkersActive = true;
 	}
 
@@ -1063,7 +1082,8 @@ void FCookDirector::ActivateMachineResourceReduction()
 	}
 	bHasReducedMachineResources = true;
 
-	// When running a multiprocess cook, we remove the Memory triggers and trigger GC based solely on PressureLevel. But keep the Soft GC settings
+	// When running a multiprocess cook, we remove the Memory triggers and trigger GC based solely on PressureLevel.
+	// But keep the Soft GC settings
 	COTFS.MemoryMaxUsedPhysical = 0;
 	COTFS.MemoryMaxUsedVirtual = 0;
 	COTFS.MemoryMinFreeVirtual = 0;
@@ -1081,7 +1101,8 @@ void FCookDirector::ActivateMachineResourceReduction()
 		COTFS.MemoryMinFreeVirtual / 1024 / 1024, COTFS.MemoryMinFreePhysical / 1024 / 1024,
 		*LexToString(COTFS.MemoryTriggerGCAtPressureLevel),
 		COTFS.bUseSoftGC ? TEXT("true") : TEXT("false"),
-		COTFS.bUseSoftGC ? *FString::Printf(TEXT(" (%d/%d)"), COTFS.SoftGCStartNumerator, COTFS.SoftGCDenominator) : TEXT("")
+		COTFS.bUseSoftGC ? *FString::Printf(TEXT(" (%d/%d)"),
+			COTFS.SoftGCStartNumerator, COTFS.SoftGCDenominator) : TEXT("")
 	);
 
 	// Set CoreLimit for updating workerthreads in this process and passing to the commandline for workers
@@ -1103,7 +1124,8 @@ void FCookDirector::ActivateMachineResourceReduction()
 
 	// Update the number of Cores and WorkerThreads for this process
 	check(IsInGameThread());
-	int32 NumBackgroundWorkers = FMath::Max(1, NumberOfWorkers - FMath::Min<int32>(GNumForegroundWorkers, NumberOfWorkers));
+	int32 NumBackgroundWorkers = FMath::Max(1,
+		NumberOfWorkers - FMath::Min<int32>(GNumForegroundWorkers, NumberOfWorkers));
 	int32 NumForegroundWorkers = FMath::Max(1, NumberOfWorkers - NumBackgroundWorkers);
 	LowLevelTasks::FScheduler::Get().RestartWorkers(NumForegroundWorkers, NumBackgroundWorkers);
 
@@ -1148,7 +1170,8 @@ void FCookDirector::TickWorkerConnects(ECookDirectorThread TickThread)
 		Status = TryReadPacket(Conn.Socket, Conn.Buffer, Messages);
 		if (Status != EConnectionStatus::Okay)
 		{
-			UE_LOG(LogCook, Warning, TEXT("Pending connection failed before sending a WorkerPacket: %s"), DescribeStatus(Status));
+			UE_LOG(LogCook, Warning,
+				TEXT("Pending connection failed before sending a WorkerPacket: %s"), DescribeStatus(Status));
 			Iter.RemoveCurrent();
 		}
 		if (Messages.Num() == 0)
@@ -1160,14 +1183,16 @@ void FCookDirector::TickWorkerConnects(ECookDirectorThread TickThread)
 
 		if (Messages[0].MessageType != FWorkerConnectMessage::MessageType)
 		{
-			UE_LOG(LogCook, Warning, TEXT("Pending connection sent a different message before sending a connection message. MessageType: %s. Connection will be ignored."),
+			UE_LOG(LogCook, Warning,
+				TEXT("Pending connection sent a different message before sending a connection message. MessageType: %s. Connection will be ignored."),
 				*Messages[0].MessageType.ToString());
 			continue;
 		}
 		FWorkerConnectMessage Message;
 		if (!Message.TryRead(MoveTemp(Messages[0].Object)))
 		{
-			UE_LOG(LogCook, Warning, TEXT("Pending connection sent an invalid Connection Message. Connection will be ignored."));
+			UE_LOG(LogCook, Warning,
+				TEXT("Pending connection sent an invalid Connection Message. Connection will be ignored."));
 			continue;
 		}
 		TRefCountPtr<FCookWorkerServer> RemoteWorker;
@@ -1187,7 +1212,8 @@ void FCookDirector::TickWorkerConnects(ECookDirectorThread TickThread)
 					}
 					ValidIndexes.RemoveSuffix(1); // Remove the terminating comma
 				}
-				UE_LOG(LogCook, Warning, TEXT("Pending connection sent a Connection Message with invalid RemoteIndex %d. ValidIndexes = {%s}. Connection will be ignored."),
+				UE_LOG(LogCook, Warning,
+					TEXT("Pending connection sent a Connection Message with invalid RemoteIndex %d. ValidIndexes = {%s}. Connection will be ignored."),
 					Message.RemoteIndex, *ValidIndexes);
 				continue;
 			}
@@ -1198,7 +1224,8 @@ void FCookDirector::TickWorkerConnects(ECookDirectorThread TickThread)
 		Messages.RemoveAt(0);
 		if (!RemoteWorker->TryHandleConnectMessage(Message, LocalSocket, MoveTemp(Messages), TickThread))
 		{
-			UE_LOG(LogCook, Warning, TEXT("Pending connection sent a Connection Message with an already in-use RemoteIndex. Connection will be ignored."));
+			UE_LOG(LogCook, Warning,
+				TEXT("Pending connection sent a Connection Message with an already in-use RemoteIndex. Connection will be ignored."));
 			Sockets::CloseSocket(LocalSocket);
 			continue;
 		}
@@ -1264,7 +1291,8 @@ FString FCookDirector::GetWorkerLogFileName(int32 ProfileId)
 	FString DirectorLogFileName = FGenericPlatformOutputDevices::GetAbsoluteLogFilename();
 	FStringView BaseFileName = FPathViews::GetBaseFilenameWithPath(DirectorLogFileName);
 	FStringView Extension = FPathViews::GetExtension(DirectorLogFileName, true /* bIncludeDot */);
-	return FString::Printf(TEXT("%.*s_Worker%d%*s"), BaseFileName.Len(), BaseFileName.GetData(), ProfileId, Extension.Len(), Extension.GetData());
+	return FString::Printf(TEXT("%.*s_Worker%d%*s"), BaseFileName.Len(), BaseFileName.GetData(), ProfileId,
+		Extension.Len(), Extension.GetData());
 }
 
 FString FCookDirector::GetWorkerCommandLine(FWorkerId WorkerId, int32 ProfileId)
@@ -1274,7 +1302,8 @@ FString FCookDirector::GetWorkerCommandLine(FWorkerId WorkerId, int32 ProfileId)
 	const TCHAR* ProjectName = FApp::GetProjectName();
 	checkf(ProjectName && ProjectName[0], TEXT("Expected UnrealEditor to be running with a non-empty project name"));
 
-	// Note that we need to handle quoted strings for e.g. a projectfile with spaces in it; FParse::Token does handle them
+	// Note that we need to handle quoted strings for e.g. a projectfile with spaces in it; FParse::Token
+	// does handle them
 	FString Token;
 	TArray<FString> Tokens;
 	while (FParse::Token(CommandLine, Token, false /* bUseEscape */))
@@ -1331,7 +1360,9 @@ FString FCookDirector::GetWorkerCommandLine(FWorkerId WorkerId, int32 ProfileId)
 	Tokens.Insert(TEXT("-cookworker"), 2);
 	Tokens.Insert(FString::Printf(TEXT("-CookProfileId=%d"), ProfileId), 3);
 	Tokens.Insert(FString::Printf(TEXT("-MultiprocessId=%d"), WorkerId.GetMultiprocessId()), 4);
-	check(!WorkerConnectAuthority.IsEmpty()); // This should have been constructed in TryCreateWorkerConnectSocket before any CookWorkerServers could exist to call GetWorkerCommandLine
+	// This should have been constructed in TryCreateWorkerConnectSocket before any CookWorkerServers could exist to
+	// call GetWorkerCommandLine
+	check(!WorkerConnectAuthority.IsEmpty());
 	Tokens.Add(FString::Printf(TEXT("-CookDirectorHost=%s"), *WorkerConnectAuthority));
 	Tokens.Add(TEXT("-unattended"));
 	Tokens.Add(FString::Printf(TEXT("-abslog=%s"), *GetWorkerLogFileName(ProfileId)));
@@ -1371,7 +1402,8 @@ bool FDirectorConnectionInfo::TryParseCommandLine()
 	}
 	if (MultiprocessId < 1 || 257 <= MultiprocessId)
 	{
-		UE_LOG(LogCook, Error, TEXT("CookWorker startup failed: commandline had invalid -MultiprocessId=%d; MultiprocessId must be in the range [1, 256]."),
+		UE_LOG(LogCook, Error,
+			TEXT("CookWorker startup failed: commandline had invalid -MultiprocessId=%d; MultiprocessId must be in the range [1, 256]."),
 			MultiprocessId);
 		return false;
 	}
@@ -1424,7 +1456,8 @@ void FCookDirector::AbortWorker(FWorkerId WorkerId, ECookDirectorThread TickThre
 	{
 		FScopeLock RemoteWorkersScopeLock(&CommunicationLock);
 		TRefCountPtr<FCookWorkerServer>& Existing = ShuttingDownWorkers.FindOrAdd(RemoteWorker.GetReference());
-		check(!Existing); // We should not be able to abort a worker twice because we removed it from RemoteWorkers above
+		// We should not be able to abort a worker twice because we removed it from RemoteWorkers above
+		check(!Existing);
 		Existing = MoveTemp(RemoteWorker);
 
 		if (TickThread != ECookDirectorThread::SchedulerThread)
@@ -1438,9 +1471,11 @@ void FCookDirector::ReassignAbortedPackages(TArray<FPackageData*>& PackagesToRea
 {
 	for (FPackageData* PackageData : PackagesToReassign)
 	{
-		check(PackageData->IsInProgress()); // Packages that were assigned to a worker should be in the AssignedToWorker state
+		// Packages that were assigned to a worker should be in the AssignedToWorker state
+		check(PackageData->IsInProgress());
 		PackageData->SetWorkerAssignment(FWorkerId::Invalid());
-		PackageData->SendToState(UE::Cook::EPackageState::Request, ESendFlags::QueueAddAndRemove, EStateChangeReason::ReassignAbortedPackages);
+		PackageData->SendToState(UE::Cook::EPackageState::Request, ESendFlags::QueueAddAndRemove,
+			EStateChangeReason::ReassignAbortedPackages);
 	}
 	PackagesToReassign.Empty();
 }
@@ -1562,7 +1597,8 @@ FCookDirector::FRetractionHandler::ERetractionState FCookDirector::FRetractionHa
 		IdleWorkerListText << Director.GetDisplayName(WorkerId) << TEXT(", ");
 	}
 	IdleWorkerListText.RemoveSuffix(2);
-	UE_LOG(LogCook, Display, TEXT("Idle CookWorkers: { %s }. Retracting %d packages from %s to distribute to the idle CookWorkers."),
+	UE_LOG(LogCook, Display,
+		TEXT("Idle CookWorkers: { %s }. Retracting %d packages from %s to distribute to the idle CookWorkers."),
 		*IdleWorkerListText, NumAssignmentsToRetract, *Director.GetDisplayName(BusiestWorker));
 	Director.DisplayRemainingPackages();
 
@@ -1595,7 +1631,8 @@ void FCookDirector::FRetractionHandler::InitializeForResultsMessage(const FWorke
 	ExpectedWorker = FromWorker;
 }
 
-void FCookDirector::FRetractionHandler::TickFromSchedulerThread(bool bAllWorkersConnected, bool bAnyIdle, int32 BusiestNumAssignments)
+void FCookDirector::FRetractionHandler::TickFromSchedulerThread(bool bAllWorkersConnected, bool bAnyIdle,
+	int32 BusiestNumAssignments)
 {
 	bool bHadStateChange;
 	int32 NumTransitions = 0;
@@ -1715,13 +1752,15 @@ FCookDirector::FRetractionHandler::ERetractionState FCookDirector::FRetractionHa
 			const TRefCountPtr<FCookWorkerServer>* RemoteWorkerPtr = Director.FindRemoteWorkerInLock(ExpectedWorker);
 			if (!RemoteWorkerPtr)
 			{
-				// The CookWorker aborted and we already reassigned all of its packages; stop waiting for a retraction message from it.
+				// The CookWorker aborted and we already reassigned all of its packages; stop waiting for a retraction
+				// message from it.
 				check(PackagesToRetract.IsEmpty()); // Otherwise WorkerWithResults would have been set
 				ExpectedWorker = FWorkerId::Invalid();
 				return ERetractionState::Idle;
 			}
 		}
-		UE_CLOG(!IsCookIgnoreTimeouts(), LogCook, Display, TEXT("%s has not responded to a RetractionRequest message for %.1f seconds. Continuing to wait..."),
+		UE_CLOG(!IsCookIgnoreTimeouts(), LogCook, Display,
+			TEXT("%s has not responded to a RetractionRequest message for %.1f seconds. Continuing to wait..."),
 			*Director.GetDisplayName(ExpectedWorker), static_cast<float>(CurrentTime - MessageSentTimeSeconds));
 		LastWarnTimeSeconds = CurrentTime;
 		return ERetractionState::WaitingForResponse;
@@ -1752,7 +1791,8 @@ FCookDirector::FRetractionHandler::ERetractionState FCookDirector::FRetractionHa
 		}
 		if (RemoteWorker)
 		{
-			// The worker(s) that sent the retraction message aborted all of the packages, so mark locally that they have been aborted
+			// The worker(s) that sent the retraction message aborted all of the packages, so mark locally that they
+			// have been aborted
 			RemoteWorker->AbortAssignments(WorkerPackageDatas, ECookDirectorThread::SchedulerThread,
 				ENotifyRemote::LocalOnly);
 		}
@@ -1819,7 +1859,8 @@ void FCookDirector::FRetractionHandler::HandleRetractionMessage(FMPCollectorServ
 }
 
 FCookDirector::FRetractionHandler::ERetractionResult
-FCookDirector::FRetractionHandler::ReassignPackages(const FWorkerId& FromWorker, TConstArrayView<FPackageData*> Packages)
+FCookDirector::FRetractionHandler::ReassignPackages(const FWorkerId& FromWorker,
+	TConstArrayView<FPackageData*> Packages)
 {
 	TArray<TRefCountPtr<FCookWorkerServer>> LocalRemoteWorkers = Director.CopyRemoteWorkers();
 
@@ -1836,7 +1877,10 @@ FCookDirector::FRetractionHandler::ReassignPackages(const FWorkerId& FromWorker,
 		if (WorkerConstraint.IsValid())
 		{
 			if (!WorkerConstraint.IsLocal() && !LocalRemoteWorkers.FindByPredicate(
-				[&WorkerConstraint](const TRefCountPtr<FCookWorkerServer>& X) { return X->GetWorkerId() == WorkerConstraint; }))
+				[&WorkerConstraint](const TRefCountPtr<FCookWorkerServer>& X)
+				{
+					return X->GetWorkerId() == WorkerConstraint;
+				}))
 			{
 				continue;
 			}
@@ -1848,13 +1892,15 @@ FCookDirector::FRetractionHandler::ReassignPackages(const FWorkerId& FromWorker,
 	}
 	if (AssignmentPackages.IsEmpty())
 	{
-		UE_LOG(LogCook, Display, TEXT("Retraction results message received from %s; no packages were available for retraction."),
+		UE_LOG(LogCook, Display,
+			TEXT("Retraction results message received from %s; no packages were available for retraction."),
 			*Director.GetDisplayName(FromWorker));
 		Director.DisplayRemainingPackages();
 		return ERetractionResult::NoneAvailable;
 	}
 
-	TArray<FWorkerId> WorkersToSplitOver = CalculateWorkersToSplitOver(AssignmentPackages.Num(), FromWorker, LocalRemoteWorkers);
+	TArray<FWorkerId> WorkersToSplitOver = CalculateWorkersToSplitOver(AssignmentPackages.Num(), FromWorker,
+		LocalRemoteWorkers);
 	if (WorkersToSplitOver.IsEmpty())
 	{
 		// Send the packages back to the Director for reassignment
@@ -1864,7 +1910,8 @@ FCookDirector::FRetractionHandler::ReassignPackages(const FWorkerId& FromWorker,
 			RestartedRequests.Add(PackageData);
 		}
 		// MPCOOKTODO: Add a method to PumpRequests long enough to assign the packages
-		UE_LOG(LogCook, Display, TEXT("%d packages retracted from %s. No workers are currently idle so the packages were assigned evenly to all CookWorkers."),
+		UE_LOG(LogCook, Display,
+			TEXT("%d packages retracted from %s. No workers are currently idle so the packages were assigned evenly to all CookWorkers."),
 			AssignmentPackages.Num(), *Director.GetDisplayName(FromWorker));
 		Director.DisplayRemainingPackages();
 		return ERetractionResult::Retracted;
@@ -1885,7 +1932,8 @@ FCookDirector::FRetractionHandler::ReassignPackages(const FWorkerId& FromWorker,
 
 	TMap<FPackageData*, TArray<FPackageData*>> RequestGraph;
 	TArray<FWorkerId> Assignments;
-	Director.AssignRequests(MoveTemp(WorkersToSplitOver), LocalRemoteWorkers, AssignmentPackages, Assignments, MoveTemp(RequestGraph));
+	Director.AssignRequests(MoveTemp(WorkersToSplitOver), LocalRemoteWorkers, AssignmentPackages, Assignments,
+		MoveTemp(RequestGraph));
 	FRequestQueue& RequestQueue = Director.COTFS.PackageDatas->GetRequestQueue();
 	bool bAssignedToLocal = false;
 	for (int32 Index = 0; Index < AssignmentPackages.Num(); ++Index)
@@ -1894,7 +1942,8 @@ FCookDirector::FRetractionHandler::ReassignPackages(const FWorkerId& FromWorker,
 		FWorkerId Assignment = Assignments[Index];
 		if (Assignment.IsInvalid())
 		{
-			Director.COTFS.DemoteToIdle(*PackageData, ESendFlags::QueueAdd, ESuppressCookReason::MultiprocessAssignmentError);
+			Director.COTFS.DemoteToIdle(*PackageData, ESendFlags::QueueAdd,
+				ESuppressCookReason::MultiprocessAssignmentError);
 		}
 		else if (Assignment.IsLocal())
 		{
@@ -1903,7 +1952,8 @@ FCookDirector::FRetractionHandler::ReassignPackages(const FWorkerId& FromWorker,
 		}
 		else
 		{
-			PackageData->SendToState(EPackageState::AssignedToWorker, ESendFlags::QueueAdd, EStateChangeReason::Retraction);
+			PackageData->SendToState(EPackageState::AssignedToWorker, ESendFlags::QueueAdd,
+				EStateChangeReason::Retraction);
 			PackageData->SetWorkerAssignment(Assignment);
 		}
 	}
@@ -1918,8 +1968,8 @@ FCookDirector::FRetractionHandler::ReassignPackages(const FWorkerId& FromWorker,
 	return ERetractionResult::Retracted;
 }
 
-TArray<FWorkerId> FCookDirector::FRetractionHandler::CalculateWorkersToSplitOver(int32 NumPackages, const FWorkerId& FromWorker,
-	TConstArrayView<TRefCountPtr<FCookWorkerServer>> LocalRemoteWorkers)
+TArray<FWorkerId> FCookDirector::FRetractionHandler::CalculateWorkersToSplitOver(int32 NumPackages,
+	const FWorkerId& FromWorker, TConstArrayView<TRefCountPtr<FCookWorkerServer>> LocalRemoteWorkers)
 {
 	TArray<TPair<FWorkerId, int32>> WorkerNumPackages;
 	if (FromWorker != FWorkerId::Local() && Director.bAllowLocalCooks)
@@ -1937,12 +1987,15 @@ TArray<FWorkerId> FCookDirector::FRetractionHandler::CalculateWorkersToSplitOver
 	{
 		return TArray<FWorkerId>();
 	}
-	WorkerNumPackages.Sort([](const TPair<FWorkerId, int32>& A, const TPair<FWorkerId, int32>& B) { return A.Value < B.Value; });
+	WorkerNumPackages.Sort([](const TPair<FWorkerId, int32>& A, const TPair<FWorkerId, int32>& B)
+		{
+			return A.Value < B.Value;
+		});
 
 	// Consider splitting the packages amonst the 1 lowest, 2 lowest, ... n lowest (not including the FromWorker)
 	// Pick the value to split over based on whichever split group results in the lowest post split maximum
-	// So splitting 500 over 0,1000,1000,1000 -> would give them all to the first, but splitting 500 over 0, 100, 1000, 1000 would
-	// split them amongst the first two.
+	// So splitting 500 over 0,1000,1000,1000 -> would give them all to the first, but splitting 500 over
+	// 0, 100, 1000, 1000 would split them amongst the first two.
 	int32 BestNumToSplitOver = 0;
 	int32 BestPostSplitValue = 0;
 	for (int32 NumToSplitOver = 1; NumToSplitOver <= WorkerNumPackages.Num(); ++NumToSplitOver)
diff --git a/Engine/Source/Editor/UnrealEd/Private/Cooker/CookDirector.h b/Engine/Source/Editor/UnrealEd/Private/Cooker/CookDirector.h
--- a/Engine/Source/Editor/UnrealEd/Private/Cooker/CookDirector.h
+++ b/Engine/Source/Editor/UnrealEd/Private/Cooker/CookDirector.h
@@ -77,9 +77,15 @@ public:
 	void TickFromSchedulerThread();
 	/** Periodic display function, called from CookOnTheFlyServer.UpdateDisplay. */
 	void UpdateDisplayDiagnostics() const;
-	/** Called when the COTFS Server has detected all packages are complete. Tells the CookWorkers to flush messages and exit. */
+	/**
+	 * Called when the COTFS Server has detected all packages are complete. Tells the CookWorkers to flush messages
+	 * and exit.
+	 */
 	void PumpCookComplete(bool& bOutCompleted);
-	/** Called when a session ends. The Director blocks on shutdown of all CookWorkers and returns state to before session started. */
+	/**
+	 * Called when a session ends. The Director blocks on shutdown of all CookWorkers and returns state to before
+	 * session started.
+	 */
 	void ShutdownCookSession();
 
 	/** Enum specifying how CookWorker log output should be shown. */
@@ -150,7 +156,10 @@ private:
 	 * The CookWorkerServers are constructed to Uninitialized; the worker process is created later.
 	 */
 	void InitializeWorkers();
-	/** Copy to snapshot variables the data required on the communication thread that can only be read from the scheduler thread. */
+	/**
+	 * Copy to snapshot variables the data required on the communication thread that can only be read from the
+	 * scheduler thread.
+	 */
 	void ConstructReadonlyThreadVariables();
 	/** Construct CookWorkerServers if necessary to replace workers that have crashed. */
 	void RecreateWorkers();
diff --git a/Engine/Source/Editor/UnrealEd/Private/Cooker/CookPackageData.cpp b/Engine/Source/Editor/UnrealEd/Private/Cooker/CookPackageData.cpp
--- a/Engine/Source/Editor/UnrealEd/Private/Cooker/CookPackageData.cpp
+++ b/Engine/Source/Editor/UnrealEd/Private/Cooker/CookPackageData.cpp
@@ -216,7 +216,8 @@ void FPackageData::AddReachablePlatforms(FRequestCluster& RequestCluster,
 void FPackageData::AddReachablePlatformsInternal(FPackageData& PackageData,
 	TConstArrayView<const ITargetPlatform*> Platforms, FInstigator&& InInstigator)
 {
-	// This is a static helper function to make it impossible to make a typo and use this->Instigator instead of InInstigator
+	// This is a static helper function to make it impossible to make a typo and use this->Instigator instead of
+	// InInstigator
 	bool bSessionPlatformModified = false;
 	for (const ITargetPlatform* Platform : Platforms)
 	{
@@ -230,7 +231,8 @@ void FPackageData::AddReachablePlatformsInternal(FPackageData& PackageData,
 	}
 }
 
-void FPackageData::QueueAsDiscovered(FInstigator&& InInstigator, FDiscoveredPlatformSet&& ReachablePlatforms, bool bUrgent)
+void FPackageData::QueueAsDiscovered(FInstigator&& InInstigator, FDiscoveredPlatformSet&& ReachablePlatforms,
+	bool bUrgent)
 {
 	QueueAsDiscoveredInternal(*this, MoveTemp(InInstigator), MoveTemp(ReachablePlatforms), bUrgent);
 }
@@ -238,7 +240,8 @@ void FPackageData::QueueAsDiscovered(FInstigator&& InInstigator, FDiscoveredPlat
 void FPackageData::QueueAsDiscoveredInternal(FPackageData& PackageData, FInstigator&& InInstigator,
 	FDiscoveredPlatformSet&& ReachablePlatforms, bool bUrgent)
 {
-	// This is a static helper function to make it impossible to make a typo and use this->Instigator instead of InInstigator
+	// This is a static helper function to make it impossible to make a typo and use this->Instigator instead of
+	// InInstigator
 	TRingBuffer<FDiscoveryQueueElement>& Queue = PackageData.PackageDatas.GetRequestQueue().GetDiscoveryQueue();
 	Queue.Add(FDiscoveryQueueElement{ &PackageData, MoveTemp(InInstigator), MoveTemp(ReachablePlatforms), bUrgent });
 }
@@ -444,12 +447,14 @@ void FPackageData::ClearCookResults(const ITargetPlatform* TargetPlatform)
 	}
 }
 
-const TSortedMap<const ITargetPlatform*, FPackagePlatformData, TInlineAllocator<1>>& FPackageData::GetPlatformDatas() const
+const TSortedMap<const ITargetPlatform*, FPackagePlatformData, TInlineAllocator<1>>&
+FPackageData::GetPlatformDatas() const
 {
 	return PlatformDatas;
 }
 
-TSortedMap<const ITargetPlatform*, FPackagePlatformData, TInlineAllocator<1>>& FPackageData::GetPlatformDatasConstKeysMutableValues()
+TSortedMap<const ITargetPlatform*, FPackagePlatformData, TInlineAllocator<1>>&
+FPackageData::GetPlatformDatasConstKeysMutableValues()
 {
 	return PlatformDatas;
 }
@@ -777,7 +782,8 @@ void FPackageData::CheckInContainer() const
 		break;
 	case EPackageState::Save:
 		// The save queue is huge and often pushed at end. Check last element first and then scan.
-		check(PackageDatas.GetSaveQueue().Num() && (PackageDatas.GetSaveQueue().Last() == this || Algo::Find(PackageDatas.GetSaveQueue(), this)));
+		check(PackageDatas.GetSaveQueue().Num() && (PackageDatas.GetSaveQueue().Last() == this
+			|| Algo::Find(PackageDatas.GetSaveQueue(), this)));
 		break;
 	default:
 		check(false);
@@ -820,7 +826,8 @@ void FPackageData::SetWorkerAssignment(FWorkerId InWorkerAssignment, ESendFlags
 {
 	if (WorkerAssignment.IsValid())
 	{
-		checkf(InWorkerAssignment.IsInvalid(), TEXT("Package %s is being assigned to worker %d while it is already assigned to worker %d."),
+		checkf(InWorkerAssignment.IsInvalid(),
+			TEXT("Package %s is being assigned to worker %d while it is already assigned to worker %d."),
 			*GetPackageName().ToString(), WorkerAssignment.GetRemoteIndex(), WorkerAssignment.GetRemoteIndex());
 		if (EnumHasAnyFlags(SendFlags, ESendFlags::QueueRemove))
 		{
@@ -832,7 +839,8 @@ void FPackageData::SetWorkerAssignment(FWorkerId InWorkerAssignment, ESendFlags
 	{
 		if (InWorkerAssignment.IsValid())
 		{
-			checkf(GetState() == EPackageState::AssignedToWorker, TEXT("Package %s is being assigned to worker %d while in a state other than AssignedToWorker. This is invalid."),
+			checkf(GetState() == EPackageState::AssignedToWorker,
+				TEXT("Package %s is being assigned to worker %d while in a state other than AssignedToWorker. This is invalid."),
 				*GetPackageName().ToString(), GetWorkerAssignment().GetRemoteIndex());
 		}
 		WorkerAssignment = InWorkerAssignment;
@@ -975,8 +983,8 @@ bool FPackageData::TryPreload()
 	{
 		if (AsyncRequest && !AsyncRequest->bHasFinished)
 		{
-			// In case of async loading, the object can be found while still being asynchronously serialized, we need to wait until 
-			// the callback is called and the async request is completely done.
+			// In case of async loading, the object can be found while still being asynchronously serialized, we need
+			// to wait until the callback is called and the async request is completely done.
 			return false;
 		}
 
@@ -1001,13 +1009,17 @@ bool FPackageData::TryPreload()
 			AsyncRequest->RequestID = LoadPackageAsync(
 				GetFileName().ToString(), 
 				FLoadPackageAsyncDelegate::CreateLambda(
-					[AsyncRequest = AsyncRequest](const FName&, UPackage*, EAsyncLoadingResult::Type) { AsyncRequest->bHasFinished = true; }
+					[AsyncRequest = AsyncRequest](const FName&, UPackage*, EAsyncLoadingResult::Type) 
+					{
+						AsyncRequest->bHasFinished = true;
+					}
 				),
 				32 /* Use arbitrary higher priority for preload as we're going to need them very soon */
 			);
 		}
 
-		// always return false so we continue to check the status of the load until FindObjectFast above finds the loaded object
+		// always return false so we continue to check the status of the load until FindObjectFast above finds the
+		// loaded object
 		return false;
 	}
 	if (!PreloadableFile.Get())
@@ -1028,7 +1040,8 @@ bool FPackageData::TryPreload()
 				// not change until this is destructed, and the destructor does not run and other threads do not read
 				// or write PreloadableFileOpenResult until after PreloadableFile.Get() has finished initialization
 				// and this callback is therefore complete.
-				// The code that accomplishes that waiting is in TryPreload (IsInitialized) and ClearPreload (ReleaseCache)
+				// The code that accomplishes that waiting is in TryPreload (IsInitialized) and ClearPreload
+				// (ReleaseCache)
 				this->GetFileName().ToString(FileNameString);
 				FPackagePath PackagePath = FPackagePath::FromLocalPath(FileNameString);
 				FOpenPackageResult Result = IPackageResourceManager::Get().OpenReadPackage(PackagePath);
@@ -1121,7 +1134,8 @@ void FPackageData::ClearPreload()
 		GetFileName().ToString(FileNameString);
 		if (IPackageResourceManager::UnRegisterPreloadableArchive(FPackagePath::FromLocalPath(FileNameString)))
 		{
-			UE_LOG(LogCook, Display, TEXT("PreloadableFile was created for %s but never used. This is wasteful and bad for cook performance."),
+			UE_LOG(LogCook, Display,
+				TEXT("PreloadableFile was created for %s but never used. This is wasteful and bad for cook performance."),
 				*PackageName.ToString());
 		}
 		FilePtr->ReleaseCache(); // ReleaseCache to conserve memory if the Linker still has a pointer to it
@@ -1175,12 +1189,14 @@ void FPackageData::CreateObjectCache()
 		TArray<UObject*> ObjectsInOuter;
 		// ignore RF_Garbage objects; they will not be serialized out so we don't need to call
 		// BeginCacheForCookedPlatformData on them
-		GetObjectsWithOuter(LocalPackage, ObjectsInOuter, true /* bIncludeNestedObjects */, RF_NoFlags, EInternalObjectFlags::Garbage);
+		GetObjectsWithOuter(LocalPackage, ObjectsInOuter, true /* bIncludeNestedObjects */, RF_NoFlags,
+			EInternalObjectFlags::Garbage);
 		CachedObjectsInOuter.Reset(ObjectsInOuter.Num());
 		for (UObject* Object : ObjectsInOuter)
 		{
 			FWeakObjectPtr ObjectWeakPointer(Object);
-			check(ObjectWeakPointer.Get()); // GetObjectsWithOuter with Garbage filtered out should only return valid-for-weakptr objects
+			// GetObjectsWithOuter with Garbage filtered out should only return valid-for-weakptr objects
+			check(ObjectWeakPointer.Get());
 			CachedObjectsInOuter.Emplace(ObjectWeakPointer);
 		}
 
@@ -1257,7 +1273,8 @@ EPollStatus FPackageData::RefreshObjectCache(bool& bOutFoundNewObjects)
 		}
 	}
 	TArray<UObject*> CurrentObjects;
-	GetObjectsWithOuter(Package.Get(), CurrentObjects, true /* bIncludeNestedObjects */, RF_NoFlags, EInternalObjectFlags::Garbage);
+	GetObjectsWithOuter(Package.Get(), CurrentObjects, true /* bIncludeNestedObjects */, RF_NoFlags, 
+		EInternalObjectFlags::Garbage);
 
 	TArray<UObject*> NewObjects = SetDifference(CurrentObjects, OldObjects);
 	bOutFoundNewObjects = NewObjects.Num() > 0;
@@ -1267,16 +1284,18 @@ EPollStatus FPackageData::RefreshObjectCache(bool& bOutFoundNewObjects)
 		for (UObject* Object : NewObjects)
 		{
 			FWeakObjectPtr ObjectWeakPointer(Object);
-			check(ObjectWeakPointer.Get()); // GetObjectsWithOuter with Garbage filtered out should only return valid-for-weakptr objects
+			// GetObjectsWithOuter with Garbage filtered out should only return valid-for-weakptr objects
+			check(ObjectWeakPointer.Get());
 			CachedObjectsInOuter.Emplace(MoveTemp(ObjectWeakPointer));
 		}
-		// GetCookedPlatformDataNextIndex is already where it should be, pointing at the first of the objects we have added
-		// Change our state back so we know we need to CallBeginCacheOnObjects again 
+		// GetCookedPlatformDataNextIndex is already where it should be, pointing at the first of the objects we have
+		// added Change our state back so we know we need to CallBeginCacheOnObjects again 
 		SetCookedPlatformDataCalled(false);
 
 		if (++GetNumRetriesBeginCacheOnObjects() > FPackageData::GetMaxNumRetriesBeginCacheOnObjects())
 		{
-			UE_LOG(LogCook, Error, TEXT("Cooker has repeatedly tried to call BeginCacheForCookedPlatformData on all objects in the package, but keeps finding new objects.\n")
+			UE_LOG(LogCook, Error,
+				TEXT("Cooker has repeatedly tried to call BeginCacheForCookedPlatformData on all objects in the package, but keeps finding new objects.\n")
 				TEXT("Aborting the save of the package; programmer needs to debug why objects keep getting added to the package.\n")
 				TEXT("Package: %s. Most recent created object: %s."),
 				*GetPackageName().ToString(), *NewObjects[0]->GetFullName());
@@ -1375,8 +1394,8 @@ void FPackageData::RemapTargetPlatforms(const TMap<ITargetPlatform*, ITargetPlat
 		NewPlatformDatas.FindOrAdd(NewKey) = MoveTemp(ExistingPair.Value);
 	}
 
-	// The save state (and maybe more in the future) by contract can depend on the order of the request platforms remaining
-	// unchanged. If we change that order due to the remap, we need to demote back to request.
+	// The save state (and maybe more in the future) by contract can depend on the order of the request platforms
+	// remaining unchanged. If we change that order due to the remap, we need to demote back to request.
 	if (IsInProgress() && GetState() != EPackageState::Request)
 	{
 		bool bDemote = true;
@@ -1567,7 +1586,8 @@ void FGenerationHelper::InitializeSave(const UObject* InSplitDataObject,
 		FName InSplitDataObjectName = *InSplitDataObject->GetFullName();
 		check(SplitDataObjectName.IsNone() || SplitDataObjectName == InSplitDataObjectName);
 		SplitDataObjectName = InSplitDataObjectName;
-		bUseInternalReferenceToAvoidGarbageCollect = CookPackageSplitterInstance->UseInternalReferenceToAvoidGarbageCollect();
+		bUseInternalReferenceToAvoidGarbageCollect =
+			CookPackageSplitterInstance->UseInternalReferenceToAvoidGarbageCollect();
 		SetOwnerPackage(GetOwner().GetPackage());
 	}
 }
@@ -1978,7 +1998,8 @@ bool FGenerationHelper::IsComplete() const
 	return RemainingToPopulate == 0;
 }
 
-void FGenerationHelper::ResetSaveState(FCookGenerationInfo& Info, UPackage* Package, EStateChangeReason ReleaseSaveReason)
+void FGenerationHelper::ResetSaveState(FCookGenerationInfo& Info, UPackage* Package,
+	EStateChangeReason ReleaseSaveReason)
 {
 	check(IsInitialized());
 	if (Info.GetSaveState() > FCookGenerationInfo::ESaveState::CallPopulate)
@@ -2043,8 +2064,9 @@ void FGenerationHelper::ResetSaveState(FCookGenerationInfo& Info, UPackage* Pack
 	{
 		if (Info.PackageData && Info.PackageData->GetCachedObjectsInOuter().Num() != 0 &&
 			IsUseInternalReferenceToAvoidGarbageCollect() &&
-			(ReleaseSaveReason != EStateChangeReason::Completed && ReleaseSaveReason != EStateChangeReason::DoneForNow &&
-			 ReleaseSaveReason != EStateChangeReason::SaveError && ReleaseSaveReason != EStateChangeReason::CookerShutdown))
+			(ReleaseSaveReason != EStateChangeReason::Completed && ReleaseSaveReason != EStateChangeReason::DoneForNow
+				&& ReleaseSaveReason != EStateChangeReason::SaveError
+				&& ReleaseSaveReason != EStateChangeReason::CookerShutdown))
 		{
 			UE_LOG(LogCook, Error,
 				TEXT("CookPackageSplitter failure: We are demoting a %s package from save and removing our references that keep its objects loaded.\n")
@@ -2301,7 +2323,8 @@ void FCookGenerationInfo::CreatePackageHash()
 	IAssetRegistry& AssetRegistry = IAssetRegistry::GetChecked();
 	for (const FAssetDependency& Dependency : PackageDependencies)
 	{
-		TOptional<FAssetPackageData> DependencyData = AssetRegistry.GetAssetPackageDataCopy(Dependency.AssetId.PackageName);
+		TOptional<FAssetPackageData> DependencyData =
+			AssetRegistry.GetAssetPackageDataCopy(Dependency.AssetId.PackageName);
 		if (DependencyData)
 		{
 			Blake3.Update(&DependencyData->GetPackageSavedHash().GetBytes(),
@@ -2316,7 +2339,8 @@ void FCookGenerationInfo::CreatePackageHash()
 	constexpr int SizeDifference = sizeof(PackageHash) - sizeof(decltype(DeclVal<UPackage>().GetGuid()));
 	if (SizeDifference > 0)
 	{
-		FMemory::Memset(((uint8*)&PackageHash.GetBytes()) + (sizeof(decltype(PackageHash.GetBytes())) - SizeDifference),
+		FMemory::Memset(((uint8*)&PackageHash.GetBytes())
+			+ (sizeof(decltype(PackageHash.GetBytes())) - SizeDifference),
 			0, SizeDifference);
 	}
 	PRAGMA_ENABLE_DEPRECATION_WARNINGS;
@@ -2334,7 +2358,8 @@ void FCookGenerationInfo::IterativeCookValidateOrClear(FGenerationHelper& Genera
 		// were detected as modified during PopulateCookedPackages.
 		for (const FAssetDependency& Dependency : this->PackageDependencies)
 		{
-			FPackageData* DependencyData = COTFS.PackageDatas->FindPackageDataByPackageName(Dependency.AssetId.PackageName);
+			FPackageData* DependencyData =
+				COTFS.PackageDatas->FindPackageDataByPackageName(Dependency.AssetId.PackageName);
 			if (!DependencyData)
 			{
 				bOutIterativelyUnmodified = false;
@@ -2428,7 +2453,8 @@ bool FPendingCookedPlatformData::PollIsComplete()
 		return true;
 	}
 	UCookOnTheFlyServer& COTFS = PackageData.GetPackageDatas().GetCookOnTheFlyServer();
-	if (COTFS.RouteIsCachedCookedPlatformDataLoaded(PackageData, LocalObject, TargetPlatform, nullptr /* ExistingEvent */))
+	if (COTFS.RouteIsCachedCookedPlatformDataLoaded(PackageData, LocalObject, TargetPlatform,
+		nullptr /* ExistingEvent */))
 	{
 		Release();
 		return true;
@@ -2442,7 +2468,8 @@ bool FPendingCookedPlatformData::PollIsComplete()
 		{
 			if (GShaderCompilingManager->HasShaderJobs() == false)
 			{
-				UE_LOG(LogCook, Warning, TEXT("Shader compiler is in a bad state!  Shader %s is finished compile but shader compiling manager did not notify shader.  "),
+				UE_LOG(LogCook, Warning,
+				TEXT("Shader compiler is in a bad state!  Shader %s is finished compile but shader compiling manager did not notify shader.  "),
 					*LocalObject->GetPathName());
 			}
 		}*/
@@ -2713,7 +2740,8 @@ void FPackageDatas::SetBeginCookConfigSettings(FStringView CookShowInstigator)
 	{
 		FString LocalPath;
 		FString PackageName;
-		if (!FPackageName::TryConvertToMountedPath(CookShowInstigator, &LocalPath, &PackageName, nullptr, nullptr, nullptr))
+		if (!FPackageName::TryConvertToMountedPath(CookShowInstigator, &LocalPath, &PackageName,
+			nullptr, nullptr, nullptr))
 		{
 			UE_LOG(LogCook, Fatal, TEXT("-CookShowInstigator argument %.*s is not a mounted filename or packagename"),
 				CookShowInstigator.Len(), CookShowInstigator.GetData());
@@ -2779,12 +2807,15 @@ FPackageData& FPackageDatas::FindOrAddPackageData(const FName& PackageName, cons
 		if (PackageDataMapAddr != nullptr)
 		{
 			FPackageData** FileNameMapAddr = FileNameToPackageData.Find(NormalizedFileName);
-			checkf(FileNameMapAddr, TEXT("Package %s is being added with filename %s, but it already exists with filename %s, ")
+			checkf(FileNameMapAddr,
+				TEXT("Package %s is being added with filename %s, but it already exists with filename %s, ")
 				TEXT("and it is not present in FileNameToPackageData map under the new name."),
-				*PackageName.ToString(), *NormalizedFileName.ToString(), *(*PackageDataMapAddr)->GetFileName().ToString());
+				*PackageName.ToString(), *NormalizedFileName.ToString(),
+				*(*PackageDataMapAddr)->GetFileName().ToString());
 			checkf(*FileNameMapAddr == *PackageDataMapAddr,
 				TEXT("Package %s is being added with filename %s, but that filename maps to a different package %s."),
-				*PackageName.ToString(), *NormalizedFileName.ToString(), *(*FileNameMapAddr)->GetPackageName().ToString());
+				*PackageName.ToString(), *NormalizedFileName.ToString(),
+				*(*FileNameMapAddr)->GetPackageName().ToString());
 			return **PackageDataMapAddr;
 		}
 
@@ -2904,7 +2935,8 @@ FPackageData* FPackageDatas::TryAddPackageDataByStandardFileName(const FName& Fi
 				return *PackageDataMapAddr;
 			}
 		}
-		UE_LOG(LogCook, Warning, TEXT("Unexpected failure to cook filename '%s'. It is mapped to PackageName '%s', but does not exist on disk and we cannot verify the extension."),
+		UE_LOG(LogCook, Warning,
+			TEXT("Unexpected failure to cook filename '%s'. It is mapped to PackageName '%s', but does not exist on disk and we cannot verify the extension."),
 			*FileName.ToString(), *PackageName.ToString());
 		return nullptr;
 	}
@@ -2966,7 +2998,8 @@ FName FPackageDatas::LookupFileNameOnDisk(FName PackageName, bool bRequireExists
 	}
 	else if (!bRequireExists)
 	{
-		FString Extension = bCreateAsMap ? FPackageName::GetMapPackageExtension() : FPackageName::GetAssetPackageExtension();
+		FString Extension = bCreateAsMap ? FPackageName::GetMapPackageExtension() :
+			FPackageName::GetAssetPackageExtension();
 		if (!FPackageName::TryConvertLongPackageNameToFilename(PackageName.ToString(), FilenameOnDisk, Extension))
 		{
 			return NAME_None;
@@ -2995,7 +3028,8 @@ bool FPackageDatas::TryLookupFileNameOnDisk(FName PackageName, FString& OutFileN
 			if (!FPackageName::TryConvertLongPackageNameToFilename(PackageNameStr, OutFileName,
 				FPackageName::GetAssetPackageExtension()))
 			{
-				UE_LOG(LogCook, Warning, TEXT("Package %s exists in memory but its PackageRoot is not mounted. It will not be cooked."),
+				UE_LOG(LogCook, Warning,
+					TEXT("Package %s exists in memory but its PackageRoot is not mounted. It will not be cooked."),
 					*PackageNameStr);
 				return false;
 			}
@@ -3189,7 +3223,8 @@ int32 FPackageDatas::GetNumCooked(ECookResult CookResult)
 	return Monitor.GetNumCooked(CookResult);
 }
 
-void FPackageDatas::GetCookedPackagesForPlatform(const ITargetPlatform* Platform, TArray<FPackageData*>& SucceededPackages,
+void FPackageDatas::GetCookedPackagesForPlatform(const ITargetPlatform* Platform,
+	TArray<FPackageData*>& SucceededPackages,
 	TArray<FPackageData*>& FailedPackages)
 {
 	LockAndEnumeratePackageDatas(
@@ -3250,7 +3285,9 @@ void FPackageDatas::ClearCookResultsForPackages(const TSet<FName>& InPackages)
 			}
 		});
 
-	UE_LOG(LogCook, Display, TEXT("Cleared the cook results of %d packages because ClearCookResultsForPackages requested them to be recooked."), AffectedPackagesCount);
+	UE_LOG(LogCook, Display,
+		TEXT("Cleared the cook results of %d packages because ClearCookResultsForPackages requested them to be recooked."),
+		AffectedPackagesCount);
 }
 
 void FPackageDatas::OnRemoveSessionPlatform(const ITargetPlatform* TargetPlatform)
@@ -3324,7 +3361,8 @@ void FPackageDatas::PollPendingCookedPlatformDatas(bool bForce, double& LastCook
 		{
 			if (bEntering)
 			{
-				CookOnTheFlyServer.SetActivePackage(Package->GetFName(), PackageAccessTrackingOps::NAME_CookerBuildObject);
+				CookOnTheFlyServer.SetActivePackage(Package->GetFName(),
+					PackageAccessTrackingOps::NAME_CookerBuildObject);
 			}
 			else
 			{
@@ -3337,7 +3375,8 @@ void FPackageDatas::PollPendingCookedPlatformDatas(bool bForce, double& LastCook
 	if (LastCookableObjectTickTime + TickCookableObjectsFrameTime <= CurrentTime)
 	{
 		UE_SCOPED_COOKTIMER(TickCookableObjects);
-		FTickableCookObject::TickObjects(static_cast<float>(CurrentTime - LastCookableObjectTickTime), false /* bTickComplete */);
+		FTickableCookObject::TickObjects(static_cast<float>(CurrentTime - LastCookableObjectTickTime),
+			false /* bTickComplete */);
 		LastCookableObjectTickTime = CurrentTime;
 	}
 
@@ -3353,7 +3392,8 @@ void FPackageDatas::PollPendingCookedPlatformDatas(bool bForce, double& LastCook
 			}
 			else
 			{
-				Data.UpdatePeriodMultiplier = FMath::Clamp(Data.UpdatePeriodMultiplier*2, 1, PendingPlatformDataMaxUpdatePeriod);
+				Data.UpdatePeriodMultiplier = FMath::Clamp(Data.UpdatePeriodMultiplier*2, 1,
+					PendingPlatformDataMaxUpdatePeriod);
 				int32 ContainerIndex = Data.UpdatePeriodMultiplier - 1;
 				while (PendingCookedPlatformDataLists.Num() <= ContainerIndex)
 				{
@@ -3433,7 +3473,8 @@ void FPackageDatas::DebugInstigator(FPackageData& PackageData)
 			ChainText << TEXT("{ ") << Instigator.ToString() << TEXT(" }");
 			bFirst = false;
 		}
-		UE_LOG(LogCook, Display, TEXT("Instigator chain of %s: %s"), *PackageData.GetPackageName().ToString(), ChainText.ToString());
+		UE_LOG(LogCook, Display, TEXT("Instigator chain of %s: %s"),
+			*PackageData.GetPackageName().ToString(), ChainText.ToString());
 	}
 	UpdateThreadsafePackageData(PackageData);
 }
@@ -3621,17 +3662,24 @@ const TCHAR* LexToString(ECachedCookedPlatformDataEvent Value)
 {
 	switch (Value)
 	{
-	case ECachedCookedPlatformDataEvent::None: return TEXT("None");
-	case ECachedCookedPlatformDataEvent::BeginCacheForCookedPlatformDataCalled: return TEXT("BeginCacheForCookedPlatformDataCalled");
-	case ECachedCookedPlatformDataEvent::IsCachedCookedPlatformDataLoadedCalled: return TEXT("IsCachedCookedPlatformDataLoadedCalled");
-	case ECachedCookedPlatformDataEvent::IsCachedCookedPlatformDataLoadedReturnedTrue: return TEXT("IsCachedCookedPlatformDataLoadedReturnedTrue");
-	case ECachedCookedPlatformDataEvent::ClearCachedCookedPlatformDataCalled: return TEXT("ClearCachedCookedPlatformDataCalled");
-	case ECachedCookedPlatformDataEvent::ClearAllCachedCookedPlatformDataCalled: return TEXT("ClearAllCachedCookedPlatformDataCalled");
+	case ECachedCookedPlatformDataEvent::None:
+		return TEXT("None");
+	case ECachedCookedPlatformDataEvent::BeginCacheForCookedPlatformDataCalled:
+		return TEXT("BeginCacheForCookedPlatformDataCalled");
+	case ECachedCookedPlatformDataEvent::IsCachedCookedPlatformDataLoadedCalled:
+		return TEXT("IsCachedCookedPlatformDataLoadedCalled");
+	case ECachedCookedPlatformDataEvent::IsCachedCookedPlatformDataLoadedReturnedTrue:
+		return TEXT("IsCachedCookedPlatformDataLoadedReturnedTrue");
+	case ECachedCookedPlatformDataEvent::ClearCachedCookedPlatformDataCalled:
+		return TEXT("ClearCachedCookedPlatformDataCalled");
+	case ECachedCookedPlatformDataEvent::ClearAllCachedCookedPlatformDataCalled:
+		return TEXT("ClearAllCachedCookedPlatformDataCalled");
 	default: return TEXT("Invalid");
 	}
 }
 
-void FPackageDatas::CachedCookedPlatformDataObjectsPostGarbageCollect(const TSet<UObject*>& SaveQueueObjectsThatStillExist)
+void FPackageDatas::CachedCookedPlatformDataObjectsPostGarbageCollect(
+	const TSet<UObject*>& SaveQueueObjectsThatStillExist)
 {
 	for (TMap<UObject*, FCachedCookedPlatformDataState>::TIterator Iter(this->CachedCookedPlatformDataObjects);
 		Iter; ++Iter)
diff --git a/Engine/Source/Editor/UnrealEd/Private/Cooker/CookRequestCluster.cpp b/Engine/Source/Editor/UnrealEd/Private/Cooker/CookRequestCluster.cpp
--- a/Engine/Source/Editor/UnrealEd/Private/Cooker/CookRequestCluster.cpp
+++ b/Engine/Source/Editor/UnrealEd/Private/Cooker/CookRequestCluster.cpp
@@ -76,7 +76,8 @@ FRequestCluster::FRequestCluster(UCookOnTheFlyServer& InCOTFS)
 		{
 			if (PlatformData->bAllowIterativeResults != bAllowIterativeResults)
 			{
-				UE_LOG(LogCook, Warning, TEXT("Full build is requested for some platforms but not others, but this is not supported. All platforms will be built full."));
+				UE_LOG(LogCook, Warning,
+					TEXT("Full build is requested for some platforms but not others, but this is not supported. All platforms will be built full."));
 				bAllowIterativeResults = false;
 			}
 		}
@@ -153,7 +154,8 @@ FRequestCluster::FRequestCluster(UCookOnTheFlyServer& InCOTFS, TRingBuffer<FDisc
 			Discovery = &PoppedDiscovery;
 			if (!PackageData.IsInProgress() && PackageData.GetPlatformsNeedingCookingNum() == 0)
 			{
-				PackageData.SendToState(EPackageState::Request, ESendFlags::QueueRemove, EStateChangeReason::RequestCluster);
+				PackageData.SendToState(EPackageState::Request, ESendFlags::QueueRemove,
+					EStateChangeReason::RequestCluster);
 				OwnedPackageDatas.Add(&PackageData, ESuppressCookReason::NotSuppressed);
 			}
 			continue;
@@ -253,7 +255,8 @@ void FRequestCluster::FetchPackageNames(const FCookerTimer& CookerTimer, bool& b
 		// The input filenames are normalized, but might be missing their extension, so allow PackageDatas
 		// to correct the filename if the package is found with a different filename
 		bool bExactMatchRequired = false;
-		FPackageData* PackageData = PackageDatas.TryAddPackageDataByStandardFileName(OriginalName, bExactMatchRequired);
+		FPackageData* PackageData = PackageDatas.TryAddPackageDataByStandardFileName(OriginalName,
+			bExactMatchRequired);
 		if (!PackageData)
 		{
 			LogCookerMessage(FString::Printf(TEXT("Could not find package at file %s!"),
@@ -320,7 +323,8 @@ void FRequestCluster::PullIntoCluster(FPackageData& PackageData)
 		}
 		else
 		{
-			PackageData.SendToState(EPackageState::Request, ESendFlags::QueueRemove, EStateChangeReason::RequestCluster);
+			PackageData.SendToState(EPackageState::Request, ESendFlags::QueueRemove,
+				EStateChangeReason::RequestCluster);
 		}
 		Existing = ESuppressCookReason::NotSuppressed;
 	}
@@ -408,7 +412,8 @@ void FRequestCluster::OnRemoveSessionPlatform(const ITargetPlatform* TargetPlatf
 		bool bComplete;
 		while (PumpExploration(CookerTimer, bComplete), !bComplete)
 		{
-			UE_LOG(LogCook, Display, TEXT("Waiting for RequestCluster to finish before removing platform from session."));
+			UE_LOG(LogCook, Display,
+				TEXT("Waiting for RequestCluster to finish before removing platform from session."));
 			FPlatformProcess::Sleep(.001f);
 		}
 	}
@@ -483,7 +488,8 @@ void FRequestCluster::PumpExploration(const FCookerTimer& CookerTimer, bool& bOu
 		}
 		else
 		{
-			TraversalTier = bAllowHardDependencies ? ETraversalTier::FollowDependencies : ETraversalTier::FetchEdgeData;
+			TraversalTier = bAllowHardDependencies ? ETraversalTier::FollowDependencies :
+				ETraversalTier::FetchEdgeData;
 		}
 		GraphSearch.Reset(new FGraphSearch(*this, TraversalTier));
 
@@ -620,8 +626,8 @@ void FRequestCluster::FGraphSearch::StartSearch()
 	{
 		FVertexData& Vertex = FindOrAddVertex(Pair.Key->GetPackageName(), *Pair.Key);
 		check(Vertex.PackageData);
-		// We're iterating over OwnedPackageDatas and the Vertex is already in the Cluster so we don't need to call AddToFrontier;
-		// just add it directly.
+		// We're iterating over OwnedPackageDatas and the Vertex is already in the Cluster so we don't need to call
+		// AddToFrontier; just add it directly.
 		check(Pair.Value != ESuppressCookReason::Invalid);
 		Frontier.Add(&Vertex);
 	}
@@ -638,8 +644,8 @@ FRequestCluster::FGraphSearch::~FGraphSearch()
 			bAsyncBatchesEmpty = AsyncQueueBatches.IsEmpty();
 			if (!bAsyncBatchesEmpty)
 			{
-				// It is safe to Reset AsyncResultsReadyEvent and wait on it later because we are inside the lock and there is a
-				// remaining batch, so it will be triggered after the Reset when that batch completes.
+				// It is safe to Reset AsyncResultsReadyEvent and wait on it later because we are inside the lock and
+				// there is a remaining batch, so it will be triggered after the Reset when that batch completes.
 				AsyncResultsReadyEvent->Reset();
 			}
 		}
@@ -698,7 +704,8 @@ void FRequestCluster::FGraphSearch::OnNewReachablePlatforms(FPackageData* Packag
 	Frontier.Add(*VertexPtr);
 }
 
-void FRequestCluster::FGraphSearch::QueueEdgesFetch(FVertexData& Vertex, TConstArrayView<const ITargetPlatform*> Platforms)
+void FRequestCluster::FGraphSearch::QueueEdgesFetch(FVertexData& Vertex,
+	TConstArrayView<const ITargetPlatform*> Platforms)
 {
 	check(!Vertex.QueryData);
 	FVertexQueryData& QueryData = *AllocateQueryData();
@@ -711,7 +718,10 @@ void FRequestCluster::FGraphSearch::QueueEdgesFetch(FVertexData& Vertex, TConstA
 	bool bHasPlatformAgnostic = false;
 	for (const ITargetPlatform* Platform : Platforms)
 	{
-		int32 Index = Algo::BinarySearchBy(FetchPlatforms, Platform, [](const FFetchPlatformData& D) { return D.Platform; });
+		int32 Index = Algo::BinarySearchBy(FetchPlatforms, Platform, [](const FFetchPlatformData& D) 
+			{
+				return D.Platform;
+			});
 		check(Index != INDEX_NONE);
 		QueryData.Platforms[Index].bActive = true;
 		if (Platform != CookerLoadingPlatformKey)
@@ -759,7 +769,8 @@ void FRequestCluster::FGraphSearch::TickExploration(bool& bOutDone)
 		{
 			if (Vertex->QueryData)
 			{
-				// Vertices that are already in the AsyncQueue can not be added again; we would clobber their QueryData. Postpone them.
+				// Vertices that are already in the AsyncQueue can not be added again; we would clobber their
+				// QueryData. Postpone them.
 				BusyVertices.Add(Vertex);
 			}
 			else
@@ -834,7 +845,8 @@ void FRequestCluster::FGraphSearch::UpdateDisplay()
 			NumVertices += Batch->PendingVertices;
 		}
 
-		UE_LOG(LogCook, Warning, TEXT("FRequestCluster waited more than %.0lfs for previous build results from the oplog. ")
+		UE_LOG(LogCook, Warning,
+			TEXT("FRequestCluster waited more than %.0lfs for previous build results from the oplog. ")
 			TEXT("NumPendingBatches == %d, NumPendingVertices == %d. Continuing to wait..."),
 			WarningTimeout, NumBatches, NumVertices);
 		LastActivityTime = FPlatformTime::Seconds();
@@ -858,7 +870,8 @@ void FRequestCluster::FGraphSearch::VisitVertex(FVertexData& Vertex)
 	const ITargetPlatform* FirstReachableSessionPlatform = nullptr;
 	ESuppressCookReason SuppressCookReason = ESuppressCookReason::Invalid;
 	bool bAllReachablesUncookable = true;
-	for (TPair<const ITargetPlatform*, FPackagePlatformData>& Pair : Vertex.PackageData->GetPlatformDatasConstKeysMutableValues())
+	for (TPair<const ITargetPlatform*, FPackagePlatformData>& Pair :
+		Vertex.PackageData->GetPlatformDatasConstKeysMutableValues())
 	{
 		FPackagePlatformData& PlatformData = Pair.Value;
 		if (Pair.Key == CookerLoadingPlatformKey)
@@ -917,8 +930,8 @@ void FRequestCluster::FGraphSearch::VisitVertex(FVertexData& Vertex)
 		Vertex.bAnyCookable = bAnyCookable;
 	}
 
-	// If any target platform is cookable, then we need to mark the CookerLoadingPlatform as reachable because we will need
-	// to load the package to cook it
+	// If any target platform is cookable, then we need to mark the CookerLoadingPlatform as reachable because we will
+	// need to load the package to cook it
 	if (bAnyCookable)
 	{
 		if (!CookerLoadingPlatform)
@@ -961,7 +974,8 @@ void FRequestCluster::FGraphSearch::VisitVertexForPlatform(FVertexData& Vertex,
 	}
 	else
 	{
-		check(SuppressCookReason != ESuppressCookReason::Invalid && SuppressCookReason != ESuppressCookReason::NotSuppressed);
+		check(SuppressCookReason != ESuppressCookReason::Invalid
+			&& SuppressCookReason != ESuppressCookReason::NotSuppressed);
 		if (AccumulatedSuppressCookReason == ESuppressCookReason::Invalid)
 		{
 			AccumulatedSuppressCookReason = SuppressCookReason;
@@ -1059,7 +1073,8 @@ void FRequestCluster::FGraphSearch::ExploreVertexEdges(FVertexData& Vertex)
 		}
 		PlatformDependencyBits.HasPlatformByIndex[PlatformIndex] = true;
 
-		// Calculate PlatformDependencyType.InstigatorType == Max(InstigatorType, PlatformDependencyType.InstigatorType)
+		// Calculate PlatformDependencyType.InstigatorType == 
+		// Max(InstigatorType, PlatformDependencyType.InstigatorType)
 		// based on the enum values, from least required to most: [ Soft, HardEditorOnly, Hard ]
 		switch (InstigatorType)
 		{
@@ -1094,9 +1109,11 @@ void FRequestCluster::FGraphSearch::ExploreVertexEdges(FVertexData& Vertex)
 
 	FQueryPlatformData& PlatformAgnosticQueryPlatformData = Vertex.QueryData->Platforms[PlatformAgnosticPlatformIndex];
 
-	auto ProcessPlatformAttachments = [this, PackageName, &PackageData, &PlatformAgnosticQueryPlatformData, &AddPlatformDependencyRange]
+	auto ProcessPlatformAttachments =
+		[this, PackageName, &PackageData, &PlatformAgnosticQueryPlatformData, &AddPlatformDependencyRange]
 		(int32 PlatformIndex, const ITargetPlatform* TargetPlatform, FFetchPlatformData& FetchPlatformData,
-			FPackagePlatformData& PackagePlatformData, FCookAttachments& PlatformAttachments, bool bExploreDependencies)
+			FPackagePlatformData& PackagePlatformData, FCookAttachments& PlatformAttachments,
+			bool bExploreDependencies)
 	{
 		bool bFoundBuildDefinitions = false;
 		ICookedPackageWriter* PackageWriter = FetchPlatformData.Writer;
@@ -1131,7 +1148,8 @@ void FRequestCluster::FGraphSearch::ExploreVertexEdges(FVertexData& Vertex)
 				}
 			}
 			bool bShouldIterativelySkip = bIterativelyUnmodified;
-			PackageWriter->UpdatePackageModificationStatus(PackageName, bIterativelyUnmodified, bShouldIterativelySkip);
+			PackageWriter->UpdatePackageModificationStatus(PackageName, bIterativelyUnmodified,
+				bShouldIterativelySkip);
 			if (bShouldIterativelySkip)
 			{
 				// Call SetPlatformCooked instead of just PackagePlatformData.SetCookResults because we might also need
@@ -1192,17 +1210,20 @@ void FRequestCluster::FGraphSearch::ExploreVertexEdges(FVertexData& Vertex)
 
 			// ITERATIVECOOK_TODO: Build dependencies need to be stored and used to mark package loads as expected
 			// But we can't use them to explore packages that will be loaded during cook because they might not be;
-			// some build dependencies might be a conservative list but unused by the asset, or unused on targetplatform
+			// some build dependencies might be a conservative list but unused by the asset, or unused on
+			// targetplatform.
 			// Adding BuildDependencies also sets up many circular dependencies, because maps declare their external
 			// actors as build dependencies and the external actors declare the map as a build or hard dependency.
-			// Topological sort done at the end of the Cluster has poor performance when there are 100k+ circular dependencies.
+			// Topological sort done at the end of the Cluster has poor performance when there are 100k+ circular
+			// dependencies.
 			constexpr bool bAddBuildDependenciesToGraph = false;
 			if (bAddBuildDependenciesToGraph)
 			{
-				Cluster.AssetRegistry.GetDependencies(PackageName, CookerLoadingDependencies, EDependencyCategory::Package,
-					EDependencyQuery::Build);
+				Cluster.AssetRegistry.GetDependencies(PackageName, CookerLoadingDependencies,
+					EDependencyCategory::Package, EDependencyQuery::Build);
 			}
-			// CookerLoadingPlatform does not cause SetInstigator so it does not modify the platformdependency's InstigatorType
+			// CookerLoadingPlatform does not cause SetInstigator so it does not modify the platformdependency's
+			// InstigatorType
 			AddPlatformDependencyRange(CookerLoadingDependencies, PlatformIndex, EInstigator::InvalidCategory);
 		}
 		else
@@ -1478,9 +1499,12 @@ void FRequestCluster::FQueryVertexBatch::Send()
 		}
 		FFetchPlatformData& FetchPlatformData = ThreadSafeOnlyVars.FetchPlatforms[PlatformIndex];
 
-		if (ThreadSafeOnlyVars.Cluster.IsIncrementalCook() // Only FetchCookAttachments if our cookmode supports it. Otherwise keep them all empty
-			&& !FetchPlatformData.bIsPlatformAgnosticPlatform // The PlatformAgnosticPlatform has no stored CookAttachments; always use empty
-			&& !FetchPlatformData.bIsCookerLoadingPlatform // The CookerLoadingPlatform has no stored CookAttachments; always use empty
+		if (ThreadSafeOnlyVars.Cluster.IsIncrementalCook() // Only FetchCookAttachments if our cookmode supports it.
+															// Otherwise keep them all empty
+			&& !FetchPlatformData.bIsPlatformAgnosticPlatform // The PlatformAgnosticPlatform has no stored
+															// CookAttachments; always use empty
+			&& !FetchPlatformData.bIsCookerLoadingPlatform // The CookerLoadingPlatform has no stored CookAttachments;
+															// always use empty
 			)
 		{
 			TUniqueFunction<void(FName PackageName, UE::TargetDomain::FCookAttachments&& Result)> Callback =
@@ -1565,13 +1589,16 @@ void FRequestCluster::IsRequestCookable(const ITargetPlatform* Platform, FName P
 }
 
 void FRequestCluster::IsRequestCookable(const ITargetPlatform* Platform, FName PackageName, FPackageData& PackageData,
-	UCookOnTheFlyServer& InCOTFS, FStringView InDLCPath, ESuppressCookReason& OutReason, bool& bOutCookable, bool& bOutExplorable)
+	UCookOnTheFlyServer& InCOTFS, FStringView InDLCPath, ESuppressCookReason& OutReason, bool& bOutCookable,
+	bool& bOutExplorable)
 {
-	check(Platform != CookerLoadingPlatformKey); // IsRequestCookable should not be called for The CookerLoadingPlatform; it has different rules
+	// IsRequestCookable should not be called for The CookerLoadingPlatform; it has different rules
+	check(Platform != CookerLoadingPlatformKey);
 
 	TStringBuilder<256> NameBuffer;
-	// We need to reject packagenames from adding themselves or their transitive dependencies using all the same rules that
-	// UCookOnTheFlyServer::ProcessRequest uses. Packages that are rejected from cook do not add their dependencies to the cook.
+	// We need to reject packagenames from adding themselves or their transitive dependencies using all the same rules
+	// that UCookOnTheFlyServer::ProcessRequest uses. Packages that are rejected from cook do not add their 
+	// dependencies to the cook.
 	PackageName.ToString(NameBuffer);
 	if (FPackageName::IsScriptPackage(NameBuffer))
 	{
@@ -1591,7 +1618,8 @@ void FRequestCluster::IsRequestCookable(const ITargetPlatform* Platform, FName P
 	FName FileName = PackageData.GetFileName();
 	if (InCOTFS.PackageTracker->NeverCookPackageList.Contains(FileName))
 	{
-		if (INDEX_NONE != UE::String::FindFirst(NameBuffer, ULevel::GetExternalActorsFolderName(), ESearchCase::IgnoreCase))
+		if (INDEX_NONE != UE::String::FindFirst(NameBuffer, ULevel::GetExternalActorsFolderName(), 
+			ESearchCase::IgnoreCase))
 		{
 			// EXTERNALACTOR_TODO: Add a separate category for ExternalActors rather than putting them in
 			// NeverCookPackageList and checking naming convention here.
@@ -1600,14 +1628,17 @@ void FRequestCluster::IsRequestCookable(const ITargetPlatform* Platform, FName P
 
 			// EXTERNALACTOR_TODO: We want to explore externalactors, because they add references to the cook that will
 			// otherwise not be found until the map package loads them and adds them as unsolicited packages
-			// But some externalactor packages will never be loaded by the generator, and we don't have a way to discover which
-			// ones will not be loaded until we load the Map and WorldPartition object.
-			// So set them to explorable = false until we implement an interface to determine which actors will be loaded up front.
+			// But some externalactor packages will never be loaded by the generator, and we don't have a way to
+			// discover which ones will not be loaded until we load the Map and WorldPartition object.
+			// So set them to explorable = false until we implement an interface to determine which actors will be
+			// loaded up front.
 			bOutExplorable = false;
 		}
 		else
 		{
-			UE_LOG(LogCook, Verbose, TEXT("Package %s is referenced but is in the never cook package list, discarding request"), *NameBuffer);
+			UE_LOG(LogCook, Verbose,
+				TEXT("Package %s is referenced but is in the never cook package list, discarding request"),
+				*NameBuffer);
 			OutReason = ESuppressCookReason::NeverCook;
 			bOutCookable = false;
 			bOutExplorable = false;
@@ -1638,11 +1669,13 @@ void FRequestCluster::IsRequestCookable(const ITargetPlatform* Platform, FName P
 			{
 				if (!PackageData.HasCookedPlatform(Platform, true /* bIncludeFailed */))
 				{
-					// AllowUncookedAssetReferences should only be used when the DLC plugin to cook is going to be mounted where uncooked packages are available.
-					// This will allow a DLC plugin to be recooked continually and mounted in an uncooked editor which is useful for CI.
+					// AllowUncookedAssetReferences should only be used when the DLC plugin to cook is going to be
+					// mounted where uncooked packages are available. This will allow a DLC plugin to be recooked 
+					// continually and mounted in an uncooked editor which is useful for CI.
 					if (!InCOTFS.CookByTheBookOptions->bAllowUncookedAssetReferences)
 					{
-						UE_LOG(LogCook, Error, TEXT("Uncooked Engine or Game content %s is being referenced by DLC!"), *NameBuffer);
+						UE_LOG(LogCook, Error,
+							TEXT("Uncooked Engine or Game content %s is being referenced by DLC!"), *NameBuffer);
 					}
 				}
 				OutReason = ESuppressCookReason::NotInCurrentPlugin;
@@ -1826,8 +1859,8 @@ void FRequestCluster::RandomizeCookOrder(TArray<FPackageData*>& InOutLeafToRootO
 	// 
 	// Notes on the practical solution:
 	// Cycles are not supposed to be a large part of our graph, so we will not seek to do a good job of randomizing
-	// them. We do a DFS as we would in an acyclic graph, and when we encounter an already-visited vertex due to a cycle
-	// we just pretend that edge does not exist.
+	// them. We do a DFS as we would in an acyclic graph, and when we encounter an already-visited vertex due to a
+	// cycle we just pretend that edge does not exist.
 	//
 	// Two DFS passes
 	// Pass 1, find all roots by iterating from root to leaf and DFSing each remaining head element.
@@ -1875,7 +1908,8 @@ void FRequestCluster::RandomizeCookOrder(TArray<FPackageData*>& InOutLeafToRootO
 	TSet<FPackageData*> ValidVertices;
 	ValidVertices.Append(InOutLeafToRootOrder);
 
-	TArray<FPackageData*> Roots = FindRootsFromLeafToRootOrderList<FPackageData*>(InOutLeafToRootOrder, Edges, ValidVertices);
+	TArray<FPackageData*> Roots = FindRootsFromLeafToRootOrderList<FPackageData*>(InOutLeafToRootOrder, Edges, 
+		ValidVertices);
 	ArrayShuffle(Roots);
 
 	TSet<FPackageData*> Visited;
@@ -1910,7 +1944,8 @@ void FRequestCluster::RandomizeCookOrder(TArray<FPackageData*>& InOutLeafToRootO
 			}
 			if (!bPushed)
 			{
-				// Add this vertex now as the most rootwards encountered; it is farther toward the root than any of the vertices it depends on
+				// Add this vertex now as the most rootwards encountered; it is farther toward the root than any of the
+				// vertices it depends on
 				OutputVisitOrder.Add(Entry.Vertex);
 				Pop();
 			}
diff --git a/Engine/Source/Editor/UnrealEd/Private/Cooker/CookTypes.h b/Engine/Source/Editor/UnrealEd/Private/Cooker/CookTypes.h
--- a/Engine/Source/Editor/UnrealEd/Private/Cooker/CookTypes.h
+++ b/Engine/Source/Editor/UnrealEd/Private/Cooker/CookTypes.h
@@ -79,7 +79,8 @@ struct TFastPointerMapKeyFuncs : public TDefaultMapKeyFuncs<KeyType, ValueType,
 
 /** A TMap which uses TFastPointerMapKeyFuncs instead of TDefaultMapKeyFuncs */
 template<typename KeyType, typename ValueType, typename SetAllocator = FDefaultSetAllocator>
-class TFastPointerMap : public TMap<KeyType, ValueType, SetAllocator, TFastPointerMapKeyFuncs<KeyType, ValueType, false>>
+class TFastPointerMap : public TMap<KeyType, ValueType, SetAllocator,
+	TFastPointerMapKeyFuncs<KeyType, ValueType, false>>
 {};
 
 /** A TSet which uses TFastPointerSetKeyFuncs instead of DefaultKeyFuncs */
@@ -100,7 +101,10 @@ namespace UE::Cook
 		using TFastPointerSet<FPackageData*>::TFastPointerSet;
 	};
 
-	/** External Requests to the cooker can either by cook requests for a specific file, or arbitrary callbacks that need to execute within the Scheduler's lock. */
+	/**
+	 * External Requests to the cooker can either by cook requests for a specific file, or arbitrary callbacks that
+	 * need to execute within the Scheduler's lock.
+	 */
 	enum class EExternalRequestType
 	{
 		None,
@@ -168,25 +172,51 @@ namespace UE::Cook
 	/** Which phase of cooking a Package is in.  */
 	enum class EPackageState
 	{
-		Idle = 0,	  /* The Package is not being operated on by the cooker, and is not in any queues.  This is the state both for packages that have never been requested and for packages that have finished cooking. */
-		Request,	  /* The Package is in the RequestQueue; it is requested for cooking but has not had any operations performed on it. */
-		AssignedToWorker, /* The Package is in the AssignedToWorkerSet; it has been sent a remote CookWorker for cooking and has not had any operations performed on it locally. */
-		LoadPrepare,  /* The Package is in the LoadPrepareQueue. Preloading is in progress. */
-		LoadReady,	  /* The package is in the LoadReadyQueue. Preloading is complete and it will be loaded when its turn comes up. */
-		Save,		  /* The Package is in the SaveQueue; it has been fully loaded and some target data may have been calculated. */
+		/**
+		 * The Package is not being operated on by the cooker, and is not in any queues. This is the state both for
+		 * packages that have never been requested and for packages that have finished cooking.
+		 */
+		Idle = 0,
+		/**
+		 * The Package is in the RequestQueue; it is requested for cooking but has not had any operations performed
+		 * on it.
+		*/
+		Request,
+		/**
+		 * The Package is in the AssignedToWorkerSet; it has been sent a remote CookWorker for cooking and has not
+		 * had any operations performed on it locally.
+		 */
+		AssignedToWorker,
+		/** The Package is in the LoadPrepareQueue. Preloading is in progress. */
+		LoadPrepare,
+		/**
+		 * The package is in the LoadReadyQueue. Preloading is complete and it will be loaded when its turn
+		 * comes up.
+		 */
+		LoadReady,
+		/** The Package is in the SaveQueue; it has been fully loaded and some target data may have been calculated. */
+		Save,
 
 		Min = Idle,
 		Max = Save,
-		Count = Max + 1, /* Number of values in this enum, not a valid value for any EPackageState variable. */
-		BitCount = 3, /* Number of bits required to store a valid EPackageState */
+		/** Number of values in this enum, not a valid value for any EPackageState variable. */
+		Count = Max + 1,
+		/** Number of bits required to store a valid EPackageState */
+		BitCount = 3,
 	};
 
 	enum class EPackageStateProperty // Bitfield
 	{
 		None		= 0,
-		InProgress	= 0x1, /* The package is being worked on by the cooker. */
-		Loading		= 0x2, /* The package is in one of the loading states and has preload data. */
-		HasPackage	= 0x4, /* The package has progressed past the loading state, and the UPackage pointer is available on the FPackageData. */
+		/** The package is being worked on by the cooker. */
+		InProgress	= 0x1,
+		/** The package is in one of the loading states and has preload data. */
+		Loading		= 0x2,
+		/**
+		 * The package has progressed past the loading state, and the UPackage pointer is available on the
+		 * FPackageData.
+		 */
+		HasPackage	= 0x4,
 
 		Min = InProgress,
 		Max = HasPackage
@@ -241,7 +271,10 @@ namespace UE::Cook
 		double LoopStartTime = 0.;
 		/** A bitmask of flags of type enum ECookOnTheSideResult that were set during the tick. */
 		uint32 ResultFlags = 0;
-		/** The CookerTimer for the current tick. Used by slow reentrant operations that need to check whether they have timed out. */
+		/**
+		 * The CookerTimer for the current tick. Used by slow reentrant operations that need to check whether they
+		 * have timed out.
+		 */
 		FCookerTimer Timer;
 		/** CookFlags describing details of the caller's desired behavior for the current tick. */
 		ECookTickFlags TickFlags;
@@ -326,7 +359,8 @@ namespace UE::Cook
 		int32 SoftGCStartNumerator;
 		int32 SoftGCDenominator;
 		TArray<FString> ConfigSettingDenyList;
-		TMap<FName, int32> MaxAsyncCacheForType; // max number of objects of a specific type which are allowed to async cache at once
+		/** max number of objects of a specific type which are allowed to async cache at once */
+		TMap<FName, int32> MaxAsyncCacheForType;
 		bool bUseSoftGC = false;
 
 		friend FCbWriter& ::operator<<(FCbWriter& Writer, const UE::Cook::FInitializeConfigSettings& Value);
@@ -382,10 +416,16 @@ public:
 	/** Create a release from this manifest and store it in the releases directory for this cgame */
 	FString							CreateReleaseVersion;
 
-	/** If we are based on a release version of the game this is the set of packages which were cooked in that release. Map from platform name to list of uncooked package filenames */
+	/**
+	 * If we are based on a release version of the game this is the set of packages which were cooked in that release.
+	 * Map from platform name to list of uncooked package filenames.
+	 */
 	TMap<FName, TArray<FName>>		BasedOnReleaseCookedPackages;
 
-	/** Mapping from source packages to their localized variants (based on the culture list in FCookByTheBookStartupOptions) */
+	/**
+	 * Mapping from source packages to their localized variants (based on the culture list in
+	 * FCookByTheBookStartupOptions)
+	 */
 	TMap<FName, TArray<FName>>		SourceToLocalizedPackageVariants;
 	/** List of all the cultures (e.g. "en") that need to be cooked */
 	TArray<FString>					AllCulturesToCook;
@@ -404,7 +444,8 @@ public:
 
 	/** error when detecting engine content being used in this cook */
 	bool							bErrorOnEngineContentUse = false;
-	bool							bAllowUncookedAssetReferences = false; // this is a flag for dlc, will allow DLC to be cook when the fixed base might be missing references.
+	/** this is a flag for dlc, will allow DLC to be cook when the fixed base might be missing references. */
+	bool							bAllowUncookedAssetReferences = false;
 	bool							bSkipHardReferences = false;
 	bool							bSkipSoftReferences = false;
 	bool							bCookAgainstFixedBase = false;
@@ -428,7 +469,10 @@ struct FCookOnTheFlyOptions
 {
 	/** Wether the network file server or the I/O store connection server should bind to any port */
 	bool bBindAnyPort = false;
-	/** Whether the network file server should use a platform-specific communication protocol instead of TCP (used when bZenStore == false) */
+	/**
+	 * Whether the network file server should use a platform-specific communication protocol instead of TCP (used when
+	 * bZenStore == false)
+	 */
 	bool bPlatformProtocol = false;
 
 	friend FCbWriter& ::operator<<(FCbWriter& Writer, const UE::Cook::FCookOnTheFlyOptions& Value);
@@ -469,7 +513,10 @@ struct FDiscoveredPlatformSet
 		TArray<const ITargetPlatform*, TInlineAllocator<ExpectedMaxNumPlatforms>>* OutBuffer);
 	/** If the current type is EmbeddedBitField, change it to EmbeddedList. */
 	void ConvertFromBitfield(TConstArrayView<const ITargetPlatform*> OrderedPlatforms);
-	/** If the current type is EmbeddedList, change it to EmbeddedBitfield. Asserts if the type is already EmbeddedBitfield. */
+	/**
+	 * If the current type is EmbeddedList, change it to EmbeddedBitfield. Asserts if the type is already
+	 * EmbeddedBitfield.
+	 */
 	void ConvertToBitfield(TConstArrayView<const ITargetPlatform*> OrderedPlatforms);
 
 private:
@@ -509,17 +556,29 @@ struct FBeginCookContextPlatform
 	bool bHasMemoryResults = false;
 	/** If true, we should delete the in-memory results from an earlier cook in the same process, if we have any. */
 	bool bClearMemoryResults = false;
-	/** If true, we should load results that previous cooks left on disk into the current cook's results; this is one way to cook iteratively. */
+	/**
+	 * If true, we should load results that previous cooks left on disk into the current cook's results; this is one
+	 * way to cook iteratively.
+	 */
 	bool bPopulateMemoryResultsFromDiskResults = false;
-	/** If true we are cooking iteratively, from results in a shared build (e.g. from buildfarm) rather than from our previous cook. */
+	/**
+	 * If true we are cooking iteratively, from results in a shared build (e.g. from buildfarm) rather than from our
+	 * previous cook.
+	 */
 	bool bIterateSharedBuild = false;
-	/** If true we are a CookWorker, and we are working on a Sandbox directory that has already been populated by a remote Director process. */
+	/**
+	 * If true we are a CookWorker, and we are working on a Sandbox directory that has already been populated by a
+	 * remote Director process.
+	 */
 	bool bWorkerOnSharedSandbox = false;
 };
 FCbWriter& operator<<(FCbWriter& Writer, const FBeginCookContextPlatform& Value);
 bool LoadFromCompactBinary(FCbFieldView Field, FBeginCookContextPlatform& Value);
 
-/** Data held on the stack and shared with multiple subfunctions when running StartCookByTheBook or StartCookOnTheFly */
+/**
+ * Data held on the stack and shared with multiple subfunctions when running StartCookByTheBook or 
+ * StartCookOnTheFly
+ */
 struct FBeginCookContext
 {
 	FBeginCookContext(UCookOnTheFlyServer& InCOTFS)
@@ -541,7 +600,10 @@ struct FBeginCookContextForWorkerPlatform
 	void Set(const FBeginCookContextPlatform& InContext);
 
 	const ITargetPlatform* TargetPlatform = nullptr;
-	/** If true, we are deleting all old results from disk and rebuilding every package. If false, we are building iteratively. */
+	/**
+	 * If true, we are deleting all old results from disk and rebuilding every package. If false, we are building
+	 * iteratively.
+	 */
 	bool bFullBuild = false;
 };
 FCbWriter& operator<<(FCbWriter& Writer, const FBeginCookContextForWorkerPlatform& Value);
diff --git a/Engine/Source/Editor/UnrealEd/Private/Cooker/CookWorkerClient.cpp b/Engine/Source/Editor/UnrealEd/Private/Cooker/CookWorkerClient.cpp
--- a/Engine/Source/Editor/UnrealEd/Private/Cooker/CookWorkerClient.cpp
+++ b/Engine/Source/Editor/UnrealEd/Private/Cooker/CookWorkerClient.cpp
@@ -57,7 +57,8 @@ FCookWorkerClient::~FCookWorkerClient()
 	if (ConnectStatus == EConnectStatus::Connected ||
 		(EConnectStatus::FlushAndAbortFirst <= ConnectStatus && ConnectStatus <= EConnectStatus::FlushAndAbortLast))
 	{
-		UE_LOG(LogCook, Warning, TEXT("CookWorker was destroyed before it finished Disconnect. The CookDirector may be missing some information."));
+		UE_LOG(LogCook, Warning,
+			TEXT("CookWorker was destroyed before it finished Disconnect. The CookDirector may be missing some information."));
 	}
 	Sockets::CloseSocket(ServerSocket);
 
@@ -206,8 +207,8 @@ void FCookWorkerClient::ReportPromoteToSaveComplete(FPackageData& PackageData)
 		FPackagePlatformData& PackagePlatformData = PackageData.FindOrAddPlatformData(TargetPlatform);
 		if (!PackagePlatformData.IsCookAttempted() || PackagePlatformData.IsReportedToDirector())
 		{
-			// We didn't attempt to cook this platform for this package, or we cooked it previously and already sent the
-			// information about it
+			// We didn't attempt to cook this platform for this package, or we cooked it previously and already sent
+			// the information about it
 			PlatformResults.SetCookResults(ECookResult::Invalid);
 		}
 		else
@@ -360,7 +361,8 @@ void FCookWorkerClient::CreateServerSocket(const FDirectorConnectionInfo& Connec
 	ISocketSubsystem* SocketSubsystem = ISocketSubsystem::Get();
 	if (!SocketSubsystem)
 	{
-		UE_LOG(LogCook, Error, TEXT("CookWorker initialization failure: platform does not support network sockets, cannot connect to CookDirector."));
+		UE_LOG(LogCook, Error,
+			TEXT("CookWorker initialization failure: platform does not support network sockets, cannot connect to CookDirector."));
 		SendToState(EConnectStatus::LostConnection);
 		return;
 	}
@@ -368,7 +370,8 @@ void FCookWorkerClient::CreateServerSocket(const FDirectorConnectionInfo& Connec
 	DirectorAddr = Sockets::GetAddressFromStringWithPort(DirectorURI);
 	if (!DirectorAddr)
 	{
-		UE_LOG(LogCook, Error, TEXT("CookWorker initialization failure: could not convert -CookDirectorHost=%s into an address, cannot connect to CookDirector."),
+		UE_LOG(LogCook, Error,
+			TEXT("CookWorker initialization failure: could not convert -CookDirectorHost=%s into an address, cannot connect to CookDirector."),
 			*DirectorURI);
 		SendToState(EConnectStatus::LostConnection);
 		return;
@@ -386,10 +389,12 @@ void FCookWorkerClient::CreateServerSocket(const FDirectorConnectionInfo& Connec
 
 	constexpr float WaitForConnectTimeout = 60.f * 10;
 	float ConditionalTimeoutSeconds = IsCookIgnoreTimeouts() ? MAX_flt : WaitForConnectTimeout;
-	bool bServerSocketReady = ServerSocket->Wait(ESocketWaitConditions::WaitForWrite, FTimespan::FromSeconds(ConditionalTimeoutSeconds));
+	bool bServerSocketReady = ServerSocket->Wait(ESocketWaitConditions::WaitForWrite,
+		FTimespan::FromSeconds(ConditionalTimeoutSeconds));
 	if (!bServerSocketReady)
 	{
-		UE_LOG(LogCook, Error, TEXT("CookWorker initialization failure: Timed out after %.0f seconds trying to connect to CookDirector."),
+		UE_LOG(LogCook, Error,
+			TEXT("CookWorker initialization failure: Timed out after %.0f seconds trying to connect to CookDirector."),
 			ConditionalTimeoutSeconds);
 		SendToState(EConnectStatus::LostConnection);
 		return;
@@ -424,7 +429,8 @@ void FCookWorkerClient::PollWriteConnectMessage()
 		if (FPlatformTime::Seconds() - ConnectStartTimeSeconds > CookWorkerClient::WaitForConnectReplyTimeout &&
 			!IsCookIgnoreTimeouts())
 		{
-			UE_LOG(LogCook, Error, TEXT("CookWorker initialization failure: timed out waiting for %fs to send ConnectMessage."),
+			UE_LOG(LogCook, Error,
+				TEXT("CookWorker initialization failure: timed out waiting for %fs to send ConnectMessage."),
 				CookWorkerClient::WaitForConnectReplyTimeout);
 			SendToState(EConnectStatus::LostConnection);
 		}
@@ -456,7 +462,8 @@ void FCookWorkerClient::PollReceiveConfigMessage()
 		if (FPlatformTime::Seconds() - ConnectStartTimeSeconds > CookWorkerClient::WaitForConnectReplyTimeout &&
 			!IsCookIgnoreTimeouts())
 		{
-			UE_LOG(LogCook, Error, TEXT("CookWorker initialization failure: timed out waiting for %fs to receive InitialConfigMessage."),
+			UE_LOG(LogCook, Error,
+				TEXT("CookWorker initialization failure: timed out waiting for %fs to receive InitialConfigMessage."),
 				CookWorkerClient::WaitForConnectReplyTimeout);
 			SendToState(EConnectStatus::LostConnection);
 		}
@@ -465,7 +472,8 @@ void FCookWorkerClient::PollReceiveConfigMessage()
 	
 	if (Messages[0].MessageType != FInitialConfigMessage::MessageType)
 	{
-		UE_LOG(LogCook, Warning, TEXT("CookWorker initialization failure: Director sent a different message before sending an InitialConfigMessage. MessageType: %s."),
+		UE_LOG(LogCook, Warning,
+			TEXT("CookWorker initialization failure: Director sent a different message before sending an InitialConfigMessage. MessageType: %s."),
 			*Messages[0].MessageType.ToString());
 		SendToState(EConnectStatus::LostConnection);
 		return;
@@ -474,7 +482,8 @@ void FCookWorkerClient::PollReceiveConfigMessage()
 	InitialConfigMessage = MakeUnique<FInitialConfigMessage>();
 	if (!InitialConfigMessage->TryRead(Messages[0].Object))
 	{
-		UE_LOG(LogCook, Warning, TEXT("CookWorker initialization failure: Director sent an invalid InitialConfigMessage."));
+		UE_LOG(LogCook, Warning,
+			TEXT("CookWorker initialization failure: Director sent an invalid InitialConfigMessage."));
 		SendToState(EConnectStatus::LostConnection);
 		return;
 	}
@@ -515,7 +524,8 @@ void FCookWorkerClient::PollReceiveConfigMessage()
 	{
 		TStringBuilder<512> StringBuilder;
 		GetPlatformDetails(StringBuilder);
-		UE_LOG(LogCook, Error, TEXT("CookWorker initialization failure: Director sent a mismatch in session platform quantity.\n%s"), *StringBuilder);
+		UE_LOG(LogCook, Error,
+			TEXT("CookWorker initialization failure: Director sent a mismatch in session platform quantity.\n%s"), *StringBuilder);
 		SendToState(EConnectStatus::LostConnection);
 		return;
 	}
@@ -534,7 +544,8 @@ void FCookWorkerClient::PollReceiveConfigMessage()
 	{
 		TStringBuilder<512> StringBuilder;
 		GetPlatformDetails(StringBuilder);
-		UE_LOG(LogCook, Error, TEXT("CookWorker initialization failure: Director sent a mismatch in session platform contents.\n%s"), *StringBuilder);
+		UE_LOG(LogCook, Error,
+			TEXT("CookWorker initialization failure: Director sent a mismatch in session platform contents.\n%s"), *StringBuilder);
 		SendToState(EConnectStatus::LostConnection);
 		return;
 	}
@@ -555,7 +566,8 @@ void FCookWorkerClient::PumpSendMessages()
 	UE::CompactBinaryTCP::EConnectionStatus Status = UE::CompactBinaryTCP::TryFlushBuffer(ServerSocket, SendBuffer);
 	if (Status == UE::CompactBinaryTCP::EConnectionStatus::Failed)
 	{
-		UE_LOG(LogCook, Error, TEXT("CookWorkerClient failed to write message to Director. We will abort the CookAsCookWorker commandlet."));
+		UE_LOG(LogCook, Error,
+			TEXT("CookWorkerClient failed to write message to Director. We will abort the CookAsCookWorker commandlet."));
 		SendToState(EConnectStatus::LostConnection);
 	}
 }
@@ -597,7 +609,8 @@ void FCookWorkerClient::PumpReceiveMessages()
 	EConnectionStatus SocketStatus = TryReadPacket(ServerSocket, ReceiveBuffer, Messages);
 	if (SocketStatus != EConnectionStatus::Okay && SocketStatus != EConnectionStatus::Incomplete)
 	{
-		UE_LOG(LogCook, Error, TEXT("CookWorkerClient failed to read from Director. We will abort the CookAsCookWorker commandlet."));
+		UE_LOG(LogCook, Error,
+			TEXT("CookWorkerClient failed to read from Director. We will abort the CookAsCookWorker commandlet."));
 		SendToState(EConnectStatus::LostConnection);
 		return;
 	}
@@ -612,11 +625,13 @@ void FCookWorkerClient::HandleReceiveMessages(TArray<UE::CompactBinaryTCP::FMars
 		{
 			if (Message.MessageType == FAbortWorkerMessage::MessageType)
 			{
-				UE_LOG(LogCook, Display, TEXT("CookWorkerClient received AbortWorker message from Director. Terminating flush and shutting down."));
+				UE_LOG(LogCook, Display,
+					TEXT("CookWorkerClient received AbortWorker message from Director. Terminating flush and shutting down."));
 				SendToState(EConnectStatus::LostConnection);
 				break;
 			}
-			UE_LOG(LogCook, Error, TEXT("CookWorkerClient received message %s from Director after receiving Abort message. Message will be ignored."),
+			UE_LOG(LogCook, Error,
+				TEXT("CookWorkerClient received message %s from Director after receiving Abort message. Message will be ignored."),
 			*Message.MessageType.ToString());
 		}
 		else
@@ -627,19 +642,22 @@ void FCookWorkerClient::HandleReceiveMessages(TArray<UE::CompactBinaryTCP::FMars
 				AbortMessage.TryRead(Message.Object);
 				if (AbortMessage.Type == FAbortWorkerMessage::EType::CookComplete)
 				{
-					UE_LOG(LogCook, Display, TEXT("CookWorkerClient received CookComplete message from Director. Flushing messages and shutting down."));
+					UE_LOG(LogCook, Display,
+						TEXT("CookWorkerClient received CookComplete message from Director. Flushing messages and shutting down."));
 					SendToState(EConnectStatus::FlushAndAbortFirst);
 				}
 				else
 				{
-					UE_LOG(LogCook, Display, TEXT("CookWorkerClient received AbortWorker message from Director. Shutting down."));
+					UE_LOG(LogCook, Display,
+						TEXT("CookWorkerClient received AbortWorker message from Director. Shutting down."));
 					SendToState(EConnectStatus::LostConnection);
 					break;
 				}
 			}
 			else if (Message.MessageType == FInitialConfigMessage::MessageType)
 			{
-				UE_LOG(LogCook, Warning, TEXT("CookWorkerClient received unexpected repeat of InitialConfigMessage. Ignoring it."));
+				UE_LOG(LogCook, Warning,
+					TEXT("CookWorkerClient received unexpected repeat of InitialConfigMessage. Ignoring it."));
 			}
 			else if (Message.MessageType == FAssignPackagesMessage::MessageType)
 			{
@@ -666,7 +684,8 @@ void FCookWorkerClient::HandleReceiveMessages(TArray<UE::CompactBinaryTCP::FMars
 				}
 				else
 				{
-					UE_LOG(LogCook, Error, TEXT("CookWorkerClient received message of unknown type %s from CookDirector. Ignoring it."),
+					UE_LOG(LogCook, Error,
+						TEXT("CookWorkerClient received message of unknown type %s from CookDirector. Ignoring it."),
 						*Message.MessageType.ToString());
 				}
 			}
@@ -698,9 +717,11 @@ void FCookWorkerClient::PumpDisconnect(FTickStackData& StackData)
 				PumpSendMessages();
 
 				constexpr float WaitForDisconnectTimeout = 60.f;
-				if (FPlatformTime::Seconds() - ConnectStartTimeSeconds > WaitForDisconnectTimeout && !IsCookIgnoreTimeouts())
+				if (FPlatformTime::Seconds() - ConnectStartTimeSeconds > WaitForDisconnectTimeout
+					&& !IsCookIgnoreTimeouts())
 				{
-					UE_LOG(LogCook, Warning, TEXT("Timedout after %.0fs waiting to send disconnect message to CookDirector."),
+					UE_LOG(LogCook, Warning,
+						TEXT("Timedout after %.0fs waiting to send disconnect message to CookDirector."),
 						WaitForDisconnectTimeout);
 					SendToState(EConnectStatus::LostConnection);
 					check(ConnectStatus == EConnectStatus::LostConnection);
@@ -753,7 +774,8 @@ void FCookWorkerClient::SendToState(EConnectStatus TargetStatus)
 
 void FCookWorkerClient::LogInvalidMessage(const TCHAR* MessageTypeName)
 {
-	UE_LOG(LogCook, Error, TEXT("CookWorkerClient received invalidly formatted message for type %s from CookDirector. Ignoring it."),
+	UE_LOG(LogCook, Error,
+		TEXT("CookWorkerClient received invalidly formatted message for type %s from CookDirector. Ignoring it."),
 		MessageTypeName);
 }
 
@@ -770,7 +792,8 @@ void FCookWorkerClient::AssignPackages(FAssignPackagesMessage& Message)
 		FPackageData& PackageData = COTFS.PackageDatas->FindOrAddPackageData(AssignData.ConstructData.PackageName,
 			AssignData.ConstructData.NormalizedFileName);
 		TConstArrayView<const ITargetPlatform*> NeedCookPlatforms = 
-			AssignData.NeedCookPlatforms.GetPlatforms(COTFS, nullptr, OrderedSessionPlatforms, &NeedCookPlatformsBuffer);
+			AssignData.NeedCookPlatforms.GetPlatforms(COTFS, nullptr, OrderedSessionPlatforms,
+				&NeedCookPlatformsBuffer);
 		if (PackageData.IsInProgress())
 		{
 			// If already in progress and no new platforms, ignore the duplicate
@@ -785,15 +808,16 @@ void FCookWorkerClient::AssignPackages(FAssignPackagesMessage& Message)
 			continue;
 		}
 
-		// We do not want CookWorkers to explore dependencies in CookRequestCluster because the Director did it already.
-		// Mark the PackageDatas we get from the Director as already explored.
+		// We do not want CookWorkers to explore dependencies in CookRequestCluster because the Director did it
+		// already. Mark the PackageDatas we get from the Director as already explored.
 		for (const ITargetPlatform* TargetPlatform : NeedCookPlatforms)
 		{
 			PackageData.FindOrAddPlatformData(TargetPlatform).MarkCookableForWorker(*this);
 		}
 		PackageData.FindOrAddPlatformData(CookerLoadingPlatformKey).MarkCookableForWorker(*this);
 		PackageData.SetInstigator(*this, FInstigator(AssignData.Instigator));
-		PackageData.SendToState(EPackageState::Request, ESendFlags::QueueAddAndRemove, EStateChangeReason::DirectorRequest);
+		PackageData.SendToState(EPackageState::Request, ESendFlags::QueueAddAndRemove,
+			EStateChangeReason::DirectorRequest);
 	}
 
 	// Clear the SoftGC diagnostic ExpectedNeverLoadPackages because we have new assigned packages
@@ -806,7 +830,8 @@ void FCookWorkerClient::Register(IMPCollector* Collector)
 	TRefCountPtr<IMPCollector>& Existing = Collectors.FindOrAdd(Collector->GetMessageType());
 	if (Existing)
 	{
-		UE_LOG(LogCook, Error, TEXT("Duplicate IMPCollectors registered. Guid: %s, Existing: %s, Registering: %s. Keeping the Existing."),
+		UE_LOG(LogCook, Error,
+			TEXT("Duplicate IMPCollectors registered. Guid: %s, Existing: %s, Registering: %s. Keeping the Existing."),
 			*Collector->GetMessageType().ToString(), Existing->GetDebugName(), Collector->GetDebugName());
 		return;
 	}
@@ -819,7 +844,8 @@ void FCookWorkerClient::Unregister(IMPCollector* Collector)
 	Collectors.RemoveAndCopyValue(Collector->GetMessageType(), Existing);
 	if (Existing && Existing.GetReference() != Collector)
 	{
-		UE_LOG(LogCook, Error, TEXT("Duplicate IMPCollector during Unregister. Guid: %s, Existing: %s, Unregistering: %s. Ignoring the Unregister."),
+		UE_LOG(LogCook, Error,
+			TEXT("Duplicate IMPCollector during Unregister. Guid: %s, Existing: %s, Unregistering: %s. Ignoring the Unregister."),
 			*Collector->GetMessageType().ToString(), Existing->GetDebugName(), Collector->GetDebugName());
 		Collectors.Add(Collector->GetMessageType(), MoveTemp(Existing));
 	}
@@ -893,7 +919,8 @@ void FCookWorkerClient::HandleAbortPackagesMessage(FMPCollectorClientMessageCont
 		FPackageData* PackageData = COTFS.PackageDatas->FindPackageDataByPackageName(PackageName);
 		if (PackageData)
 		{
-			COTFS.DemoteToIdle(*PackageData, ESendFlags::QueueAddAndRemove, ESuppressCookReason::RetractedByCookDirector);
+			COTFS.DemoteToIdle(*PackageData, ESendFlags::QueueAddAndRemove,
+				ESuppressCookReason::RetractedByCookDirector);
 		}
 	}
 }
@@ -931,7 +958,8 @@ void FCookWorkerClient::HandleHeartbeatMessage(FMPCollectorClientMessageContext&
 		return;
 	}
 
-	UE_LOG(LogCook, Display, TEXT("%.*s %d"), HeartbeatCategoryText.Len(), HeartbeatCategoryText.GetData(), Message.HeartbeatNumber);
+	UE_LOG(LogCook, Display, TEXT("%.*s %d"),
+		HeartbeatCategoryText.Len(), HeartbeatCategoryText.GetData(), Message.HeartbeatNumber);
 	SendMessage(FHeartbeatMessage(Message.HeartbeatNumber));
 }
 
diff --git a/Engine/Source/Editor/UnrealEd/Private/Cooker/CookWorkerServer.cpp b/Engine/Source/Editor/UnrealEd/Private/Cooker/CookWorkerServer.cpp
--- a/Engine/Source/Editor/UnrealEd/Private/Cooker/CookWorkerServer.cpp
+++ b/Engine/Source/Editor/UnrealEd/Private/Cooker/CookWorkerServer.cpp
@@ -39,9 +39,11 @@ FCookWorkerServer::~FCookWorkerServer()
 	checkf(PendingPackages.IsEmpty() && PackagesToAssign.IsEmpty(),
 		TEXT("CookWorkerServer still has assigned packages when it is being destroyed; we will leak them and block the cook."));
 
-	if (ConnectStatus == EConnectStatus::Connected || ConnectStatus == EConnectStatus::PumpingCookComplete || ConnectStatus == EConnectStatus::WaitForDisconnect)
+	if (ConnectStatus == EConnectStatus::Connected || ConnectStatus == EConnectStatus::PumpingCookComplete
+		|| ConnectStatus == EConnectStatus::WaitForDisconnect)
 	{
-		UE_LOG(LogCook, Error, TEXT("CookWorkerServer %d was destroyed before it finished Disconnect. The remote process may linger and may interfere with writes of future packages."),
+		UE_LOG(LogCook, Error,
+			TEXT("CookWorkerServer %d was destroyed before it finished Disconnect. The remote process may linger and may interfere with writes of future packages."),
 			ProfileId);
 	}
 	DetachFromRemoteProcess();
@@ -66,7 +68,8 @@ void FCookWorkerServer::DetachFromRemoteProcess()
 	}
 }
 
-bool TryParseLogCategoryVerbosityMessage(FStringView Line, FName& OutCategory, ELogVerbosity::Type& OutVerbosity, FStringView& OutMessage)
+bool TryParseLogCategoryVerbosityMessage(FStringView Line, FName& OutCategory, ELogVerbosity::Type& OutVerbosity,
+	FStringView& OutMessage)
 {
 	TPair<FStringView, ELogVerbosity::Type> VerbosityMarkers[]{
 		{ TEXTVIEW(": Fatal:"), ELogVerbosity::Fatal },
@@ -77,7 +80,8 @@ bool TryParseLogCategoryVerbosityMessage(FStringView Line, FName& OutCategory, E
 	};
 
 
-	// Find the first colon not in brackets and look for ": <Verbosity>:". This is complicated by Log verbosity not printing out the Verbosity:
+	// Find the first colon not in brackets and look for ": <Verbosity>:". This is complicated by Log verbosity not
+	// printing out the Verbosity:
 	// [2023.03.20-16.32.48:878][  0]LogCook: MessageText
 	// [2023.03.20-16.32.48:878][  0]LogCook: Display: MessageText
 
@@ -139,9 +143,12 @@ bool TryParseLogCategoryVerbosityMessage(FStringView Line, FName& OutCategory, E
 void FCookWorkerServer::SendCrashDiagnostics()
 {
 	FString LogFileName = Director.GetWorkerLogFileName(ProfileId);
-	UE_LOG(LogCook, Display, TEXT("LostConnection to CookWorker %d. Log messages written after communication loss:"), ProfileId);
+	UE_LOG(LogCook, Display,
+		TEXT("LostConnection to CookWorker %d. Log messages written after communication loss:"), ProfileId);
 	FString LogText;
-	int32 ReadFlags = FILEREAD_AllowWrite; // To be able to open a file for read that might be open for write from another process, we have to specify FILEREAD_AllowWrite
+	// To be able to open a file for read that might be open for write from another process,
+	// we have to specify FILEREAD_AllowWrite
+	int32 ReadFlags = FILEREAD_AllowWrite;
 	bool bLoggedErrorMessage = false;
 	if (!FFileHelper::LoadFileToString(LogText, *LogFileName, FFileHelper::EHashOptions::None, ReadFlags))
 	{
@@ -149,8 +156,8 @@ void FCookWorkerServer::SendCrashDiagnostics()
 	}
 	else
 	{
-		FString LastSentHeartbeat = FString::Printf(TEXT("%.*s %d"), HeartbeatCategoryText.Len(), HeartbeatCategoryText.GetData(),
-			LastReceivedHeartbeatNumber);
+		FString LastSentHeartbeat = FString::Printf(TEXT("%.*s %d"), HeartbeatCategoryText.Len(),
+			HeartbeatCategoryText.GetData(), LastReceivedHeartbeatNumber);
 		int32 StartIndex = INDEX_NONE;
 		for (FStringView MarkerText : { FStringView(LastSentHeartbeat),
 			HeartbeatCategoryText, TEXTVIEW("Connection to CookDirector successful") })
@@ -206,8 +213,8 @@ void FCookWorkerServer::SendCrashDiagnostics()
 		}
 		else
 		{
-			// When we already logged an error from the crashed worker, log the what-went-wrong as a warning rather than an error,
-			// to avoid making it seem like a separate issue.
+			// When we already logged an error from the crashed worker, log the what-went-wrong as a warning rather
+			// than an error, to avoid making it seem like a separate issue.
 			UE_LOG(LogCook, Warning, TEXT("%s"), *CrashDiagnosticsError);
 		}
 	}
@@ -351,7 +358,8 @@ bool FCookWorkerServer::IsConnected() const
 bool FCookWorkerServer::IsShuttingDown() const
 {
 	FScopeLock CommunicationScopeLock(&CommunicationLock);
-	return ConnectStatus == EConnectStatus::PumpingCookComplete || ConnectStatus == EConnectStatus::WaitForDisconnect || ConnectStatus == EConnectStatus::LostConnection;
+	return ConnectStatus == EConnectStatus::PumpingCookComplete || ConnectStatus == EConnectStatus::WaitForDisconnect
+		|| ConnectStatus == EConnectStatus::LostConnection;
 }
 
 bool FCookWorkerServer::IsFlushingBeforeShutdown() const
@@ -401,7 +409,8 @@ int32 FCookWorkerServer::GetPackagesRetiredFenceMarker() const
 	return PackagesRetiredFenceMarker;
 }
 
-bool FCookWorkerServer::TryHandleConnectMessage(FWorkerConnectMessage& Message, FSocket* InSocket, TArray<UE::CompactBinaryTCP::FMarshalledMessage>&& OtherPacketMessages, ECookDirectorThread TickThread)
+bool FCookWorkerServer::TryHandleConnectMessage(FWorkerConnectMessage& Message, FSocket* InSocket,
+	TArray<UE::CompactBinaryTCP::FMarshalledMessage>&& OtherPacketMessages, ECookDirectorThread TickThread)
 {
 	FCommunicationScopeLock ScopeLock(this, TickThread, ETickAction::Tick);
 
@@ -463,11 +472,13 @@ void FCookWorkerServer::TickCommunication(ECookDirectorThread TickThread)
 			{
 				PumpSendMessages();
 				constexpr float WaitForPumpCompleteTimeout = 10.f * 60;
-				if (FPlatformTime::Seconds() - ConnectStartTimeSeconds <= WaitForPumpCompleteTimeout || IsCookIgnoreTimeouts())
+				if (FPlatformTime::Seconds() - ConnectStartTimeSeconds <= WaitForPumpCompleteTimeout 
+					|| IsCookIgnoreTimeouts())
 				{
 					return; // Try again later
 				}
-				UE_LOG(LogCook, Error, TEXT("CookWorker process of CookWorkerServer %d failed to finalize its cook within %.0f seconds; we will tell it to shutdown."),
+				UE_LOG(LogCook, Error,
+					TEXT("CookWorker process of CookWorkerServer %d failed to finalize its cook within %.0f seconds; we will tell it to shutdown."),
 					ProfileId, WaitForPumpCompleteTimeout);
 				SendMessageInLock(FAbortWorkerMessage(FAbortWorkerMessage::EType::Abort));
 				SendToState(EConnectStatus::WaitForDisconnect);
@@ -534,7 +545,8 @@ void FCookWorkerServer::LaunchProcess()
 		nullptr /* PipeWriteChild */);
 	if (CookWorkerHandle.IsValid())
 	{
-		UE_LOG(LogCook, Display, TEXT("CookWorkerServer %d launched CookWorker as WorkerId %d and PID %u with commandline \"%s\"."),
+		UE_LOG(LogCook, Display,
+			TEXT("CookWorkerServer %d launched CookWorker as WorkerId %d and PID %u with commandline \"%s\"."),
 			ProfileId, WorkerId.GetRemoteIndex(), CookWorkerProcessId, *LaunchInfo.WorkerCommandLine);
 		FCoreDelegates::OnMultiprocessWorkerCreated.Broadcast({WorkerId.GetMultiprocessId()});
 		SendToState(EConnectStatus::WaitForConnect);
@@ -542,7 +554,8 @@ void FCookWorkerServer::LaunchProcess()
 	else
 	{
 		// GetLastError information was logged by CreateProc
-		CrashDiagnosticsError = FString::Printf(TEXT("CookWorkerCrash: Failed to create process for CookWorker %d. Assigned packages will be returned to the director."),
+		CrashDiagnosticsError = FString::Printf(
+			TEXT("CookWorkerCrash: Failed to create process for CookWorker %d. Assigned packages will be returned to the director."),
 			ProfileId);
 		bNeedCrashDiagnostics = true;
 		SendToState(EConnectStatus::LostConnection);
@@ -554,14 +567,16 @@ void FCookWorkerServer::TickWaitForConnect()
 	constexpr float TestProcessExistencePeriod = 1.f;
 	constexpr float WaitForConnectTimeout = 60.f * 20;
 
-	check(!Socket); // When the Socket is assigned we leave the WaitForConnect state, and we set it to null before entering
+	// When the Socket is assigned we leave the WaitForConnect state, and we set it to null before entering
+	check(!Socket);
 
 	double CurrentTime = FPlatformTime::Seconds();
 	if (CurrentTime - ConnectTestStartTimeSeconds > TestProcessExistencePeriod)
 	{
 		if (!FPlatformProcess::IsProcRunning(CookWorkerHandle))
 		{
-			CrashDiagnosticsError = FString::Printf(TEXT("CookWorkerCrash: CookWorker %d process terminated before connecting. Assigned packages will be returned to the director."),
+			CrashDiagnosticsError = FString::Printf(
+				TEXT("CookWorkerCrash: CookWorker %d process terminated before connecting. Assigned packages will be returned to the director."),
 				ProfileId);
 			bNeedCrashDiagnostics = true;
 			SendToState(EConnectStatus::LostConnection);
@@ -572,7 +587,8 @@ void FCookWorkerServer::TickWaitForConnect()
 
 	if (CurrentTime - ConnectStartTimeSeconds > WaitForConnectTimeout && !IsCookIgnoreTimeouts())
 	{
-		CrashDiagnosticsError = FString::Printf(TEXT("CookWorkerCrash: CookWorker %d process failed to connect within %.0f seconds. Assigned packages will be returned to the director."),
+		CrashDiagnosticsError = FString::Printf(
+			TEXT("CookWorkerCrash: CookWorker %d process failed to connect within %.0f seconds. Assigned packages will be returned to the director."),
 			ProfileId, WaitForConnectTimeout);
 		bNeedCrashDiagnostics = true;
 		ShutdownRemoteProcess();
@@ -602,7 +618,8 @@ void FCookWorkerServer::TickWaitForDisconnect()
 	TArray<UE::CompactBinaryTCP::FMarshalledMessage> Messages;
 	TryReadPacket(Socket, ReceiveBuffer, Messages);
 
-	if (bTerminateImmediately || (CurrentTime - ConnectStartTimeSeconds > WaitForDisconnectTimeout && !IsCookIgnoreTimeouts()))
+	if (bTerminateImmediately ||
+		(CurrentTime - ConnectStartTimeSeconds > WaitForDisconnectTimeout && !IsCookIgnoreTimeouts()))
 	{
 		UE_CLOG(!bTerminateImmediately, LogCook, Warning,
 			TEXT("CookWorker process of CookWorkerServer %d failed to disconnect within %.0f seconds; we will terminate it."),
@@ -617,7 +634,8 @@ void FCookWorkerServer::PumpSendMessages()
 	UE::CompactBinaryTCP::EConnectionStatus Status = UE::CompactBinaryTCP::TryFlushBuffer(Socket, SendBuffer);
 	if (Status == UE::CompactBinaryTCP::EConnectionStatus::Failed)
 	{
-		UE_LOG(LogCook, Error, TEXT("CookWorkerCrash: CookWorker %d failed to write to socket, we will shutdown the remote process. Assigned packages will be returned to the director."),
+		UE_LOG(LogCook, Error,
+			TEXT("CookWorkerCrash: CookWorker %d failed to write to socket, we will shutdown the remote process. Assigned packages will be returned to the director."),
 			ProfileId);
 		bNeedCrashDiagnostics = true;
 		SendToState(EConnectStatus::WaitForDisconnect);
@@ -666,7 +684,8 @@ void FCookWorkerServer::PumpReceiveMessages()
 	EConnectionStatus SocketStatus = TryReadPacket(Socket, ReceiveBuffer, Messages);
 	if (SocketStatus != EConnectionStatus::Okay && SocketStatus != EConnectionStatus::Incomplete)
 	{
-		CrashDiagnosticsError = FString::Printf(TEXT("CookWorkerCrash: CookWorker %d failed to read from socket, we will shutdown the remote process. Assigned packages will be returned to the director."),
+		CrashDiagnosticsError = FString::Printf(
+			TEXT("CookWorkerCrash: CookWorker %d failed to read from socket, we will shutdown the remote process. Assigned packages will be returned to the director."),
 			ProfileId);
 		bNeedCrashDiagnostics = true;
 		SendToState(EConnectStatus::WaitForDisconnect);
@@ -695,9 +714,11 @@ void FCookWorkerServer::HandleReceiveMessagesInternal()
 		if (PeekMessage.MessageType == FAbortWorkerMessage::MessageType)
 		{
 			UE::CompactBinaryTCP::FMarshalledMessage Message = ReceiveMessages.PopFrontValue();
-			if (ConnectStatus != EConnectStatus::PumpingCookComplete && ConnectStatus != EConnectStatus::WaitForDisconnect)
+			if (ConnectStatus != EConnectStatus::PumpingCookComplete
+				&& ConnectStatus != EConnectStatus::WaitForDisconnect)
 			{
-				CrashDiagnosticsError = FString::Printf(TEXT("CookWorkerCrash: CookWorker %d remote process shut down unexpectedly. Assigned packages will be returned to the director."),
+				CrashDiagnosticsError = FString::Printf(
+					TEXT("CookWorkerCrash: CookWorker %d remote process shut down unexpectedly. Assigned packages will be returned to the director."),
 					ProfileId);
 				bNeedCrashDiagnostics = true;
 			}
@@ -756,15 +777,16 @@ void FCookWorkerServer::HandleReceiveMessagesInternal()
 			}
 			else
 			{
-				UE_LOG(LogCook, Error, TEXT("CookWorkerServer received message of unknown type %s from CookWorker. Ignoring it."),
+				UE_LOG(LogCook, Error,
+					TEXT("CookWorkerServer received message of unknown type %s from CookWorker. Ignoring it."),
 					*Message.MessageType.ToString());
 			}
 		}
 	}
 }
 
-void FCookWorkerServer::HandleReceivedPackagePlatformMessages(FPackageData& PackageData, const ITargetPlatform* TargetPlatform,
-	TArray<UE::CompactBinaryTCP::FMarshalledMessage>&& Messages)
+void FCookWorkerServer::HandleReceivedPackagePlatformMessages(FPackageData& PackageData,
+	const ITargetPlatform* TargetPlatform, TArray<UE::CompactBinaryTCP::FMarshalledMessage>&& Messages)
 {
 	check(TickState.TickThread == ECookDirectorThread::SchedulerThread);
 	if (Messages.IsEmpty())
@@ -790,7 +812,8 @@ void FCookWorkerServer::HandleReceivedPackagePlatformMessages(FPackageData& Pack
 		}
 		else
 		{
-			UE_LOG(LogCook, Error, TEXT("CookWorkerServer received PackageMessage of unknown type %s from CookWorker. Ignoring it."),
+			UE_LOG(LogCook, Error,
+				TEXT("CookWorkerServer received PackageMessage of unknown type %s from CookWorker. Ignoring it."),
 				*Message.MessageType.ToString());
 		}
 	}
@@ -825,21 +848,23 @@ void FCookWorkerServer::RecordResults(FPackageResultsMessage& Message)
 		FPackageData* PackageData = COTFS.PackageDatas->FindPackageDataByPackageName(Result.GetPackageName());
 		if (!PackageData)
 		{
-			UE_LOG(LogCook, Warning, TEXT("CookWorkerServer %d received FPackageResultsMessage for invalid package %s. Ignoring it."),
+			UE_LOG(LogCook, Warning,
+				TEXT("CookWorkerServer %d received FPackageResultsMessage for invalid package %s. Ignoring it."),
 				ProfileId, *Result.GetPackageName().ToString());
 			continue;
 		}
 		if (PendingPackages.Remove(PackageData) != 1)
 		{
-			UE_LOG(LogCook, Display, TEXT("CookWorkerServer %d received FPackageResultsMessage for package %s which is not a pending package. Ignoring it."),
+			UE_LOG(LogCook, Display,
+				TEXT("CookWorkerServer %d received FPackageResultsMessage for package %s which is not a pending package. Ignoring it."),
 				ProfileId, *Result.GetPackageName().ToString());
 			continue;
 		}
 		bRetiredAnyPackages = true;
 		PackageData->SetWorkerAssignment(FWorkerId::Invalid(), ESendFlags::QueueNone);
 
-		// MPCOOKTODO: Refactor FSaveCookedPackageContext::FinishPlatform and ::FinishPackage so we can call them from here
-		// to reduce duplication
+		// MPCOOKTODO: Refactor FSaveCookedPackageContext::FinishPlatform and ::FinishPackage so we can call them from
+		// here to reduce duplication
 		if (Result.GetSuppressCookReason() == ESuppressCookReason::NotSuppressed)
 		{
 			int32 NumPlatforms = OrderedSessionPlatforms.Num();
@@ -873,7 +898,8 @@ void FCookWorkerServer::RecordResults(FPackageResultsMessage& Message)
 					{
 						PackageData->SetPlatformCooked(TargetPlatform, PlatformResult.GetCookResults());
 					}
-					HandleReceivedPackagePlatformMessages(*PackageData, TargetPlatform, PlatformResult.ReleaseMessages());
+					HandleReceivedPackagePlatformMessages(*PackageData, TargetPlatform,
+						PlatformResult.ReleaseMessages());
 				}
 			}
 			COTFS.RecordExternalActorDependencies(Result.GetExternalActorDependencies());
@@ -897,7 +923,8 @@ void FCookWorkerServer::RecordResults(FPackageResultsMessage& Message)
 
 void FCookWorkerServer::LogInvalidMessage(const TCHAR* MessageTypeName)
 {
-	UE_LOG(LogCook, Error, TEXT("CookWorkerServer received invalidly formatted message for type %s from CookWorker. Ignoring it."),
+	UE_LOG(LogCook, Error,
+		TEXT("CookWorkerServer received invalidly formatted message for type %s from CookWorker. Ignoring it."),
 		MessageTypeName);
 }
 
@@ -919,7 +946,8 @@ void FCookWorkerServer::QueueDiscoveredPackage(FDiscoveredPackageReplication&& D
 	}
 	else
 	{
-		DiscoveredPlatforms = Platforms.GetPlatforms(COTFS, &Instigator, OrderedSessionAndSpecialPlatforms, &BufferPlatforms);
+		DiscoveredPlatforms = Platforms.GetPlatforms(COTFS, &Instigator, OrderedSessionAndSpecialPlatforms,
+			&BufferPlatforms);
 	}
 
 	if (Instigator.Category != EInstigator::ForceExplorableSaveTimeSoftDependency &&
@@ -952,7 +980,8 @@ void FCookWorkerServer::QueueDiscoveredPackage(FDiscoveredPackageReplication&& D
 	}
 	Director.ResetFinalIdleHeartbeatFence();
 	Platforms.ConvertFromBitfield(OrderedSessionAndSpecialPlatforms);
-	COTFS.QueueDiscoveredPackageOnDirector(PackageData, MoveTemp(Instigator), MoveTemp(Platforms), false /* bUrgent */);
+	COTFS.QueueDiscoveredPackageOnDirector(PackageData, MoveTemp(Instigator), MoveTemp(Platforms),
+		false /* bUrgent */);
 }
 
 FCookWorkerServer::FTickState::FTickState()
@@ -961,7 +990,8 @@ FCookWorkerServer::FTickState::FTickState()
 	TickAction = ETickAction::Invalid;
 }
 
-FCookWorkerServer::FCommunicationScopeLock::FCommunicationScopeLock(FCookWorkerServer* InServer, ECookDirectorThread TickThread, ETickAction TickAction)
+FCookWorkerServer::FCommunicationScopeLock::FCommunicationScopeLock(FCookWorkerServer* InServer,
+	ECookDirectorThread TickThread, ETickAction TickAction)
 	: ScopeLock(&InServer->CommunicationLock)
 	, Server(*InServer)
 {
@@ -1323,7 +1353,8 @@ void FLogMessagesMessageHandler::ClientTick(FMPCollectorClientTickContext& Conte
 	}
 }
 
-void FLogMessagesMessageHandler::ServerReceiveMessage(FMPCollectorServerMessageContext& Context, FCbObjectView InMessage)
+void FLogMessagesMessageHandler::ServerReceiveMessage(FMPCollectorServerMessageContext& Context,
+	FCbObjectView InMessage)
 {
 	TArray<FReplicatedLogData> Messages;
 	if (!LoadFromCompactBinary(InMessage["Messages"], Messages))
@@ -1403,7 +1434,8 @@ void FPackageWriterMPCollector::ServerReceiveMessage(FMPCollectorServerMessageCo
 	ICookedPackageWriter& PackageWriter = COTFS.FindOrCreatePackageWriter(TargetPlatform);
 	if (!PackageWriter.TryReadMPCookMessageForPackage(PackageName, Message))
 	{
-		UE_LOG(LogCook, Error, TEXT("CookWorkerServer received invalidly formatted PackageWriter message from CookWorker %d. Ignoring it."),
+		UE_LOG(LogCook, Error,
+			TEXT("CookWorkerServer received invalidly formatted PackageWriter message from CookWorker %d. Ignoring it."),
 			Context.GetProfileId());
 	}
 }
diff --git a/Engine/Source/Editor/UnrealEd/Private/Cooker/CookWorkerServer.h b/Engine/Source/Editor/UnrealEd/Private/Cooker/CookWorkerServer.h
--- a/Engine/Source/Editor/UnrealEd/Private/Cooker/CookWorkerServer.h
+++ b/Engine/Source/Editor/UnrealEd/Private/Cooker/CookWorkerServer.h
@@ -61,10 +61,14 @@ public:
 	 * Report all packages that were removed.
 	 */
 	void AbortAllAssignments(TSet<FPackageData*>& OutPendingPackages, ECookDirectorThread TickThread);
-	/** AbortAllAssignments and tell the connected Client to gracefully terminate. Report all packages that were unassigned. */
+	/**
+	 * AbortAllAssignments and tell the connected Client to gracefully terminate. Report all packages that were
+	 * unassigned.
+	 */
 	void AbortWorker(TSet<FPackageData*>& OutPendingPackages, ECookDirectorThread TickThread);
 	/** Take over the Socket for a CookWorker that has just connected. */
-	bool TryHandleConnectMessage(FWorkerConnectMessage& Message, FSocket* InSocket, TArray<UE::CompactBinaryTCP::FMarshalledMessage>&& OtherPacketMessages, ECookDirectorThread TickThread);
+	bool TryHandleConnectMessage(FWorkerConnectMessage& Message, FSocket* InSocket,
+		TArray<UE::CompactBinaryTCP::FMarshalledMessage>&& OtherPacketMessages, ECookDirectorThread TickThread);
 
 	/** Send the message immediately to the Socket. If cannot complete immediately, it will be finished during Tick. */
 	void SendMessage(const IMPCollectorMessage& Message, ECookDirectorThread TickThread);
@@ -73,7 +77,7 @@ public:
 	void TickCommunication(ECookDirectorThread TickThread);
 	/** Called when the COTFS wants to send a heartbeat message to the Client. */
 	void SignalHeartbeat(ECookDirectorThread TickThread, int32 HeartbeatNumber);
-	/** Called when the COTFS Server has detected all packages are complete. Tell the CookWorker to flush messages and exit. */
+	/** Called when the Server detects all packages are complete. Tell the CookWorker to flush messages and exit. */
 	void SignalCookComplete(ECookDirectorThread TickThread);
 	/**
 	 * Execute the respond for all messages that have been received and that can be executed from the given thread. 
@@ -85,7 +89,7 @@ public:
 	bool IsConnected() const;
 	/** Is this either shutting down or completed shutdown of its remote Client? */
 	bool IsShuttingDown() const;
-	/** Is this executing the portion of graceful shutdown where it waits for the CookWorker to transfer remaining messages? */
+	/** Is this executing graceful shutdown and is waiting for the CookWorker to transfer remaining messages? */
 	bool IsFlushingBeforeShutdown() const;
 	/** Is this not yet or no longer connected to a remote Client? */
 	bool IsShutdownComplete() const;
@@ -96,7 +100,10 @@ public:
 
 	/** Get the LastReceivedHeartbeatNumber. */
 	int32 GetLastReceivedHeartbeatNumber() const;
-	/** Set the LastReceivedHeartbeatNumber. Assumes lock is already entered; can only be called from with a HandleReceivedMessages callback */
+	/**
+	 * Set the LastReceivedHeartbeatNumber. Assumes lock is already entered; can only be called from with a
+	 * HandleReceivedMessages callback.
+	 */
 	void SetLastReceivedHeartbeatNumberInLock(int32 InHeartbeatNumber);
 
 	int32 GetPackagesAssignedFenceMarker() const;
@@ -166,7 +173,8 @@ private:
 	/** The main implementation of HandleReceiveMessages, only callable from inside the lock. */
 	void HandleReceiveMessagesInternal();
 	/** Helper for PumpReceiveMessages: dispatch the messages received from the socket. */
-	void HandleReceivedPackagePlatformMessages(FPackageData& PackageData, const ITargetPlatform* TargetPlatform, TArray<UE::CompactBinaryTCP::FMarshalledMessage>&& Messages);
+	void HandleReceivedPackagePlatformMessages(FPackageData& PackageData, const ITargetPlatform* TargetPlatform,
+		TArray<UE::CompactBinaryTCP::FMarshalledMessage>&& Messages);
 	/** Add results from the client to the local CookOnTheFlyServer. */
 	void RecordResults(FPackageResultsMessage& Message);
 	void LogInvalidMessage(const TCHAR* MessageTypeName);
@@ -376,7 +384,8 @@ public:
 
 	// FOutputDevice
 	virtual void Serialize(const TCHAR* V, ELogVerbosity::Type Verbosity, const FName& Category) override;
-	virtual void Serialize(const TCHAR* V, ELogVerbosity::Type Verbosity, const FName& Category, const double Time) override;
+	virtual void Serialize(const TCHAR* V, ELogVerbosity::Type Verbosity, const FName& Category,
+		const double Time) override;
 	virtual bool CanBeUsedOnAnyThread() const override { return true; }
 	virtual bool CanBeUsedOnMultipleThreads() const override { return true; }
 
@@ -390,10 +399,10 @@ private:
 };
 
 /**
- * Message from Director to CookWorker or CookWorker to Director that reports a heartbeat number, in addition to reporting the
- * machine is still alive just by the presence of the message.
- * The Director intiates a heartbeat message; the CookWorker always responds to a heartbeat message with its own heartbeat message
- * in reply, with the same number.
+ * Message from Director to CookWorker or CookWorker to Director that reports a heartbeat number, in addition to
+ * reporting the machine is still alive just by the presence of the message.
+ * The Director intiates a heartbeat message; the CookWorker always responds to a heartbeat message with its own
+ * heartbeat message in reply, with the same number.
  */
 struct FHeartbeatMessage : public IMPCollectorMessage
 {
diff --git a/Engine/Source/Editor/UnrealEd/Private/Cooker/IWorkerRequests.h b/Engine/Source/Editor/UnrealEd/Private/Cooker/IWorkerRequests.h
--- a/Engine/Source/Editor/UnrealEd/Private/Cooker/IWorkerRequests.h
+++ b/Engine/Source/Editor/UnrealEd/Private/Cooker/IWorkerRequests.h
@@ -14,7 +14,7 @@ namespace UE::Cook
 {
 
 /**
- * Interface for Requests sent to this processes' Cooker, and for information this process needs to send back to its Director.
+ * Interface for Requests sent to this processes' Cooker, and for information it needs to send back to its Director.
  * In SingleProcess cooks the functions are passed through to the local CookOnTheFlyServer.
  * In MultiProcess cooks the functions are implemented as interprocess messages to/from the CookDirector.
  **/
@@ -26,27 +26,38 @@ public:
 	// Reading Callbacks And Requests
 	/**
 	 * Lockless value for whether the Worker has ExternalRequests it has not yet reported to the Cooker.
-	 * May be out of date after calling; do not assume a true return value means Requests are actually present or a false value means no Requests are present.
-	 * Intended usage is for the Scheduler to be the only consumer of requests, and to use this value for periodic checking of whether there is any work that justifies the expense of taking the lock.
-	 * In a single-consumer case, HasRequests will eventually correctly return true as long as the consumer is not consuming.
+	 * May be out of date after calling; do not assume a true return value means Requests are actually present or a
+	 * false value means no Requests are present.
+	 * Intended usage is for the Scheduler to be the only consumer of requests, and to use this value for periodic
+	 * checking of whether there is any work that justifies the expense of taking the lock.
+	 * In a single-consumer case, HasRequests will eventually correctly return true as long as the consumer is not
+	 * consuming.
 	 */
 	virtual bool HasExternalRequests() const = 0;
 	/**
 	 * Lockless value for the number of ExternalRequests the Worker has which it has not yet reported to the Cooker.
-	 * May be out of date after calling; do not assume the number of actual requests is any one of equal, greater than, or less than the returned value.
-	 * Intended usage is for the Scheduler to be the only consumer of requests, and to use this value for rough reporting of periodic progress.
+	 * May be out of date after calling; do not assume the number of actual requests is any one of equal, greater than,
+	 * or less than the returned value.
+	 * Intended usage is for the Scheduler to be the only consumer of requests, and to use this value for rough
+	 * reporting of periodic progress.
 	 */
 	virtual int32 GetNumExternalRequests() const = 0;
 	/**
-	 * If this Worker has any callbacks, dequeue them all into OutCallbacks and return EExternalRequestType::Callback; Callbacks take priority over cook requests.
-	 * Otherwise, if there are any cook requests, dequeue them all into OutBuildRequests and return EExternalRequestType::Cook.
+	 * If this Worker has any callbacks, dequeue them all into OutCallbacks and return EExternalRequestType::Callback;
+	 * Callbacks take priority over cook requests.
+	 * Otherwise, if there are any cook requests, dequeue them all into OutBuildRequests and return
+	 * EExternalRequestType::Cook.
 	 * Otherwise, return EExternalRequestType::None.
 	 */
 	virtual EExternalRequestType DequeueNextCluster(TArray<FSchedulerCallback>& OutCallbacks, TArray<FFilePlatformRequest>& OutBuildRequests) = 0;
 	/* Move any existing callbacks onto OutCallbacks, and return whether any were added. */
 	virtual bool DequeueSchedulerCallbacks(TArray<FSchedulerCallback>& OutCallbacks) = 0;
-	/** Move all callbacks into OutCallbacks, and all cook requests into OutCookRequests. This is used when canceling a cook session. */
-	virtual void DequeueAllExternal(TArray<FSchedulerCallback>& OutCallbacks, TArray<FFilePlatformRequest>& OutCookRequests) = 0;
+	/**
+	 * Move all callbacks into OutCallbacks, and all cook requests into OutCookRequests. This is used when canceling
+	 * a cook session.
+	 */
+	virtual void DequeueAllExternal(TArray<FSchedulerCallback>& OutCallbacks,
+		TArray<FFilePlatformRequest>& OutCookRequests) = 0;
 
 	// Writing Packages during Cook
 	virtual void QueueDiscoveredPackage(UCookOnTheFlyServer& COTFS, FPackageData& PackageData,
@@ -76,8 +87,10 @@ public:
 	virtual void ReportPromoteToSaveComplete(UE::Cook::FPackageData& PackageData) = 0;
 
 	// Functions that read from local for Directors but read from Director for Workers
-	virtual void GetInitializeConfigSettings(UCookOnTheFlyServer& COTFS, const FString& OutputDirectoryOverride, UE::Cook::FInitializeConfigSettings& Settings) = 0;
-	virtual void GetBeginCookConfigSettings(UCookOnTheFlyServer& COTFS, FBeginCookContext& BeginContext, UE::Cook::FBeginCookConfigSettings& Settings) = 0;
+	virtual void GetInitializeConfigSettings(UCookOnTheFlyServer& COTFS, const FString& OutputDirectoryOverride,
+		UE::Cook::FInitializeConfigSettings& Settings) = 0;
+	virtual void GetBeginCookConfigSettings(UCookOnTheFlyServer& COTFS, FBeginCookContext& BeginContext,
+		UE::Cook::FBeginCookConfigSettings& Settings) = 0;
 	virtual void GetBeginCookIterativeFlags(UCookOnTheFlyServer& COTFS, FBeginCookContext& BeginContext) = 0;
 	virtual ECookMode::Type GetDirectorCookMode(UCookOnTheFlyServer& COTFS) = 0;
 
diff --git a/Engine/Source/Editor/UnrealEd/Private/Cooker/WorkerRequestsRemote.h b/Engine/Source/Editor/UnrealEd/Private/Cooker/WorkerRequestsRemote.h
--- a/Engine/Source/Editor/UnrealEd/Private/Cooker/WorkerRequestsRemote.h
+++ b/Engine/Source/Editor/UnrealEd/Private/Cooker/WorkerRequestsRemote.h
@@ -22,7 +22,10 @@ namespace UE::Cook
 struct FInstigator;
 struct FPackageData;
 
-/** An IWorkerRequests for CookWorkers in MultiProcess cooks: functions are implemented as interprocess messages to/from the CookDirector. */
+/**
+ * An IWorkerRequests for CookWorkers in MultiProcess cooks: functions are implemented as interprocess messages
+ * to/from the CookDirector.
+ */
 class FWorkerRequestsRemote : public IWorkerRequests
 {
 public:
