ResavePackages commandlet - When processing a WorldPartition map, make sure to resave external objects packages
* External actors were already resaved properly

diff --git a/Engine/Source/Editor/UnrealEd/Private/Commandlets/ContentCommandlets.cpp b/Engine/Source/Editor/UnrealEd/Private/Commandlets/ContentCommandlets.cpp
--- a/Engine/Source/Editor/UnrealEd/Private/Commandlets/ContentCommandlets.cpp
+++ b/Engine/Source/Editor/UnrealEd/Private/Commandlets/ContentCommandlets.cpp
@@ -77,6 +77,7 @@ DEFINE_LOG_CATEGORY(LogContentCommandlet);
 #include "Particles/TypeData/ParticleModuleTypeDataMesh.h"
 #include "Engine/LevelStreaming.h"
 #include "EditorBuildUtils.h"
+#include "ExternalPackageHelper.h"
 
 // for UResavePackagesCommandlet::PerformAdditionalOperations building lighting code
 #include "LightingBuildOptions.h"
@@ -1905,11 +1906,11 @@ void UResavePackagesCommandlet::PerformAdditionalOperations(class UWorld* World,
 	const bool bShouldCheckoutDirtyPackageOnly = (bShouldBuildHLOD || bShouldBuildNavigationData) && !bBuildingNonHLODData;
 
 	UWorldPartition* WorldPartition = World->GetWorldPartition();
-	const bool bResaveWorldPartitionExternalActors = !!WorldPartition;
+	const bool bResaveWorldPartitionExternalPackages = !!WorldPartition;
 	const int32 DefaultExternalActorGCFreq = 2048;
 
 	// Load and Save Level's external packages
- 	if (!bResaveWorldPartitionExternalActors)
+ 	if (!bResaveWorldPartitionExternalPackages)
 	{
 		// Use a default GC frequency for external actors if GarbageCollectionFrequency is 0.
 		TGuardValue<int32> ScopedGCFreq(GarbageCollectionFrequency, GarbageCollectionFrequency ? GarbageCollectionFrequency : DefaultExternalActorGCFreq);
@@ -1925,7 +1926,7 @@ void UResavePackagesCommandlet::PerformAdditionalOperations(class UWorld* World,
 		World->RemoveFromRoot();
 	}
 
-	if (!bBuildingNonHLODData && !bShouldBuildHLOD && !bShouldBuildNavigationData && !bResaveWorldPartitionExternalActors)
+	if (!bBuildingNonHLODData && !bShouldBuildHLOD && !bShouldBuildNavigationData && !bResaveWorldPartitionExternalPackages)
 	{
 		return;
 	}
@@ -1942,28 +1943,40 @@ void UResavePackagesCommandlet::PerformAdditionalOperations(class UWorld* World,
 	FScopedEditorWorld EditorWorld(World, IVS);
 
 	// Load and Save world partition actor packages
-	if (bResaveWorldPartitionExternalActors && !bShouldBuildNavigationData)
+	if (bResaveWorldPartitionExternalPackages && !bShouldBuildNavigationData)
 	{
 		// Use a default GC frequency for external actors if GarbageCollectionFrequency is 0.
 		TGuardValue<int32> ScopedGCFreq(GarbageCollectionFrequency, GarbageCollectionFrequency ? GarbageCollectionFrequency : DefaultExternalActorGCFreq);
 
-		FWorldPartitionHelpers::ForEachActorDescInstance(WorldPartition, [this, WorldPartition](const FWorldPartitionActorDescInstance* ActorDescInstance)
+		auto ResaveExternalPackage = [this](const UPackage* Package)
 		{
 			++TotalPackagesForResave;
-			// Load & Register World Partition Actor
-			FWorldPartitionReference LoadedActor(WorldPartition, ActorDescInstance->GetGuid());
-			AActor* Actor = LoadedActor.GetActor();
-			UPackage* Package = Actor ? Actor->GetExternalPackage() : nullptr;
 			if (Package == nullptr)
 			{
 				check(bCanIgnoreFails);
-				return true;
+				return;
 			}
 			const FString PackageFilename = Package->GetLoadedPath().GetLocalFullPath();
 			check(FLinkerLoad::FindExistingLinkerForPackage(Package));
 			LoadAndSaveOnePackage(PackageFilename);
+		};
+
+		// Resave all external actors packages
+		FWorldPartitionHelpers::ForEachActorDescInstance(WorldPartition, [this, WorldPartition, ResaveExternalPackage](const FWorldPartitionActorDescInstance* ActorDescInstance)
+		{
+			// Load & Register World Partition Actor
+			FWorldPartitionReference LoadedActor(WorldPartition, ActorDescInstance->GetGuid());
+			AActor* Actor = LoadedActor.GetActor();
+			UPackage* Package = Actor ? Actor->GetExternalPackage() : nullptr;
+			ResaveExternalPackage(Package);
 			return true;
 		});
+
+		// Resave all external objects packages
+		FExternalPackageHelper::LoadObjectsFromExternalPackages<UObject>(World, [this, ResaveExternalPackage](UObject* ExternalObject)
+		{
+			ResaveExternalPackage(ExternalObject->GetPackage());
+		});
 	}
 
 	if (bBuildingNonHLODData || bShouldBuildHLOD || bShouldBuildNavigationData)
