Interchange fix conversion between legacy fbx and interchange when re-importing custom lod.

#rnx

diff --git a/Engine/Plugins/Interchange/Editor/Source/InterchangeEditor/Private/InterchangeFbxAssetImportDataConverter.cpp b/Engine/Plugins/Interchange/Editor/Source/InterchangeEditor/Private/InterchangeFbxAssetImportDataConverter.cpp
--- a/Engine/Plugins/Interchange/Editor/Source/InterchangeEditor/Private/InterchangeFbxAssetImportDataConverter.cpp
+++ b/Engine/Plugins/Interchange/Editor/Source/InterchangeEditor/Private/InterchangeFbxAssetImportDataConverter.cpp
@@ -98,6 +98,126 @@ namespace UE::Interchange::Private
 		}
 	}
 
+	void FillFbxStaticMeshImportData(const UInterchangeGenericAssetsPipeline* GenericAssetPipeline, UFbxStaticMeshImportData* DestinationStaticMeshImportData)
+	{
+		DestinationStaticMeshImportData->bAutoGenerateCollision = GenericAssetPipeline->MeshPipeline->bImportCollision;
+		DestinationStaticMeshImportData->bBuildNanite = GenericAssetPipeline->MeshPipeline->bBuildNanite;
+		DestinationStaticMeshImportData->bBuildReversedIndexBuffer = GenericAssetPipeline->MeshPipeline->bBuildReversedIndexBuffer;
+		DestinationStaticMeshImportData->bCombineMeshes = GenericAssetPipeline->MeshPipeline->bCombineStaticMeshes;
+		DestinationStaticMeshImportData->bGenerateLightmapUVs = GenericAssetPipeline->MeshPipeline->bGenerateLightmapUVs;
+		DestinationStaticMeshImportData->bOneConvexHullPerUCX = GenericAssetPipeline->MeshPipeline->bOneConvexHullPerUCX;
+		DestinationStaticMeshImportData->bRemoveDegenerates = GenericAssetPipeline->CommonMeshesProperties->bRemoveDegenerates;
+		DestinationStaticMeshImportData->DistanceFieldResolutionScale = GenericAssetPipeline->MeshPipeline->DistanceFieldResolutionScale;
+		DestinationStaticMeshImportData->StaticMeshLODGroup = GenericAssetPipeline->MeshPipeline->LodGroup;
+		if (GenericAssetPipeline->CommonMeshesProperties->VertexColorImportOption == EInterchangeVertexColorImportOption::IVCIO_Ignore)
+		{
+			DestinationStaticMeshImportData->VertexColorImportOption = EVertexColorImportOption::Ignore;
+		}
+		else if (GenericAssetPipeline->CommonMeshesProperties->VertexColorImportOption == EInterchangeVertexColorImportOption::IVCIO_Override)
+		{
+			DestinationStaticMeshImportData->VertexColorImportOption = EVertexColorImportOption::Override;
+		}
+		else if (GenericAssetPipeline->CommonMeshesProperties->VertexColorImportOption == EInterchangeVertexColorImportOption::IVCIO_Replace)
+		{
+			DestinationStaticMeshImportData->VertexColorImportOption = EVertexColorImportOption::Replace;
+		}
+		DestinationStaticMeshImportData->VertexOverrideColor = GenericAssetPipeline->CommonMeshesProperties->VertexOverrideColor;
+	}
+
+	void FillFbxSkeletalMeshImportData(const UInterchangeGenericAssetsPipeline* GenericAssetPipeline, UFbxSkeletalMeshImportData* DestinationSkeletalMeshImportData)
+	{
+		DestinationSkeletalMeshImportData->bImportMeshesInBoneHierarchy = GenericAssetPipeline->CommonSkeletalMeshesAndAnimationsProperties->bImportMeshesInBoneHierarchy;
+		DestinationSkeletalMeshImportData->bImportMorphTargets = GenericAssetPipeline->MeshPipeline->bImportMorphTargets;
+		DestinationSkeletalMeshImportData->bImportVertexAttributes = GenericAssetPipeline->MeshPipeline->bImportVertexAttributes;
+		DestinationSkeletalMeshImportData->bKeepSectionsSeparate = GenericAssetPipeline->CommonMeshesProperties->bKeepSectionsSeparate;
+		DestinationSkeletalMeshImportData->bPreserveSmoothingGroups = true;
+		DestinationSkeletalMeshImportData->bUpdateSkeletonReferencePose = GenericAssetPipeline->MeshPipeline->bUpdateSkeletonReferencePose;
+		DestinationSkeletalMeshImportData->bUseT0AsRefPose = GenericAssetPipeline->CommonSkeletalMeshesAndAnimationsProperties->bUseT0AsRefPose;
+		if (GenericAssetPipeline->MeshPipeline->SkeletalMeshImportContentType == EInterchangeSkeletalMeshContentType::All)
+		{
+			DestinationSkeletalMeshImportData->ImportContentType = EFBXImportContentType::FBXICT_All;
+		}
+		else if (GenericAssetPipeline->MeshPipeline->SkeletalMeshImportContentType == EInterchangeSkeletalMeshContentType::Geometry)
+		{
+			DestinationSkeletalMeshImportData->ImportContentType = EFBXImportContentType::FBXICT_Geometry;
+		}
+		else if (GenericAssetPipeline->MeshPipeline->SkeletalMeshImportContentType == EInterchangeSkeletalMeshContentType::SkinningWeights)
+		{
+			DestinationSkeletalMeshImportData->ImportContentType = EFBXImportContentType::FBXICT_SkinningWeights;
+		}
+
+		if (GenericAssetPipeline->MeshPipeline->LastSkeletalMeshImportContentType == EInterchangeSkeletalMeshContentType::All)
+		{
+			DestinationSkeletalMeshImportData->LastImportContentType = EFBXImportContentType::FBXICT_All;
+		}
+		else if (GenericAssetPipeline->MeshPipeline->LastSkeletalMeshImportContentType == EInterchangeSkeletalMeshContentType::Geometry)
+		{
+			DestinationSkeletalMeshImportData->LastImportContentType = EFBXImportContentType::FBXICT_Geometry;
+		}
+		else if (GenericAssetPipeline->MeshPipeline->LastSkeletalMeshImportContentType == EInterchangeSkeletalMeshContentType::SkinningWeights)
+		{
+			DestinationSkeletalMeshImportData->LastImportContentType = EFBXImportContentType::FBXICT_SkinningWeights;
+		}
+
+		DestinationSkeletalMeshImportData->MorphThresholdPosition = GenericAssetPipeline->MeshPipeline->MorphThresholdPosition;
+		DestinationSkeletalMeshImportData->ThresholdPosition = GenericAssetPipeline->MeshPipeline->ThresholdPosition;
+		DestinationSkeletalMeshImportData->ThresholdTangentNormal = GenericAssetPipeline->MeshPipeline->ThresholdTangentNormal;
+		DestinationSkeletalMeshImportData->ThresholdUV = GenericAssetPipeline->MeshPipeline->ThresholdUV;
+
+		if (GenericAssetPipeline->CommonMeshesProperties->VertexColorImportOption == EInterchangeVertexColorImportOption::IVCIO_Ignore)
+		{
+			DestinationSkeletalMeshImportData->VertexColorImportOption = EVertexColorImportOption::Ignore;
+		}
+		else if (GenericAssetPipeline->CommonMeshesProperties->VertexColorImportOption == EInterchangeVertexColorImportOption::IVCIO_Override)
+		{
+			DestinationSkeletalMeshImportData->VertexColorImportOption = EVertexColorImportOption::Override;
+		}
+		else if (GenericAssetPipeline->CommonMeshesProperties->VertexColorImportOption == EInterchangeVertexColorImportOption::IVCIO_Replace)
+		{
+			DestinationSkeletalMeshImportData->VertexColorImportOption = EVertexColorImportOption::Replace;
+		}
+		DestinationSkeletalMeshImportData->VertexOverrideColor = GenericAssetPipeline->CommonMeshesProperties->VertexOverrideColor;
+	}
+
+	void FillFbxAnimSequenceImportData(const UInterchangeGenericAssetsPipeline* GenericAssetPipeline, UFbxAnimSequenceImportData* DestinationAnimSequenceImportData)
+	{
+		switch (GenericAssetPipeline->AnimationPipeline->AnimationRange)
+		{
+		case EInterchangeAnimationRange::Timeline:
+		{
+			DestinationAnimSequenceImportData->AnimationLength = EFBXAnimationLengthImportType::FBXALIT_ExportedTime;
+			break;
+		}
+		case EInterchangeAnimationRange::Animated:
+		{
+			DestinationAnimSequenceImportData->AnimationLength = EFBXAnimationLengthImportType::FBXALIT_AnimatedKey;
+			break;
+		}
+		case EInterchangeAnimationRange::SetRange:
+		{
+			DestinationAnimSequenceImportData->AnimationLength = EFBXAnimationLengthImportType::FBXALIT_SetRange;
+			break;
+		}
+		}
+		DestinationAnimSequenceImportData->bAddCurveMetadataToSkeleton = GenericAssetPipeline->AnimationPipeline->bAddCurveMetadataToSkeleton;
+		DestinationAnimSequenceImportData->bDeleteExistingCustomAttributeCurves = GenericAssetPipeline->AnimationPipeline->bDeleteExistingCustomAttributeCurves;
+		DestinationAnimSequenceImportData->bDeleteExistingMorphTargetCurves = GenericAssetPipeline->AnimationPipeline->bDeleteExistingMorphTargetCurves;
+		DestinationAnimSequenceImportData->bDeleteExistingNonCurveCustomAttributes = GenericAssetPipeline->AnimationPipeline->bDeleteExistingNonCurveCustomAttributes;
+		DestinationAnimSequenceImportData->bDoNotImportCurveWithZero = GenericAssetPipeline->AnimationPipeline->bDoNotImportCurveWithZero;
+		DestinationAnimSequenceImportData->bImportBoneTracks = GenericAssetPipeline->AnimationPipeline->bImportBoneTracks;
+		DestinationAnimSequenceImportData->bImportCustomAttribute = GenericAssetPipeline->AnimationPipeline->bImportCustomAttribute;
+		DestinationAnimSequenceImportData->bImportMeshesInBoneHierarchy = GenericAssetPipeline->CommonSkeletalMeshesAndAnimationsProperties->bImportMeshesInBoneHierarchy;
+		DestinationAnimSequenceImportData->bPreserveLocalTransform = false;
+		DestinationAnimSequenceImportData->bRemoveRedundantKeys = GenericAssetPipeline->AnimationPipeline->bRemoveCurveRedundantKeys;
+		DestinationAnimSequenceImportData->bSetMaterialDriveParameterOnCustomAttribute = GenericAssetPipeline->AnimationPipeline->bSetMaterialDriveParameterOnCustomAttribute;
+		DestinationAnimSequenceImportData->bSnapToClosestFrameBoundary = GenericAssetPipeline->AnimationPipeline->bSnapToClosestFrameBoundary;
+		DestinationAnimSequenceImportData->bUseDefaultSampleRate = GenericAssetPipeline->AnimationPipeline->bUse30HzToBakeBoneAnimation;
+		DestinationAnimSequenceImportData->CustomSampleRate = GenericAssetPipeline->AnimationPipeline->CustomBoneAnimationSampleRate;
+		DestinationAnimSequenceImportData->FrameImportRange = GenericAssetPipeline->AnimationPipeline->FrameImportRange;
+		DestinationAnimSequenceImportData->MaterialCurveSuffixes = GenericAssetPipeline->AnimationPipeline->MaterialCurveSuffixes;
+		DestinationAnimSequenceImportData->SourceAnimationName = GenericAssetPipeline->AnimationPipeline->SourceAnimationName;
+	}
+
 	void FillInterchangeGenericAssetsPipelineFromFbxMeshImportData(UInterchangeGenericAssetsPipeline* GenericAssetPipeline, const UFbxMeshImportData* LegacyMeshImportData)
 	{
 		if (!LegacyMeshImportData || !GenericAssetPipeline)
@@ -119,20 +239,20 @@ namespace UE::Interchange::Private
 			GenericAssetPipeline->CommonMeshesProperties->bUseMikkTSpace = false;
 		}
 
-		if (LegacyMeshImportData->NormalImportMethod == EFBXNormalImportMethod::FBXNIM_ComputeNormals)
+		if (LegacyMeshImportData->NormalImportMethod == EFBXNormalImportMethod::FBXNIM_ImportNormalsAndTangents)
 		{
-			GenericAssetPipeline->CommonMeshesProperties->bRecomputeNormals = true;
+			GenericAssetPipeline->CommonMeshesProperties->bRecomputeNormals = false;
+			GenericAssetPipeline->CommonMeshesProperties->bRecomputeTangents = false;
 		}
-		else
+		else if (LegacyMeshImportData->NormalImportMethod == EFBXNormalImportMethod::FBXNIM_ImportNormals)
 		{
-			if (LegacyMeshImportData->NormalImportMethod == EFBXNormalImportMethod::FBXNIM_ImportNormals)
-			{
-				GenericAssetPipeline->CommonMeshesProperties->bRecomputeTangents = true;
-			}
-			else
-			{
-				GenericAssetPipeline->CommonMeshesProperties->bRecomputeTangents = false;
-			}
+			GenericAssetPipeline->CommonMeshesProperties->bRecomputeNormals = false;
+			GenericAssetPipeline->CommonMeshesProperties->bRecomputeTangents = true;
+		}
+		else if (LegacyMeshImportData->NormalImportMethod == EFBXNormalImportMethod::FBXNIM_ComputeNormals)
+		{
+			GenericAssetPipeline->CommonMeshesProperties->bRecomputeNormals = true;
+			GenericAssetPipeline->CommonMeshesProperties->bRecomputeTangents = true;
 		}
 	}
 
@@ -314,30 +434,7 @@ namespace UE::Interchange::Private
 
 				FillFbxAssetImportData(InterchangeFbxTranslatorSettings, GenericAssetPipeline, DestinationStaticMeshImportData);
 				FillFbxMeshImportData(GenericAssetPipeline, DestinationStaticMeshImportData);
-
-				DestinationStaticMeshImportData->bAutoGenerateCollision = GenericAssetPipeline->MeshPipeline->bImportCollision;
-				DestinationStaticMeshImportData->bBuildNanite = GenericAssetPipeline->MeshPipeline->bBuildNanite;
-				DestinationStaticMeshImportData->bBuildReversedIndexBuffer = GenericAssetPipeline->MeshPipeline->bBuildReversedIndexBuffer;
-				DestinationStaticMeshImportData->bCombineMeshes = GenericAssetPipeline->MeshPipeline->bCombineStaticMeshes;
-				DestinationStaticMeshImportData->bGenerateLightmapUVs = GenericAssetPipeline->MeshPipeline->bGenerateLightmapUVs;
-				DestinationStaticMeshImportData->bOneConvexHullPerUCX = GenericAssetPipeline->MeshPipeline->bOneConvexHullPerUCX;
-				DestinationStaticMeshImportData->bRemoveDegenerates = GenericAssetPipeline->CommonMeshesProperties->bRemoveDegenerates;
-				DestinationStaticMeshImportData->DistanceFieldResolutionScale = GenericAssetPipeline->MeshPipeline->DistanceFieldResolutionScale;
-				DestinationStaticMeshImportData->StaticMeshLODGroup = GenericAssetPipeline->MeshPipeline->LodGroup;
-				if (GenericAssetPipeline->CommonMeshesProperties->VertexColorImportOption == EInterchangeVertexColorImportOption::IVCIO_Ignore)
-				{
-					DestinationStaticMeshImportData->VertexColorImportOption = EVertexColorImportOption::Ignore;
-				}
-				else if (GenericAssetPipeline->CommonMeshesProperties->VertexColorImportOption == EInterchangeVertexColorImportOption::IVCIO_Override)
-				{
-					DestinationStaticMeshImportData->VertexColorImportOption = EVertexColorImportOption::Override;
-				}
-				else if (GenericAssetPipeline->CommonMeshesProperties->VertexColorImportOption == EInterchangeVertexColorImportOption::IVCIO_Replace)
-				{
-					DestinationStaticMeshImportData->VertexColorImportOption = EVertexColorImportOption::Replace;
-				}
-				DestinationStaticMeshImportData->VertexOverrideColor = GenericAssetPipeline->CommonMeshesProperties->VertexOverrideColor;
-
+				FillFbxStaticMeshImportData(GenericAssetPipeline, DestinationStaticMeshImportData);
 				//Fill the reimport material match data and section data
 				FImportMeshLodSectionsData SectionData;
 				for (const FStaticMaterial& Material : StaticMesh->GetStaticMaterials())
@@ -378,58 +475,7 @@ namespace UE::Interchange::Private
 			{
 				FillFbxAssetImportData(InterchangeFbxTranslatorSettings, GenericAssetPipeline, DestinationSkeletalMeshImportData);
 				FillFbxMeshImportData(GenericAssetPipeline, DestinationSkeletalMeshImportData);
-				DestinationSkeletalMeshImportData->bImportMeshesInBoneHierarchy = GenericAssetPipeline->CommonSkeletalMeshesAndAnimationsProperties->bImportMeshesInBoneHierarchy;
-				DestinationSkeletalMeshImportData->bImportMorphTargets = GenericAssetPipeline->MeshPipeline->bImportMorphTargets;
-				DestinationSkeletalMeshImportData->bImportVertexAttributes = GenericAssetPipeline->MeshPipeline->bImportVertexAttributes;
-				DestinationSkeletalMeshImportData->bKeepSectionsSeparate = GenericAssetPipeline->CommonMeshesProperties->bKeepSectionsSeparate;
-				DestinationSkeletalMeshImportData->bPreserveSmoothingGroups = true;
-				DestinationSkeletalMeshImportData->bUpdateSkeletonReferencePose = GenericAssetPipeline->MeshPipeline->bUpdateSkeletonReferencePose;
-				DestinationSkeletalMeshImportData->bUseT0AsRefPose = GenericAssetPipeline->CommonSkeletalMeshesAndAnimationsProperties->bUseT0AsRefPose;
-				if (GenericAssetPipeline->MeshPipeline->SkeletalMeshImportContentType == EInterchangeSkeletalMeshContentType::All)
-				{
-					DestinationSkeletalMeshImportData->ImportContentType = EFBXImportContentType::FBXICT_All;
-				}
-				else if (GenericAssetPipeline->MeshPipeline->SkeletalMeshImportContentType == EInterchangeSkeletalMeshContentType::Geometry)
-				{
-					DestinationSkeletalMeshImportData->ImportContentType = EFBXImportContentType::FBXICT_Geometry;
-				}
-				else if (GenericAssetPipeline->MeshPipeline->SkeletalMeshImportContentType == EInterchangeSkeletalMeshContentType::SkinningWeights)
-				{
-					DestinationSkeletalMeshImportData->ImportContentType = EFBXImportContentType::FBXICT_SkinningWeights;
-				}
-
-				if (GenericAssetPipeline->MeshPipeline->LastSkeletalMeshImportContentType == EInterchangeSkeletalMeshContentType::All)
-				{
-					DestinationSkeletalMeshImportData->LastImportContentType = EFBXImportContentType::FBXICT_All;
-				}
-				else if (GenericAssetPipeline->MeshPipeline->LastSkeletalMeshImportContentType == EInterchangeSkeletalMeshContentType::Geometry)
-				{
-					DestinationSkeletalMeshImportData->LastImportContentType = EFBXImportContentType::FBXICT_Geometry;
-				}
-				else if (GenericAssetPipeline->MeshPipeline->LastSkeletalMeshImportContentType == EInterchangeSkeletalMeshContentType::SkinningWeights)
-				{
-					DestinationSkeletalMeshImportData->LastImportContentType = EFBXImportContentType::FBXICT_SkinningWeights;
-				}
-
-				DestinationSkeletalMeshImportData->MorphThresholdPosition = GenericAssetPipeline->MeshPipeline->MorphThresholdPosition;
-				DestinationSkeletalMeshImportData->ThresholdPosition = GenericAssetPipeline->MeshPipeline->ThresholdPosition;
-				DestinationSkeletalMeshImportData->ThresholdTangentNormal = GenericAssetPipeline->MeshPipeline->ThresholdTangentNormal;
-				DestinationSkeletalMeshImportData->ThresholdUV = GenericAssetPipeline->MeshPipeline->ThresholdUV;
-
-				if (GenericAssetPipeline->CommonMeshesProperties->VertexColorImportOption == EInterchangeVertexColorImportOption::IVCIO_Ignore)
-				{
-					DestinationSkeletalMeshImportData->VertexColorImportOption = EVertexColorImportOption::Ignore;
-				}
-				else if (GenericAssetPipeline->CommonMeshesProperties->VertexColorImportOption == EInterchangeVertexColorImportOption::IVCIO_Override)
-				{
-					DestinationSkeletalMeshImportData->VertexColorImportOption = EVertexColorImportOption::Override;
-				}
-				else if (GenericAssetPipeline->CommonMeshesProperties->VertexColorImportOption == EInterchangeVertexColorImportOption::IVCIO_Replace)
-				{
-					DestinationSkeletalMeshImportData->VertexColorImportOption = EVertexColorImportOption::Replace;
-				}
-				DestinationSkeletalMeshImportData->VertexOverrideColor = GenericAssetPipeline->CommonMeshesProperties->VertexOverrideColor;
-
+				FillFbxSkeletalMeshImportData(GenericAssetPipeline, DestinationSkeletalMeshImportData);
 				//Fill the reimport material match data and section data
 				FImportMeshLodSectionsData SectionData;
 				for (const FSkeletalMaterial& Material : SkeletalMesh->GetMaterials())
@@ -469,47 +515,117 @@ namespace UE::Interchange::Private
 			if (const UInterchangeGenericAssetsPipeline* GenericAssetPipeline = Cast<UInterchangeGenericAssetsPipeline>(Pipeline))
 			{
 				FillFbxAssetImportData(InterchangeFbxTranslatorSettings, GenericAssetPipeline, DestinationAnimSequenceImportData);
-
-				switch (GenericAssetPipeline->AnimationPipeline->AnimationRange)
-				{
-				case EInterchangeAnimationRange::Timeline:
-				{
-					DestinationAnimSequenceImportData->AnimationLength = EFBXAnimationLengthImportType::FBXALIT_ExportedTime;
-					break;
-				}
-				case EInterchangeAnimationRange::Animated:
-				{
-					DestinationAnimSequenceImportData->AnimationLength = EFBXAnimationLengthImportType::FBXALIT_AnimatedKey;
-					break;
-				}
-				case EInterchangeAnimationRange::SetRange:
-				{
-					DestinationAnimSequenceImportData->AnimationLength = EFBXAnimationLengthImportType::FBXALIT_SetRange;
-					break;
-				}
-				}
-				DestinationAnimSequenceImportData->bAddCurveMetadataToSkeleton = GenericAssetPipeline->AnimationPipeline->bAddCurveMetadataToSkeleton;
-				DestinationAnimSequenceImportData->bDeleteExistingCustomAttributeCurves = GenericAssetPipeline->AnimationPipeline->bDeleteExistingCustomAttributeCurves;
-				DestinationAnimSequenceImportData->bDeleteExistingMorphTargetCurves = GenericAssetPipeline->AnimationPipeline->bDeleteExistingMorphTargetCurves;
-				DestinationAnimSequenceImportData->bDeleteExistingNonCurveCustomAttributes = GenericAssetPipeline->AnimationPipeline->bDeleteExistingNonCurveCustomAttributes;
-				DestinationAnimSequenceImportData->bDoNotImportCurveWithZero = GenericAssetPipeline->AnimationPipeline->bDoNotImportCurveWithZero;
-				DestinationAnimSequenceImportData->bImportBoneTracks = GenericAssetPipeline->AnimationPipeline->bImportBoneTracks;
-				DestinationAnimSequenceImportData->bImportCustomAttribute = GenericAssetPipeline->AnimationPipeline->bImportCustomAttribute;
-				DestinationAnimSequenceImportData->bImportMeshesInBoneHierarchy = GenericAssetPipeline->CommonSkeletalMeshesAndAnimationsProperties->bImportMeshesInBoneHierarchy;
-				DestinationAnimSequenceImportData->bPreserveLocalTransform = false;
-				DestinationAnimSequenceImportData->bRemoveRedundantKeys = GenericAssetPipeline->AnimationPipeline->bRemoveCurveRedundantKeys;
-				DestinationAnimSequenceImportData->bSetMaterialDriveParameterOnCustomAttribute = GenericAssetPipeline->AnimationPipeline->bSetMaterialDriveParameterOnCustomAttribute;
-				DestinationAnimSequenceImportData->bSnapToClosestFrameBoundary = GenericAssetPipeline->AnimationPipeline->bSnapToClosestFrameBoundary;
-				DestinationAnimSequenceImportData->bUseDefaultSampleRate = GenericAssetPipeline->AnimationPipeline->bUse30HzToBakeBoneAnimation;
-				DestinationAnimSequenceImportData->CustomSampleRate = GenericAssetPipeline->AnimationPipeline->CustomBoneAnimationSampleRate;
-				DestinationAnimSequenceImportData->FrameImportRange = GenericAssetPipeline->AnimationPipeline->FrameImportRange;
-				DestinationAnimSequenceImportData->MaterialCurveSuffixes = GenericAssetPipeline->AnimationPipeline->MaterialCurveSuffixes;
-				DestinationAnimSequenceImportData->SourceAnimationName = GenericAssetPipeline->AnimationPipeline->SourceAnimationName;
+				FillFbxAnimSequenceImportData(GenericAssetPipeline, DestinationAnimSequenceImportData);
 			}
 		}
 		return DestinationAnimSequenceImportData;
 	}
 
+	UFbxImportUI* ConvertToLegacyFbx(UObject* Owner, const UInterchangeAssetImportData* InterchangeSourceData)
+	{
+		if (!InterchangeSourceData || !Owner)
+		{
+			return nullptr;
+		}
+		//Create a fbx asset import data and fill the options
+		UFbxImportUI* DestinationData = NewObject<UFbxImportUI>(Owner);
+
+		if (!DestinationData)
+		{
+			return nullptr;
+		}
+
+		const UInterchangeFbxTranslatorSettings* InterchangeFbxTranslatorSettings = Cast<UInterchangeFbxTranslatorSettings>(InterchangeSourceData->GetTranslatorSettings());
+		const UInterchangeGenericAssetsPipeline* GenericAssetPipeline = nullptr;
+		//Now find the generic asset pipeline
+		for (UObject* Pipeline : InterchangeSourceData->GetPipelines())
+		{
+			if (const UInterchangeGenericAssetsPipeline* AssetPipeline = Cast<UInterchangeGenericAssetsPipeline>(Pipeline))
+			{
+				GenericAssetPipeline = AssetPipeline;
+			}
+		}
+
+		if (!GenericAssetPipeline)
+		{
+			//Since we did not find any generic asset pipeline we fallback on the generic pipeline CDO settings
+			GenericAssetPipeline = NewObject<UInterchangeGenericAssetsPipeline>(GetTransientPackage());
+		}
+
+		FillFbxAssetImportData(InterchangeFbxTranslatorSettings, GenericAssetPipeline, DestinationData->StaticMeshImportData);
+		FillFbxAssetImportData(InterchangeFbxTranslatorSettings, GenericAssetPipeline, DestinationData->SkeletalMeshImportData);
+		FillFbxAssetImportData(InterchangeFbxTranslatorSettings, GenericAssetPipeline, DestinationData->AnimSequenceImportData);
+		FillFbxAssetImportData(InterchangeFbxTranslatorSettings, GenericAssetPipeline, DestinationData->TextureImportData);
+
+		FillFbxMeshImportData(GenericAssetPipeline, DestinationData->StaticMeshImportData);
+		FillFbxMeshImportData(GenericAssetPipeline, DestinationData->SkeletalMeshImportData);
+
+		FillFbxStaticMeshImportData(GenericAssetPipeline, DestinationData->StaticMeshImportData);
+		FillFbxSkeletalMeshImportData(GenericAssetPipeline, DestinationData->SkeletalMeshImportData);
+
+		FillFbxAnimSequenceImportData(GenericAssetPipeline, DestinationData->AnimSequenceImportData);
+
+		DestinationData->bOverrideFullName = GenericAssetPipeline->bUseSourceNameForAsset;
+
+		//Material Options
+		DestinationData->bImportMaterials = GenericAssetPipeline->MaterialPipeline->bImportMaterials;
+		switch (GenericAssetPipeline->MaterialPipeline->SearchLocation)
+		{
+		case EInterchangeMaterialSearchLocation::Local:
+			DestinationData->TextureImportData->MaterialSearchLocation = EMaterialSearchLocation::Local;
+			break;
+		case EInterchangeMaterialSearchLocation::UnderParent:
+			DestinationData->TextureImportData->MaterialSearchLocation = EMaterialSearchLocation::UnderParent;
+			break;
+		case EInterchangeMaterialSearchLocation::UnderRoot:
+			DestinationData->TextureImportData->MaterialSearchLocation = EMaterialSearchLocation::UnderRoot;
+			break;
+		case EInterchangeMaterialSearchLocation::AllAssets:
+			DestinationData->TextureImportData->MaterialSearchLocation = EMaterialSearchLocation::AllAssets;
+			break;
+		case EInterchangeMaterialSearchLocation::DoNotSearch:
+			DestinationData->TextureImportData->MaterialSearchLocation = EMaterialSearchLocation::DoNotSearch;
+			break;
+		}
+
+		if (GenericAssetPipeline->MaterialPipeline->ParentMaterial.IsAsset())
+		{
+			DestinationData->TextureImportData->bUseBaseMaterial = true;
+			DestinationData->TextureImportData->BaseMaterialName = GenericAssetPipeline->MaterialPipeline->ParentMaterial;
+		}
+		else
+		{
+			DestinationData->TextureImportData->bUseBaseMaterial = false;
+			DestinationData->TextureImportData->BaseMaterialName.Reset();
+		}
+
+		//Texture Options
+		DestinationData->bImportTextures = GenericAssetPipeline->MaterialPipeline->TexturePipeline->bImportTextures;
+		DestinationData->TextureImportData->bInvertNormalMaps = GenericAssetPipeline->MaterialPipeline->TexturePipeline->bFlipNormalMapGreenChannel;
+
+		//Discover if we must import something in particular
+		if (GenericAssetPipeline->MeshPipeline->CommonMeshesProperties->ForceAllMeshAsType == EInterchangeForceMeshType::IFMT_SkeletalMesh)
+		{
+			DestinationData->MeshTypeToImport = EFBXImportType::FBXIT_SkeletalMesh;
+			DestinationData->bImportAsSkeletal = true;
+			DestinationData->bImportAnimations = GenericAssetPipeline->AnimationPipeline->bImportAnimations;
+		}
+		else if (GenericAssetPipeline->MeshPipeline->CommonMeshesProperties->ForceAllMeshAsType == EInterchangeForceMeshType::IFMT_StaticMesh)
+		{
+			DestinationData->MeshTypeToImport = EFBXImportType::FBXIT_StaticMesh;
+
+			DestinationData->bImportAsSkeletal = false;
+			DestinationData->bImportAnimations = false;
+		}
+		else
+		{
+			DestinationData->bAutomatedImportShouldDetectType = true;
+		}
+		DestinationData->Skeleton = GenericAssetPipeline->CommonSkeletalMeshesAndAnimationsProperties->Skeleton.Get();
+
+		return DestinationData;
+	}
+
 	UAssetImportData* ConvertToInterchange(UObject* Obj, const UFbxAssetImportData* FbxAssetImportData)
 	{
 		if (!FbxAssetImportData || !Obj)
@@ -806,7 +922,7 @@ namespace UE::Interchange::Private
 	}
 } //ns: UE::Interchange::Private
 
-bool UInterchangeFbxAssetImportDataConverter::ConvertImportData(UObject* Obj, const FString& TargetExtension) const
+bool UInterchangeFbxAssetImportDataConverter::ConvertImportData(UObject* Asset, const FString& TargetExtension) const
 {
 	bool bResult = false;
 	const FString TargetExtensionLower = TargetExtension.ToLower();
@@ -815,11 +931,11 @@ bool UInterchangeFbxAssetImportDataConverter::ConvertImportData(UObject* Obj, co
 	
 	UAssetImportData* OldAssetData = nullptr;
 	TArray<FString> InterchangeSupportedExtensions;
-	if (Obj->IsA(UStaticMesh::StaticClass()) || Obj->IsA(USkeletalMesh::StaticClass()))
+	if (Asset->IsA(UStaticMesh::StaticClass()) || Asset->IsA(USkeletalMesh::StaticClass()))
 	{
 		InterchangeSupportedExtensions = InterchangeManager.GetSupportedAssetTypeFormats(EInterchangeTranslatorAssetType::Meshes);
 	}
-	else if (Obj->IsA(UAnimSequence::StaticClass()))
+	else if (Asset->IsA(UAnimSequence::StaticClass()))
 	{
 		InterchangeSupportedExtensions = InterchangeManager.GetSupportedAssetTypeFormats(EInterchangeTranslatorAssetType::Animations);
 	}
@@ -837,7 +953,7 @@ bool UInterchangeFbxAssetImportDataConverter::ConvertImportData(UObject* Obj, co
 	
 	if (TargetExtensionLower.Equals(TEXT("fbx")) || bInterchangeSupportTargetExtension)
 	{
-		if (UStaticMesh* StaticMesh = Cast<UStaticMesh>(Obj))
+		if (UStaticMesh* StaticMesh = Cast<UStaticMesh>(Asset))
 		{
 			if (UAssetImportData* ConvertedAssetData = UE::Interchange::Private::ConvertData(StaticMesh, StaticMesh->GetAssetImportData(), bInterchangeSupportTargetExtension))
 			{
@@ -846,7 +962,7 @@ bool UInterchangeFbxAssetImportDataConverter::ConvertImportData(UObject* Obj, co
 				bResult = true;
 			}
 		}
-		else if (USkeletalMesh* SkeletalMesh = Cast<USkeletalMesh>(Obj))
+		else if (USkeletalMesh* SkeletalMesh = Cast<USkeletalMesh>(Asset))
 		{
 			if (UAssetImportData* ConvertedAssetData = UE::Interchange::Private::ConvertData(SkeletalMesh, SkeletalMesh->GetAssetImportData(), bInterchangeSupportTargetExtension))
 			{
@@ -855,7 +971,7 @@ bool UInterchangeFbxAssetImportDataConverter::ConvertImportData(UObject* Obj, co
 				bResult = true;
 			}
 		}
-		else if (UAnimSequence* AnimSequence = Cast<UAnimSequence>(Obj))
+		else if (UAnimSequence* AnimSequence = Cast<UAnimSequence>(Asset))
 		{
 			if (UAssetImportData* ConvertedAssetData = UE::Interchange::Private::ConvertData(AnimSequence, AnimSequence->AssetImportData, bInterchangeSupportTargetExtension))
 			{
@@ -891,11 +1007,98 @@ bool UInterchangeFbxAssetImportDataConverter::ConvertImportData(const UObject* S
 		*DestinationImportData = UE::Interchange::Private::ConvertToInterchange(GetTransientPackage(), Cast<UFbxImportUI>(SourceImportData));
 		bResult = true;
 	}
+	else if (SourceImportData->IsA<UFbxAssetImportData>())
+	{
+		//We convert the UFbxAssetImportData into a UFbxImportUI
+		auto FillFbxAssetImportData = [](const UFbxAssetImportData* SourceAssetImportData, UFbxAssetImportData* DestinationAssetImportData)
+			{
+				DestinationAssetImportData->bConvertScene = SourceAssetImportData->bConvertScene;
+				DestinationAssetImportData->bConvertSceneUnit = SourceAssetImportData->bConvertSceneUnit;
+				DestinationAssetImportData->bForceFrontXAxis = SourceAssetImportData->bForceFrontXAxis;
+				DestinationAssetImportData->bImportAsScene = SourceAssetImportData->bImportAsScene;
+				DestinationAssetImportData->ImportRotation = SourceAssetImportData->ImportRotation;
+				DestinationAssetImportData->ImportTranslation = SourceAssetImportData->ImportTranslation;
+				DestinationAssetImportData->ImportUniformScale = SourceAssetImportData->ImportUniformScale;
+			};
+
+		UFbxImportUI* TempFbxImportUI = NewObject<UFbxImportUI>(GetTransientPackage());
+		TempFbxImportUI->bImportMaterials = false;
+		TempFbxImportUI->bImportAsSkeletal = false;
+		TempFbxImportUI->bImportMesh = false;
+		TempFbxImportUI->bImportAnimations = false;
+		TempFbxImportUI->bImportRigidMesh = false;
+		TempFbxImportUI->bImportTextures = false;
+		TempFbxImportUI->bIsObjImport = false;
+		TempFbxImportUI->bIsReimport = false;
+		TempFbxImportUI->bCreatePhysicsAsset = false;
+		TempFbxImportUI->PhysicsAsset = nullptr;
+		TempFbxImportUI->Skeleton = nullptr;
+		if (SourceImportData->IsA<UFbxSkeletalMeshImportData>())
+		{
+			TempFbxImportUI->SkeletalMeshImportData = const_cast<UFbxSkeletalMeshImportData*>(Cast<UFbxSkeletalMeshImportData>(SourceImportData));
+			TempFbxImportUI->MeshTypeToImport = EFBXImportType::FBXIT_SkeletalMesh;
+			TempFbxImportUI->bImportAsSkeletal = true;
+			TempFbxImportUI->bImportMesh = true;
+
+			FillFbxAssetImportData(TempFbxImportUI->SkeletalMeshImportData, TempFbxImportUI->StaticMeshImportData);
+			FillFbxAssetImportData(TempFbxImportUI->SkeletalMeshImportData, TempFbxImportUI->AnimSequenceImportData);
+			FillFbxAssetImportData(TempFbxImportUI->SkeletalMeshImportData, TempFbxImportUI->TextureImportData);
+		}
+		else if (SourceImportData->IsA<UFbxStaticMeshImportData>())
+		{
+			TempFbxImportUI->MeshTypeToImport = EFBXImportType::FBXIT_StaticMesh;
+			TempFbxImportUI->bImportMesh = true;
+
+			TempFbxImportUI->StaticMeshImportData = const_cast<UFbxStaticMeshImportData*>(Cast<UFbxStaticMeshImportData>(SourceImportData));
+			FillFbxAssetImportData(TempFbxImportUI->StaticMeshImportData, TempFbxImportUI->SkeletalMeshImportData);
+			FillFbxAssetImportData(TempFbxImportUI->StaticMeshImportData, TempFbxImportUI->AnimSequenceImportData);
+			FillFbxAssetImportData(TempFbxImportUI->StaticMeshImportData, TempFbxImportUI->TextureImportData);
+		}
+		else if (SourceImportData->IsA<UFbxAnimSequenceImportData>())
+		{
+			TempFbxImportUI->MeshTypeToImport = EFBXImportType::FBXIT_Animation;
+			TempFbxImportUI->bImportAsSkeletal = true;
+			TempFbxImportUI->bImportMesh = false;
+			TempFbxImportUI->AnimSequenceImportData = const_cast<UFbxAnimSequenceImportData*>(Cast<UFbxAnimSequenceImportData>(SourceImportData));
+			FillFbxAssetImportData(TempFbxImportUI->AnimSequenceImportData, TempFbxImportUI->SkeletalMeshImportData);
+			FillFbxAssetImportData(TempFbxImportUI->AnimSequenceImportData, TempFbxImportUI->StaticMeshImportData);
+			FillFbxAssetImportData(TempFbxImportUI->AnimSequenceImportData, TempFbxImportUI->TextureImportData);
+		}
+		else
+		{
+			ensureMsgf(false, TEXT("Fbx interchange converter: miss match between CanConvertClass and the convertion capacity"));
+		}
+		*DestinationImportData = UE::Interchange::Private::ConvertToInterchange(GetTransientPackage(), TempFbxImportUI);
+		bResult = true;
+	}
 	else if (SourceImportData->IsA<UInterchangeAssetImportData>())
 	{
-		//TODO Convert Interchange to Legacy Fbx
-		//*DestinationImportData = UE::Interchange::Private::ConvertToLegacyFbx(GetTransientPackage(), Cast<UInterchangeAssetImportData>(SourceImportData));
+		*DestinationImportData = UE::Interchange::Private::ConvertToLegacyFbx(GetTransientPackage(), Cast<UInterchangeAssetImportData>(SourceImportData));
+		bResult = true;
 	}
 	return bResult;
 }
 
+bool UInterchangeFbxAssetImportDataConverter::CanConvertClass(const UClass* SourceClass, const UClass* DestinationClass) const
+{
+	if (SourceClass->IsChildOf(UFbxImportUI::StaticClass()))
+	{
+		return DestinationClass->IsChildOf(UInterchangeAssetImportData::StaticClass());
+	}
+
+	if (SourceClass->IsChildOf(UFbxAssetImportData::StaticClass()))
+	{
+		if (SourceClass->IsChildOf(UFbxSkeletalMeshImportData::StaticClass())
+			|| SourceClass->IsChildOf(UFbxStaticMeshImportData::StaticClass())
+			|| SourceClass->IsChildOf(UFbxAnimSequenceImportData::StaticClass()))
+		{
+			return DestinationClass->IsChildOf(UInterchangeAssetImportData::StaticClass());
+		}
+	}
+
+	if (SourceClass->IsChildOf(UInterchangeAssetImportData::StaticClass()))
+	{
+		return DestinationClass->IsChildOf(UFbxImportUI::StaticClass());
+	}
+	return false;
+}
\ No newline at end of file
diff --git a/Engine/Plugins/Interchange/Editor/Source/InterchangeEditor/Public/InterchangeFbxAssetImportDataConverter.h b/Engine/Plugins/Interchange/Editor/Source/InterchangeEditor/Public/InterchangeFbxAssetImportDataConverter.h
--- a/Engine/Plugins/Interchange/Editor/Source/InterchangeEditor/Public/InterchangeFbxAssetImportDataConverter.h
+++ b/Engine/Plugins/Interchange/Editor/Source/InterchangeEditor/Public/InterchangeFbxAssetImportDataConverter.h
@@ -15,6 +15,8 @@ class INTERCHANGEEDITOR_API UInterchangeFbxAssetImportDataConverter : public UIn
 	GENERATED_BODY()
 
 public:
-	virtual bool ConvertImportData(UObject* Object, const FString& ToExtension) const override;
-	virtual bool ConvertImportData(const UObject* SourceImportData, UObject** DestinationImportDataClass) const override;
+	virtual bool CanConvertClass(const UClass* SourceClass, const UClass* DestinationClass) const override;
+
+	virtual bool ConvertImportData(UObject* Asset, const FString& ToExtension) const override;
+	virtual bool ConvertImportData(const UObject* SourceImportData, UObject** DestinationImportData) const override;
 };
diff --git a/Engine/Plugins/Interchange/Editor/Source/Utilities/Private/InterchangeOpenFileDialog.cpp b/Engine/Plugins/Interchange/Editor/Source/Utilities/Private/InterchangeOpenFileDialog.cpp
--- a/Engine/Plugins/Interchange/Editor/Source/Utilities/Private/InterchangeOpenFileDialog.cpp
+++ b/Engine/Plugins/Interchange/Editor/Source/Utilities/Private/InterchangeOpenFileDialog.cpp
@@ -2,21 +2,48 @@
 #include "InterchangeOpenFileDialog.h"
 
 #include "DesktopPlatformModule.h"
+#include "Factories/Factory.h"
 #include "Framework/Application/SlateApplication.h"
 #include "IDesktopPlatform.h"
 #include "InterchangeManager.h"
 #include "ObjectTools.h"
+#include "UObject/UObjectIterator.h"
 
 #include UE_INLINE_GENERATED_CPP_BY_NAME(InterchangeOpenFileDialog)
 
 namespace UE::Interchange::Utilities::Private
 {
-	FString GetOpenFileDialogExtensions(const TArray<FString>& TranslatorFormats)
+	FString GetOpenFileDialogExtensions(const TArray<FString>& TranslatorFormats, bool bShowAllFactoriesExtension)
 	{
 		FString FileTypes;
 		FString Extensions;
+		TMultiMap<uint32, UFactory*> FilterIndexToFactory;
 
-		ObjectTools::AppendFormatsFileExtensions(TranslatorFormats, FileTypes, Extensions);
+		if (bShowAllFactoriesExtension)
+		{
+			TArray<UFactory*> Factories;
+
+			// Get the list of valid factories
+			for (TObjectIterator<UClass> It; It; ++It)
+			{
+				UClass* CurrentClass = (*It);
+
+				if (CurrentClass->IsChildOf(UFactory::StaticClass()) && !(CurrentClass->HasAnyClassFlags(CLASS_Abstract)))
+				{
+					UFactory* Factory = Cast<UFactory>(CurrentClass->GetDefaultObject());
+					if (Factory->bEditorImport)
+					{
+						Factories.Add(Factory);
+					}
+				}
+			}
+
+			
+			// Generate the file types and extensions represented by the selected factories
+			ObjectTools::GenerateFactoryFileExtensions(Factories, FileTypes, Extensions, FilterIndexToFactory);
+		}
+
+		ObjectTools::AppendFormatsFileExtensions(TranslatorFormats, FileTypes, Extensions, FilterIndexToFactory);
 
 		const FString FormatString = FString::Printf(TEXT("All Files (%s)|%s|%s"), *Extensions, *Extensions, *FileTypes);
 
@@ -50,13 +77,12 @@ namespace UE::Interchange::Utilities::Private
 
 bool UInterchangeFilePickerGeneric::FilePickerForTranslatorAssetType(const EInterchangeTranslatorAssetType TranslatorAssetType, FInterchangeFilePickerParameters& Parameters, TArray<FString>& OutFilenames)
 {
-	FString Extensions = UE::Interchange::Utilities::Private::GetOpenFileDialogExtensions(UInterchangeManager::GetInterchangeManager().GetSupportedAssetTypeFormats(TranslatorAssetType));
+	FString Extensions = UE::Interchange::Utilities::Private::GetOpenFileDialogExtensions(UInterchangeManager::GetInterchangeManager().GetSupportedAssetTypeFormats(TranslatorAssetType), Parameters.bShowAllFactoriesExtension);
 	return UE::Interchange::Utilities::Private::FilePickerDialog(Extensions, Parameters, OutFilenames);
 }
 
 bool UInterchangeFilePickerGeneric::FilePickerForTranslatorType(const EInterchangeTranslatorType TranslatorType, FInterchangeFilePickerParameters& Parameters, TArray<FString>& OutFilenames)
 {
-	FString Extensions = UE::Interchange::Utilities::Private::GetOpenFileDialogExtensions(UInterchangeManager::GetInterchangeManager().GetSupportedFormats(TranslatorType));
+	FString Extensions = UE::Interchange::Utilities::Private::GetOpenFileDialogExtensions(UInterchangeManager::GetInterchangeManager().GetSupportedFormats(TranslatorType), Parameters.bShowAllFactoriesExtension);
 	return UE::Interchange::Utilities::Private::FilePickerDialog(Extensions, Parameters, OutFilenames);
 }
-
diff --git a/Engine/Source/Editor/UnrealEd/Private/FbxMeshUtils.cpp b/Engine/Source/Editor/UnrealEd/Private/FbxMeshUtils.cpp
--- a/Engine/Source/Editor/UnrealEd/Private/FbxMeshUtils.cpp
+++ b/Engine/Source/Editor/UnrealEd/Private/FbxMeshUtils.cpp
@@ -56,14 +56,25 @@ namespace FbxMeshUtils
 			FMessageDialog::Open(EAppMsgType::Ok, FText::Format(LOCTEXT("LODImport_Failure", "Failed to import LOD{0}"), FText::AsNumber(LodIndex)));
 		}
 
-		void SetupFbxImportOptions(const UStaticMesh* BaseStaticMesh, UnFbx::FBXImportOptions* ImportOptions)
+		UFbxStaticMeshImportData* SetupFbxImportOptions(UStaticMesh* BaseStaticMesh, UnFbx::FBXImportOptions* ImportOptions)
 		{
 			check(BaseStaticMesh);
 
 			UFbxStaticMeshImportData* ImportData = Cast<UFbxStaticMeshImportData>(BaseStaticMesh->AssetImportData);
-			if (ImportData != nullptr)
+			if (ImportData == nullptr)
 			{
+				//Convert the static mesh import data to legacy fbx import data, because the legacy code will override it if not of the good type
+				if (UInterchangeAssetImportData* InterchangeAssetImportData = Cast<UInterchangeAssetImportData>(BaseStaticMesh->GetAssetImportData()))
+				{
+					if (UInterchangeManager::GetInterchangeManager().ConvertImportData(BaseStaticMesh, TEXT("fbx")))
+					{
+						ImportData = Cast<UFbxStaticMeshImportData>(BaseStaticMesh->AssetImportData);
+					}
+				}
+			}
 
+			if (ImportData != nullptr)
+			{
 				UFbxImportUI* ReimportUI = NewObject<UFbxImportUI>();
 				ReimportUI->MeshTypeToImport = FBXIT_StaticMesh;
 				UnFbx::FBXImportOptions::ResetOptions(ImportOptions);
@@ -71,7 +82,7 @@ namespace FbxMeshUtils
 				ReimportUI->StaticMeshImportData = ImportData;
 				ApplyImportUIToImportOptions(ReimportUI, *ImportOptions);
 			}
-			else if (BaseStaticMesh->IsSourceModelValid(0))
+			else
 			{
 				// Use the lod 0 to set the import settings
 				const FStaticMeshSourceModel& SourceModel = BaseStaticMesh->GetSourceModel(0);
@@ -91,6 +102,8 @@ namespace FbxMeshUtils
 			ImportOptions->bImportTextures = false;
 
 			ImportOptions->bAutoComputeLodDistances = true; //Setting auto compute distance to true will avoid changing the staticmesh flag
+
+			return ImportData;
 		}
 
 		bool CopyHighResMeshDescription(UStaticMesh* SrcStaticMesh, UStaticMesh* BaseStaticMesh)
@@ -189,32 +202,36 @@ namespace FbxMeshUtils
 		if (UInterchangeManager::IsInterchangeImportEnabled() && SelectedInterchangeAssetImportData)
 		{
 			UInterchangeSourceData* SourceData = UInterchangeManager::GetInterchangeManager().CreateSourceData(Filename);
-			//Call interchange mesh utilities to import custom LOD
-			UInterchangeMeshUtilities::ImportCustomLod(BaseStaticMesh, LODLevel, SourceData).Then([BaseStaticMesh, LODLevel](TFuture<bool> Result)
-				{
-					bool bResult = Result.Get();
-					Async(EAsyncExecution::TaskGraphMainThread, [BaseStaticMesh, LODLevel, bResult]()
-						{
-							if (bResult)
-							{
-								// Notification of success
-								FNotificationInfo NotificationInfo(FText::GetEmpty());
-								NotificationInfo.Text = FText::Format(NSLOCTEXT("UnrealEd", "StaticMeshLODImportSuccessful", "Static mesh LOD {0} imported successfully!"), FText::AsNumber(LODLevel));
-								NotificationInfo.ExpireDuration = 5.0f;
-								FSlateNotificationManager::Get().AddNotification(NotificationInfo);
-							}
-							else
+			//Only use interchange if we can translate the source data. Use the legacy path in other cases
+			if (UInterchangeManager::GetInterchangeManager().CanTranslateSourceData(SourceData))
+			{
+				//Call interchange mesh utilities to import custom LOD
+				UInterchangeMeshUtilities::ImportCustomLod(BaseStaticMesh, LODLevel, SourceData).Then([BaseStaticMesh, LODLevel](TFuture<bool> Result)
+					{
+						bool bResult = Result.Get();
+						Async(EAsyncExecution::TaskGraphMainThread, [BaseStaticMesh, LODLevel, bResult]()
 							{
-								// Notification of failure
-								FNotificationInfo NotificationInfo(FText::GetEmpty());
-								NotificationInfo.Text = FText::Format(NSLOCTEXT("UnrealEd", "StaticMeshLODImportFail", "Failed to import static mesh LOD {0}!"), FText::AsNumber(LODLevel));
-								NotificationInfo.ExpireDuration = 5.0f;
-								FSlateNotificationManager::Get().AddNotification(NotificationInfo);
-							}
-						});
-				});
+								if (bResult)
+								{
+									// Notification of success
+									FNotificationInfo NotificationInfo(FText::GetEmpty());
+									NotificationInfo.Text = FText::Format(NSLOCTEXT("UnrealEd", "StaticMeshLODImportSuccessful", "Static mesh LOD {0} imported successfully!"), FText::AsNumber(LODLevel));
+									NotificationInfo.ExpireDuration = 5.0f;
+									FSlateNotificationManager::Get().AddNotification(NotificationInfo);
+								}
+								else
+								{
+									// Notification of failure
+									FNotificationInfo NotificationInfo(FText::GetEmpty());
+									NotificationInfo.Text = FText::Format(NSLOCTEXT("UnrealEd", "StaticMeshLODImportFail", "Failed to import static mesh LOD {0}!"), FText::AsNumber(LODLevel));
+									NotificationInfo.ExpireDuration = 5.0f;
+									FSlateNotificationManager::Get().AddNotification(NotificationInfo);
+								}
+							});
+					});
 
-			return true;
+				return true;
+			}
 		}
 
 		bool bSuccess = false;
@@ -228,9 +245,7 @@ namespace FbxMeshUtils
 		UnFbx::FFbxLoggerSetter Logger(FFbxImporter);
 
 		UnFbx::FBXImportOptions* ImportOptions = FFbxImporter->GetImportOptions();
-		Private::SetupFbxImportOptions(BaseStaticMesh, ImportOptions);
-
-		UFbxStaticMeshImportData* ImportData = Cast<UFbxStaticMeshImportData>(BaseStaticMesh->AssetImportData);
+		UFbxStaticMeshImportData* ImportData = Private::SetupFbxImportOptions(BaseStaticMesh, ImportOptions);
 
 		const bool bIsReimport = BaseStaticMesh->GetRenderData()->LODResources.Num() > LODLevel;
 
@@ -494,11 +509,11 @@ namespace FbxMeshUtils
 			return false;
 		}
 
-		//We will use interchange only if interchange is enable and the skeletalmesh we want to add a LOD was import with interchange
-		UInterchangeAssetImportData* SelectedInterchangeAssetImportData = Cast<UInterchangeAssetImportData>(SelectedSkelMesh->GetAssetImportData());
-		if (UInterchangeManager::IsInterchangeImportEnabled() && SelectedInterchangeAssetImportData)
+		UInterchangeManager& InterchangeManager = UInterchangeManager::GetInterchangeManager();
+		const UInterchangeSourceData* SourceData = InterchangeManager.CreateSourceData(Filename);
+		const bool bInterchangeCanImportSourceData = InterchangeManager.CanTranslateSourceData(SourceData);
+		if (bInterchangeCanImportSourceData)
 		{
-			UInterchangeSourceData* SourceData = UInterchangeManager::GetInterchangeManager().CreateSourceData(Filename);
 			//Call interchange mesh utilities to import custom LOD
 			UInterchangeMeshUtilities::ImportCustomLod(SelectedSkelMesh, LODLevel, SourceData).Then([SelectedSkelMesh, LODLevel](TFuture<bool> Result)
 				{
@@ -571,284 +586,299 @@ namespace FbxMeshUtils
 			return false;
 		}
 
-		FScopedSkeletalMeshPostEditChange ScopePostEditChange(SelectedSkelMesh);
-		UnFbx::FFbxScopedOperation FbxScopedOperation(FFbxImporter);
-
-		//If the imported LOD already exist, we will need to reimport all the skin weight profiles
-		bool bMustReimportAlternateSkinWeightProfile = false;
+		//Import LOD using fbx importer
+		{
+			FScopedSkeletalMeshPostEditChange ScopePostEditChange(SelectedSkelMesh);
+			UnFbx::FFbxScopedOperation FbxScopedOperation(FFbxImporter);
 
-		// Get a list of all the clothing assets affecting this LOD so we can re-apply later
-		TArray<ClothingAssetUtils::FClothingAssetMeshBinding> ClothingBindings;
-		TArray<UClothingAssetBase*> ClothingAssetsInUse;
-		TArray<int32> ClothingAssetSectionIndices;
-		TArray<int32> ClothingAssetInternalLodIndices;
+			//If the imported LOD already exist, we will need to reimport all the skin weight profiles
+			bool bMustReimportAlternateSkinWeightProfile = false;
 
-		FSkeletalMeshModel* ImportedResource = SelectedSkelMesh->GetImportedModel();
-		if(ImportedResource && ImportedResource->LODModels.IsValidIndex(LODLevel))
-		{
-			bMustReimportAlternateSkinWeightProfile = true;
-			FLODUtilities::UnbindClothingAndBackup(SelectedSkelMesh, ClothingBindings, LODLevel);
-		}
+			// Get a list of all the clothing assets affecting this LOD so we can re-apply later
+			TArray<ClothingAssetUtils::FClothingAssetMeshBinding> ClothingBindings;
+			TArray<UClothingAssetBase*> ClothingAssetsInUse;
+			TArray<int32> ClothingAssetSectionIndices;
+			TArray<int32> ClothingAssetInternalLodIndices;
 
-		//Lambda to call to re-apply the clothing
-		auto ReapplyClothing = [&SelectedSkelMesh, &ClothingBindings, &ImportedResource, &LODLevel]()
-		{
+			FSkeletalMeshModel* ImportedResource = SelectedSkelMesh->GetImportedModel();
 			if (ImportedResource && ImportedResource->LODModels.IsValidIndex(LODLevel))
 			{
-				// Re-apply our clothing assets
-				FLODUtilities::RestoreClothingFromBackup(SelectedSkelMesh, ClothingBindings, LODLevel);
+				bMustReimportAlternateSkinWeightProfile = true;
+				FLODUtilities::UnbindClothingAndBackup(SelectedSkelMesh, ClothingBindings, LODLevel);
 			}
-		};
 
-		// don't import material and animation
-		UnFbx::FBXImportOptions* ImportOptions = FFbxImporter->GetImportOptions();
-			
-		//Set the skeletal mesh import data from the base mesh, this make sure the import rotation transform is use when importing a LOD
-		UFbxSkeletalMeshImportData* TempAssetImportData = NULL;
+			//Lambda to call to re-apply the clothing
+			auto ReapplyClothing = [&SelectedSkelMesh, &ClothingBindings, &ImportedResource, &LODLevel]()
+				{
+					if (ImportedResource && ImportedResource->LODModels.IsValidIndex(LODLevel))
+					{
+						// Re-apply our clothing assets
+						FLODUtilities::RestoreClothingFromBackup(SelectedSkelMesh, ClothingBindings, LODLevel);
+					}
+				};
 
-		UFbxAssetImportData *FbxAssetImportData = Cast<UFbxAssetImportData>(SelectedSkelMesh->GetAssetImportData());
-		if (FbxAssetImportData != nullptr)
-		{
-			UFbxSkeletalMeshImportData* ImportData = Cast<UFbxSkeletalMeshImportData>(FbxAssetImportData);
-			if (ImportData)
-			{
-				TempAssetImportData = ImportData;
-				UnFbx::FBXImportOptions::ResetOptions(ImportOptions);
-				// Prepare the import options
-				UFbxImportUI* ReimportUI = NewObject<UFbxImportUI>();
-				ReimportUI->MeshTypeToImport = FBXIT_SkeletalMesh;
-				ReimportUI->Skeleton = SelectedSkelMesh->GetSkeleton();
-				ReimportUI->PhysicsAsset = SelectedSkelMesh->GetPhysicsAsset();
-				// Import data already exists, apply it to the fbx import options
-				ReimportUI->SkeletalMeshImportData = ImportData;
-				//Some options not supported with skeletal mesh
-				ReimportUI->SkeletalMeshImportData->bBakePivotInVertex = false;
-				ReimportUI->SkeletalMeshImportData->bTransformVertexToAbsolute = true;
-				ApplyImportUIToImportOptions(ReimportUI, *ImportOptions);
-			}
-			ImportOptions->bImportMaterials = false;
-			ImportOptions->bImportTextures = false;
-		}
-		ImportOptions->bImportAnimations = false;
-		//Adjust the option in case we import only the skinning or the geometry
-		if (ImportOptions->bImportAsSkeletalSkinning)
-		{
-			ImportOptions->bImportMaterials = false;
-			ImportOptions->bImportTextures = false;
-			ImportOptions->bImportLOD = false;
-			ImportOptions->bImportSkeletalMeshLODs = false;
-			ImportOptions->bImportAnimations = false;
-			ImportOptions->bImportMorph = false;
-		}
-		else if (ImportOptions->bImportAsSkeletalGeometry)
-		{
-			ImportOptions->bImportAnimations = false;
-			ImportOptions->bUpdateSkeletonReferencePose = false;
-		}
+			// don't import material and animation
+			UnFbx::FBXImportOptions* ImportOptions = FFbxImporter->GetImportOptions();
 
-		if ( !FFbxImporter->ImportFromFile( *Filename, FPaths::GetExtension( Filename ), true ) )
-		{
-			ReapplyClothing();
-			// Log the error message and fail the import.
-			FFbxImporter->AddTokenizedErrorMessage(FTokenizedMessage::Create(EMessageSeverity::Error, LOCTEXT("FBXImport_ParseFailed", "FBX file parsing failed.")), FFbxErrors::Generic_FBXFileParseFailed);
-		}
-		else
-		{
-			bool bUseLODs = true;
-			int32 MaxLODLevel = 0;
-			TArray<FbxNode*>* MeshObject = NULL;
+			//Set the skeletal mesh import data from the base mesh, this make sure the import rotation transform is use when importing a LOD
+			UFbxSkeletalMeshImportData* TempAssetImportData = NULL;
 
-			//Set the build options if the BuildDat is not available so it is the same option we use to import the LOD
-			if (ImportedResource && ImportedResource->LODModels.IsValidIndex(LODLevel) && !SelectedSkelMesh->HasMeshDescription(LODLevel))
+			UFbxAssetImportData* FbxAssetImportData = Cast<UFbxAssetImportData>(SelectedSkelMesh->GetAssetImportData());
+			if (!FbxAssetImportData)
 			{
-				FSkeletalMeshLODInfo* LODInfo = SelectedSkelMesh->GetLODInfo(LODLevel);
-				if (LODInfo)
+				//Convert the data if its Interchange import data, we convert the full asset in this case since legacy fbx workflow will create one if missing
+				if (UInterchangeAssetImportData* InterchangeAssetImportData = Cast<UInterchangeAssetImportData>(SelectedSkelMesh->GetAssetImportData()))
 				{
-					LODInfo->BuildSettings.bRecomputeNormals = !ImportOptions->ShouldImportNormals();
-					LODInfo->BuildSettings.bRecomputeTangents = !ImportOptions->ShouldImportTangents();
-					LODInfo->BuildSettings.bUseMikkTSpace = (ImportOptions->NormalGenerationMethod == EFBXNormalGenerationMethod::MikkTSpace) && (!ImportOptions->ShouldImportNormals() || !ImportOptions->ShouldImportTangents());
-					LODInfo->BuildSettings.bComputeWeightedNormals = ImportOptions->bComputeWeightedNormals;
-					LODInfo->BuildSettings.bRemoveDegenerates = ImportOptions->bRemoveDegenerates;
-					LODInfo->BuildSettings.ThresholdPosition = ImportOptions->OverlappingThresholds.ThresholdPosition;
-					LODInfo->BuildSettings.ThresholdTangentNormal = ImportOptions->OverlappingThresholds.ThresholdTangentNormal;
-					LODInfo->BuildSettings.ThresholdUV = ImportOptions->OverlappingThresholds.ThresholdUV;
-					LODInfo->BuildSettings.MorphThresholdPosition = ImportOptions->OverlappingThresholds.MorphThresholdPosition;
+					if (InterchangeManager.ConvertImportData(SelectedSkelMesh, TEXT("fbx")))
+					{
+						FbxAssetImportData = Cast<UFbxAssetImportData>(SelectedSkelMesh->GetAssetImportData());
+					}
 				}
 			}
 
-			// Populate the mesh array
-			FFbxImporter->FillFbxSkelMeshArrayInScene(FFbxImporter->Scene->GetRootNode(), MeshArray, false, ImportOptions->bImportAsSkeletalGeometry || ImportOptions->bImportAsSkeletalSkinning, ImportOptions->bImportScene);
-
-			// Nothing found, error out
-			if (MeshArray.Num() == 0)
-			{
-				ReapplyClothing();
-				FFbxImporter->AddTokenizedErrorMessage(FTokenizedMessage::Create(EMessageSeverity::Error, LOCTEXT("FBXImport_NoMesh", "No meshes were found in file.")), FFbxErrors::Generic_MeshNotFound);
-				CleanUpScene();
-				return false;
-			}
-
-			MeshObject = MeshArray[0];
-
-			// check if there is LODGroup for this skeletal mesh
-			for (int32 j = 0; j < MeshObject->Num(); j++)
+			if (FbxAssetImportData)
 			{
-				FbxNode* Node = (*MeshObject)[j];
-				if (Node->GetNodeAttribute() && Node->GetNodeAttribute()->GetAttributeType() == FbxNodeAttribute::eLODGroup)
+				UFbxSkeletalMeshImportData* ImportData = Cast<UFbxSkeletalMeshImportData>(FbxAssetImportData);
+				if (ImportData)
 				{
-					// get max LODgroup level
-					if (MaxLODLevel < (Node->GetChildCount() - 1))
-					{
-						MaxLODLevel = Node->GetChildCount() - 1;
-					}
+					TempAssetImportData = ImportData;
+					UnFbx::FBXImportOptions::ResetOptions(ImportOptions);
+					// Prepare the import options
+					UFbxImportUI* ReimportUI = NewObject<UFbxImportUI>();
+					ReimportUI->MeshTypeToImport = FBXIT_SkeletalMesh;
+					ReimportUI->Skeleton = SelectedSkelMesh->GetSkeleton();
+					ReimportUI->PhysicsAsset = SelectedSkelMesh->GetPhysicsAsset();
+					// Import data already exists, apply it to the fbx import options
+					ReimportUI->SkeletalMeshImportData = ImportData;
+					//Some options not supported with skeletal mesh
+					ReimportUI->SkeletalMeshImportData->bBakePivotInVertex = false;
+					ReimportUI->SkeletalMeshImportData->bTransformVertexToAbsolute = true;
+					ApplyImportUIToImportOptions(ReimportUI, *ImportOptions);
 				}
+				ImportOptions->bImportMaterials = false;
+				ImportOptions->bImportTextures = false;
 			}
-
-			// No LODs found, switch to supporting a mesh array containing meshes instead of LODs
-			if (MaxLODLevel == 0)
+			ImportOptions->bImportAnimations = false;
+			//Adjust the option in case we import only the skinning or the geometry
+			if (ImportOptions->bImportAsSkeletalSkinning)
 			{
-				bUseLODs = false;
-				MaxLODLevel = SelectedSkelMesh->GetLODNum();
+				ImportOptions->bImportMaterials = false;
+				ImportOptions->bImportTextures = false;
+				ImportOptions->bImportLOD = false;
+				ImportOptions->bImportSkeletalMeshLODs = false;
+				ImportOptions->bImportAnimations = false;
+				ImportOptions->bImportMorph = false;
+			}
+			else if (ImportOptions->bImportAsSkeletalGeometry)
+			{
+				ImportOptions->bImportAnimations = false;
+				ImportOptions->bUpdateSkeletonReferencePose = false;
 			}
 
-			int32 SelectedLOD = LODLevel;
-			if (SelectedLOD > SelectedSkelMesh->GetLODNum())
+			if (!FFbxImporter->ImportFromFile(*Filename, FPaths::GetExtension(Filename), true))
 			{
 				ReapplyClothing();
-				// Make sure they don't manage to select a bad LOD index
-				FFbxImporter->AddTokenizedErrorMessage(FTokenizedMessage::Create(EMessageSeverity::Warning, FText::Format(LOCTEXT("FBXImport_InvalidLODIdx", "Invalid mesh LOD index {0}, no prior LOD index exists"), FText::AsNumber(SelectedLOD))), FFbxErrors::Generic_Mesh_LOD_InvalidIndex);
+				// Log the error message and fail the import.
+				FFbxImporter->AddTokenizedErrorMessage(FTokenizedMessage::Create(EMessageSeverity::Error, LOCTEXT("FBXImport_ParseFailed", "FBX file parsing failed.")), FFbxErrors::Generic_FBXFileParseFailed);
 			}
 			else
 			{
-				TArray<FbxNode*> SkelMeshNodeArray;
+				bool bUseLODs = true;
+				int32 MaxLODLevel = 0;
+				TArray<FbxNode*>* MeshObject = NULL;
 
-				if (bUseLODs || ImportOptions->bImportMorph)
+				//Set the build options if the BuildDat is not available so it is the same option we use to import the LOD
+				if (ImportedResource && ImportedResource->LODModels.IsValidIndex(LODLevel) && !SelectedSkelMesh->HasMeshDescription(LODLevel))
 				{
-					for (int32 j = 0; j < MeshObject->Num(); j++)
+					FSkeletalMeshLODInfo* LODInfo = SelectedSkelMesh->GetLODInfo(LODLevel);
+					if (LODInfo)
 					{
-						FbxNode* Node = (*MeshObject)[j];
-						if (Node->GetNodeAttribute() && Node->GetNodeAttribute()->GetAttributeType() == FbxNodeAttribute::eLODGroup)
-						{
-							TArray<FbxNode*> NodeInLod;
-							if (Node->GetChildCount() > SelectedLOD)
-							{
-								FFbxImporter->FindAllLODGroupNode(NodeInLod, Node, SelectedLOD);
-							}
-							else // in less some LODGroups have less level, use the last level
-							{
-								FFbxImporter->FindAllLODGroupNode(NodeInLod, Node, Node->GetChildCount() - 1);
-							}
-
-							for (FbxNode *MeshNode : NodeInLod)
-							{
-								SkelMeshNodeArray.Add(MeshNode);
-							}
-						}
-						else
-						{
-							SkelMeshNodeArray.Add(Node);
-						}
+						LODInfo->BuildSettings.bRecomputeNormals = !ImportOptions->ShouldImportNormals();
+						LODInfo->BuildSettings.bRecomputeTangents = !ImportOptions->ShouldImportTangents();
+						LODInfo->BuildSettings.bUseMikkTSpace = (ImportOptions->NormalGenerationMethod == EFBXNormalGenerationMethod::MikkTSpace) && (!ImportOptions->ShouldImportNormals() || !ImportOptions->ShouldImportTangents());
+						LODInfo->BuildSettings.bComputeWeightedNormals = ImportOptions->bComputeWeightedNormals;
+						LODInfo->BuildSettings.bRemoveDegenerates = ImportOptions->bRemoveDegenerates;
+						LODInfo->BuildSettings.ThresholdPosition = ImportOptions->OverlappingThresholds.ThresholdPosition;
+						LODInfo->BuildSettings.ThresholdTangentNormal = ImportOptions->OverlappingThresholds.ThresholdTangentNormal;
+						LODInfo->BuildSettings.ThresholdUV = ImportOptions->OverlappingThresholds.ThresholdUV;
+						LODInfo->BuildSettings.MorphThresholdPosition = ImportOptions->OverlappingThresholds.MorphThresholdPosition;
 					}
 				}
 
-				// Import mesh
-				USkeletalMesh* TempSkelMesh = NULL;
-				TArray<FName> OrderedMaterialNames;
+				// Populate the mesh array
+				FFbxImporter->FillFbxSkelMeshArrayInScene(FFbxImporter->Scene->GetRootNode(), MeshArray, false, ImportOptions->bImportAsSkeletalGeometry || ImportOptions->bImportAsSkeletalSkinning, ImportOptions->bImportScene);
+
+				// Nothing found, error out
+				if (MeshArray.Num() == 0)
 				{
-					int32 NoneNameCount = 0;
-					for (const FSkeletalMaterial &Material : SelectedSkelMesh->GetMaterials())
-					{
-						if (Material.ImportedMaterialSlotName == NAME_None)
-							NoneNameCount++;
+					ReapplyClothing();
+					FFbxImporter->AddTokenizedErrorMessage(FTokenizedMessage::Create(EMessageSeverity::Error, LOCTEXT("FBXImport_NoMesh", "No meshes were found in file.")), FFbxErrors::Generic_MeshNotFound);
+					CleanUpScene();
+					return false;
+				}
 
-						OrderedMaterialNames.Add(Material.ImportedMaterialSlotName);
-					}
-					if (NoneNameCount >= OrderedMaterialNames.Num())
+				MeshObject = MeshArray[0];
+
+				// check if there is LODGroup for this skeletal mesh
+				for (int32 j = 0; j < MeshObject->Num(); j++)
+				{
+					FbxNode* Node = (*MeshObject)[j];
+					if (Node->GetNodeAttribute() && Node->GetNodeAttribute()->GetAttributeType() == FbxNodeAttribute::eLODGroup)
 					{
-						OrderedMaterialNames.Empty();
+						// get max LODgroup level
+						if (MaxLODLevel < (Node->GetChildCount() - 1))
+						{
+							MaxLODLevel = Node->GetChildCount() - 1;
+						}
 					}
 				}
-					
-				TSharedPtr<FExistingSkelMeshData> SkelMeshDataPtr;
-				if (SelectedSkelMesh->GetLODNum() > SelectedLOD)
+
+				// No LODs found, switch to supporting a mesh array containing meshes instead of LODs
+				if (MaxLODLevel == 0)
 				{
-					SelectedSkelMesh->PreEditChange(NULL);
-					SkelMeshDataPtr = SkeletalMeshImportUtils::SaveExistingSkelMeshData(SelectedSkelMesh, true, SelectedLOD);
+					bUseLODs = false;
+					MaxLODLevel = SelectedSkelMesh->GetLODNum();
 				}
 
-				//Original fbx data storage
-				TArray<FName> ImportMaterialOriginalNameData;
-				TArray<FImportMeshLodSectionsData> ImportMeshLodData;
-				ImportMeshLodData.AddZeroed();
-				FSkeletalMeshImportData OutData;
-
-				UnFbx::FFbxImporter::FImportSkeletalMeshArgs ImportSkeletalMeshArgs;
-				ImportSkeletalMeshArgs.InParent = SelectedSkelMesh->GetOutermost();
-				ImportSkeletalMeshArgs.NodeArray = bUseLODs ? SkelMeshNodeArray : *MeshObject;
-				ImportSkeletalMeshArgs.Name = NAME_None;
-				ImportSkeletalMeshArgs.Flags = RF_Transient;
-				ImportSkeletalMeshArgs.TemplateImportData = TempAssetImportData;
-				ImportSkeletalMeshArgs.LodIndex = SelectedLOD;
-				ImportSkeletalMeshArgs.OrderedMaterialNames = OrderedMaterialNames.Num() > 0 ? &OrderedMaterialNames : nullptr;
-				ImportSkeletalMeshArgs.ImportMaterialOriginalNameData = &ImportMaterialOriginalNameData;
-				ImportSkeletalMeshArgs.ImportMeshSectionsData = &ImportMeshLodData[0];
-				ImportSkeletalMeshArgs.OutData = &OutData;
-
-				TempSkelMesh = (USkeletalMesh*)FFbxImporter->ImportSkeletalMesh( ImportSkeletalMeshArgs );
-				// Add the new imported LOD to the existing model (check skeleton compatibility)
-				if( TempSkelMesh  && FFbxImporter->ImportSkeletalMeshLOD(TempSkelMesh, SelectedSkelMesh, SelectedLOD, TempAssetImportData))
+				int32 SelectedLOD = LODLevel;
+				if (SelectedLOD > SelectedSkelMesh->GetLODNum())
 				{
-					//Update the import data for this lod
-					UnFbx::FFbxImporter::UpdateSkeletalMeshImportData(SelectedSkelMesh, nullptr, SelectedLOD, &ImportMaterialOriginalNameData, &ImportMeshLodData);
+					ReapplyClothing();
+					// Make sure they don't manage to select a bad LOD index
+					FFbxImporter->AddTokenizedErrorMessage(FTokenizedMessage::Create(EMessageSeverity::Warning, FText::Format(LOCTEXT("FBXImport_InvalidLODIdx", "Invalid mesh LOD index {0}, no prior LOD index exists"), FText::AsNumber(SelectedLOD))), FFbxErrors::Generic_Mesh_LOD_InvalidIndex);
+				}
+				else
+				{
+					TArray<FbxNode*> SkelMeshNodeArray;
 
-					const FString SourceImportFilename = UAssetImportData::SanitizeImportFilename(Filename, nullptr);
-					if (SkelMeshDataPtr)
+					if (bUseLODs || ImportOptions->bImportMorph)
 					{
-						//Setting the source filename allow to not restore the reduction settings in case we import a custom LOD over a generated LOD.
-						//This value will be wipe during the restore but we put it back just after.
-						SelectedSkelMesh->GetLODInfo(SelectedLOD)->SourceImportFilename = SourceImportFilename;
-						SkeletalMeshImportUtils::RestoreExistingSkelMeshData(SkelMeshDataPtr, SelectedSkelMesh, SelectedLOD, false, ImportOptions->bImportAsSkeletalSkinning, ImportOptions->bResetToFbxOnMaterialConflict);
+						for (int32 j = 0; j < MeshObject->Num(); j++)
+						{
+							FbxNode* Node = (*MeshObject)[j];
+							if (Node->GetNodeAttribute() && Node->GetNodeAttribute()->GetAttributeType() == FbxNodeAttribute::eLODGroup)
+							{
+								TArray<FbxNode*> NodeInLod;
+								if (Node->GetChildCount() > SelectedLOD)
+								{
+									FFbxImporter->FindAllLODGroupNode(NodeInLod, Node, SelectedLOD);
+								}
+								else // in less some LODGroups have less level, use the last level
+								{
+									FFbxImporter->FindAllLODGroupNode(NodeInLod, Node, Node->GetChildCount() - 1);
+								}
+
+								for (FbxNode* MeshNode : NodeInLod)
+								{
+									SkelMeshNodeArray.Add(MeshNode);
+								}
+							}
+							else
+							{
+								SkelMeshNodeArray.Add(Node);
+							}
+						}
+					}
+
+					// Import mesh
+					USkeletalMesh* TempSkelMesh = NULL;
+					TArray<FName> OrderedMaterialNames;
+					{
+						int32 NoneNameCount = 0;
+						for (const FSkeletalMaterial& Material : SelectedSkelMesh->GetMaterials())
+						{
+							if (Material.ImportedMaterialSlotName == NAME_None)
+								NoneNameCount++;
+
+							OrderedMaterialNames.Add(Material.ImportedMaterialSlotName);
+						}
+						if (NoneNameCount >= OrderedMaterialNames.Num())
+						{
+							OrderedMaterialNames.Empty();
+						}
 					}
 
-					if (ImportOptions->bImportMorph)
+					TSharedPtr<FExistingSkelMeshData> SkelMeshDataPtr;
+					if (SelectedSkelMesh->GetLODNum() > SelectedLOD)
 					{
-						FFbxImporter->ImportFbxMorphTarget(SkelMeshNodeArray, SelectedSkelMesh, SelectedLOD, OutData, ImportSkeletalMeshArgs.bMapMorphTargetToTimeZero);
+						SelectedSkelMesh->PreEditChange(NULL);
+						SkelMeshDataPtr = SkeletalMeshImportUtils::SaveExistingSkelMeshData(SelectedSkelMesh, true, SelectedLOD);
 					}
 
-					bSuccess = true;
+					//Original fbx data storage
+					TArray<FName> ImportMaterialOriginalNameData;
+					TArray<FImportMeshLodSectionsData> ImportMeshLodData;
+					ImportMeshLodData.AddZeroed();
+					FSkeletalMeshImportData OutData;
+
+					UnFbx::FFbxImporter::FImportSkeletalMeshArgs ImportSkeletalMeshArgs;
+					ImportSkeletalMeshArgs.InParent = SelectedSkelMesh->GetOutermost();
+					ImportSkeletalMeshArgs.NodeArray = bUseLODs ? SkelMeshNodeArray : *MeshObject;
+					ImportSkeletalMeshArgs.Name = NAME_None;
+					ImportSkeletalMeshArgs.Flags = RF_Transient;
+					ImportSkeletalMeshArgs.TemplateImportData = TempAssetImportData;
+					ImportSkeletalMeshArgs.LodIndex = SelectedLOD;
+					ImportSkeletalMeshArgs.OrderedMaterialNames = OrderedMaterialNames.Num() > 0 ? &OrderedMaterialNames : nullptr;
+					ImportSkeletalMeshArgs.ImportMaterialOriginalNameData = &ImportMaterialOriginalNameData;
+					ImportSkeletalMeshArgs.ImportMeshSectionsData = &ImportMeshLodData[0];
+					ImportSkeletalMeshArgs.OutData = &OutData;
+
+					TempSkelMesh = (USkeletalMesh*)FFbxImporter->ImportSkeletalMesh(ImportSkeletalMeshArgs);
+					// Add the new imported LOD to the existing model (check skeleton compatibility)
+					if (TempSkelMesh && FFbxImporter->ImportSkeletalMeshLOD(TempSkelMesh, SelectedSkelMesh, SelectedLOD, TempAssetImportData))
+					{
+						//Update the import data for this lod
+						UnFbx::FFbxImporter::UpdateSkeletalMeshImportData(SelectedSkelMesh, nullptr, SelectedLOD, &ImportMaterialOriginalNameData, &ImportMeshLodData);
 
-					// Set LOD source filename
-					SelectedSkelMesh->GetLODInfo(SelectedLOD)->SourceImportFilename = SourceImportFilename;
-					SelectedSkelMesh->GetLODInfo(SelectedLOD)->bImportWithBaseMesh = false;
+						const FString SourceImportFilename = UAssetImportData::SanitizeImportFilename(Filename, nullptr);
+						if (SkelMeshDataPtr)
+						{
+							//Setting the source filename allow to not restore the reduction settings in case we import a custom LOD over a generated LOD.
+							//This value will be wipe during the restore but we put it back just after.
+							SelectedSkelMesh->GetLODInfo(SelectedLOD)->SourceImportFilename = SourceImportFilename;
+							SkeletalMeshImportUtils::RestoreExistingSkelMeshData(SkelMeshDataPtr, SelectedSkelMesh, SelectedLOD, false, ImportOptions->bImportAsSkeletalSkinning, ImportOptions->bResetToFbxOnMaterialConflict);
+						}
 
-					ReapplyClothing();
+						if (ImportOptions->bImportMorph)
+						{
+							FFbxImporter->ImportFbxMorphTarget(SkelMeshNodeArray, SelectedSkelMesh, SelectedLOD, OutData, ImportSkeletalMeshArgs.bMapMorphTargetToTimeZero);
+						}
+
+						bSuccess = true;
+
+						// Set LOD source filename
+						SelectedSkelMesh->GetLODInfo(SelectedLOD)->SourceImportFilename = SourceImportFilename;
+						SelectedSkelMesh->GetLODInfo(SelectedLOD)->bImportWithBaseMesh = false;
 
-					//Must be the last step because it cleanup the fbx importer to import the alternate skinning FBX
-					if (bMustReimportAlternateSkinWeightProfile)
+						ReapplyClothing();
+
+						//Must be the last step because it cleanup the fbx importer to import the alternate skinning FBX
+						if (bMustReimportAlternateSkinWeightProfile)
+						{
+							//We cannot use anymore the FFbxImporter after the cleanup
+							CleanUpScene();
+							FSkinWeightsUtilities::ReimportAlternateSkinWeight(SelectedSkelMesh, SelectedLOD);
+						}
+
+						// Notification of success
+						FNotificationInfo NotificationInfo(FText::GetEmpty());
+						NotificationInfo.Text = FText::Format(NSLOCTEXT("UnrealEd", "LODImportSuccessful", "Mesh for LOD {0} imported successfully!"), FText::AsNumber(SelectedLOD));
+						NotificationInfo.ExpireDuration = 5.0f;
+						FSlateNotificationManager::Get().AddNotification(NotificationInfo);
+					}
+					else
 					{
-						//We cannot use anymore the FFbxImporter after the cleanup
-						CleanUpScene();
-						FSkinWeightsUtilities::ReimportAlternateSkinWeight(SelectedSkelMesh, SelectedLOD);
+						ReapplyClothing();
+						// Notification of failure
+						FNotificationInfo NotificationInfo(FText::GetEmpty());
+						NotificationInfo.Text = FText::Format(NSLOCTEXT("UnrealEd", "MeshLODImportFail2", "Failed to import mesh for LOD {0}!"), FText::AsNumber(SelectedLOD));
+						NotificationInfo.ExpireDuration = 5.0f;
+						FSlateNotificationManager::Get().AddNotification(NotificationInfo);
 					}
-
-					// Notification of success
-					FNotificationInfo NotificationInfo(FText::GetEmpty());
-					NotificationInfo.Text = FText::Format(NSLOCTEXT("UnrealEd", "LODImportSuccessful", "Mesh for LOD {0} imported successfully!"), FText::AsNumber(SelectedLOD));
-					NotificationInfo.ExpireDuration = 5.0f;
-					FSlateNotificationManager::Get().AddNotification(NotificationInfo);
-				}
-				else
-				{
-					ReapplyClothing();
-					// Notification of failure
-					FNotificationInfo NotificationInfo(FText::GetEmpty());
-					NotificationInfo.Text = FText::Format(NSLOCTEXT("UnrealEd", "MeshLODImportFail2", "Failed to import mesh for LOD {0}!"), FText::AsNumber(SelectedLOD));
-					NotificationInfo.ExpireDuration = 5.0f;
-					FSlateNotificationManager::Get().AddNotification(NotificationInfo);
 				}
 			}
+			CleanUpScene();
 		}
-		CleanUpScene();
 		return bSuccess;
 	}
 
@@ -915,8 +945,6 @@ namespace FbxMeshUtils
 		UStaticMesh* StaticMesh = Cast<UStaticMesh>(SelectedMesh);
 
 		FString FilenameToImport("");
-		UInterchangeAssetImportData* SelectedInterchangeAssetImportData = nullptr;
-
 		//Make sure the LODLevel is valid, it should not be more then one over the existing lod count
 		bool bInvalidLodIndex = false;
 		if (SkeletalMesh)
@@ -933,7 +961,6 @@ namespace FbxMeshUtils
 						SkeletalMesh->GetLODInfo(LODLevel)->SourceImportFilename :
 						UAssetImportData::ResolveImportFilename(SkeletalMesh->GetLODInfo(LODLevel)->SourceImportFilename, nullptr);
 				}
-				SelectedInterchangeAssetImportData = Cast<UInterchangeAssetImportData>(SkeletalMesh->GetAssetImportData());
 			}
 		}
 		else if (StaticMesh)
@@ -951,7 +978,6 @@ namespace FbxMeshUtils
 						SourceModel.SourceImportFilename :
 						UAssetImportData::ResolveImportFilename(SourceModel.SourceImportFilename, nullptr);
 				}
-				SelectedInterchangeAssetImportData = Cast<UInterchangeAssetImportData>(StaticMesh->GetAssetImportData());
 			}
 		}
 		else
@@ -961,7 +987,6 @@ namespace FbxMeshUtils
 			return Promise->GetFuture();
 		}
 
-
 		if (bInvalidLodIndex)
 		{
 			UE_LOG(LogExportMeshUtils, Warning, TEXT("ImportMeshLODDialog: Invalid mesh LOD index %d, no prior LOD index exists."), LODLevel);
@@ -971,24 +996,91 @@ namespace FbxMeshUtils
 		}
 
 		// Check the file exists first
-		const bool bSourceFileExists = FPaths::FileExists(FilenameToImport);
-		// We'll give the user a chance to choose a new file if a previously set file fails to import
-		const bool bPromptOnFail = bSourceFileExists;
+		bool bSourceFileExists = FPaths::FileExists(FilenameToImport);
 
-		//We will use interchange only if interchange is enable and the skeletalmesh we want to add a LOD was import with interchange
-		if(UInterchangeManager::IsInterchangeImportEnabled() && SelectedInterchangeAssetImportData)
+		if (!bSourceFileExists)
 		{
-			TFuture<bool> Result;
-			if(!bSourceFileExists)
+			//Pop a file picker that join both interchange and other format
+			//Ask the user for a file path
+			const UInterchangeProjectSettings* InterchangeProjectSettings = GetDefault<UInterchangeProjectSettings>();
+			UInterchangeFilePickerBase* FilePicker = nullptr;
+
+			//In runtime we do not have any pipeline configurator
+#if WITH_EDITORONLY_DATA
+			TSoftClassPtr <UInterchangeFilePickerBase> FilePickerClass = InterchangeProjectSettings->FilePickerClass;
+			if (FilePickerClass.IsValid())
 			{
-				//Call interchange mesh utilities to import custom LOD
-				Result = UInterchangeMeshUtilities::ImportCustomLod(SelectedMesh, LODLevel);
+				UClass* FilePickerClassLoaded = FilePickerClass.LoadSynchronous();
+				if (FilePickerClassLoaded)
+				{
+					FilePicker = NewObject<UInterchangeFilePickerBase>(GetTransientPackage(), FilePickerClassLoaded, NAME_None, RF_NoFlags);
+				}
 			}
-			else
+#endif
+			if (FilePicker)
+			{
+				FInterchangeFilePickerParameters Parameters;
+				Parameters.bAllowMultipleFiles = false;
+				Parameters.Title = FText::Format(NSLOCTEXT("Interchange", "ImportCustomLodAsync_FilePickerTitle", "Choose a file to import a custom LOD for LOD{0}"), FText::AsNumber(LODLevel));
+				Parameters.bShowAllFactoriesExtension = true;
+				TArray<FString> Filenames;
+				bool bFilePickerResult = FilePicker->ScriptedFilePickerForTranslatorAssetType(EInterchangeTranslatorAssetType::Meshes, Parameters, Filenames);
+				if (bFilePickerResult)
+				{
+					if (Filenames.Num() > 0)
+					{
+						FilenameToImport = Filenames[0];
+						bSourceFileExists = FPaths::FileExists(FilenameToImport);
+					}
+					else
+					{
+						UE_LOG(LogExportMeshUtils, Display, TEXT("ImportMeshLODDialog: Error when picking a file to import LOD index %d."), LODLevel);
+						Promise->SetValue(false);
+						return Promise->GetFuture();
+					}
+				}
+				else
+				{
+					UE_LOG(LogExportMeshUtils, Display, TEXT("ImportMeshLODDialog: User cancel import LOD index %d."), LODLevel);
+					Promise->SetValue(false);
+					return Promise->GetFuture();
+				}
+			}
+		}
+
+		if (!bSourceFileExists)
+		{
+			UE_LOG(LogExportMeshUtils, Display, TEXT("ImportMeshLODDialog: Cannot import LOD index %d. The filename do not exist."), LODLevel);
+			Promise->SetValue(false);
+			return Promise->GetFuture();
+		}
+
+		//Convert the import data if necessary
+		UInterchangeAssetImportData* SelectedInterchangeAssetImportData = nullptr;
+		SelectedInterchangeAssetImportData = nullptr;
+		if (SkeletalMesh)
+		{
+			if (SkeletalMesh->IsValidLODIndex(LODLevel))
 			{
-				UInterchangeSourceData* SourceData = UInterchangeManager::GetInterchangeManager().CreateSourceData(FilenameToImport);
-				Result = UInterchangeMeshUtilities::ImportCustomLod(SelectedMesh, LODLevel, SourceData);
+				SelectedInterchangeAssetImportData = Cast<UInterchangeAssetImportData>(SkeletalMesh->GetAssetImportData());
 			}
+		}
+		else if (StaticMesh)
+		{
+			if (LODLevel >= 0 && LODLevel <= StaticMesh->GetNumSourceModels())
+			{
+				SelectedInterchangeAssetImportData = Cast<UInterchangeAssetImportData>(StaticMesh->GetAssetImportData());
+			}
+		}
+
+		UInterchangeManager& InterchangeManager = UInterchangeManager::GetInterchangeManager();
+		const UInterchangeSourceData* SourceData = InterchangeManager.CreateSourceData(FilenameToImport);
+		const bool bInterchangeCanImportSourceData = InterchangeManager.CanTranslateSourceData(SourceData);
+
+		if (bInterchangeCanImportSourceData)
+		{
+			TFuture<bool> Result = UInterchangeMeshUtilities::ImportCustomLod(SelectedMesh, LODLevel, SourceData);
+
 			Result.Then([Promise, bNotifyCB, SkeletalMesh, StaticMesh, LODLevel](TFuture<bool> FutureResult)
 				{
 					check(IsInGameThread());
@@ -1025,29 +1117,12 @@ namespace FbxMeshUtils
 					}
 					Promise->SetValue(bResult);
 				});
-			
 
+			//Since we start an asynchronous lod import we can return the future of the promise.
 			return Promise->GetFuture();
 		}
 		
-		if(!bSourceFileExists || FilenameToImport.IsEmpty())
-		{
-			FText PromptTitle;
-
-			if(FilenameToImport.IsEmpty())
-			{
-				PromptTitle = FText::Format(LOCTEXT("LODImportPrompt_NoSource", "Choose a file to import for LOD {0}"), FText::AsNumber(LODLevel));
-			}
-			else if(!bSourceFileExists)
-			{
-				PromptTitle = FText::Format(LOCTEXT("LODImportPrompt_SourceNotFound", "LOD {0} Source file not found. Choose new file."), FText::AsNumber(LODLevel));
-			}
-
-			FilenameToImport = PromptForLODImportFile(PromptTitle);
-		}
-		
 		bool bImportSuccess = false;
-
 		if(!FilenameToImport.IsEmpty())
 		{
 			if(SkeletalMesh)
@@ -1060,26 +1135,6 @@ namespace FbxMeshUtils
 			}
 		}
 
-		if(!bImportSuccess && bPromptOnFail)
-		{
-			FMessageDialog::Open(EAppMsgType::Ok, FText::Format(LOCTEXT("LODImport_SourceMissingDialog", "Failed to import LOD{0} as the source file failed to import, please select a new source file."), FText::AsNumber(LODLevel)));
-
-			FText PromptTitle = FText::Format(LOCTEXT("LODImportPrompt_SourceFailed", "Failed to import source file for LOD {0}, choose a new file"), FText::AsNumber(LODLevel));
-			FilenameToImport = PromptForLODImportFile(PromptTitle);
-
-			if(FilenameToImport.Len() > 0 && FPaths::FileExists(FilenameToImport))
-			{
-				if(SkeletalMesh)
-				{
-					bImportSuccess = ImportSkeletalMeshLOD(SkeletalMesh, FilenameToImport, LODLevel);
-				}
-				else if(StaticMesh)
-				{
-					bImportSuccess = ImportStaticMeshLOD(StaticMesh, FilenameToImport, LODLevel);
-				}
-			}
-		}
-
 		//If the filename is empty it mean the user cancel the file selection
 		if(!bImportSuccess && !FilenameToImport.IsEmpty())
 		{
diff --git a/Engine/Source/Editor/UnrealEd/Private/SkinWeightsUtilities.cpp b/Engine/Source/Editor/UnrealEd/Private/SkinWeightsUtilities.cpp
--- a/Engine/Source/Editor/UnrealEd/Private/SkinWeightsUtilities.cpp
+++ b/Engine/Source/Editor/UnrealEd/Private/SkinWeightsUtilities.cpp
@@ -75,10 +75,6 @@ bool FSkinWeightsUtilities::ImportAlternateSkinWeight(USkeletalMesh* SkeletalMes
 	FScopedSuspendAlternateSkinWeightPreview ScopedSuspendAlternateSkinnWeightPreview(SkeletalMesh);
 	FScopedSkeletalMeshPostEditChange ScopePostEditChange(SkeletalMesh);
 
-	//If Interchange is enable use it if not use the old path
-
-	bool bUseInterchangeFramework = UInterchangeManager::IsInterchangeImportEnabled();
-	UInterchangeManager& InterchangeManager = UInterchangeManager::GetInterchangeManager();
 	const FString FileExtension = FPaths::GetExtension(AbsoluteFilePath);
 
 	FAssetToolsModule& AssetToolsModule = FModuleManager::LoadModuleChecked<FAssetToolsModule>("AssetTools");
@@ -112,10 +108,19 @@ bool FSkinWeightsUtilities::ImportAlternateSkinWeight(USkeletalMesh* SkeletalMes
 	UObject* ImportedObject = nullptr;
 
 	bool bCreateTransaction = false;
-	//Only use interchange if the base skeletal mesh was imported with interchange
-	const UInterchangeAssetImportData* SelectedInterchangeAssetImportData = Cast<UInterchangeAssetImportData>(SkeletalMesh->GetAssetImportData());
-	if (bUseInterchangeFramework && SelectedInterchangeAssetImportData)
+
+	UInterchangeManager& InterchangeManager = UInterchangeManager::GetInterchangeManager();
+	const UInterchangeSourceData* SourceData = InterchangeManager.CreateSourceData(AbsoluteFilePath);
+	const bool bInterchangeCanImportSourceData = InterchangeManager.CanTranslateSourceData(SourceData);
+	if (bInterchangeCanImportSourceData)
 	{
+		UInterchangeAssetImportData* SelectedInterchangeAssetImportData = Cast<UInterchangeAssetImportData>(SkeletalMesh->GetAssetImportData());
+		if (!SelectedInterchangeAssetImportData)
+		{
+			//Try to convert the asset import data
+			InterchangeManager.ConvertImportData(SkeletalMesh->GetAssetImportData(), UInterchangeAssetImportData::StaticClass(), reinterpret_cast<UObject**>(&SelectedInterchangeAssetImportData));
+		}
+
 		UE::Interchange::FScopedSourceData ScopedSourceData(AbsoluteFilePath);
 		const UInterchangeProjectSettings* InterchangeProjectSettings = GetDefault<UInterchangeProjectSettings>();
 		FImportAssetParameters ImportAssetParameters;
@@ -124,14 +129,17 @@ bool FSkinWeightsUtilities::ImportAlternateSkinWeight(USkeletalMesh* SkeletalMes
 		if (const UClass* GenericPipelineClass = InterchangeProjectSettings->GenericPipelineClass.LoadSynchronous())
 		{
 			UInterchangePipelineBase* GenericPipeline = nullptr;
-			for (UObject* PipelineObject : SelectedInterchangeAssetImportData->GetPipelines())
+			if (SelectedInterchangeAssetImportData)
 			{
-				if (PipelineObject->GetClass()->IsChildOf(GenericPipelineClass))
+				for (UObject* PipelineObject : SelectedInterchangeAssetImportData->GetPipelines())
 				{
-					if (UInterchangePipelineBase* ImportPipeline = Cast<UInterchangePipelineBase>(PipelineObject))
+					if (PipelineObject->GetClass()->IsChildOf(GenericPipelineClass))
 					{
-						GenericPipeline = Cast<UInterchangePipelineBase>(StaticDuplicateObject(ImportPipeline, GetTransientPackage()));
-						break;
+						if (UInterchangePipelineBase* ImportPipeline = Cast<UInterchangePipelineBase>(PipelineObject))
+						{
+							GenericPipeline = Cast<UInterchangePipelineBase>(StaticDuplicateObject(ImportPipeline, GetTransientPackage()));
+							break;
+						}
 					}
 				}
 			}
@@ -160,12 +168,32 @@ bool FSkinWeightsUtilities::ImportAlternateSkinWeight(USkeletalMesh* SkeletalMes
 	}
 	else if(FileExtension.Equals(TEXT("fbx"), ESearchCase::IgnoreCase))
 	{
+		UFbxSkeletalMeshImportData* OriginalSkeletalMeshImportData = Cast<UFbxSkeletalMeshImportData>(SkeletalMesh->GetAssetImportData());
+		if (!OriginalSkeletalMeshImportData)
+		{
+			//Convert the data if its Interchange import data
+			if (UInterchangeAssetImportData* InterchangeAssetImportData = Cast<UInterchangeAssetImportData>(SkeletalMesh->GetAssetImportData()))
+			{
+				UFbxImportUI* FbxImportUI = nullptr;
+				InterchangeManager.ConvertImportData(InterchangeAssetImportData, UFbxImportUI::StaticClass(), reinterpret_cast<UObject**>(&FbxImportUI));
+				if (FbxImportUI)
+				{
+					OriginalSkeletalMeshImportData = FbxImportUI->SkeletalMeshImportData;
+				}
+			}
+
+			if (!OriginalSkeletalMeshImportData)
+			{
+				//This will reset the import data
+				OriginalSkeletalMeshImportData = UFbxSkeletalMeshImportData::GetImportDataForSkeletalMesh(SkeletalMesh, nullptr);
+			}
+		}
+
 		//Import the alternate fbx into a temporary skeletal mesh using the same import options
 		UFbxFactory* FbxFactory = NewObject<UFbxFactory>(UFbxFactory::StaticClass());
 		FbxFactory->AddToRoot();
 
 		FbxFactory->ImportUI = NewObject<UFbxImportUI>(FbxFactory);
-		UFbxSkeletalMeshImportData* OriginalSkeletalMeshImportData = UFbxSkeletalMeshImportData::GetImportDataForSkeletalMesh(SkeletalMesh, nullptr);
 		if (OriginalSkeletalMeshImportData != nullptr)
 		{
 			//Copy the skeletal mesh import data options
diff --git a/Engine/Source/Runtime/Interchange/Engine/Private/InterchangeManager.cpp b/Engine/Source/Runtime/Interchange/Engine/Private/InterchangeManager.cpp
--- a/Engine/Source/Runtime/Interchange/Engine/Private/InterchangeManager.cpp
+++ b/Engine/Source/Runtime/Interchange/Engine/Private/InterchangeManager.cpp
@@ -1052,16 +1052,39 @@ bool UInterchangeManager::ConvertImportData(UObject* Object, const FString& Exte
 	return false;
 }
 
-bool UInterchangeManager::ConvertImportData(const UObject* SourceImportData, FImportAssetParameters& ImportAssetParameters) const
+bool UInterchangeManager::ConvertImportData(const UObject* SourceImportData, const UClass* DestinationClass, UObject** DestinationImportData) const
 {
-	UObject* DestinationImportData = nullptr;
 	for (TPair<TObjectPtr<const UClass>, TObjectPtr<UInterchangeAssetImportDataConverterBase>> RegisteredConverter : RegisteredConverters)
 	{
-		if (!RegisteredConverter.Value->ConvertImportData(SourceImportData, &DestinationImportData))
+		if (RegisteredConverter.Value->CanConvertClass(SourceImportData->GetClass(), DestinationClass))
 		{
-			return false;
+			if (RegisteredConverter.Value->ConvertImportData(SourceImportData, DestinationImportData))
+			{
+				break;
+			}
 		}
 	}
+
+	if (!(*DestinationImportData))
+	{
+		return false;
+	}
+	return true;
+}
+
+bool UInterchangeManager::ConvertImportData(const UObject* SourceImportData, FImportAssetParameters& ImportAssetParameters) const
+{
+	UObject* DestinationImportData = nullptr;
+	if (!ConvertImportData(SourceImportData, UInterchangeAssetImportData::StaticClass(), &DestinationImportData))
+	{
+		return false;
+	}
+
+	if (!ensure(DestinationImportData))
+	{
+		return false;
+	}
+
 	if (UInterchangeAssetImportData* AssetImportData = Cast<UInterchangeAssetImportData>(DestinationImportData))
 	{
 		//We can use the default pipeline stack, if it contain a pipeline that match the converted pipeline class
@@ -1520,7 +1543,10 @@ UInterchangeManager::ImportInternal(const FString& ContentPath, const UInterchan
 	auto EarlyExit = [&PreReturn]()
 	{
 		PreReturn();
-		return TTuple<UE::Interchange::FAssetImportResultRef, UE::Interchange::FSceneImportResultRef>{ MakeShared< UE::Interchange::FImportResult, ESPMode::ThreadSafe >(), MakeShared< UE::Interchange::FImportResult, ESPMode::ThreadSafe >() };
+		TTuple<UE::Interchange::FAssetImportResultRef, UE::Interchange::FSceneImportResultRef> ImportResult = TTuple<UE::Interchange::FAssetImportResultRef, UE::Interchange::FSceneImportResultRef>{ MakeShared< UE::Interchange::FImportResult, ESPMode::ThreadSafe >(), MakeShared< UE::Interchange::FImportResult, ESPMode::ThreadSafe >() };
+		ImportResult.Key->SetDone();
+		ImportResult.Value->SetDone();
+		return ImportResult;
 	};
 
 	if (FEngineAnalytics::IsAvailable())
diff --git a/Engine/Source/Runtime/Interchange/Engine/Private/InterchangeMeshUtilities.cpp b/Engine/Source/Runtime/Interchange/Engine/Private/InterchangeMeshUtilities.cpp
--- a/Engine/Source/Runtime/Interchange/Engine/Private/InterchangeMeshUtilities.cpp
+++ b/Engine/Source/Runtime/Interchange/Engine/Private/InterchangeMeshUtilities.cpp
@@ -97,60 +97,6 @@ bool FInterchangeSkeletalMeshAlternateSkinWeightPostImportTask::AddLodToReimport
 	return true;
 }
 
-TFuture<bool> UInterchangeMeshUtilities::ImportCustomLod(UObject* MeshObject, const int32 LodIndex)
-{
-	TSharedPtr<TPromise<bool>> Promise = MakeShared<TPromise<bool>>();
-	if (!MeshObject)
-	{
-		UE_LOG(LogInterchangeEngine, Warning, TEXT("FInterchangeMeshUtilities::ImportCustomLod: The MeshObject parameter cannot be null."));
-		Promise->SetValue(false);
-		return Promise->GetFuture();
-	}
-	if (!IsInGameThread())
-	{
-		UE_LOG(LogInterchangeEngine, Warning, TEXT("FInterchangeMeshUtilities::ImportCustomLod: Cannot ask the user for a file path outside of the game thread."));
-		Promise->SetValue(false);
-		return Promise->GetFuture();
-	}
-	UInterchangeManager& InterchangeManager = UInterchangeManager::GetInterchangeManager();
-
-	//Ask the user for a file path
-	const UInterchangeProjectSettings* InterchangeProjectSettings = GetDefault<UInterchangeProjectSettings>();
-	UInterchangeFilePickerBase* FilePicker = nullptr;
-
-	//In runtime we do not have any pipeline configurator
-#if WITH_EDITORONLY_DATA
-	TSoftClassPtr <UInterchangeFilePickerBase> FilePickerClass = InterchangeProjectSettings->FilePickerClass;
-	if (FilePickerClass.IsValid())
-	{
-		UClass* FilePickerClassLoaded = FilePickerClass.LoadSynchronous();
-		if (FilePickerClassLoaded)
-		{
-			FilePicker = NewObject<UInterchangeFilePickerBase>(GetTransientPackage(), FilePickerClassLoaded, NAME_None, RF_NoFlags);
-		}
-	}
-#endif
-	if(FilePicker)
-	{
-		FInterchangeFilePickerParameters Parameters;
-		Parameters.bAllowMultipleFiles = false;
-		Parameters.Title = FText::Format(NSLOCTEXT("Interchange", "ImportCustomLodAsync_FilePickerTitle", "Choose a file to import a custom LOD for LOD{0}"), FText::AsNumber(LodIndex));
-		TArray<FString> Filenames;
-		if (FilePicker->ScriptedFilePickerForTranslatorAssetType(EInterchangeTranslatorAssetType::Meshes, Parameters, Filenames))
-		{
-			//We set bAllowMultipleFile to false, we should have only one result
-			if (ensure(Filenames.Num() == 1))
-			{
-				const UInterchangeSourceData* SourceData = InterchangeManager.CreateSourceData(Filenames[0]);
-				return InternalImportCustomLod(Promise, MeshObject, LodIndex, SourceData);
-			}
-		}
-	}
-
-	Promise->SetValue(false);
-	return Promise->GetFuture();
-}
-
 TFuture<bool> UInterchangeMeshUtilities::ImportCustomLod(UObject* MeshObject, const int32 LodIndex, const UInterchangeSourceData* SourceData)
 {
 	TSharedPtr<TPromise<bool>> Promise = MakeShared<TPromise<bool>>();
@@ -168,10 +114,12 @@ TFuture<bool> UInterchangeMeshUtilities::InternalImportCustomLod(TSharedPtr<TPro
 	UStaticMesh* StaticMesh = Cast<UStaticMesh>(MeshObject);
 	EInterchangePipelineContext ImportType = EInterchangePipelineContext::AssetCustomLODImport;
 	bool bInvalidLodIndex = false;
+	UObject* SourceImportData = nullptr;
 	if (SkeletalMesh)
 	{
-		InterchangeAssetImportData = Cast<UInterchangeAssetImportData>(SkeletalMesh->GetAssetImportData());
-		if (SkeletalMesh->GetLODNum() > LodIndex && InterchangeAssetImportData)
+		SourceImportData = SkeletalMesh->GetAssetImportData();
+		InterchangeAssetImportData = Cast<UInterchangeAssetImportData>(SourceImportData);
+		if (SkeletalMesh->GetLODNum() > LodIndex)
 		{
 			ImportType = EInterchangePipelineContext::AssetCustomLODReimport;
 		}
@@ -182,8 +130,9 @@ TFuture<bool> UInterchangeMeshUtilities::InternalImportCustomLod(TSharedPtr<TPro
 	}
 	else if (StaticMesh)
 	{
-		InterchangeAssetImportData = Cast<UInterchangeAssetImportData>(StaticMesh->GetAssetImportData());
-		if (StaticMesh->GetNumSourceModels() > LodIndex && InterchangeAssetImportData)
+		SourceImportData = StaticMesh->GetAssetImportData();
+		InterchangeAssetImportData = Cast<UInterchangeAssetImportData>(SourceImportData);
+		if (StaticMesh->GetNumSourceModels() > LodIndex)
 		{
 			ImportType = EInterchangePipelineContext::AssetCustomLODReimport;
 		}
@@ -206,6 +155,22 @@ TFuture<bool> UInterchangeMeshUtilities::InternalImportCustomLod(TSharedPtr<TPro
 		return Promise->GetFuture();
 	}
 
+	const bool bInterchangeCanImportSourceData = InterchangeManager.CanTranslateSourceData(SourceData);
+
+	if (!bInterchangeCanImportSourceData)
+	{
+		UE_LOG(LogInterchangeEngine, Warning, TEXT("FInterchangeMeshUtilities::InternalImportCustomLod: Cannot import mesh LOD index %d, no interchange translator support this source file. [%s]"), LodIndex, *(SourceData->GetFilename()));
+		Promise->SetValue(false);
+		return Promise->GetFuture();
+	}
+
+	//Convert the asset import data if needed
+	if (!InterchangeAssetImportData)
+	{
+		//Try to convert the asset import data
+		InterchangeManager.ConvertImportData(SourceImportData, UInterchangeAssetImportData::StaticClass(), reinterpret_cast<UObject**>(&InterchangeAssetImportData));
+	}
+
 	FImportAssetParameters ImportAssetParameters;
 	ImportAssetParameters.bIsAutomated = true;
 	if (InterchangeAssetImportData)
@@ -229,6 +194,22 @@ TFuture<bool> UInterchangeMeshUtilities::InternalImportCustomLod(TSharedPtr<TPro
 			}
 		}
 	}
+	else
+	{
+		//Create import data
+		InterchangeAssetImportData = NewObject<UInterchangeAssetImportData>();
+		const UInterchangeProjectSettings* InterchangeProjectSettings = GetDefault<UInterchangeProjectSettings>();
+
+		if (const UClass* GenericPipelineClass = InterchangeProjectSettings->GenericPipelineClass.LoadSynchronous())
+		{
+			if (UInterchangePipelineBase* GenericPipeline = NewObject<UInterchangePipelineBase>(GetTransientPackage(), GenericPipelineClass))
+			{
+				GenericPipeline->ClearFlags(EObjectFlags::RF_Standalone | EObjectFlags::RF_Public);
+				GenericPipeline->AdjustSettingsForContext(ImportType, nullptr);
+				ImportAssetParameters.OverridePipelines.Add(GenericPipeline);
+			}
+		}
+	}
 
 	FString ImportAssetPath = TEXT("/Engine/TempEditor/Interchange/") + FGuid::NewGuid().ToString(EGuidFormats::Base36Encoded);
 	UE::Interchange::FAssetImportResultRef AssetImportResult = InterchangeManager.ImportAssetAsync(ImportAssetPath, SourceData, ImportAssetParameters);
diff --git a/Engine/Source/Runtime/Interchange/Engine/Public/InterchangeAssetImportData.h b/Engine/Source/Runtime/Interchange/Engine/Public/InterchangeAssetImportData.h
--- a/Engine/Source/Runtime/Interchange/Engine/Public/InterchangeAssetImportData.h
+++ b/Engine/Source/Runtime/Interchange/Engine/Public/InterchangeAssetImportData.h
@@ -199,14 +199,20 @@ class UInterchangeAssetImportDataConverterBase : public UObject
 {
 	GENERATED_BODY()
 public:
+
+	/* Converter must return which class they can convert to and from UInterchangeAssetImportData. */
+	virtual bool CanConvertClass(const UClass* SourceClass, const UClass* DestinationClass) const { return false; }
+
 	/**
 	 * Convert the asset import data from the one that is in the Object to
 	 * one that supports the target extension (for example, legacy FBX to Interchange or vice-versa)
 	 * The function should return true only if it has converted the asset import data, or false otherwise.
 	 * 
 	 * The system will call all objects that derive from this class until one converts the data.
+	 * 
+	 * @Param Asset - Represent an asset that have a member AssetImportData (StaticMesh, SkeletalMesh, Texture....)
 	 */
-	virtual bool ConvertImportData(UObject* Object, const FString& TargetExtension) const
+	virtual bool ConvertImportData(UObject* Asset, const FString& TargetExtension) const
 	{
 		return false;
 	}
@@ -217,7 +223,7 @@ public:
 	 *
 	 * The system will call all object deriving from this class until one convert the data.
 	 */
-	virtual bool ConvertImportData(const UObject* SourceImportData, UObject** DestinationImportDataClass) const
+	virtual bool ConvertImportData(const UObject* SourceImportData, UObject** DestinationImportData) const
 	{
 		return false;
 	}
diff --git a/Engine/Source/Runtime/Interchange/Engine/Public/InterchangeFilePickerBase.h b/Engine/Source/Runtime/Interchange/Engine/Public/InterchangeFilePickerBase.h
--- a/Engine/Source/Runtime/Interchange/Engine/Public/InterchangeFilePickerBase.h
+++ b/Engine/Source/Runtime/Interchange/Engine/Public/InterchangeFilePickerBase.h
@@ -25,6 +25,10 @@ struct FInterchangeFilePickerParameters
 	//Set the default open path that the dialog will show to the user.
 	UPROPERTY(EditAnywhere, Category = "Interchange | File Picker")
 	FString DefaultPath = TEXT("");
+
+	//If true, the user will be able to select any unreal editor factory + interchange file types.
+	UPROPERTY(EditAnywhere, Category = "Interchange | File Picker")
+	bool bShowAllFactoriesExtension = false;
 };
 
 UCLASS(Abstract, BlueprintType, Blueprintable, MinimalAPI)
diff --git a/Engine/Source/Runtime/Interchange/Engine/Public/InterchangeManager.h b/Engine/Source/Runtime/Interchange/Engine/Public/InterchangeManager.h
--- a/Engine/Source/Runtime/Interchange/Engine/Public/InterchangeManager.h
+++ b/Engine/Source/Runtime/Interchange/Engine/Public/InterchangeManager.h
@@ -485,20 +485,29 @@ public:
 
 	/**
 	 * Call all the registered converters to see if any converter can convert the data.
-	 * @Param Object - The Object to convert the import data.
+	 * @Param Asset - The asset we want to convert the import data.
 	 * @Param Extension - The file extension we want to import.
 	 * @return true if one of the converter has converted the data, or false otherwise.
 	 */
-	INTERCHANGEENGINE_API bool ConvertImportData(UObject* Object, const FString& Extension) const;
+	INTERCHANGEENGINE_API bool ConvertImportData(UObject* Asset, const FString& Extension) const;
 
 	/**
-	 * Call all the registered converter, if one converter want
+	 * Call all the registered converter, until one can convert the source import data to UInterchangeAssetImportData
 	 * @Param SourceImportData - The source import data options.
-	 * @Param DestinationImportData - The destination import data options.
+	 * @Param ImportAssetParameters - The interchange import asset parameters.
 	 * @return true if one of the converter has convert the data. False otherwise.
 	 */
 	INTERCHANGEENGINE_API bool ConvertImportData(const UObject* SourceImportData, FImportAssetParameters& ImportAssetParameters) const;
 
+	/**
+	 * Call all the registered converter, until one can convert the source data to the destination
+	 * @Param SourceImportData - The source import data options.
+	 * @Param DestinationImportData - The destination import data options we have to create and fill.
+	 * @Param DestinationClass - The class representing the DestinationImportData
+	 * @return true if one of the converter has convert the data. False otherwise.
+	 */
+	INTERCHANGEENGINE_API bool ConvertImportData(const UObject* SourceImportData, const UClass* DestinationClass, UObject** DestinationImportData) const;
+
 	/**
 	 * Returns the list of supported formats for a given translator type.
 	 */
diff --git a/Engine/Source/Runtime/Interchange/Engine/Public/InterchangeMeshUtilities.h b/Engine/Source/Runtime/Interchange/Engine/Public/InterchangeMeshUtilities.h
--- a/Engine/Source/Runtime/Interchange/Engine/Public/InterchangeMeshUtilities.h
+++ b/Engine/Source/Runtime/Interchange/Engine/Public/InterchangeMeshUtilities.h
@@ -40,17 +40,6 @@ class UInterchangeMeshUtilities : public UObject
 {
 	GENERATED_BODY()
 public:
-	/**
-	 * This function ask to select a file and import a mesh from it. After it add/replace the MeshObject LOD (at LodIndex) with the imported mesh LOD data.
-	 *
-	 * @Param MeshObject - The Mesh we want to add the lod
-	 * @Param LodIndex - The index of the LOD we want to replace or add
-	 * @Return - return true if it successfully add or replace the MeshObject LOD at LodIndex with the imported data.
-	 *
-	 * @Note - This function will sort out the SourceData and call ImportCustomLodAsync and wait until the result is available.
-	 */
-	static INTERCHANGEENGINE_API TFuture<bool> ImportCustomLod(UObject* MeshObject, const int32 LodIndex);
-
 
 	/**
 	 * This function import a mesh from the source data and add/replace the MeshObject LOD (at LodIndex) with the imported mesh LOD data.
