World Partition - Coversion Commandlet: When calling the conversion commandlet from the editor, make sure to use the commandlet result as the process exit code.
* Without this, any unrelated error logged during the execution will lead to the process exiting with an error

diff --git a/Engine/Source/Editor/UnrealEd/Private/Commandlets/WorldPartitionConvertCommandlet.cpp b/Engine/Source/Editor/UnrealEd/Private/Commandlets/WorldPartitionConvertCommandlet.cpp
--- a/Engine/Source/Editor/UnrealEd/Private/Commandlets/WorldPartitionConvertCommandlet.cpp
+++ b/Engine/Source/Editor/UnrealEd/Private/Commandlets/WorldPartitionConvertCommandlet.cpp
@@ -706,6 +706,12 @@ int32 UWorldPartitionConvertCommandlet::Main(const FString& Params)
 		LogWorldPartitionConvertCommandlet.SetVerbosity(ELogVerbosity::Verbose);
 	}
 
+	if (Switches.Contains(TEXT("RunningFromUnrealEd")))
+	{
+		UseCommandletResultAsExitCode = true;	// The process return code will match the return code of the commandlet
+		FastExit = true;						// Faster exit which avoids crash during shutdown. The engine isn't shutdown cleanly.
+	}
+
 	bConversionSuffix = Switches.Contains(TEXT("ConversionSuffix"));
 
 	// Load configuration file
