Experimental option to import a Project Launcher custom profile to the Turnkey custom builds

 - disabled by default. needs enabling via editor experimental settings

diff --git a/Engine/Source/Developer/DeveloperToolSettings/Classes/Settings/ProjectPackagingSettings.h b/Engine/Source/Developer/DeveloperToolSettings/Classes/Settings/ProjectPackagingSettings.h
--- a/Engine/Source/Developer/DeveloperToolSettings/Classes/Settings/ProjectPackagingSettings.h
+++ b/Engine/Source/Developer/DeveloperToolSettings/Classes/Settings/ProjectPackagingSettings.h
@@ -163,6 +163,8 @@ struct FProjectBuildSettings
 	 *   {iniif:-iostore:bUseIoStore}
 	 * Additionally, the ini settings can have an optional search and replace modifier, to easily modify the string. The Replace can be blank:
 	 *   {inivalue:BuildConfiguration|PPBC_=} - This will get the BuildConfiguration from the settings, and then remove the PPBC_ enum prefix from the string, to just get say Development
+	 * 
+	 * Adding -device={DeviceId} will show the command per-device in the platforms menu
 	 */
 	UPROPERTY(EditAnywhere, Category="Packaging")
 	FString BuildCookRunParams;
@@ -633,7 +635,7 @@ public:
 	/**
 	 * A list of custom builds that will show up in the Platforms menu to allow customized builds that make sense for your project. Will show up near Package Project in the Platforms menu.
 	 */
-	UPROPERTY(config, EditAnywhere, Category=Packaging, meta=(DisplayName = "Additional builds for this project."))
+	UPROPERTY(config, EditAnywhere, Category=CustomBuilds, meta=(DisplayName = "Additional builds for this project.", TitleProperty="Name"))
 	TArray<FProjectBuildSettings> ProjectCustomBuilds;
 
 	/** If set, platforms that destructively edit the iostore containers during packaging will save a copy prior to doing so. */
diff --git a/Engine/Source/Developer/LauncherServices/Private/Launcher/LauncherWorker.cpp b/Engine/Source/Developer/LauncherServices/Private/Launcher/LauncherWorker.cpp
--- a/Engine/Source/Developer/LauncherServices/Private/Launcher/LauncherWorker.cpp
+++ b/Engine/Source/Developer/LauncherServices/Private/Launcher/LauncherWorker.cpp
@@ -387,10 +387,10 @@ static FString Join(const TSet<FString>& Tokens, const FString& Delimeter)
 	return Result.RightChop(Delimeter.Len());
 }
 
-FString FLauncherWorker::CreateUATCommand( const ILauncherProfileRef& InProfile, const TArray<FString>& InPlatforms, TArray<FCommandDesc>& OutCommands, FString& CommandStart )
+FString FLauncherWorker::CreateUATCommand( const ILauncherProfileRef& InProfile, const TArray<FString>& InPlatforms, TArray<FCommandDesc>& OutCommands, FString& CommandStart, bool bForTurnkeyCustomBuild )
 {
 	CommandStart = TEXT("");
-	FString UATCommand = TEXT(" -utf8output");
+	FString UATCommand = bForTurnkeyCustomBuild ? TEXT("") : TEXT(" -utf8output");
 	FGuid SessionId(FGuid::NewGuid());
 	FString InitialMap = InProfile->GetDefaultLaunchRole()->GetInitialMap();
 	if (InitialMap.IsEmpty() && InProfile->GetCookedMaps().Num() == 1)
@@ -420,6 +420,10 @@ FString FLauncherWorker::CreateUATCommand( const ILauncherProfileRef& InProfile,
 	{
 		// Platform info for the given platform
 		const PlatformInfo::FTargetPlatformInfo* PlatformInfo = PlatformInfo::FindPlatformInfo(FName(*InPlatforms[PlatformIndex]));
+		if (bForTurnkeyCustomBuild && PlatformInfo == nullptr)
+		{
+			continue;
+		}
 
 		if (ensure(PlatformInfo))
 		{
@@ -510,7 +514,7 @@ FString FLauncherWorker::CreateUATCommand( const ILauncherProfileRef& InProfile,
 
 	bool bVsyncAdded = false;
 
-	if (DeviceGroup.IsValid())
+	if (DeviceGroup.IsValid() && DeviceProxyManager.IsValid())
 	{
 		const TArray<FString>& Devices = DeviceGroup->GetDeviceIDs();
 
@@ -552,8 +556,13 @@ FString FLauncherWorker::CreateUATCommand( const ILauncherProfileRef& InProfile,
 	}
 
 	// game command line
-	FString CommandLine = FString::Printf(TEXT(" -cmdline=\"%s -Messaging\""),
-		*InitialMap);
+	FString CommandLine = FString::Printf(TEXT(" -cmdline=\"%s%s\""),
+		*InitialMap, 
+		bForTurnkeyCustomBuild ? TEXT("") : TEXT(" -Messaging") );
+	if (CommandLine.EndsWith(TEXT("\"\"")))
+	{
+		CommandLine.Reset(); // don't bother with -cmdline if it's empty
+	}
 
 	// localization command line
 	FString LocalizationCommands;
@@ -567,22 +576,33 @@ FString FLauncherWorker::CreateUATCommand( const ILauncherProfileRef& InProfile,
 #endif	// WITH_EDITOR
 
 	// to reduce UECommandLine.txt churn (timestamp causing extra work), for LaunchOn (ie iterative deploy) we use a single session guid
-	if (InProfile->GetDeploymentMode() == ELauncherProfileDeploymentModes::CopyToDevice && Profile->IsDeployingIncrementally())
+	if (InProfile->GetDeploymentMode() == ELauncherProfileDeploymentModes::CopyToDevice && InProfile->IsDeployingIncrementally())
 	{
 		static FGuid StaticGuid(FGuid::NewGuid());
 		SessionId = StaticGuid;
 	}
 
 	// additional commands to be sent to the commandline
-	FString SessionName = InProfile->GetName().Replace(TEXT("\'"), TEXT("_")).Replace(TEXT("\'"), TEXT("_"));
-	FString SessionOwner = FString(FPlatformProcess::UserName(false)).Replace(TEXT("\'"), TEXT("_")).Replace(TEXT("\'"), TEXT("_"));;
-	FString AdditionalCommandLine = FString::Printf(TEXT(" -addcmdline=\"-SessionId=%s -SessionOwner='%s' -SessionName='%s'%s%s %s\""),
+	FString SessionCommands;
+	if (!bForTurnkeyCustomBuild)
+	{
+		FString SessionName = InProfile->GetName().Replace(TEXT("\'"), TEXT("_")).Replace(TEXT("\'"), TEXT("_"));
+		FString SessionOwner = FString(FPlatformProcess::UserName(false)).Replace(TEXT("\'"), TEXT("_")).Replace(TEXT("\'"), TEXT("_"));;
+		SessionCommands = FString::Printf(TEXT("-SessionId=%s -SessionOwner='%s' -SessionName='%s' "),
 		*SessionId.ToString(),
 		*SessionOwner,
-		*SessionName,
+		*SessionName);
+	}
+	FString AdditionalCommandLine = FString::Printf(TEXT(" -addcmdline=\"%s%s %s %s\""),
+		*SessionCommands,
 		*RoleCommands,
 		*LocalizationCommands,
 		*InProfile->GetAdditionalCommandLineParameters());
+	if (AdditionalCommandLine.EndsWith(TEXT("\"  \"")))
+	{
+		AdditionalCommandLine.Reset(); // don't bother with -addcmdline if it's empty
+	}
+
 
 	// map list
 	FString MapList;
@@ -921,7 +941,7 @@ FString FLauncherWorker::CreateUATCommand( const ILauncherProfileRef& InProfile,
 
 		case ELauncherProfileDeploymentModes::CopyToDevice:
 			{
-				if (Profile->IsDeployingIncrementally())
+				if (InProfile->IsDeployingIncrementally())
 				{
 					UATCommand += " -iterativedeploy";
 				}
@@ -995,7 +1015,7 @@ FString FLauncherWorker::CreateUATCommand( const ILauncherProfileRef& InProfile,
 
 		if (InProfile->IsArchiving())
 		{
-			UATCommand += FString::Printf(TEXT(" -archive -archivedirectory=\"%s\""), *Profile->GetArchiveDirectory());
+			UATCommand += FString::Printf(TEXT(" -archive -archivedirectory=\"%s\""), *InProfile->GetArchiveDirectory());
 
 			FCommandDesc Desc;
 			FText Command = FText::Format(LOCTEXT("LauncherArchiveDesc", "Archiving content for {0}"), FText::FromString(Platforms.RightChop(1)));
@@ -1021,6 +1041,50 @@ FString FLauncherWorker::CreateUATCommand( const ILauncherProfileRef& InProfile,
 	return UATCommand;
 }
 
+FString FLauncherWorker::MakeBuildCookRunParamsForProjectCustomBuild(const ILauncherProfileRef& InProfile, const TArray<FString>& InPlatforms)
+{
+	// get the basic UAT command line
+	FLauncherWorker TempWorker(NoInit);
+	TArray<FCommandDesc> Commands;
+	FString CommandStart;
+	FString BuildCookRunParams = TempWorker.CreateUATCommand(InProfile, InPlatforms, Commands, CommandStart, true);
+
+	// optional project
+	if (InProfile->HasProjectSpecified())
+	{
+		BuildCookRunParams += FString::Printf( TEXT(" -project=%s"), *FPaths::ConvertRelativePathToFull(InProfile->GetProjectPath()) );
+	}
+	else
+	{
+		BuildCookRunParams += TEXT(" -project={project}");
+	}
+
+	// optional platform
+	if (InPlatforms.Num() == 0)
+	{
+		BuildCookRunParams += TEXT(" -platform={platform}");
+	}
+
+	// optional device
+	if ( InProfile->GetDeploymentMode() == ELauncherProfileDeploymentModes::CopyToDevice || 
+		 InProfile->GetLaunchMode() != ELauncherProfileLaunchModes::DoNotLaunch ||
+		 InProfile->GetDeployedDeviceGroup() != nullptr )
+	{
+		BuildCookRunParams += TEXT(" -device={DeviceId}");
+	}
+
+	// configuration
+	if (InProfile->GetBuildConfiguration() != EBuildConfiguration::Unknown)
+	{
+		BuildCookRunParams += FString::Printf( TEXT(" -configuration=%s"), LexToString(InProfile->GetBuildConfiguration()) );
+	}
+
+	return BuildCookRunParams;
+}
+
+
+
+
 /* FLauncherWorker implementation
  *****************************************************************************/
 
diff --git a/Engine/Source/Developer/LauncherServices/Private/Launcher/LauncherWorker.h b/Engine/Source/Developer/LauncherServices/Private/Launcher/LauncherWorker.h
--- a/Engine/Source/Developer/LauncherServices/Private/Launcher/LauncherWorker.h
+++ b/Engine/Source/Developer/LauncherServices/Private/Launcher/LauncherWorker.h
@@ -105,6 +105,12 @@ public:
 		CachedDevicePackagePair.Add(Device, Package);
 	}
 
+	/**
+	 * Returns best-guess BuildCookRun parameters for the given profile in the format expected by FProjectBuidSettings::BuildCookRunParams
+	 *  Used by the Project custom Build's "import from Project Launcher" menu
+	 */
+	static FString MakeBuildCookRunParamsForProjectCustomBuild(const ILauncherProfileRef& InProfile, const TArray<FString>& InPlatforms);
+
 protected:
 
 	/**
@@ -117,7 +123,7 @@ protected:
 	void OnTaskStarted(const FString& TaskName);
 	void OnTaskCompleted(const FString& TaskName);
 
-	FString CreateUATCommand( const ILauncherProfileRef& InProfile, const TArray<FString>& InPlatforms, TArray<FCommandDesc>& OutCommands, FString& CommandStart );
+	FString CreateUATCommand( const ILauncherProfileRef& InProfile, const TArray<FString>& InPlatforms, TArray<FCommandDesc>& OutCommands, FString& CommandStart, bool bForTurnkeyCustomBuild = false );
 
 	//start listening to the device discovery routine
 	void EnableDeviceDiscoveryListener();
@@ -128,6 +134,9 @@ protected:
 	//Cancel the currently running application on all devices
 	bool TerminateLaunchedProcess();
 
+	// just used by MakeBuildCookRunParamsForProjectCustomBuild
+	FLauncherWorker(ENoInit) {};
+
 private:
 
 	// Holds a critical section to lock access to selected properties.
diff --git a/Engine/Source/Developer/LauncherServices/Private/Profiles/LauncherProfileManager.cpp b/Engine/Source/Developer/LauncherServices/Private/Profiles/LauncherProfileManager.cpp
--- a/Engine/Source/Developer/LauncherServices/Private/Profiles/LauncherProfileManager.cpp
+++ b/Engine/Source/Developer/LauncherServices/Private/Profiles/LauncherProfileManager.cpp
@@ -7,6 +7,7 @@
 #include "Serialization/JsonReader.h"
 #include "Serialization/JsonSerializer.h"
 #include "Launcher/LauncherProjectPath.h"
+#include "Launcher/LauncherWorker.h" // just for MakeBuildCookRunParamsForProjectCustomBuild
 #include "Profiles/LauncherDeviceGroup.h"
 #include "Profiles/LauncherProfile.h"
 #include "GameProjectHelper.h"
@@ -723,3 +724,9 @@ void FLauncherProfileManager::SaveProfiles( )
 		SaveJSONProfile((*It).ToSharedRef());
 	}
 }
+
+FString FLauncherProfileManager::MakeBuildCookRunParamsForProjectCustomBuild( const ILauncherProfileRef& InProfile, const TArray<FString>& InPlatforms ) const
+{
+	FString Result = FLauncherWorker::MakeBuildCookRunParamsForProjectCustomBuild(InProfile, InPlatforms);
+	return Result;
+}
diff --git a/Engine/Source/Developer/LauncherServices/Private/Profiles/LauncherProfileManager.h b/Engine/Source/Developer/LauncherServices/Private/Profiles/LauncherProfileManager.h
--- a/Engine/Source/Developer/LauncherServices/Private/Profiles/LauncherProfileManager.h
+++ b/Engine/Source/Developer/LauncherServices/Private/Profiles/LauncherProfileManager.h
@@ -121,6 +121,8 @@ public:
 
 	virtual const TArray<FString>& GetAllExplicitBuildTargetNames() const override;
 
+	virtual FString MakeBuildCookRunParamsForProjectCustomBuild(const ILauncherProfileRef& InProfile, const TArray<FString>& InPlatforms) const override;
+
 	//~ End ILauncherProfileManager Interface
 
 protected:
diff --git a/Engine/Source/Developer/LauncherServices/Public/ILauncherProfileManager.h b/Engine/Source/Developer/LauncherServices/Public/ILauncherProfileManager.h
--- a/Engine/Source/Developer/LauncherServices/Public/ILauncherProfileManager.h
+++ b/Engine/Source/Developer/LauncherServices/Public/ILauncherProfileManager.h
@@ -289,6 +289,12 @@ public:
 	*/
 	virtual const TArray<ILauncherProfileWizardPtr>& GetProfileWizards() const = 0;
 
+	/**
+	 * Returns best-guess BuildCookRun parameters for the given profile in the format expected by FProjectBuidSettings::BuildCookRunParams
+	 *  Used by the Project custom Build's "import from Project Launcher" menu
+	 */
+	virtual FString MakeBuildCookRunParamsForProjectCustomBuild(const ILauncherProfileRef& InProfile, const TArray<FString>& InPlatforms) const = 0;
+
 public:
 
 	/**
diff --git a/Engine/Source/Editor/UnrealEd/Classes/Settings/EditorExperimentalSettings.h b/Engine/Source/Editor/UnrealEd/Classes/Settings/EditorExperimentalSettings.h
--- a/Engine/Source/Editor/UnrealEd/Classes/Settings/EditorExperimentalSettings.h
+++ b/Engine/Source/Editor/UnrealEd/Classes/Settings/EditorExperimentalSettings.h
@@ -112,6 +112,10 @@ public:
 	UPROPERTY(EditAnywhere, config, Category = Tools, meta = (ConfigRestartRequired = true))
 	bool bFacialAnimationImporter;
 
+	/** Shows the option to import a custom project launcher profile in the project packaging 'custom builds' section */
+	UPROPERTY(EditAnywhere, config, Category = Tools, meta = (DisplayName = "Enable tools in Packaging Settings -> Additional builds for this project"))
+	bool bProjectCustomBuildTools;
+
 	/** Enable experimental PIE preview device launch */
 	UPROPERTY(EditAnywhere, config, Category = PIE, meta = (DisplayName = "Enable mobile PIE with preview device launch options."))
 	bool bMobilePIEPreviewDeviceLaunch;
diff --git a/Engine/Source/Editor/UnrealEd/Private/Settings/ProjectPackagingSettingsCustomization.h b/Engine/Source/Editor/UnrealEd/Private/Settings/ProjectPackagingSettingsCustomization.h
--- a/Engine/Source/Editor/UnrealEd/Private/Settings/ProjectPackagingSettingsCustomization.h
+++ b/Engine/Source/Editor/UnrealEd/Private/Settings/ProjectPackagingSettingsCustomization.h
@@ -10,12 +10,14 @@
 #include "Widgets/SWidget.h"
 #include "Widgets/SCompoundWidget.h"
 #include "Widgets/SBoxPanel.h"
+#include "Widgets/Input/SComboButton.h"
 #include "Styling/SlateTypes.h"
 #include "Misc/Paths.h"
 #include "Styling/CoreStyle.h"
 #include "HAL/FileManager.h"
 #include "Styling/AppStyle.h"
 #include "IDetailCustomization.h"
+#include "IPropertyUtilities.h"
 #include "PropertyHandle.h"
 #include "DetailLayoutBuilder.h"
 #include "DetailCategoryBuilder.h"
@@ -25,9 +27,12 @@
 #include "Widgets/Images/SImage.h"
 #include "Widgets/Input/SCheckBox.h"
 #include "Settings/ProjectPackagingSettings.h"
+#include "Settings/EditorExperimentalSettings.h"
 #include "PropertyRestriction.h"
 #include "Widgets/Views/SMultipleOptionTable.h"
 #include "DesktopPlatformModule.h"
+#include "ILauncherServicesModule.h"
+
 
 #define LOCTEXT_NAMESPACE "FProjectPackagingSettingsCustomization"
 
@@ -102,6 +107,7 @@ public:
 	{
 		CustomizeProjectCategory(LayoutBuilder);
 		CustomizePackagingCategory(LayoutBuilder);
+		CustomizeCustomBuildsCategory(LayoutBuilder);
 	}
 
 public:
@@ -408,6 +414,143 @@ protected:
 		return false;
 	}
 
+
+	/**
+	 * Customizes the Custom Builds property category.
+	 *
+	 * @param LayoutBuilder The layout builder.
+	 */
+	void CustomizeCustomBuildsCategory( IDetailLayoutBuilder& LayoutBuilder )
+	{
+		TSharedPtr<IPropertyUtilities> PropertyUtilities = LayoutBuilder.GetPropertyUtilities();
+		ILauncherServicesModule& LauncherServicesModule = FModuleManager::LoadModuleChecked<ILauncherServicesModule>("LauncherServices");
+		ILauncherProfileManagerRef LauncherProfileManager = LauncherServicesModule.GetProfileManager();
+
+		IDetailCategoryBuilder& CustomBuildsCategory = LayoutBuilder.EditCategory("CustomBuilds", FText::GetEmpty(), ECategoryPriority::Important);
+		{
+			CustomBuildsCategory.AddCustomRow(FText::GetEmpty(), false)
+				.WholeRowWidget
+				[
+					SNew(SBorder)
+					.BorderImage(FAppStyle::GetBrush("ToolPanel.GroupBorder"))
+					[
+						SNew(SHorizontalBox)
+
+						// combo button to import from project launcher. hidden if there are no custom laucher profiles
+						+SHorizontalBox::Slot()
+						.AutoWidth()
+						.VAlign(VAlign_Center)
+						.Padding(8.4)
+						[
+							SNew(SComboButton)
+							.ComboButtonStyle(FAppStyle::Get(), "SimpleComboButton")
+							.OnGetMenuContent_Lambda( [this, PropertyUtilities] { return CreateImportFromProjectLauncherMenu(PropertyUtilities); } )
+							.Visibility_Lambda( [LauncherProfileManager] { return LauncherProfileManager->GetAllProfiles().IsEmpty() ? EVisibility::Collapsed : EVisibility::Visible; })
+							.ButtonContent()
+							[
+								SNew(SHorizontalBox)
+
+								+SHorizontalBox::Slot()
+								.AutoWidth()
+								.VAlign(VAlign_Center)
+								.HAlign(HAlign_Center)
+								[
+									SNew(SImage)
+									.Image(FAppStyle::GetBrush("Launcher.TabIcon"))
+								]
+
+								+SHorizontalBox::Slot()
+								.AutoWidth()
+								.VAlign(VAlign_Center)
+								.Padding(4,0,0,0)
+								[
+									SNew(STextBlock)
+									.Text(LOCTEXT("ImportProjectLauncher", "Import From Project Launcher"))
+								]
+							]
+						]
+					]
+				]
+				.Visibility(TAttribute<EVisibility>::Create(TAttribute<EVisibility>::FGetter::CreateLambda([]() 
+				{
+					// only visible if the user has enabled this in the experimental settings
+					return GetDefault<UEditorExperimentalSettings>()->bProjectCustomBuildTools ? EVisibility::Visible : EVisibility::Collapsed; 
+				})))
+				;				
+		}
+	}
+
+	TSharedRef<SWidget> CreateImportFromProjectLauncherMenu( TSharedPtr<IPropertyUtilities> PropertyUtilities )
+	{
+		const bool bCloseAfterSelection = true;
+		FMenuBuilder MenuBuilder(bCloseAfterSelection, nullptr, nullptr, true);
+
+		ILauncherServicesModule& LauncherServicesModule = FModuleManager::LoadModuleChecked<ILauncherServicesModule>("LauncherServices");
+		for (const ILauncherProfilePtr& LauncherProfilePtr : LauncherServicesModule.GetProfileManager()->GetAllProfiles())
+		{
+			MenuBuilder.AddMenuEntry(
+				FText::FromString(LauncherProfilePtr->GetName()),
+				FText::FromString(LauncherProfilePtr->GetDescription()),
+				FSlateIcon(),
+				FUIAction
+				(
+					FExecuteAction::CreateRaw( this, &FProjectPackagingSettingsCustomization::ImportFromLauncherProfile, LauncherProfilePtr, PropertyUtilities )
+				),
+				NAME_None,
+				EUserInterfaceActionType::Button);
+		}
+
+		return MenuBuilder.MakeWidget();
+	}
+
+	void ImportFromLauncherProfile( const ILauncherProfilePtr LauncherProfilePtr, TSharedPtr<IPropertyUtilities> PropertyUtilities )
+	{
+		ILauncherServicesModule& LauncherServicesModule = FModuleManager::LoadModuleChecked<ILauncherServicesModule>("LauncherServices");
+
+		// grab the project packaging settings
+		UProjectPackagingSettings* ProjectPackagingSettings = UProjectPackagingSettings::StaticClass()->GetDefaultObject<UProjectPackagingSettings>();
+		if (ProjectPackagingSettings == nullptr)
+		{
+			return;
+		}
+
+		// ensure the name is unique (Turnkey builds a dictionary using the Name)
+		FString ProfileName = LauncherProfilePtr->GetName();
+		int UniqueId = 1;
+		while (ProjectPackagingSettings->ProjectCustomBuilds.ContainsByPredicate( [ProfileName](const FProjectBuildSettings& Other) { return ProfileName == Other.Name; } ) )
+		{
+			ProfileName = FString::Printf( TEXT("%s %d"), *LauncherProfilePtr->GetName(), UniqueId++ );
+		}
+
+		// add a new item
+		FProjectBuildSettings& ProjectBuildSettings = ProjectPackagingSettings->ProjectCustomBuilds.AddDefaulted_GetRef();
+		ProjectBuildSettings.Name               = ProfileName;
+		ProjectBuildSettings.HelpText           = LauncherProfilePtr->GetDescription().IsEmpty() ? LauncherProfilePtr->GetName() : LauncherProfilePtr->GetDescription();
+		ProjectBuildSettings.SpecificPlatforms  = LauncherProfilePtr->GetCookedPlatforms();
+		if (ProjectBuildSettings.SpecificPlatforms.Num() == 0 && !LauncherProfilePtr->GetDefaultDeployPlatform().IsNone())
+		{
+			ProjectBuildSettings.SpecificPlatforms.Add(LauncherProfilePtr->GetDefaultDeployPlatform().ToString());
+		}
+		ProjectBuildSettings.BuildCookRunParams = LauncherServicesModule.GetProfileManager()->MakeBuildCookRunParamsForProjectCustomBuild(LauncherProfilePtr.ToSharedRef(), ProjectBuildSettings.SpecificPlatforms);
+
+
+		// signal that the property has changed
+		FNotifyHook* NotifyHook = PropertyUtilities->GetNotifyHook();
+		FProperty* Property = UProjectPackagingSettings::StaticClass()->FindPropertyByName("ProjectCustomBuilds");
+		if (NotifyHook != nullptr && Property != nullptr)
+		{
+			TArray<const UObject*> NotifyTopLevelObjects;
+			NotifyTopLevelObjects.Add(ProjectPackagingSettings);
+
+			FEditPropertyChain PropertyChain;
+			PropertyChain.AddHead(Property);
+
+			FPropertyChangedEvent ChangeEvent(Property, EPropertyChangeType::ValueSet, MakeArrayView(NotifyTopLevelObjects));
+			NotifyHook->NotifyPostChange(ChangeEvent, &PropertyChain);
+		}
+	}
+
+
 private:
 	TArray<FCulturePtr> CultureList;
 	TSharedPtr<IPropertyHandle> CulturesPropertyHandle;
