Interchange: fix static mesh workflow when switching from legacy to interchange and do re-import/reimport LODs/...

#rnx

diff --git a/Engine/Plugins/Interchange/Editor/Source/InterchangeEditor/Private/InterchangeFbxAssetImportDataConverter.cpp b/Engine/Plugins/Interchange/Editor/Source/InterchangeEditor/Private/InterchangeFbxAssetImportDataConverter.cpp
--- a/Engine/Plugins/Interchange/Editor/Source/InterchangeEditor/Private/InterchangeFbxAssetImportDataConverter.cpp
+++ b/Engine/Plugins/Interchange/Editor/Source/InterchangeEditor/Private/InterchangeFbxAssetImportDataConverter.cpp
@@ -992,7 +992,7 @@ bool UInterchangeFbxAssetImportDataConverter::ConvertImportData(UObject* Asset,
 	return bResult;
 }
 
-bool UInterchangeFbxAssetImportDataConverter::ConvertImportData(const UObject* SourceImportData, UObject** DestinationImportData) const
+bool UInterchangeFbxAssetImportDataConverter::ConvertImportData(const UObject* SourceImportData, const UClass* DestinationClass, UObject** DestinationImportData) const
 {
 	bool bResult = false;
 	if (!SourceImportData || !DestinationImportData)
@@ -1001,80 +1001,113 @@ bool UInterchangeFbxAssetImportDataConverter::ConvertImportData(const UObject* S
 	}
 
 	UInterchangeManager& InterchangeManager = UInterchangeManager::GetInterchangeManager();
-	if (SourceImportData->IsA<UFbxImportUI>())
+	if (SourceImportData->IsA<UInterchangeAssetImportData>())
 	{
-		//Convert Legacy Fbx to Interchange
-		*DestinationImportData = UE::Interchange::Private::ConvertToInterchange(GetTransientPackage(), Cast<UFbxImportUI>(SourceImportData));
-		bResult = true;
+		if (UFbxImportUI* FbxImportUI = UE::Interchange::Private::ConvertToLegacyFbx(GetTransientPackage(), Cast<UInterchangeAssetImportData>(SourceImportData)))
+		{
+			bResult = true;
+			if (DestinationClass->IsChildOf<UFbxImportUI>())
+			{
+				*DestinationImportData = FbxImportUI;
+			}
+			else if (DestinationClass->IsChildOf<UFbxStaticMeshImportData>())
+			{
+				(*DestinationImportData) = FbxImportUI->StaticMeshImportData;
+			}
+			else if (DestinationClass->IsChildOf<UFbxSkeletalMeshImportData>())
+			{
+				(*DestinationImportData) = FbxImportUI->SkeletalMeshImportData;
+			}
+			else if (DestinationClass->IsChildOf<UFbxAnimSequenceImportData>())
+			{
+				(*DestinationImportData) = FbxImportUI->AnimSequenceImportData;
+			}
+			else
+			{
+				bResult = false;
+			}
+		}
 	}
-	else if (SourceImportData->IsA<UFbxAssetImportData>())
+	else
 	{
-		//We convert the UFbxAssetImportData into a UFbxImportUI
-		auto FillFbxAssetImportData = [](const UFbxAssetImportData* SourceAssetImportData, UFbxAssetImportData* DestinationAssetImportData)
+		const UFbxImportUI* FbxImportUI = nullptr;
+		if (SourceImportData->IsA<UFbxImportUI>())
+		{
+			FbxImportUI = const_cast<UFbxImportUI*>(Cast<UFbxImportUI>(SourceImportData));
+			
+		}
+		else if (SourceImportData->IsA<UFbxAssetImportData>())
+		{
+			//We convert the UFbxAssetImportData into a UFbxImportUI
+			auto FillFbxAssetImportData = [](const UFbxAssetImportData* SourceAssetImportData, UFbxAssetImportData* DestinationAssetImportData)
+				{
+					DestinationAssetImportData->bConvertScene = SourceAssetImportData->bConvertScene;
+					DestinationAssetImportData->bConvertSceneUnit = SourceAssetImportData->bConvertSceneUnit;
+					DestinationAssetImportData->bForceFrontXAxis = SourceAssetImportData->bForceFrontXAxis;
+					DestinationAssetImportData->bImportAsScene = SourceAssetImportData->bImportAsScene;
+					DestinationAssetImportData->ImportRotation = SourceAssetImportData->ImportRotation;
+					DestinationAssetImportData->ImportTranslation = SourceAssetImportData->ImportTranslation;
+					DestinationAssetImportData->ImportUniformScale = SourceAssetImportData->ImportUniformScale;
+				};
+
+			UFbxImportUI* TempFbxImportUI = NewObject<UFbxImportUI>(GetTransientPackage());
+			TempFbxImportUI->bImportMaterials = false;
+			TempFbxImportUI->bImportAsSkeletal = false;
+			TempFbxImportUI->bImportMesh = false;
+			TempFbxImportUI->bImportAnimations = false;
+			TempFbxImportUI->bImportRigidMesh = false;
+			TempFbxImportUI->bImportTextures = false;
+			TempFbxImportUI->bIsObjImport = false;
+			TempFbxImportUI->bIsReimport = false;
+			TempFbxImportUI->bCreatePhysicsAsset = false;
+			TempFbxImportUI->PhysicsAsset = nullptr;
+			TempFbxImportUI->Skeleton = nullptr;
+			if (SourceImportData->IsA<UFbxSkeletalMeshImportData>())
 			{
-				DestinationAssetImportData->bConvertScene = SourceAssetImportData->bConvertScene;
-				DestinationAssetImportData->bConvertSceneUnit = SourceAssetImportData->bConvertSceneUnit;
-				DestinationAssetImportData->bForceFrontXAxis = SourceAssetImportData->bForceFrontXAxis;
-				DestinationAssetImportData->bImportAsScene = SourceAssetImportData->bImportAsScene;
-				DestinationAssetImportData->ImportRotation = SourceAssetImportData->ImportRotation;
-				DestinationAssetImportData->ImportTranslation = SourceAssetImportData->ImportTranslation;
-				DestinationAssetImportData->ImportUniformScale = SourceAssetImportData->ImportUniformScale;
-			};
+				TempFbxImportUI->SkeletalMeshImportData = const_cast<UFbxSkeletalMeshImportData*>(Cast<UFbxSkeletalMeshImportData>(SourceImportData));
+				TempFbxImportUI->MeshTypeToImport = EFBXImportType::FBXIT_SkeletalMesh;
+				TempFbxImportUI->bImportAsSkeletal = true;
+				TempFbxImportUI->bImportMesh = true;
+
+				FillFbxAssetImportData(TempFbxImportUI->SkeletalMeshImportData, TempFbxImportUI->StaticMeshImportData);
+				FillFbxAssetImportData(TempFbxImportUI->SkeletalMeshImportData, TempFbxImportUI->AnimSequenceImportData);
+				FillFbxAssetImportData(TempFbxImportUI->SkeletalMeshImportData, TempFbxImportUI->TextureImportData);
+			}
+			else if (SourceImportData->IsA<UFbxStaticMeshImportData>())
+			{
+				TempFbxImportUI->MeshTypeToImport = EFBXImportType::FBXIT_StaticMesh;
+				TempFbxImportUI->bImportMesh = true;
 
-		UFbxImportUI* TempFbxImportUI = NewObject<UFbxImportUI>(GetTransientPackage());
-		TempFbxImportUI->bImportMaterials = false;
-		TempFbxImportUI->bImportAsSkeletal = false;
-		TempFbxImportUI->bImportMesh = false;
-		TempFbxImportUI->bImportAnimations = false;
-		TempFbxImportUI->bImportRigidMesh = false;
-		TempFbxImportUI->bImportTextures = false;
-		TempFbxImportUI->bIsObjImport = false;
-		TempFbxImportUI->bIsReimport = false;
-		TempFbxImportUI->bCreatePhysicsAsset = false;
-		TempFbxImportUI->PhysicsAsset = nullptr;
-		TempFbxImportUI->Skeleton = nullptr;
-		if (SourceImportData->IsA<UFbxSkeletalMeshImportData>())
-		{
-			TempFbxImportUI->SkeletalMeshImportData = const_cast<UFbxSkeletalMeshImportData*>(Cast<UFbxSkeletalMeshImportData>(SourceImportData));
-			TempFbxImportUI->MeshTypeToImport = EFBXImportType::FBXIT_SkeletalMesh;
-			TempFbxImportUI->bImportAsSkeletal = true;
-			TempFbxImportUI->bImportMesh = true;
-
-			FillFbxAssetImportData(TempFbxImportUI->SkeletalMeshImportData, TempFbxImportUI->StaticMeshImportData);
-			FillFbxAssetImportData(TempFbxImportUI->SkeletalMeshImportData, TempFbxImportUI->AnimSequenceImportData);
-			FillFbxAssetImportData(TempFbxImportUI->SkeletalMeshImportData, TempFbxImportUI->TextureImportData);
-		}
-		else if (SourceImportData->IsA<UFbxStaticMeshImportData>())
-		{
-			TempFbxImportUI->MeshTypeToImport = EFBXImportType::FBXIT_StaticMesh;
-			TempFbxImportUI->bImportMesh = true;
-
-			TempFbxImportUI->StaticMeshImportData = const_cast<UFbxStaticMeshImportData*>(Cast<UFbxStaticMeshImportData>(SourceImportData));
-			FillFbxAssetImportData(TempFbxImportUI->StaticMeshImportData, TempFbxImportUI->SkeletalMeshImportData);
-			FillFbxAssetImportData(TempFbxImportUI->StaticMeshImportData, TempFbxImportUI->AnimSequenceImportData);
-			FillFbxAssetImportData(TempFbxImportUI->StaticMeshImportData, TempFbxImportUI->TextureImportData);
-		}
-		else if (SourceImportData->IsA<UFbxAnimSequenceImportData>())
-		{
-			TempFbxImportUI->MeshTypeToImport = EFBXImportType::FBXIT_Animation;
-			TempFbxImportUI->bImportAsSkeletal = true;
-			TempFbxImportUI->bImportMesh = false;
-			TempFbxImportUI->AnimSequenceImportData = const_cast<UFbxAnimSequenceImportData*>(Cast<UFbxAnimSequenceImportData>(SourceImportData));
-			FillFbxAssetImportData(TempFbxImportUI->AnimSequenceImportData, TempFbxImportUI->SkeletalMeshImportData);
-			FillFbxAssetImportData(TempFbxImportUI->AnimSequenceImportData, TempFbxImportUI->StaticMeshImportData);
-			FillFbxAssetImportData(TempFbxImportUI->AnimSequenceImportData, TempFbxImportUI->TextureImportData);
+				TempFbxImportUI->StaticMeshImportData = const_cast<UFbxStaticMeshImportData*>(Cast<UFbxStaticMeshImportData>(SourceImportData));
+				FillFbxAssetImportData(TempFbxImportUI->StaticMeshImportData, TempFbxImportUI->SkeletalMeshImportData);
+				FillFbxAssetImportData(TempFbxImportUI->StaticMeshImportData, TempFbxImportUI->AnimSequenceImportData);
+				FillFbxAssetImportData(TempFbxImportUI->StaticMeshImportData, TempFbxImportUI->TextureImportData);
+			}
+			else if (SourceImportData->IsA<UFbxAnimSequenceImportData>())
+			{
+				TempFbxImportUI->MeshTypeToImport = EFBXImportType::FBXIT_Animation;
+				TempFbxImportUI->bImportAsSkeletal = true;
+				TempFbxImportUI->bImportMesh = false;
+				TempFbxImportUI->AnimSequenceImportData = const_cast<UFbxAnimSequenceImportData*>(Cast<UFbxAnimSequenceImportData>(SourceImportData));
+				FillFbxAssetImportData(TempFbxImportUI->AnimSequenceImportData, TempFbxImportUI->SkeletalMeshImportData);
+				FillFbxAssetImportData(TempFbxImportUI->AnimSequenceImportData, TempFbxImportUI->StaticMeshImportData);
+				FillFbxAssetImportData(TempFbxImportUI->AnimSequenceImportData, TempFbxImportUI->TextureImportData);
+			}
+			else
+			{
+				ensureMsgf(false, TEXT("Fbx interchange converter: miss match between CanConvertClass and the convertion capacity"));
+				TempFbxImportUI = nullptr;
+			}
+			//Assign to the const pointer we use to convert the data
+			FbxImportUI = TempFbxImportUI;
 		}
-		else
+
+		if (FbxImportUI)
 		{
-			ensureMsgf(false, TEXT("Fbx interchange converter: miss match between CanConvertClass and the convertion capacity"));
+			//Convert Legacy Fbx to Interchange
+			*DestinationImportData = UE::Interchange::Private::ConvertToInterchange(GetTransientPackage(), FbxImportUI);
+			bResult = true;
 		}
-		*DestinationImportData = UE::Interchange::Private::ConvertToInterchange(GetTransientPackage(), TempFbxImportUI);
-		bResult = true;
-	}
-	else if (SourceImportData->IsA<UInterchangeAssetImportData>())
-	{
-		*DestinationImportData = UE::Interchange::Private::ConvertToLegacyFbx(GetTransientPackage(), Cast<UInterchangeAssetImportData>(SourceImportData));
-		bResult = true;
 	}
 	return bResult;
 }
@@ -1098,7 +1131,13 @@ bool UInterchangeFbxAssetImportDataConverter::CanConvertClass(const UClass* Sour
 
 	if (SourceClass->IsChildOf(UInterchangeAssetImportData::StaticClass()))
 	{
-		return DestinationClass->IsChildOf(UFbxImportUI::StaticClass());
+		if (DestinationClass->IsChildOf(UFbxImportUI::StaticClass())
+			|| DestinationClass->IsChildOf(UFbxSkeletalMeshImportData::StaticClass())
+			|| DestinationClass->IsChildOf(UFbxStaticMeshImportData::StaticClass())
+			|| DestinationClass->IsChildOf(UFbxAnimSequenceImportData::StaticClass()))
+		{
+			return true;
+		}
 	}
 	return false;
 }
\ No newline at end of file
diff --git a/Engine/Plugins/Interchange/Editor/Source/InterchangeEditor/Public/InterchangeFbxAssetImportDataConverter.h b/Engine/Plugins/Interchange/Editor/Source/InterchangeEditor/Public/InterchangeFbxAssetImportDataConverter.h
--- a/Engine/Plugins/Interchange/Editor/Source/InterchangeEditor/Public/InterchangeFbxAssetImportDataConverter.h
+++ b/Engine/Plugins/Interchange/Editor/Source/InterchangeEditor/Public/InterchangeFbxAssetImportDataConverter.h
@@ -18,5 +18,5 @@ public:
 	virtual bool CanConvertClass(const UClass* SourceClass, const UClass* DestinationClass) const override;
 
 	virtual bool ConvertImportData(UObject* Asset, const FString& ToExtension) const override;
-	virtual bool ConvertImportData(const UObject* SourceImportData, UObject** DestinationImportData) const override;
+	virtual bool ConvertImportData(const UObject* SourceImportData, const UClass* DestinationClass, UObject** DestinationImportData) const override;
 };
diff --git a/Engine/Plugins/Interchange/Editor/Source/Utilities/Private/InterchangeOpenFileDialog.cpp b/Engine/Plugins/Interchange/Editor/Source/Utilities/Private/InterchangeOpenFileDialog.cpp
--- a/Engine/Plugins/Interchange/Editor/Source/Utilities/Private/InterchangeOpenFileDialog.cpp
+++ b/Engine/Plugins/Interchange/Editor/Source/Utilities/Private/InterchangeOpenFileDialog.cpp
@@ -13,7 +13,7 @@
 
 namespace UE::Interchange::Utilities::Private
 {
-	FString GetOpenFileDialogExtensions(const TArray<FString>& TranslatorFormats, bool bShowAllFactoriesExtension)
+	FString GetOpenFileDialogExtensions(const TArray<FString>& TranslatorFormats, bool bShowAllFactoriesExtension, const TArray<FString>& ExtraFormats)
 	{
 		FString FileTypes;
 		FString Extensions;
@@ -44,6 +44,11 @@ namespace UE::Interchange::Utilities::Private
 		}
 
 		ObjectTools::AppendFormatsFileExtensions(TranslatorFormats, FileTypes, Extensions, FilterIndexToFactory);
+		
+		if (!ExtraFormats.IsEmpty())
+		{
+			ObjectTools::AppendFormatsFileExtensions(ExtraFormats, FileTypes, Extensions, FilterIndexToFactory);
+		}
 
 		const FString FormatString = FString::Printf(TEXT("All Files (%s)|%s|%s"), *Extensions, *Extensions, *FileTypes);
 
@@ -77,12 +82,19 @@ namespace UE::Interchange::Utilities::Private
 
 bool UInterchangeFilePickerGeneric::FilePickerForTranslatorAssetType(const EInterchangeTranslatorAssetType TranslatorAssetType, FInterchangeFilePickerParameters& Parameters, TArray<FString>& OutFilenames)
 {
-	FString Extensions = UE::Interchange::Utilities::Private::GetOpenFileDialogExtensions(UInterchangeManager::GetInterchangeManager().GetSupportedAssetTypeFormats(TranslatorAssetType), Parameters.bShowAllFactoriesExtension);
+	FString Extensions = UE::Interchange::Utilities::Private::GetOpenFileDialogExtensions
+	(
+		UInterchangeManager::GetInterchangeManager().GetSupportedAssetTypeFormats(TranslatorAssetType), Parameters.bShowAllFactoriesExtension, Parameters.ExtraFormats
+	);
+
 	return UE::Interchange::Utilities::Private::FilePickerDialog(Extensions, Parameters, OutFilenames);
 }
 
 bool UInterchangeFilePickerGeneric::FilePickerForTranslatorType(const EInterchangeTranslatorType TranslatorType, FInterchangeFilePickerParameters& Parameters, TArray<FString>& OutFilenames)
 {
-	FString Extensions = UE::Interchange::Utilities::Private::GetOpenFileDialogExtensions(UInterchangeManager::GetInterchangeManager().GetSupportedFormats(TranslatorType), Parameters.bShowAllFactoriesExtension);
+	FString Extensions = UE::Interchange::Utilities::Private::GetOpenFileDialogExtensions
+	(
+		UInterchangeManager::GetInterchangeManager().GetSupportedFormats(TranslatorType), Parameters.bShowAllFactoriesExtension, Parameters.ExtraFormats
+	);
 	return UE::Interchange::Utilities::Private::FilePickerDialog(Extensions, Parameters, OutFilenames);
 }
diff --git a/Engine/Plugins/Interchange/Runtime/Source/Import/Private/Mesh/InterchangeStaticMeshFactory.cpp b/Engine/Plugins/Interchange/Runtime/Source/Import/Private/Mesh/InterchangeStaticMeshFactory.cpp
--- a/Engine/Plugins/Interchange/Runtime/Source/Import/Private/Mesh/InterchangeStaticMeshFactory.cpp
+++ b/Engine/Plugins/Interchange/Runtime/Source/Import/Private/Mesh/InterchangeStaticMeshFactory.cpp
@@ -96,13 +96,21 @@ namespace UE::Interchange::Private::StaticMesh
 				bool bFoundMatch = false;
 				for (int32 MaterialIndex = 0; MaterialIndex < Materials.Num(); ++MaterialIndex)
 				{
+					FName MaterialName = Materials[MaterialIndex].ImportedMaterialSlotName;
 					if (RemapMaterialIndexes[MaterialIndex] != INDEX_NONE)
 					{
 						//This material was already matched
+
+						//If the name match say we found the match)
+						if (MaterialName == ImportMaterialName)
+						{
+							bFoundMatch = true;
+							break;
+						}
 						continue;
 					}
 					int32& RemapIndex = RemapMaterialIndexes[MaterialIndex];
-					FName MaterialName = Materials[MaterialIndex].ImportedMaterialSlotName;
+					
 					if (MaterialName == ImportMaterialName)
 					{
 						RemapIndex = ReorderMaterialArray.Add(Materials[MaterialIndex]);
diff --git a/Engine/Plugins/Interchange/Runtime/Source/Parsers/Fbx/Private/FbxAPI.cpp b/Engine/Plugins/Interchange/Runtime/Source/Parsers/Fbx/Private/FbxAPI.cpp
--- a/Engine/Plugins/Interchange/Runtime/Source/Parsers/Fbx/Private/FbxAPI.cpp
+++ b/Engine/Plugins/Interchange/Runtime/Source/Parsers/Fbx/Private/FbxAPI.cpp
@@ -15,6 +15,7 @@
 #if WITH_ENGINE
 #include "Mesh/InterchangeMeshPayload.h"
 #endif
+#include "Misc/Paths.h"
 #include "Nodes/InterchangeBaseNodeContainer.h"
 #include "Nodes/InterchangeSourceNode.h"
 #include "Misc/SecureHash.h"
@@ -130,6 +131,10 @@ namespace UE
 
 				bool bStatus = SDKImporter->Import(SDKScene);
 
+				//To be able to re-import legacy fbx imported skeleton hierarchy we need to rename bones the same way legacy was doing, so this rename precede the CleanupFbxData renaming
+				constexpr bool bRemovePath = true;
+				EnsureNodeNameAreValid(FPaths::GetBaseFilename(Filename, bRemovePath));
+
 				//We always convert scene to UE axis and units
 				FFbxConvert::ConvertScene(SDKScene, bConvertScene, bForceFrontXAxis, bConvertSceneUnit, FileDetails.AxisDirection, FileDetails.UnitSystem);
 
@@ -256,6 +261,60 @@ namespace UE
 				}
 			}
 
+			void FFbxParser::EnsureNodeNameAreValid(const FString& BaseFilename)
+			{
+				const bool bKeepNamespace = false;//Todo use: GetDefault<UEditorPerProjectUserSettings>()->bKeepFbxNamespace;
+
+				TSet<FString> AllNodeName;
+				int32 CurrentNameIndex = 1;
+				for (int32 NodeIndex = 0; NodeIndex < SDKScene->GetNodeCount(); ++NodeIndex)
+				{
+					FbxNode* Node = SDKScene->GetNode(NodeIndex);
+					FString NodeName = UTF8_TO_TCHAR(Node->GetName());
+					if (NodeName.IsEmpty())
+					{
+						do
+						{
+							NodeName = TEXT("ncl1_") + FString::FromInt(CurrentNameIndex++);
+						} while (AllNodeName.Contains(NodeName));
+
+						Node->SetName(TCHAR_TO_UTF8(*NodeName));
+						if (!GIsAutomationTesting)
+						{
+							UInterchangeResultDisplay_Generic* Message = AddMessage<UInterchangeResultDisplay_Generic>();
+							Message->Text = FText::Format(LOCTEXT("EnsureNodeNameAreValid_NoNodeName", "Interchange FBX file Loading: Found node with no name, new node name is '{0}'"), FText::FromString(NodeName));
+						}
+					}
+					if (bKeepNamespace)
+					{
+						if (NodeName.Contains(TEXT(":")))
+						{
+							NodeName = NodeName.Replace(TEXT(":"), TEXT("_"));
+							Node->SetName(TCHAR_TO_UTF8(*NodeName));
+						}
+					}
+					// Do not allow node to be named same as filename as this creates problems later on (reimport)
+					if (AllNodeName.Contains(NodeName))
+					{
+						FString UniqueNodeName;
+						do
+						{
+							UniqueNodeName = NodeName + FString::FromInt(CurrentNameIndex++);
+						} while (AllNodeName.Contains(UniqueNodeName));
+
+						FbxString UniqueName(TCHAR_TO_UTF8(*UniqueNodeName));
+						Node->SetName(UniqueName);
+
+						if (!GIsAutomationTesting)
+						{
+							UInterchangeResultDisplay_Generic* Message = AddMessage<UInterchangeResultDisplay_Generic>();
+							Message->Text = FText::Format(LOCTEXT("EnsureNodeNameAreValid_NodeNameClash", "FBX File Loading: Found name clash, node '{0}' was renamed to '{1}'"), FText::FromString(NodeName), FText::FromString(UniqueNodeName));
+						}
+					}
+					AllNodeName.Add(NodeName);
+				}
+			}
+
 			void FFbxParser::CleanupFbxData()
 			{
 				//////////////////////////////////////////////////////////////////////////
diff --git a/Engine/Plugins/Interchange/Runtime/Source/Parsers/Fbx/Private/FbxAPI.h b/Engine/Plugins/Interchange/Runtime/Source/Parsers/Fbx/Private/FbxAPI.h
--- a/Engine/Plugins/Interchange/Runtime/Source/Parsers/Fbx/Private/FbxAPI.h
+++ b/Engine/Plugins/Interchange/Runtime/Source/Parsers/Fbx/Private/FbxAPI.h
@@ -108,6 +108,7 @@ namespace UE
 				FCriticalSection PayloadCriticalSection;
 			private:
 
+				void EnsureNodeNameAreValid(const FString& BaseFilename);
 				void CleanupFbxData();
 				void ProcessExtraInformation(UInterchangeBaseNodeContainer& NodeContainer);
 
diff --git a/Engine/Plugins/Interchange/Runtime/Source/Pipelines/Private/InterchangeGenericAssetsPipeline.cpp b/Engine/Plugins/Interchange/Runtime/Source/Pipelines/Private/InterchangeGenericAssetsPipeline.cpp
--- a/Engine/Plugins/Interchange/Runtime/Source/Pipelines/Private/InterchangeGenericAssetsPipeline.cpp
+++ b/Engine/Plugins/Interchange/Runtime/Source/Pipelines/Private/InterchangeGenericAssetsPipeline.cpp
@@ -926,7 +926,7 @@ TArray<FInterchangeConflictInfo> UInterchangeGenericAssetsPipeline::GetConflictI
 		CreateMaterialConflict(StaticMesh, SkeletalMesh, TransientBaseNodeContainer);
 	}
 
-	if (SpecifiedSkeleton || SkeletalMesh)
+	if (SpecifiedSkeleton && SkeletalMesh)
 	{
 		CreateSkeletonConflict(SpecifiedSkeleton, SkeletalMesh, TransientBaseNodeContainer);
 	}
diff --git a/Engine/Source/Editor/UnrealEd/Private/FbxMeshUtils.cpp b/Engine/Source/Editor/UnrealEd/Private/FbxMeshUtils.cpp
--- a/Engine/Source/Editor/UnrealEd/Private/FbxMeshUtils.cpp
+++ b/Engine/Source/Editor/UnrealEd/Private/FbxMeshUtils.cpp
@@ -63,13 +63,10 @@ namespace FbxMeshUtils
 			UFbxStaticMeshImportData* ImportData = Cast<UFbxStaticMeshImportData>(BaseStaticMesh->AssetImportData);
 			if (ImportData == nullptr)
 			{
-				//Convert the static mesh import data to legacy fbx import data, because the legacy code will override it if not of the good type
+				//Convert interchange asset import data to legacy fbx static mesh import data.
 				if (UInterchangeAssetImportData* InterchangeAssetImportData = Cast<UInterchangeAssetImportData>(BaseStaticMesh->GetAssetImportData()))
 				{
-					if (UInterchangeManager::GetInterchangeManager().ConvertImportData(BaseStaticMesh, TEXT("fbx")))
-					{
-						ImportData = Cast<UFbxStaticMeshImportData>(BaseStaticMesh->AssetImportData);
-					}
+					UInterchangeManager::GetInterchangeManager().ConvertImportData(InterchangeAssetImportData, UFbxStaticMeshImportData::StaticClass(), reinterpret_cast<UObject**>(&ImportData));
 				}
 			}
 
@@ -197,41 +194,37 @@ namespace FbxMeshUtils
 			return false;
 		}
 
-		//We will use interchange only if interchange is enabled and the mesh we want to add a LOD was imported with interchange
-		const UInterchangeAssetImportData* SelectedInterchangeAssetImportData = Cast<UInterchangeAssetImportData>(BaseStaticMesh->GetAssetImportData());
-		if (UInterchangeManager::IsInterchangeImportEnabled() && SelectedInterchangeAssetImportData)
+		UInterchangeManager& InterchangeManager = UInterchangeManager::GetInterchangeManager();
+		const UInterchangeSourceData* SourceData = InterchangeManager.CreateSourceData(Filename);
+		const bool bInterchangeCanImportSourceData = InterchangeManager.CanTranslateSourceData(SourceData);
+		if (bInterchangeCanImportSourceData)
 		{
-			UInterchangeSourceData* SourceData = UInterchangeManager::GetInterchangeManager().CreateSourceData(Filename);
-			//Only use interchange if we can translate the source data. Use the legacy path in other cases
-			if (UInterchangeManager::GetInterchangeManager().CanTranslateSourceData(SourceData))
-			{
-				//Call interchange mesh utilities to import custom LOD
-				UInterchangeMeshUtilities::ImportCustomLod(BaseStaticMesh, LODLevel, SourceData).Then([BaseStaticMesh, LODLevel](TFuture<bool> Result)
-					{
-						bool bResult = Result.Get();
-						Async(EAsyncExecution::TaskGraphMainThread, [BaseStaticMesh, LODLevel, bResult]()
+			//Call interchange mesh utilities to import custom LOD
+			UInterchangeMeshUtilities::ImportCustomLod(BaseStaticMesh, LODLevel, SourceData).Then([BaseStaticMesh, LODLevel](TFuture<bool> Result)
+				{
+					bool bResult = Result.Get();
+					Async(EAsyncExecution::TaskGraphMainThread, [BaseStaticMesh, LODLevel, bResult]()
+						{
+							if (bResult)
 							{
-								if (bResult)
-								{
-									// Notification of success
-									FNotificationInfo NotificationInfo(FText::GetEmpty());
-									NotificationInfo.Text = FText::Format(NSLOCTEXT("UnrealEd", "StaticMeshLODImportSuccessful", "Static mesh LOD {0} imported successfully!"), FText::AsNumber(LODLevel));
-									NotificationInfo.ExpireDuration = 5.0f;
-									FSlateNotificationManager::Get().AddNotification(NotificationInfo);
-								}
-								else
-								{
-									// Notification of failure
-									FNotificationInfo NotificationInfo(FText::GetEmpty());
-									NotificationInfo.Text = FText::Format(NSLOCTEXT("UnrealEd", "StaticMeshLODImportFail", "Failed to import static mesh LOD {0}!"), FText::AsNumber(LODLevel));
-									NotificationInfo.ExpireDuration = 5.0f;
-									FSlateNotificationManager::Get().AddNotification(NotificationInfo);
-								}
-							});
-					});
+								// Notification of success
+								FNotificationInfo NotificationInfo(FText::GetEmpty());
+								NotificationInfo.Text = FText::Format(NSLOCTEXT("UnrealEd", "StaticMeshLODImportSuccessful", "Static mesh LOD {0} imported successfully!"), FText::AsNumber(LODLevel));
+								NotificationInfo.ExpireDuration = 5.0f;
+								FSlateNotificationManager::Get().AddNotification(NotificationInfo);
+							}
+							else
+							{
+								// Notification of failure
+								FNotificationInfo NotificationInfo(FText::GetEmpty());
+								NotificationInfo.Text = FText::Format(NSLOCTEXT("UnrealEd", "StaticMeshLODImportFail", "Failed to import static mesh LOD {0}!"), FText::AsNumber(LODLevel));
+								NotificationInfo.ExpireDuration = 5.0f;
+								FSlateNotificationManager::Get().AddNotification(NotificationInfo);
+							}
+						});
+				});
 
-				return true;
-			}
+			return true;
 		}
 
 		bool bSuccess = false;
@@ -373,25 +366,28 @@ namespace FbxMeshUtils
 			return false;
 		}
 
-		//We will use interchange only if interchange is enabled and the mesh we want to add a LOD was imported with interchange
-		const UInterchangeAssetImportData* SelectedInterchangeAssetImportData = Cast<UInterchangeAssetImportData>(BaseStaticMesh->GetAssetImportData());
-		if (UInterchangeManager::IsInterchangeImportEnabled() && SelectedInterchangeAssetImportData)
+		UInterchangeManager& InterchangeManager = UInterchangeManager::GetInterchangeManager();
+		const UInterchangeSourceData* SourceData = InterchangeManager.CreateSourceData(Filename);
+		const bool bInterchangeCanImportSourceData = InterchangeManager.CanTranslateSourceData(SourceData);
+		if (bInterchangeCanImportSourceData)
 		{
 			UStaticMesh* TempStaticMesh = NewObject<UStaticMesh>(GetTransientPackage(), NAME_None, RF_Transient | RF_Public | RF_Standalone);
 			TempStaticMesh->AddSourceModel();
-
-			// Since it is Async, any action on the mesh should be locked
-
-			UInterchangeSourceData* SourceData = UInterchangeManager::GetInterchangeManager().CreateSourceData(Filename);
+			//Set the asset import data to pass the correct import options
+			TempStaticMesh->SetAssetImportData(BaseStaticMesh->GetAssetImportData());
 			//Call interchange mesh utilities to import custom LOD
-			UInterchangeMeshUtilities::ImportCustomLod(TempStaticMesh, 0, SourceData).Then([BaseStaticMesh,TempStaticMesh](TFuture<bool> Result)
+			UInterchangeMeshUtilities::ImportCustomLod(TempStaticMesh, 0, SourceData).Then([BaseStaticMesh,TempStaticMesh, Filename](TFuture<bool> Result)
 			{
 				bool bResult = Result.Get();
-				Async(EAsyncExecution::TaskGraphMainThread, [BaseStaticMesh, TempStaticMesh, bResult]()
+				Async(EAsyncExecution::TaskGraphMainThread, [BaseStaticMesh, TempStaticMesh, Filename, bResult]()
 				{
 					// Copy high res mesh from temporary static mesh to targeted one
 					if (bResult && Private::CopyHighResMeshDescription(TempStaticMesh, BaseStaticMesh))
 					{
+						FStaticMeshSourceModel& SourceModel = BaseStaticMesh->GetHiResSourceModel();
+						SourceModel.SourceImportFilename = UAssetImportData::SanitizeImportFilename(Filename, nullptr);
+						SourceModel.bImportWithBaseMesh = false;
+
 						// Notification of success
 						FNotificationInfo NotificationInfo(FText::GetEmpty());
 						NotificationInfo.Text = NSLOCTEXT("UnrealEd", "ImportStaticMeshHiResSourceModelSuccessful", "High res mesh imported successfully!");
@@ -426,12 +422,10 @@ namespace FbxMeshUtils
 		UnFbx::FFbxLoggerSetter Logger(FFbxImporter);
 
 		UnFbx::FBXImportOptions* ImportOptions = FFbxImporter->GetImportOptions();
-		Private::SetupFbxImportOptions(BaseStaticMesh, ImportOptions);
+		UFbxStaticMeshImportData* ImportData = Private::SetupFbxImportOptions(BaseStaticMesh, ImportOptions);
 		ImportOptions->StaticMeshLODGroup = NAME_None;
 		ImportOptions->bImportLOD = false;
 
-		UFbxStaticMeshImportData* ImportData = Cast<UFbxStaticMeshImportData>(BaseStaticMesh->AssetImportData);
-
 		const bool bPreventMaterialNameClash = true;
 		if (!FFbxImporter->ImportFromFile(*Filename, FPaths::GetExtension(Filename), bPreventMaterialNameClash))
 		{
@@ -621,40 +615,30 @@ namespace FbxMeshUtils
 			UnFbx::FBXImportOptions* ImportOptions = FFbxImporter->GetImportOptions();
 
 			//Set the skeletal mesh import data from the base mesh, this make sure the import rotation transform is use when importing a LOD
-			UFbxSkeletalMeshImportData* TempAssetImportData = NULL;
-
-			UFbxAssetImportData* FbxAssetImportData = Cast<UFbxAssetImportData>(SelectedSkelMesh->GetAssetImportData());
-			if (!FbxAssetImportData)
+			UFbxSkeletalMeshImportData* FbxImportData = Cast<UFbxSkeletalMeshImportData>(SelectedSkelMesh->GetAssetImportData());
+			if (!FbxImportData)
 			{
-				//Convert the data if its Interchange import data, we convert the full asset in this case since legacy fbx workflow will create one if missing
+				//Convert the data if its Interchange import data
 				if (UInterchangeAssetImportData* InterchangeAssetImportData = Cast<UInterchangeAssetImportData>(SelectedSkelMesh->GetAssetImportData()))
 				{
-					if (InterchangeManager.ConvertImportData(SelectedSkelMesh, TEXT("fbx")))
-					{
-						FbxAssetImportData = Cast<UFbxAssetImportData>(SelectedSkelMesh->GetAssetImportData());
-					}
+					InterchangeManager.ConvertImportData(InterchangeAssetImportData, UFbxSkeletalMeshImportData::StaticClass(), reinterpret_cast<UObject**>(&FbxImportData));
 				}
 			}
 
-			if (FbxAssetImportData)
+			if (FbxImportData)
 			{
-				UFbxSkeletalMeshImportData* ImportData = Cast<UFbxSkeletalMeshImportData>(FbxAssetImportData);
-				if (ImportData)
-				{
-					TempAssetImportData = ImportData;
-					UnFbx::FBXImportOptions::ResetOptions(ImportOptions);
-					// Prepare the import options
-					UFbxImportUI* ReimportUI = NewObject<UFbxImportUI>();
-					ReimportUI->MeshTypeToImport = FBXIT_SkeletalMesh;
-					ReimportUI->Skeleton = SelectedSkelMesh->GetSkeleton();
-					ReimportUI->PhysicsAsset = SelectedSkelMesh->GetPhysicsAsset();
-					// Import data already exists, apply it to the fbx import options
-					ReimportUI->SkeletalMeshImportData = ImportData;
-					//Some options not supported with skeletal mesh
-					ReimportUI->SkeletalMeshImportData->bBakePivotInVertex = false;
-					ReimportUI->SkeletalMeshImportData->bTransformVertexToAbsolute = true;
-					ApplyImportUIToImportOptions(ReimportUI, *ImportOptions);
-				}
+				UnFbx::FBXImportOptions::ResetOptions(ImportOptions);
+				// Prepare the import options
+				UFbxImportUI* ReimportUI = NewObject<UFbxImportUI>();
+				ReimportUI->MeshTypeToImport = FBXIT_SkeletalMesh;
+				ReimportUI->Skeleton = SelectedSkelMesh->GetSkeleton();
+				ReimportUI->PhysicsAsset = SelectedSkelMesh->GetPhysicsAsset();
+				// Import data already exists, apply it to the fbx import options
+				ReimportUI->SkeletalMeshImportData = FbxImportData;
+				//Some options not supported with skeletal mesh
+				ReimportUI->SkeletalMeshImportData->bBakePivotInVertex = false;
+				ReimportUI->SkeletalMeshImportData->bTransformVertexToAbsolute = true;
+				ApplyImportUIToImportOptions(ReimportUI, *ImportOptions);
 				ImportOptions->bImportMaterials = false;
 				ImportOptions->bImportTextures = false;
 			}
@@ -816,7 +800,7 @@ namespace FbxMeshUtils
 					ImportSkeletalMeshArgs.NodeArray = bUseLODs ? SkelMeshNodeArray : *MeshObject;
 					ImportSkeletalMeshArgs.Name = NAME_None;
 					ImportSkeletalMeshArgs.Flags = RF_Transient;
-					ImportSkeletalMeshArgs.TemplateImportData = TempAssetImportData;
+					ImportSkeletalMeshArgs.TemplateImportData = FbxImportData;
 					ImportSkeletalMeshArgs.LodIndex = SelectedLOD;
 					ImportSkeletalMeshArgs.OrderedMaterialNames = OrderedMaterialNames.Num() > 0 ? &OrderedMaterialNames : nullptr;
 					ImportSkeletalMeshArgs.ImportMaterialOriginalNameData = &ImportMaterialOriginalNameData;
@@ -825,7 +809,7 @@ namespace FbxMeshUtils
 
 					TempSkelMesh = (USkeletalMesh*)FFbxImporter->ImportSkeletalMesh(ImportSkeletalMeshArgs);
 					// Add the new imported LOD to the existing model (check skeleton compatibility)
-					if (TempSkelMesh && FFbxImporter->ImportSkeletalMeshLOD(TempSkelMesh, SelectedSkelMesh, SelectedLOD, TempAssetImportData))
+					if (TempSkelMesh && FFbxImporter->ImportSkeletalMeshLOD(TempSkelMesh, SelectedSkelMesh, SelectedLOD, FbxImportData))
 					{
 						//Update the import data for this lod
 						UnFbx::FFbxImporter::UpdateSkeletalMeshImportData(SelectedSkelMesh, nullptr, SelectedLOD, &ImportMaterialOriginalNameData, &ImportMeshLodData);
@@ -1022,7 +1006,9 @@ namespace FbxMeshUtils
 				FInterchangeFilePickerParameters Parameters;
 				Parameters.bAllowMultipleFiles = false;
 				Parameters.Title = FText::Format(NSLOCTEXT("Interchange", "ImportCustomLodAsync_FilePickerTitle", "Choose a file to import a custom LOD for LOD{0}"), FText::AsNumber(LODLevel));
-				Parameters.bShowAllFactoriesExtension = true;
+				Parameters.bShowAllFactoriesExtension = false;
+				//Lod import support interchange format and we force fbx in case it is disabled
+				Parameters.ExtraFormats = { TEXT("fbx;Filmbox") };
 				TArray<FString> Filenames;
 				bool bFilePickerResult = FilePicker->ScriptedFilePickerForTranslatorAssetType(EInterchangeTranslatorAssetType::Meshes, Parameters, Filenames);
 				if (bFilePickerResult)
@@ -1057,7 +1043,6 @@ namespace FbxMeshUtils
 
 		//Convert the import data if necessary
 		UInterchangeAssetImportData* SelectedInterchangeAssetImportData = nullptr;
-		SelectedInterchangeAssetImportData = nullptr;
 		if (SkeletalMesh)
 		{
 			if (SkeletalMesh->IsValidLODIndex(LODLevel))
diff --git a/Engine/Source/Runtime/Interchange/Engine/Private/InterchangeManager.cpp b/Engine/Source/Runtime/Interchange/Engine/Private/InterchangeManager.cpp
--- a/Engine/Source/Runtime/Interchange/Engine/Private/InterchangeManager.cpp
+++ b/Engine/Source/Runtime/Interchange/Engine/Private/InterchangeManager.cpp
@@ -1058,7 +1058,7 @@ bool UInterchangeManager::ConvertImportData(const UObject* SourceImportData, con
 	{
 		if (RegisteredConverter.Value->CanConvertClass(SourceImportData->GetClass(), DestinationClass))
 		{
-			if (RegisteredConverter.Value->ConvertImportData(SourceImportData, DestinationImportData))
+			if (RegisteredConverter.Value->ConvertImportData(SourceImportData, DestinationClass, DestinationImportData))
 			{
 				break;
 			}
diff --git a/Engine/Source/Runtime/Interchange/Engine/Public/InterchangeAssetImportData.h b/Engine/Source/Runtime/Interchange/Engine/Public/InterchangeAssetImportData.h
--- a/Engine/Source/Runtime/Interchange/Engine/Public/InterchangeAssetImportData.h
+++ b/Engine/Source/Runtime/Interchange/Engine/Public/InterchangeAssetImportData.h
@@ -223,7 +223,7 @@ public:
 	 *
 	 * The system will call all object deriving from this class until one convert the data.
 	 */
-	virtual bool ConvertImportData(const UObject* SourceImportData, UObject** DestinationImportData) const
+	virtual bool ConvertImportData(const UObject* SourceImportData, const UClass* DestinationClass, UObject** DestinationImportData) const
 	{
 		return false;
 	}
diff --git a/Engine/Source/Runtime/Interchange/Engine/Public/InterchangeFilePickerBase.h b/Engine/Source/Runtime/Interchange/Engine/Public/InterchangeFilePickerBase.h
--- a/Engine/Source/Runtime/Interchange/Engine/Public/InterchangeFilePickerBase.h
+++ b/Engine/Source/Runtime/Interchange/Engine/Public/InterchangeFilePickerBase.h
@@ -29,6 +29,10 @@ struct FInterchangeFilePickerParameters
 	//If true, the user will be able to select any unreal editor factory + interchange file types.
 	UPROPERTY(EditAnywhere, Category = "Interchange | File Picker")
 	bool bShowAllFactoriesExtension = false;
+
+	//Add some extension to the picker. Format text item that way TEXT("fbx;Filmbox")
+	UPROPERTY(EditAnywhere, Category = "Interchange | File Picker")
+	TArray<FString> ExtraFormats;
 };
 
 UCLASS(Abstract, BlueprintType, Blueprintable, MinimalAPI)
