Change Camera Ortho Height Adjustment to never be larger than half the Ortho Width, otherwise we overcompensate the camera origin update and potentially clip too soon in top down views. Also fixed ortho wireframe projection clipping too soon + mismatched measurement tool.

diff --git a/Engine/Source/Editor/UnrealEd/Private/EditorViewportClient.cpp b/Engine/Source/Editor/UnrealEd/Private/EditorViewportClient.cpp
--- a/Engine/Source/Editor/UnrealEd/Private/EditorViewportClient.cpp
+++ b/Engine/Source/Editor/UnrealEd/Private/EditorViewportClient.cpp
@@ -1129,21 +1129,9 @@ FSceneView* FEditorViewportClient::CalcSceneView(FSceneViewFamily* ViewFamily, c
 
 			//The divisor for the matrix needs to match the translation code.
 			const float Zoom = GetOrthoUnitsPerPixel(Viewport);
-			const float OrthoZoom = FMath::Clamp(Zoom * 500.0f, 0.0f, UE_LARGE_HALF_WORLD_MAX);
+			float OrthoWidth = FMath::Clamp(Zoom * ViewportSize.X/2.0f, 0.0f, UE_LARGE_HALF_WORLD_MAX);
+			float OrthoHeight = FMath::Clamp(Zoom * ViewportSize.Y/2.0f, 0.0f, UE_LARGE_HALF_WORLD_MAX);
 
-			float ViewportAspectRatio = float(ViewportSize.X) / (ViewportSize.Y > 0 ? ViewportSize.Y : 1);
-
-			float OrthoWidth, OrthoHeight;
-			if(ViewportAspectRatio >= 1.0f)
-			{
-				OrthoWidth = OrthoZoom;
-				OrthoHeight = OrthoWidth / ViewportAspectRatio;
-			}
-			else
-			{
-				OrthoHeight = OrthoZoom;
-				OrthoWidth = OrthoHeight * ViewportAspectRatio;
-			}
 
 			if (EffectiveViewportType == LVT_OrthoXY)
 			{
@@ -1207,21 +1195,27 @@ FSceneView* FEditorViewportClient::CalcSceneView(FSceneViewFamily* ViewFamily, c
 				check(false);
 			}
 
-			FMinimalViewInfo CalculatePlanesViewInfo;
-			CalculatePlanesViewInfo.ProjectionMode = ECameraProjectionMode::Orthographic;
-			CalculatePlanesViewInfo.bAutoCalculateOrthoPlanes = true;
-			CalculatePlanesViewInfo.OrthoWidth = OrthoWidth;
-			CalculatePlanesViewInfo.Rotation = ViewInitOptions.ViewRotationMatrix.Rotator();
-			CalculatePlanesViewInfo.AspectRatio = ViewportAspectRatio;
-			CalculatePlanesViewInfo.bConstrainAspectRatio = false;
-			CalculatePlanesViewInfo.OrthoNearClipPlane = OrthoWidth * -CVarOrthoEditorDebugClipPlaneScale.GetValueOnAnyThread();
-			CalculatePlanesViewInfo.OrthoFarClipPlane = FarPlane - NearPlane + CalculatePlanesViewInfo.OrthoNearClipPlane;
-			CalculatePlanesViewInfo.AutoCalculateOrthoPlanes(ViewInitOptions);
-
-			ViewInitOptions.UpdateOrthoPlanes(CalculatePlanesViewInfo);
-
-			const float ZScale = 1.0f / (CalculatePlanesViewInfo.OrthoFarClipPlane - CalculatePlanesViewInfo.OrthoNearClipPlane);
-			const float ZOffset = -CalculatePlanesViewInfo.OrthoNearClipPlane;
+			FMatrix::FReal ZScale = 0.5f / UE_OLD_WORLD_MAX;
+			FMatrix::FReal ZOffset = UE_OLD_WORLD_MAX;
+			if(!ViewFamily->EngineShowFlags.Wireframe)
+			{
+				FMinimalViewInfo CalculatePlanesViewInfo;
+				CalculatePlanesViewInfo.ProjectionMode = ECameraProjectionMode::Orthographic;
+				CalculatePlanesViewInfo.bAutoCalculateOrthoPlanes = true;
+				CalculatePlanesViewInfo.OrthoWidth = OrthoWidth;
+				CalculatePlanesViewInfo.Rotation = ViewInitOptions.ViewRotationMatrix.Rotator();
+				CalculatePlanesViewInfo.AspectRatio = AspectRatio;
+				CalculatePlanesViewInfo.bConstrainAspectRatio = false;
+				CalculatePlanesViewInfo.OrthoNearClipPlane = OrthoWidth * -CVarOrthoEditorDebugClipPlaneScale.GetValueOnAnyThread();
+				CalculatePlanesViewInfo.OrthoFarClipPlane = FarPlane - NearPlane + CalculatePlanesViewInfo.OrthoNearClipPlane;
+
+				CalculatePlanesViewInfo.AutoCalculateOrthoPlanes(ViewInitOptions);
+				if(ViewInitOptions.UpdateOrthoPlanes(CalculatePlanesViewInfo))
+				{
+					ZScale = 1.0f / (CalculatePlanesViewInfo.OrthoFarClipPlane - CalculatePlanesViewInfo.OrthoNearClipPlane);
+					ZOffset = -CalculatePlanesViewInfo.OrthoNearClipPlane;
+				}
+			}			
 
 			ViewInitOptions.ProjectionMatrix = FReversedZOrthoMatrix(
 				OrthoWidth,
diff --git a/Engine/Source/Runtime/Engine/Private/Camera/CameraStackTypes.cpp b/Engine/Source/Runtime/Engine/Private/Camera/CameraStackTypes.cpp
--- a/Engine/Source/Runtime/Engine/Private/Camera/CameraStackTypes.cpp
+++ b/Engine/Source/Runtime/Engine/Private/Camera/CameraStackTypes.cpp
@@ -302,7 +302,7 @@ void FMinimalViewInfo::CalculateProjectionMatrixGivenViewRectangle(FMinimalViewI
 			float FarPlane = ViewInfo.OrthoFarClipPlane;
 			float NearPlane = ViewInfo.OrthoNearClipPlane;
 
-			InOutProjectionData.UpdateOrthoPlanes(NearPlane, FarPlane);
+			InOutProjectionData.UpdateOrthoPlanes(NearPlane, FarPlane, OrthoWidth);
 
 			const float ZScale = 1.0f / (FarPlane - NearPlane);
 			const float ZOffset = -NearPlane;
diff --git a/Engine/Source/Runtime/Engine/Private/SceneView.cpp b/Engine/Source/Runtime/Engine/Private/SceneView.cpp
--- a/Engine/Source/Runtime/Engine/Private/SceneView.cpp
+++ b/Engine/Source/Runtime/Engine/Private/SceneView.cpp
@@ -570,7 +570,7 @@ FVector4f CreateInvDeviceZToWorldZTransform(const FMatrix& ProjMatrix)
 	}
 }
 
-bool FSceneViewProjectionData::UpdateOrthoPlanes(FSceneViewProjectionData* InOutProjectionData, float& NearPlane, float& FarPlane)
+bool FSceneViewProjectionData::UpdateOrthoPlanes(FSceneViewProjectionData* InOutProjectionData, float& NearPlane, float& FarPlane, float HalfOrthoWidth)
 {
 	if (!InOutProjectionData || !CVarAllowOrthoNearPlaneCorrection.GetValueOnAnyThread())
 	{
@@ -597,7 +597,11 @@ bool FSceneViewProjectionData::UpdateOrthoPlanes(FSceneViewProjectionData* InOut
 		* Use the height of the camera as a sort of pseudo CameraToViewTarget to remove camera position as much as possible so we minimise the near clip plane distance.
 		* Depends on view direction being top down as when we get to a side view, the scale grows much larger and it becomes redundant.
 		*/
-		float CameraHeightAdjustment = CVarOrthoCameraHeightAsViewTarget.GetValueOnAnyThread() ? FMath::Abs(InOutProjectionData->ViewOrigin.Z) * FMath::Abs((ViewForward.Dot(FVector(0, 0, -1.0f)))) : 0.0f;
+		float CameraHeightAdjustment = 0.0f;
+		if (CVarOrthoCameraHeightAsViewTarget.GetValueOnAnyThread())
+		{
+			CameraHeightAdjustment = FMath::Abs(FMath::Min(InOutProjectionData->ViewOrigin.Z, HalfOrthoWidth)) * FMath::Abs((ViewForward.Dot(FVector(0, 0, -1.0f))));
+		}
 		InOutProjectionData->ViewOrigin += ViewForward * (CameraHeightAdjustment + NearPlane);
 	}
 	NearPlane = GDefaultUpdateOrthoNearPlane;
@@ -608,7 +612,7 @@ bool FSceneViewProjectionData::UpdateOrthoPlanes(FSceneViewProjectionData* InOut
 
 bool FSceneViewProjectionData::UpdateOrthoPlanes(FMinimalViewInfo& MinimalViewInfo)
 {
-	return UpdateOrthoPlanes(MinimalViewInfo.OrthoNearClipPlane, MinimalViewInfo.OrthoFarClipPlane);
+	return UpdateOrthoPlanes(MinimalViewInfo.OrthoNearClipPlane, MinimalViewInfo.OrthoFarClipPlane, MinimalViewInfo.OrthoWidth/2.0f);
 }
 
 bool FSceneViewProjectionData::UpdateOrthoPlanes()
@@ -634,7 +638,12 @@ bool FSceneViewProjectionData::UpdateOrthoPlanes()
 		return false;
 	}
 
-	UpdateOrthoPlanes(NearPlane, FarPlane);
+	float HalfInvOrthoWidth = static_cast<float>(ProjectionMatrix.M[0][0]);
+	if (HalfInvOrthoWidth == 0.0f)
+	{
+		return false;
+	}
+	UpdateOrthoPlanes(NearPlane, FarPlane, 1.0f / HalfInvOrthoWidth);
 
 	const float ZScale = 1.0f / (FarPlane - NearPlane);
 	const float ZOffset = -NearPlane;
diff --git a/Engine/Source/Runtime/Engine/Public/SceneView.h b/Engine/Source/Runtime/Engine/Public/SceneView.h
--- a/Engine/Source/Runtime/Engine/Public/SceneView.h
+++ b/Engine/Source/Runtime/Engine/Public/SceneView.h
@@ -128,10 +128,10 @@ public:
 	}
 
 	// Function for correcting Ortho camera near plane locations to avoid artifacts behind camera view origin
-	static ENGINE_API bool UpdateOrthoPlanes(FSceneViewProjectionData* InOutProjectionData, float& NearPlane, float& FarPlane);
-	ENGINE_API inline bool UpdateOrthoPlanes(float& NearPlane, float& FarPlane)
+	static ENGINE_API bool UpdateOrthoPlanes(FSceneViewProjectionData* InOutProjectionData, float& NearPlane, float& FarPlane, float HalfOrthoWidth);
+	ENGINE_API inline bool UpdateOrthoPlanes(float& NearPlane, float& FarPlane, float HalfOrthoWidth)
 	{
-		return UpdateOrthoPlanes(this, NearPlane, FarPlane);
+		return UpdateOrthoPlanes(this, NearPlane, FarPlane, HalfOrthoWidth);
 	}
 	ENGINE_API bool UpdateOrthoPlanes(FMinimalViewInfo& MinimalViewInfo);
 	ENGINE_API bool UpdateOrthoPlanes();
