MPCook And CookPackageSplitter: Implement on demand creation and splitting of generator packages rather than only creating them during save.
#rnx

diff --git a/Engine/Config/BaseEditor.ini b/Engine/Config/BaseEditor.ini
--- a/Engine/Config/BaseEditor.ini
+++ b/Engine/Config/BaseEditor.ini
@@ -395,6 +395,10 @@ CookProcessCount=1
 ; A negative value is ignored and CookProcessCount is used instead.
 CookProcessCountFromEditor=-1
 
+; Debug setting to specify that the packages generated by a generator package must be saved on the same cookworker
+; that saved the generator package. This is bad for performance, but is provided in case there is a bug with moving
+; them to other workers. Values = AnyWorker, AllOnSameWorker, SomeOnSameWorker, NoneOnSameWorker
+MPCookGeneratorSplit=AllOnSameWorker
 
 +ExcludedIniParams=Game:/Script/BlueprintContext.McpContext
 +ExcludedIniParams=Editor:CookSettings
diff --git a/Engine/Source/Editor/UnrealEd/Classes/CookOnTheSide/CookOnTheFlyServer.h b/Engine/Source/Editor/UnrealEd/Classes/CookOnTheSide/CookOnTheFlyServer.h
--- a/Engine/Source/Editor/UnrealEd/Classes/CookOnTheSide/CookOnTheFlyServer.h
+++ b/Engine/Source/Editor/UnrealEd/Classes/CookOnTheSide/CookOnTheFlyServer.h
@@ -153,6 +153,18 @@ enum class ECookTickFlags : uint8
 };
 ENUM_CLASS_FLAGS(ECookTickFlags);
 
+namespace UE::Cook
+{
+/** MPCook Behavior set from config/commandline that decides where generatedpackages should be assigned. */
+enum class EMPCookGeneratorSplit : uint8
+{
+	AnyWorker,
+	AllOnSameWorker,
+	SomeOnSameWorker,
+	NoneOnSameWorker,
+};
+}
+
 namespace UE::Cook
 {
 	class FAssetRegistryMPCollector;
@@ -1342,13 +1354,11 @@ private:
 	/** Generates long package names for all files to be cooked */
 	void GenerateLongPackageNames(TArray<FName>& FilesInPath, TMap<FName, UE::Cook::FInstigator>& Instigators);
 
-	UE::Cook::EPollStatus ConditionalCreateGeneratorPackage(UE::Cook::FPackageData& PackageData, bool bPrecaching);
-
 	/** Generate the list of cook-time-created packages created by the generator package. */
 	UE::Cook::EPollStatus QueueGeneratedPackages(UE::Cook::FGenerationHelper& GenerationHelper,
 		UE::Cook::FPackageData& PackageData);
-	/** Run additional steps in PrepareSave that are required when the package has a GenerationHelper. */
-	UE::Cook::EPollStatus PrepareSaveGeneratedPackage(UE::Cook::FGenerationHelper& GenerationHelper,
+	/** Run additional steps in PrepareSave required when the package is a Generator or a GeneratedPackage. */
+	UE::Cook::EPollStatus PrepareSaveGenerationPackage(UE::Cook::FGenerationHelper& GenerationHelper,
 		UE::Cook::FPackageData& PackageData, UE::Cook::FCookerTimer& Timer, bool bPrecaching);
 	/**
 	 * Call BeginCacheForCookedPlatformData on objects the CookPackageSplitter plans to move
@@ -1466,6 +1476,7 @@ private:
 	TSet<FName> CookFilterIncludedAssetClasses;
 
 	ELogVerbosity::Type CookerIdleWarningSeverity = ELogVerbosity::Warning;
+	UE::Cook::EMPCookGeneratorSplit MPCookGeneratorSplit = UE::Cook::EMPCookGeneratorSplit::AnyWorker;
 	/** True when PumpLoads has detected it is blocked on async work and CookOnTheFlyServer should do work elsewhere. */
 	bool bLoadBusy = false;
 	/** True when PumpSaves has detected it is blocked on async work and CookOnTheFlyServer should do work elsewhere. */
diff --git a/Engine/Source/Editor/UnrealEd/Private/CookOnTheFlyServer.cpp b/Engine/Source/Editor/UnrealEd/Private/CookOnTheFlyServer.cpp
--- a/Engine/Source/Editor/UnrealEd/Private/CookOnTheFlyServer.cpp
+++ b/Engine/Source/Editor/UnrealEd/Private/CookOnTheFlyServer.cpp
@@ -2629,6 +2629,10 @@ void UCookOnTheFlyServer::PumpRequests(UE::Cook::FTickStackData& StackData, int3
 		++NumInBatch;
 	}
 	OutNumPushed += NumInBatch;
+	if (DiscoveryQueue.IsEmpty() && RequestClusters.IsEmpty() && RestartedRequests.IsEmpty())
+	{
+		RequestQueue.NotifyRequestFencePassed(*PackageDatas);
+	}
 }
 
 void UCookOnTheFlyServer::AssignRequests(TArrayView<UE::Cook::FPackageData*> Requests, UE::Cook::FRequestQueue& RequestQueue,
@@ -2881,12 +2885,17 @@ void UCookOnTheFlyServer::LoadPackageInQueue(UE::Cook::FPackageData& PackageData
 	}
 	else
 	{
-		FGenerationHelper* GenerationHelper = PackageData.GetGeneratedOwner();
+		TRefCountPtr<FGenerationHelper> GenerationHelper = PackageData.TryCreateValidParentGenerationHelper();
 		if (!GenerationHelper)
 		{
 			UE_LOG(LogCook, Error,
-				TEXT("Package %s is an out-of-date generated package with a no-longer-available generator. It can not be loaded."),
-				*PackageFileName.ToString());
+				TEXT("Package %s is a generated package, but ParentGenerator '%s' is not a generator package. The generated package cannot be loaded."),
+				*PackageFileName.ToString(), *PackageData.GetParentGenerator().ToString());
+			RejectPackageToLoad(PackageData, TEXT("is an orphaned generated package"), ESuppressCookReason::OrphanedGenerated);
+			return;
+		}
+		if (!GenerationHelper->TryGenerateList())
+		{
 			RejectPackageToLoad(PackageData, TEXT("is an orphaned generated package"), ESuppressCookReason::OrphanedGenerated);
 			return;
 		}
@@ -2894,42 +2903,12 @@ void UCookOnTheFlyServer::LoadPackageInQueue(UE::Cook::FPackageData& PackageData
 		if (!Info)
 		{
 			UE_LOG(LogCook, Error,
-				TEXT("Package %s is a generated package but its generator no longer has a record of it. It can not be loaded."),
+				TEXT("Package %s is a generated package but its generator does not have a record of it. It can not be loaded."),
 				*PackageFileName.ToString());
 			RejectPackageToLoad(PackageData, TEXT("is an orphaned generated package"), ESuppressCookReason::OrphanedGenerated);
 			return;
 		}
 
-		FPackageData& OwnerPackageData = GenerationHelper->GetOwner();
-		UPackage* OwnerPackage = GenerationHelper->GetOwnerPackage();
-		if (!OwnerPackage)
-		{
-			OwnerPackage = FindObject<UPackage>(nullptr, *OwnerPackageData.GetPackageName().ToString());
-		}
-		if (!OwnerPackage || !OwnerPackage->IsFullyLoaded())
-		{
-			bool bLoadFullySuccessful = LoadPackageForCooking(OwnerPackageData, OwnerPackage, &PackageData);
-			UObject* SplitterDataObject = nullptr;
-			if (bLoadFullySuccessful)
-			{
-				SplitterDataObject = GenerationHelper->FindSplitDataObject();
-			}
-			if (!SplitterDataObject)
-			{
-				ResultFlags |= COSR_ErrorLoadingPackage;
-				UE_LOG(LogCook, Error,
-					TEXT("Package %s is a generated package and we could not load its generator package %s. It can not be loaded."),
-					*PackageFileName.ToString(), *OwnerPackageData.GetFileName().ToString());
-				RejectPackageToLoad(PackageData, TEXT("is a generated package which could not load its generator"),
-					ESuppressCookReason::LoadError);
-				return;
-			}
-			FScopedActivePackage ScopedActivePackage(*this, OwnerPackageData.GetPackageName(),
-				PackageAccessTrackingOps::NAME_CookerBuildObject);
-			GenerationHelper->GetCookPackageSplitterInstance()->OnOwnerReloaded(OwnerPackage, SplitterDataObject);
-			GenerationHelper->SetOwnerPackage(OwnerPackage);
-		}
-
 		LoadedPackage = TryCreateGeneratedPackage(*GenerationHelper, *Info);
 		if (!LoadedPackage)
 		{
@@ -3270,121 +3249,12 @@ void UCookOnTheFlyServer::TickNetwork()
 	}
 }
 
-UE::Cook::EPollStatus UCookOnTheFlyServer::ConditionalCreateGeneratorPackage(UE::Cook::FPackageData& PackageData, bool bPrecaching)
-{
-	using namespace UE::Cook;
-
-	Private::FRegisteredCookPackageSplitter* Splitter = nullptr;
-	UObject* SplitDataObject = nullptr;
-	bool bGeneratorExists = false;
-	bool bIncomplete = false;
-	ON_SCOPE_EXIT
-	{
-		if (!bGeneratorExists && !bIncomplete)
-		{
-			// Destroy any old GenerationHelper if we no longer find we need one
-			PackageData.DestroyGenerationHelper();
-		}
-	};
-
-	TArray<Private::FRegisteredCookPackageSplitter*> FoundRegisteredSplitters;
-
-	for (FCachedObjectInOuter& CachedObjectInOuter : PackageData.GetCachedObjectsInOuter())
-	{
-		UObject* Obj = CachedObjectInOuter.Object.Get();
-		if (!Obj)
-		{
-			continue;
-		}
-	
-		FoundRegisteredSplitters.Reset();
-		RegisteredSplitDataClasses.MultiFind(Obj->GetClass(), FoundRegisteredSplitters);
-
-		for (Private::FRegisteredCookPackageSplitter* SplitterForObject: FoundRegisteredSplitters)
-		{
-			if (SplitterForObject && SplitterForObject->ShouldSplitPackage(Obj))
-			{
-				if (!Obj->HasAnyFlags(RF_Public))
-				{
-					UE_LOG(LogCook, Error, TEXT("SplitterData object %s must be publicly referenceable so we can keep them from being garbage collected"), *Obj->GetFullName());
-					return EPollStatus::Error;
-				}
-
-				if (Splitter)
-				{
-					UE_LOG(LogCook, Error, TEXT("Found more than one registered Cook Package Splitter for package %s."), *PackageData.GetPackageName().ToString());
-					return EPollStatus::Error;
-				}
-
-				Splitter = SplitterForObject;
-				SplitDataObject = Obj;
-			}
-		}
-	}
-	if (!Splitter)
-	{
-		return EPollStatus::Success;
-	}
-
-	if (bPrecaching)
-	{
-		bIncomplete = true;
-		return EPollStatus::Incomplete;
-	}
-
-	// TODO: Add support for cooking in the editor. Possibly moot since we plan to deprecate cooking in the editor.
-	if (IsCookingInEditor())
-	{
-		// CookPackageSplitters allow destructive changes to the generator package. e.g. moving UObjects out
-		// of it into the streaming packages. To allow its use in the editor, we will need to make it non-destructive
-		// (by e.g. copying to new packages), or restore the package after the changes have been made.
-		UE_LOG(LogCook, Error, TEXT("Cooking in editor doesn't support Cook Package Splitters."));
-		return EPollStatus::Error;
-	}
-
-	UE_LOG(LogCook, Display, TEXT("Splitting Package %s with splitter %s acting on object %s."),
-		*PackageData.GetPackageName().ToString(), *Splitter->GetSplitterDebugName(), *SplitDataObject->GetFullName());
-
-	// Create instance of CookPackageSplitter class
-	ICookPackageSplitter* SplitterInstance = Splitter->CreateInstance(SplitDataObject);
-	if (!SplitterInstance)
-	{
-		UE_LOG(LogCook, Error, TEXT("Error instantiating Cook Package Splitter %s for object %s."),
-			*Splitter->GetSplitterDebugName(), *SplitDataObject->GetFullName());
-		return EPollStatus::Error;
-	}
-
-	// Create an FGenerationHelper struct using this CookPackageSplitter instance
-	bGeneratorExists = true;
-	PackageData.CreateGenerationHelper(SplitDataObject, SplitterInstance);
-	return EPollStatus::Success;
-}
-
 UE::Cook::EPollStatus UCookOnTheFlyServer::QueueGeneratedPackages(UE::Cook::FGenerationHelper& GenerationHelper,
 	UE::Cook::FPackageData& PackageData)
 {
 	using namespace UE::Cook;
 
-	ICookPackageSplitter* Splitter = GenerationHelper.GetCookPackageSplitterInstance();
-	UObject* SplitObject = GenerationHelper.FindSplitDataObject();
 	FCookGenerationInfo& Info = GenerationHelper.GetOwnerInfo();
-	if (!SplitObject)
-	{
-		UE_LOG(LogCook, Error, TEXT("Could not find SplitDataObject %s"),
-			*GenerationHelper.GetSplitDataObjectName().ToString());
-		return EPollStatus::Error;
-	}
-
-	if (Info.GetSaveState() <= FCookGenerationInfo::ESaveState::GenerateList)
-	{
-		// Call the splitter to generate the list
-		if (!GenerationHelper.TryGenerateList(SplitObject, *PackageDatas))
-		{
-			return EPollStatus::Error;
-		}
-		GenerationHelper.SetOwnerPackage(PackageData.GetPackage());
-		Info.SetSaveStateComplete(FCookGenerationInfo::ESaveState::GenerateList);
-	}
 
 	if (Info.GetSaveState() <= FCookGenerationInfo::ESaveState::ClearOldPackagesLastAttempt)
 	{
@@ -3416,6 +3286,7 @@ UE::Cook::EPollStatus UCookOnTheFlyServer::QueueGeneratedPackages(UE::Cook::FGen
 	FName OwnerName = Owner->GetFName();
 	if (Info.GetSaveState() <= FCookGenerationInfo::ESaveState::QueueGeneratedPackages)
 	{
+		GenerationHelper.StartQueueGeneratedPackages(*this);
 		TArray<const ITargetPlatform*, TInlineAllocator<ExpectedMaxNumPlatforms>> ReachablePlatforms;
 		PackageData.GetReachablePlatforms(ReachablePlatforms);
 		for (const FCookGenerationInfo& ChildInfo: GenerationHelper.GetPackagesToGenerate())
@@ -3431,12 +3302,14 @@ UE::Cook::EPollStatus UCookOnTheFlyServer::QueueGeneratedPackages(UE::Cook::FGen
 			QueueDiscoveredPackage(*ChildPackageData, FInstigator(ChildPackageData->GetInstigator()),
 				EDiscoveredPlatformSet::CopyFromInstigator, bUrgent);
 		}
+		GenerationHelper.EndQueueGeneratedPackages(*this);
+
 		Info.SetSaveStateComplete(FCookGenerationInfo::ESaveState::QueueGeneratedPackages);
 	}
 	return EPollStatus::Success;
 }
 
-UE::Cook::EPollStatus UCookOnTheFlyServer::PrepareSaveGeneratedPackage(UE::Cook::FGenerationHelper& GenerationHelper,
+UE::Cook::EPollStatus UCookOnTheFlyServer::PrepareSaveGenerationPackage(UE::Cook::FGenerationHelper& GenerationHelper,
 	UE::Cook::FPackageData& PackageData, UE::Cook::FCookerTimer& Timer, bool bPrecaching)
 {
 	using namespace UE::Cook;
@@ -3452,6 +3325,8 @@ UE::Cook::EPollStatus UCookOnTheFlyServer::PrepareSaveGeneratedPackage(UE::Cook:
 
 	if (Info.GetSaveState() <= FCookGenerationInfo::ESaveState::FinishCachePreMove)
 	{
+		// Both Generator packages and Generated packages should wait for all IsCachedCookedPlatformData
+		// to finish before they start BeginCache calls on the objects to move.
 		if (PackageData.GetNumPendingCookedPlatformData() > 0)
 		{
 			return EPollStatus::Incomplete;
@@ -3488,7 +3363,7 @@ UE::Cook::EPollStatus UCookOnTheFlyServer::PrepareSaveGeneratedPackage(UE::Cook:
 		{
 			// Call this function recursively to reexecute CallBeginCacheOnObjects in BeginCacheObjectsToMove.
 			// Note that RefreshPackageObjects checked for too many recursive calls and ErrorExited if so.
-			return PrepareSaveGeneratedPackage(GenerationHelper, PackageData, Timer, bPrecaching);
+			return PrepareSaveGenerationPackage(GenerationHelper, PackageData, Timer, bPrecaching);
 		}
 		Info.SetSaveStateComplete(FCookGenerationInfo::ESaveState::FinishCacheObjectsToMove);
 	}
@@ -3498,7 +3373,7 @@ UE::Cook::EPollStatus UCookOnTheFlyServer::PrepareSaveGeneratedPackage(UE::Cook:
 		if (bPrecaching)
 		{
 			// We're not allowed to populate when precaching, because we want to avoid 
-			// garbagecollection in between Populating and PostSaving the populates package,
+			// garbagecollection in between Populating and PostSaving the populated package,
 			// so we need to not Populate until we're ready to save
 			return EPollStatus::Incomplete;
 		}
@@ -3546,7 +3421,7 @@ UE::Cook::EPollStatus UCookOnTheFlyServer::PrepareSaveGeneratedPackage(UE::Cook:
 		{
 			// Call this function recursively to reexecute CallBeginCacheOnObjects in BeginCachePostMove
 			// Note that RefreshPackageObjects checked for too many recursive calls and ErrorExited if so.
-			return PrepareSaveGeneratedPackage(GenerationHelper, PackageData, Timer, bPrecaching);
+			return PrepareSaveGenerationPackage(GenerationHelper, PackageData, Timer, bPrecaching);
 		}
 
 		Info.SetSaveStateComplete(FCookGenerationInfo::ESaveState::FinishCachePostMove);
@@ -3566,7 +3441,7 @@ UE::Cook::EPollStatus UCookOnTheFlyServer::BeginCacheObjectsToMove(UE::Cook::FGe
 	FPackageData& PackageData(*Info.PackageData);
 	UPackage* Package = PackageData.GetPackage();
 	ICookPackageSplitter* Splitter = GenerationHelper.GetCookPackageSplitterInstance();
-	UObject* SplitDataObject = GenerationHelper.FindSplitDataObject();
+	UObject* SplitDataObject = GenerationHelper.FindOrLoadSplitDataObject();
 	if (!Package || !Splitter || !SplitDataObject)
 	{
 		UE_LOG(LogCook, Error,
@@ -3631,7 +3506,7 @@ UE::Cook::EPollStatus UCookOnTheFlyServer::PreSaveGeneratorPackage(UE::Cook::FPa
 
 	UPackage* Package = PackageData.GetPackage();
 	ICookPackageSplitter* Splitter = GenerationHelper.GetCookPackageSplitterInstance();
-	UObject* SplitDataObject = GenerationHelper.FindSplitDataObject();
+	UObject* SplitDataObject = GenerationHelper.FindOrLoadSplitDataObject();
 	if (!Package || !Splitter || !SplitDataObject)
 	{
 		UE_LOG(LogCook, Error, TEXT("CookPackageSplitter is missing %s during PreSaveGeneratorPackage. PackageName: %s."),
@@ -3700,7 +3575,7 @@ UE::Cook::EPollStatus UCookOnTheFlyServer::BeginCachePostMove(UE::Cook::FGenerat
 	UE::Cook::FPackageData& PackageData(*Info.PackageData);
 	UPackage* Package = PackageData.GetPackage();
 	ICookPackageSplitter* Splitter = GenerationHelper.GetCookPackageSplitterInstance();
-	UObject* SplitDataObject = GenerationHelper.FindSplitDataObject();
+	UObject* SplitDataObject = GenerationHelper.FindOrLoadSplitDataObject();
 	if (!Package || !Splitter || !SplitDataObject)
 	{
 		UE_LOG(LogCook, Error, TEXT("CookPackageSplitter is missing %s during BeginCachePostMove. PackageName: %s."),
@@ -3765,15 +3640,18 @@ UE::Cook::EPollStatus UCookOnTheFlyServer::BeginCachePostMove(UE::Cook::FGenerat
 UPackage* UCookOnTheFlyServer::TryCreateGeneratedPackage(UE::Cook::FGenerationHelper& GenerationHelper, UE::Cook::FCookGenerationInfo& Info)
 {
 	using namespace UE::Cook;
-	// Caller is responsible for validating OwnerPackage and Generated PackageData
 	check(Info.PackageData); // Caller is responsible for validating
 	UE::Cook::FPackageData& GeneratedPackageData = *Info.PackageData;
-	UPackage* OwnerPackage = GenerationHelper.GetOwnerPackage();
-	check(OwnerPackage); // Caller is responsible for validating
-
 	const FString GeneratedPackageName = GeneratedPackageData.GetPackageName().ToString();
+	UPackage* OwnerPackage = GenerationHelper.FindOrLoadOwnerPackage(*this);
+	if (!OwnerPackage)
+	{
+		UE_LOG(LogCook, Error, TEXT("TryCreateGeneratedPackage: could not load ParentGeneratorPackage %s for GeneratedPackage %s"),
+			*GenerationHelper.GetOwner().GetPackageName().ToString(), *GeneratedPackageName);
+		return nullptr;
+	}
+
 	UPackage* GeneratedPackage = FindObject<UPackage>(nullptr, *GeneratedPackageName);
-	bool bPopulatedByPreSave = false;
 	if (GeneratedPackage)
 	{
 		if (!Info.HasCreatedPackage())
@@ -3788,9 +3666,8 @@ UPackage* UCookOnTheFlyServer::TryCreateGeneratedPackage(UE::Cook::FGenerationHe
 			FReferenceChainSearch RefChainSearch(GeneratedPackage, SearchMode);
 			return nullptr;
 		}
-		// Otherwise this is the package that was created and passed to presave, and it is still valid because there has not been a GC since
-		// we created it. Mark its state and use it
-		bPopulatedByPreSave = true;
+		// Otherwise this is the package that was created for the generator's presave, and it is still valid because
+		// there has not been a GC since we created it.
 	}
 	else
 	{
@@ -3804,17 +3681,23 @@ UE::Cook::EPollStatus UCookOnTheFlyServer::TryPopulateGeneratedPackage(UE::Cook:
 {
 	using namespace UE::Cook;
 
-	UPackage* OwnerPackage = GenerationHelper.GetOwnerPackage();
 	check(GeneratedInfo.PackageData); // Caller already checked this
 	UE::Cook::FPackageData& GeneratedPackageData = *GeneratedInfo.PackageData;
 	const FString GeneratedPackageName = GeneratedPackageData.GetPackageName().ToString();
+	UPackage* OwnerPackage = GenerationHelper.FindOrLoadOwnerPackage(*this);
+	if (!OwnerPackage)
+	{
+		UE_LOG(LogCook, Error, TEXT("TryPopulateGeneratedPackage: could not load ParentGeneratorPackage %s for GeneratedPackage %s"),
+			*GenerationHelper.GetOwner().GetPackageName().ToString(), *GeneratedPackageName);
+		return EPollStatus::Error;
+	}
 	UPackage* GeneratedPackage = GeneratedPackageData.GetPackage();
 	check(GeneratedPackage); // We would have been kicked out of save if the package were gone
 
-	UObject* OwnerObject = GenerationHelper.FindSplitDataObject();
+	UObject* OwnerObject = GenerationHelper.FindOrLoadSplitDataObject();
 	if (!OwnerObject)
 	{
-		UE_LOG(LogCook, Error, TEXT("PopulateGeneratedPacakge could not find the original splitting object. Generated package can not be created. Splitter=%s, Generated=%s."),
+		UE_LOG(LogCook, Error, TEXT("PopulateGeneratedPackage could not find the original splitting object. Generated package can not be created. Splitter=%s, Generated=%s."),
 			*GenerationHelper.GetSplitDataObjectName().ToString(), *GeneratedPackageName);
 		return EPollStatus::Error;
 	}
@@ -3905,7 +3788,7 @@ UE::Cook::EPollStatus UCookOnTheFlyServer::PrepareSaveInternal(UE::Cook::FPackag
 	UPackage* Package = PackageData.GetPackage();
 	check(Package && Package->IsFullyLoaded());
 	check(PackageData.GetState() == EPackageState::Save);
-	FGenerationHelper* GenerationHelper = nullptr;
+	TRefCountPtr<FGenerationHelper> GenerationHelper;
 
 	if (!PackageData.GetCookedPlatformDataCalled())
 	{
@@ -3965,14 +3848,46 @@ UE::Cook::EPollStatus UCookOnTheFlyServer::PrepareSaveInternal(UE::Cook::FPackag
 		// Check for whether the Package has a Splitter and initialize its list if so
 		if (!PackageData.HasInitializedGeneratorSave())
 		{
-			Result = ConditionalCreateGeneratorPackage(PackageData, bPrecaching);
-			if (Result != EPollStatus::Success)
+			// The GenerationHelper might have already been created by a child generated package;
+			// or it might have been created and not initialized by iterative cook startup.
+			// If not created or initialized, try looking for it
+			GenerationHelper = PackageData.TryCreateValidGenerationHelper();
+			if (GenerationHelper)
 			{
-				return Result;
+				// Keep it referenced even if we are only precaching, so we do not recreate it
+				GenerationHelper->SetKeepForGeneratorSave();
+				if (bPrecaching)
+				{
+					// Do not proceed to GetGenerateList when precaching; do that only when we're ready to save the package
+					return EPollStatus::Incomplete;
+				}
+				else
+				{
+					// TODO: Add support for cooking in the editor. Possibly moot since we plan to deprecate cooking in the editor.
+					if (IsCookingInEditor())
+					{
+						// CookPackageSplitters allow destructive changes to the generator package. e.g. moving UObjects out
+						// of it into the streaming packages. To allow its use in the editor, we will need to make it non-destructive
+						// (by e.g. copying to new packages), or restore the package after the changes have been made.
+						UE_LOG(LogCook, Error, TEXT("Can not cook package %s: cooking in editor doesn't support Cook Package Splitters."),
+							*PackageData.GetPackageName().ToString());
+						return EPollStatus::Error;
+					}
+					if (!GenerationHelper->TryGenerateList())
+					{
+						return EPollStatus::Error;
+					}
+					// The earlier exit from SaveState should have reset the progress back to StartPopulate or earlier
+					check(GenerationHelper->GetOwnerInfo().GetSaveState() <= FCookGenerationInfo::ESaveState::StartPopulate);
+					GenerationHelper->StartOwnerSave();
+				}
 			}
 			PackageData.SetInitializedGeneratorSave(true);
 		}
-		GenerationHelper = PackageData.GetGenerationHelper();
+		else
+		{
+			GenerationHelper = PackageData.GetGenerationHelperIfValid();
+		}
 		if (GenerationHelper)
 		{
 			Result = QueueGeneratedPackages(*GenerationHelper, PackageData);
@@ -3986,12 +3901,12 @@ UE::Cook::EPollStatus UCookOnTheFlyServer::PrepareSaveInternal(UE::Cook::FPackag
 	}
 	else
 	{
-		GenerationHelper = PackageData.GetGenerationHelper();
+		GenerationHelper = PackageData.GetGenerationHelperIfValid();
 	}
 
 	if (GenerationHelper)
 	{
-		EPollStatus Result = PrepareSaveGeneratedPackage(*GenerationHelper, PackageData, Timer, bPrecaching);
+		EPollStatus Result = PrepareSaveGenerationPackage(*GenerationHelper, PackageData, Timer, bPrecaching);
 		if (Result != EPollStatus::Success)
 		{
 			return Result;
@@ -3999,15 +3914,16 @@ UE::Cook::EPollStatus UCookOnTheFlyServer::PrepareSaveInternal(UE::Cook::FPackag
 	}
 	else if (PackageData.IsGenerated())
 	{
-		FGenerationHelper* ParentGenerator = PackageData.GetGeneratedOwner();
-		if (!ParentGenerator)
+		TRefCountPtr<FGenerationHelper> ParentGenerationHelper = PackageData.GetParentGenerationHelper();
+		if (!ParentGenerationHelper || !ParentGenerationHelper->IsValid())
 		{
-			UE_LOG(LogCook, Error, TEXT("Generated package %s is missing its Parent generator package and cannot be saved."),
-				*PackageData.GetPackageName().ToString());
+			UE_LOG(LogCook, Error, TEXT("Generated package %s %s ParentGenerator package %s and cannot be saved."),
+				(!ParentGenerationHelper ? TEXT("is missing its") : TEXT("has an invalid")),
+				*PackageData.GetPackageName().ToString(), *PackageData.GetParentGenerator().ToString());
 			return EPollStatus::Error;
 		}
 
-		EPollStatus Result = PrepareSaveGeneratedPackage(*ParentGenerator, PackageData, Timer, bPrecaching);
+		EPollStatus Result = PrepareSaveGenerationPackage(*ParentGenerationHelper, PackageData, Timer, bPrecaching);
 		if (Result != EPollStatus::Success)
 		{
 			return Result;
@@ -4136,12 +4052,13 @@ void UCookOnTheFlyServer::ReleaseCookedPlatformData(UE::Cook::FPackageData& Pack
 		return;
 	}
 
-	FGenerationHelper* GenerationHelper = PackageData.GetGenerationHelper();
+	TRefCountPtr<FGenerationHelper> GenerationHelper = PackageData.GetGenerationHelper();
 	if (!GenerationHelper)
 	{
-		GenerationHelper = PackageData.GetGeneratedOwner();
+		GenerationHelper = PackageData.GetParentGenerationHelper();
 	}
-	FCookGenerationInfo* GenerationInfo = GenerationHelper ? GenerationHelper->FindInfo(PackageData) : nullptr;
+	FCookGenerationInfo* GenerationInfo = (GenerationHelper && GenerationHelper->IsInitialized()) ?
+		GenerationHelper->FindInfo(PackageData) : nullptr;
 
 	// For every BeginCacheForCookedPlatformData call we made we need to call ClearAllCachedCookedPlatformData
 	if (ReleaseSaveReason == EStateChangeReason::Completed)
@@ -4227,25 +4144,6 @@ void UCookOnTheFlyServer::ReleaseCookedPlatformData(UE::Cook::FPackageData& Pack
 		{
 			PackageData.SetInitializedGeneratorSave(false);
 		}
-
-		if (ReleaseSaveReason == EStateChangeReason::Completed)
-		{
-			GenerationHelper->SetPackageSaved(*GenerationInfo, PackageData);
-			if (GenerationHelper->IsComplete())
-			{
-				if (GenerationInfo->IsGenerator())
-				{
-					PackageData.DestroyGenerationHelper();
-				}
-				else
-				{
-					GenerationHelper->GetOwner().DestroyGenerationHelper();
-				}
-				// Clear now-dangling pointers
-				GenerationHelper = nullptr;
-				GenerationInfo = nullptr;
-			}
-		}
 	}
 
 	PackageData.ClearCookedPlatformData();
@@ -5315,17 +5213,20 @@ void UCookOnTheFlyServer::PreGarbageCollect()
 		FGenerationHelper* GenerationHelper = PackageData->GetGenerationHelper();
 		if (!GenerationHelper)
 		{
-			GenerationHelper = PackageData->GetGeneratedOwner();
+			GenerationHelper = PackageData->GetParentGenerationHelper();
 		}
-		FCookGenerationInfo* Info = GenerationHelper ? GenerationHelper->FindInfo(*PackageData) : nullptr;
-		if (Info)
+		if (GenerationHelper && GenerationHelper->IsInitialized())
 		{
-			bool bShouldDemote;
-			GenerationHelper->PreGarbageCollect(*Info, GCKeepObjects, GCKeepPackages,
-				GCKeepPackageDatas, bShouldDemote);
-			if (bShouldDemote)
+			FCookGenerationInfo* Info = GenerationHelper->FindInfo(*PackageData);
+			if (Info)
 			{
-				ReleaseCookedPlatformData(*PackageData, UE::Cook::EStateChangeReason::GeneratorPreGarbageCollected);
+				bool bShouldDemote;
+				GenerationHelper->PreGarbageCollect(*Info, GCKeepObjects, GCKeepPackages,
+					GCKeepPackageDatas, bShouldDemote);
+				if (bShouldDemote)
+				{
+					ReleaseCookedPlatformData(*PackageData, UE::Cook::EStateChangeReason::GeneratorPreGarbageCollected);
+				}
 			}
 		}
 		if (PackageData->GetIsCookLast())
@@ -5564,7 +5465,7 @@ void UCookOnTheFlyServer::PostGarbageCollect()
 
 	PackageDatas->LockAndEnumeratePackageDatas([](FPackageData* PackageData)
 	{
-		if (FGenerationHelper* GenerationHelper = PackageData->GetGenerationHelper())
+		if (TRefCountPtr<FGenerationHelper> GenerationHelper = PackageData->GetGenerationHelper())
 		{
 			GenerationHelper->PostGarbageCollect();
 		}
@@ -6314,14 +6215,14 @@ void FSaveCookedPackageContext::SetupPlatform(const ITargetPlatform* InTargetPla
 		if (TSet<FName>* NeverCookPackages = COTFS.PackageTracker->PlatformSpecificNeverCookPackages.Find(TargetPlatform))
 		{
 			FGenerationHelper* GenerationHelper =
-				PackageData.IsGenerated() ? PackageData.GetGeneratedOwner() : nullptr;
+				PackageData.IsGenerated() ? PackageData.GetParentGenerationHelper() : nullptr;
 
 			if (NeverCookPackages->Find(Package->GetFName()) ||
 				(GenerationHelper && NeverCookPackages->Find(GenerationHelper->GetOwner().GetPackageName())))
 			{
 				SavePackageResult = ESavePackageResult::ContainsEditorOnlyData;
 				UE_LOG(LogCook, Display, TEXT("Excluding %s"), *PackageName);
-				return;				
+				return;
 			}
 		}
 	}
@@ -6470,7 +6371,7 @@ void FSaveCookedPackageContext::FinishPlatform()
 				COTFS.RecordExternalActorDependencies(GenerationHelper->GetExternalActorDependencies());
 			}
 		}
-		else if (GenerationHelper = PackageData.GetGeneratedOwner(); GenerationHelper)
+		else if (GenerationHelper = PackageData.GetParentGenerationHelper(); GenerationHelper)
 		{
 			FCookGenerationInfo* GeneratedInfo = GenerationHelper->FindInfo(PackageData);
 			if (!GeneratedInfo)
@@ -6967,6 +6868,33 @@ void FInitializeConfigSettings::LoadLocal(const FString& InOutputDirectoryOverri
 	}
 }
 
+static EMPCookGeneratorSplit ParseMPCookGeneratorSplitFromString(const FString& Text)
+{
+	if (Text.IsEmpty() || Text == TEXTVIEW("AnyWorker"))
+	{
+		return EMPCookGeneratorSplit::AnyWorker;
+	}
+	else if (Text == TEXTVIEW("AllOnSameWorker"))
+	{
+		return EMPCookGeneratorSplit::AllOnSameWorker;
+	}
+	if (Text == TEXTVIEW("SomeOnSameWorker"))
+	{
+		return EMPCookGeneratorSplit::SomeOnSameWorker;
+	}
+	if (Text == TEXTVIEW("NoneOnSameWorker"))
+	{
+		return EMPCookGeneratorSplit::NoneOnSameWorker;
+	}
+	else
+	{
+		UE_LOG(LogCook, Error,
+			TEXT("Invalid value -MPCookGeneratorSplit=%s. Valid values: { AnyWorker, AllOnSameWorker, SomeOnSameWorker, NoneOnSameWorker }."),
+			*Text);
+		return EMPCookGeneratorSplit::AnyWorker;
+	}
+}
+
 }
 
 void UCookOnTheFlyServer::SetInitializeConfigSettings(UE::Cook::FInitializeConfigSettings&& Settings)
@@ -7143,6 +7071,11 @@ void UCookOnTheFlyServer::SetInitializeConfigSettings(UE::Cook::FInitializeConfi
 	{
 		GWarn->TreatWarningsAsErrors = true;
 	}
+
+	FString GeneratorSplit;
+	GConfig->GetString(TEXT("CookSettings"), TEXT("MPCookGeneratorSplit"), GeneratorSplit, GEditorIni);
+	FParse::Value(FCommandLine::Get(), TEXT("-MPCookGeneratorSplit="), GeneratorSplit);
+	MPCookGeneratorSplit = UE::Cook::ParseMPCookGeneratorSplitFromString(GeneratorSplit);
 }
 
 void UCookOnTheFlyServer::ParseCookFilters()
@@ -8559,7 +8492,8 @@ void UCookOnTheFlyServer::PopulateCookedPackages(TArrayView<const ITargetPlatfor
 					*Pair.Key.ToString());
 				continue;
 			}
-			Generator->CreateGenerationHelper(nullptr, nullptr).SetPreviousGeneratedPackages(MoveTemp(Pair.Value.Generated));
+			TRefCountPtr<FGenerationHelper> GenerationHelper = Generator->CreateUninitializedGenerationHelper();
+			GenerationHelper->SetPreviousGeneratedPackages(MoveTemp(Pair.Value.Generated));
 		}
 
 		PlatformAssetRegistry.SetPreviousAssetRegistry(MoveTemp(PreviousAssetRegistry));
@@ -10101,6 +10035,7 @@ void UCookOnTheFlyServer::CookByTheBookFinishedInternal()
 	check(IsCookByTheBookMode());
 	check(IsInSession());
 	check(PackageDatas->GetRequestQueue().IsEmpty());
+	check(PackageDatas->GetRequestQueue().GetDiscoveryQueue().IsEmpty());
 	check(PackageDatas->GetAssignedToWorkerSet().IsEmpty());
 	check(PackageDatas->GetLoadPrepareQueue().IsEmpty());
 	check(PackageDatas->GetLoadReadyQueue().IsEmpty());
@@ -10377,11 +10312,6 @@ void UCookOnTheFlyServer::ShutdownCookSession()
 		CookAsCookWorkerFinished();
 	}
 
-	PackageDatas->LockAndEnumeratePackageDatas([](UE::Cook::FPackageData* PackageData)
-	{
-		PackageData->DestroyGenerationHelper();
-	});
-
 	if (IsCookByTheBookMode())
 	{
 		UnregisterCookByTheBookDelegates();
@@ -11534,15 +11464,15 @@ void UCookOnTheFlyServer::GetPackagesToRetract(int32 NumToRetract, TArray<FName>
 		return;
 	}
 
-	auto AddPackageIfPossibleAndReportDone = [&OutRetractionPackages, NumToRetract](FPackageData* PackageData)
+	auto AddPackageIfPossibleAndReportDone = [&OutRetractionPackages, NumToRetract, this](FPackageData* PackageData)
 	{
 		if (OutRetractionPackages.Num() >= NumToRetract)
 		{
 			return true;
 		}
 
-		if (PackageData->GetWorkerAssignmentConstraint().IsValid() || PackageData->IsGenerated() ||
-			PackageData->GetGenerationHelper())
+		if (PackageData->GetWorkerAssignmentConstraint().IsValid() ||
+			(MPCookGeneratorSplit == UE::Cook::EMPCookGeneratorSplit::AllOnSameWorker && PackageData->IsGenerated()))
 		{
 			// Don't send back Packages that are constrained to this worker. Doing so will just
 			// cause the CookDirector to send it back to us, and this can cause the cooker to crash
diff --git a/Engine/Source/Editor/UnrealEd/Private/Cooker/CookDirector.cpp b/Engine/Source/Editor/UnrealEd/Private/Cooker/CookDirector.cpp
--- a/Engine/Source/Editor/UnrealEd/Private/Cooker/CookDirector.cpp
+++ b/Engine/Source/Editor/UnrealEd/Private/Cooker/CookDirector.cpp
@@ -340,13 +340,47 @@ void FCookDirector::AssignRequests(TArrayView<FPackageData*> Requests, TArray<FW
 		WorkerIds.Add(RemoteWorker->GetWorkerId());
 	}
 
-	AssignRequests(MoveTemp(WorkerIds), LocalRemoteWorkers, Requests, OutAssignments, MoveTemp(RequestGraph));
+	AssignRequests(MoveTemp(WorkerIds), LocalRemoteWorkers, Requests, OutAssignments, MoveTemp(RequestGraph),
+		true /* bInitialAssignment */);
+}
+
+TMap<FPackageData*, FAssignPackageExtraData> FCookDirector::GetAssignPackageExtraDatas(
+	TConstArrayView<FPackageData*> Requests) const
+{
+	TMap<FPackageData*, FAssignPackageExtraData> Results;
+	for (FPackageData* Request : Requests)
+	{
+		TRefCountPtr<FGenerationHelper> GenerationHelper = Request->GetGenerationHelper();
+		if (GenerationHelper && !GenerationHelper->GetPreviousGeneratedPackages().IsEmpty())
+		{
+			FAssignPackageExtraData& ExtraData = Results.FindOrAdd(Request);
+			ExtraData.GeneratorPreviousGeneratedPackages = GenerationHelper->GetPreviousGeneratedPackages();
+		}
+	}
+	return Results;
+}
+
+TArray<FPackageData*> FCookDirector::GetInfoPackagesForRequests(TConstArrayView<FPackageData*> Requests) const
+{
+	TSet<FPackageData*> InfoPackages;
+	for (FPackageData* Request : Requests)
+	{
+		if (!Request->GetParentGenerator().IsNone())
+		{
+			FPackageData* Generator = COTFS.PackageDatas->FindPackageDataByPackageName(Request->GetParentGenerator());
+			if (Generator)
+			{
+				InfoPackages.Add(Generator);
+			}
+		}
+	}
+	return InfoPackages.Array();
 }
 
 void FCookDirector::AssignRequests(TArray<FWorkerId>&& InWorkers,
 	TArray<TRefCountPtr<FCookWorkerServer>>& InRemoteWorkers,
-	TArrayView<FPackageData*> Requests, TArray<FWorkerId>& OutAssignments, TMap<FPackageData*,
-	TArray<FPackageData*>>&& RequestGraph)
+	TArrayView<FPackageData*> Requests, TArray<FWorkerId>& OutAssignments,
+	TMap<FPackageData*, TArray<FPackageData*>>&& RequestGraph, bool bInitialAssignment)
 {
 	check(InWorkers.Num() > 0);
 	if (InWorkers.Num() <= 1)
@@ -393,7 +427,8 @@ void FCookDirector::AssignRequests(TArray<FWorkerId>&& InWorkers,
 				RequestBuffer = RequestSet.Array();
 				RequestsToSend = RequestBuffer;
 			}
-			(*RemoteWorker)->AppendAssignments(RequestsToSend, ECookDirectorThread::SchedulerThread);
+			(*RemoteWorker)->AppendAssignments(RequestsToSend, GetAssignPackageExtraDatas(RequestsToSend),
+				GetInfoPackagesForRequests(RequestsToSend), ECookDirectorThread::SchedulerThread);
 		}
 		return;
 	}
@@ -418,22 +453,71 @@ void FCookDirector::AssignRequests(TArray<FWorkerId>&& InWorkers,
 		}
 	}
 
+	// Debug function for generated packages. We want to be able to force their assignment to all,some,none
+	// on the same worker that saved their generator.
+	int32 MPCookGeneratorNextWorker = -1;
+	auto GetNextWorkerForGeneratedPackage = [&MPCookGeneratorNextWorker, &InWorkers](FWorkerId WorkerNotToUse)
+		{
+			MPCookGeneratorNextWorker = (MPCookGeneratorNextWorker + 1) % InWorkers.Num();
+			if (InWorkers[MPCookGeneratorNextWorker] == WorkerNotToUse)
+			{
+				MPCookGeneratorNextWorker = (MPCookGeneratorNextWorker + 1) % InWorkers.Num();
+			}
+			return InWorkers[MPCookGeneratorNextWorker];
+		};
 	for (int32 RequestIndex = 0; RequestIndex < Requests.Num(); ++RequestIndex)
 	{
 		FWorkerId& WorkerId = OutAssignments[RequestIndex];
 		// Override the loadbalancer's assignment if the Package has a WorkerAssignmentConstraint
 		// This allows us to guarantee that generated packages will be cooked on the worker that cooked
 		// their generator package
-		FWorkerId WorkerIdConstraint = Requests[RequestIndex]->GetWorkerAssignmentConstraint();
+		FPackageData* RequestPackage = Requests[RequestIndex];
+		FWorkerId WorkerIdConstraint = RequestPackage->GetWorkerAssignmentConstraint();
 		if (WorkerIdConstraint.IsValid())
 		{
 			WorkerId = WorkerIdConstraint;
 		}
 		// Override the loadbalancer's assignment to force it local if the Package is urgent
-		else if (Requests[RequestIndex]->GetIsUrgent())
+		else if (RequestPackage->GetIsUrgent())
 		{
 			WorkerId = FWorkerId::Local();
 		}
+		else if (RequestPackage->IsGenerated() && bInitialAssignment)
+		{
+			TRefCountPtr<FGenerationHelper> GenerationHelper = RequestPackage->GetParentGenerationHelper();
+			if (GenerationHelper)
+			{
+				FWorkerId GeneratorWorker = GenerationHelper->GetWorkerIdThatSavedGenerator();
+				if (GeneratorWorker.IsValid())
+				{
+					switch (COTFS.MPCookGeneratorSplit)
+					{
+					default: // fallthrough
+					case EMPCookGeneratorSplit::AnyWorker:
+						// If we're allowed to put them on any worker, put them on the GeneratorWorker anyway
+						// for the initial assignment since it already has the GenerationHelper created
+						WorkerId = GeneratorWorker;
+						break;
+					case EMPCookGeneratorSplit::AllOnSameWorker:
+						WorkerId = GeneratorWorker;
+						break;
+					case EMPCookGeneratorSplit::SomeOnSameWorker:
+						if ((GenerationHelper->GetMPCookNextAssignmentIndex()++) % 2 == 0)
+						{
+							WorkerId = GeneratorWorker;
+						}
+						else
+						{
+							WorkerId = GetNextWorkerForGeneratedPackage(GeneratorWorker);
+						}
+						break;
+					case EMPCookGeneratorSplit::NoneOnSameWorker:
+						WorkerId = GetNextWorkerForGeneratedPackage(GeneratorWorker);
+						break;
+					}
+				}
+			}
+		}
 
 		if (!WorkerId.IsLocal())
 		{
@@ -464,8 +548,9 @@ void FCookDirector::AssignRequests(TArray<FWorkerId>&& InWorkers,
 			TRefCountPtr<FCookWorkerServer>* RemoteWorker = InRemoteWorkers.FindByPredicate(
 				[&WorkerId](const TRefCountPtr<FCookWorkerServer>& X) { return X->GetWorkerId() == WorkerId; });
 			check(RemoteWorker);
-			(*RemoteWorker)->AppendAssignments(RemoteBatches[WorkerId.GetRemoteIndex()],
-				ECookDirectorThread::SchedulerThread);
+			TArray<FPackageData*>& RemoteBatch = RemoteBatches[WorkerId.GetRemoteIndex()];
+			(*RemoteWorker)->AppendAssignments(RemoteBatch,
+				GetAssignPackageExtraDatas(RemoteBatch), GetInfoPackagesForRequests(RemoteBatch), ECookDirectorThread::SchedulerThread);
 		}
 	}
 
@@ -506,6 +591,21 @@ void FCookDirector::RemoveFromWorker(FPackageData& PackageData)
 	OwningWorker->AbortAssignment(PackageData, ECookDirectorThread::SchedulerThread);
 }
 
+void FCookDirector::BroadcastGeneratorFencePassed(FGenerationHelper& GenerationHelper)
+{
+	FName PackageName = GenerationHelper.GetOwner().GetPackageName();
+	FGeneratorEventMessage Message(EGeneratorEvent::QueuedGeneratedPackagesFencePassed, PackageName);
+	{
+		FScopeLock CommunicationScopeLock(&CommunicationLock);
+		InitializeWorkers();
+	}
+	TArray<TRefCountPtr<FCookWorkerServer>> LocalRemoteWorkers = CopyRemoteWorkers();
+	for (TRefCountPtr<FCookWorkerServer>& RemoteWorker : LocalRemoteWorkers)
+	{
+		RemoteWorker->AppendMessage(Message, ECookDirectorThread::SchedulerThread);
+	}
+}
+
 void FCookDirector::TickFromSchedulerThread()
 {
 	double CurrentTime = FPlatformTime::Seconds();
@@ -1933,7 +2033,7 @@ FCookDirector::FRetractionHandler::ReassignPackages(const FWorkerId& FromWorker,
 	TMap<FPackageData*, TArray<FPackageData*>> RequestGraph;
 	TArray<FWorkerId> Assignments;
 	Director.AssignRequests(MoveTemp(WorkersToSplitOver), LocalRemoteWorkers, AssignmentPackages, Assignments,
-		MoveTemp(RequestGraph));
+		MoveTemp(RequestGraph), false /* bInitialAssignment */);
 	FRequestQueue& RequestQueue = Director.COTFS.PackageDatas->GetRequestQueue();
 	bool bAssignedToLocal = false;
 	for (int32 Index = 0; Index < AssignmentPackages.Num(); ++Index)
diff --git a/Engine/Source/Editor/UnrealEd/Private/Cooker/CookDirector.h b/Engine/Source/Editor/UnrealEd/Private/Cooker/CookDirector.h
--- a/Engine/Source/Editor/UnrealEd/Private/Cooker/CookDirector.h
+++ b/Engine/Source/Editor/UnrealEd/Private/Cooker/CookDirector.h
@@ -28,6 +28,7 @@ class FCbWriter;
 class FRunnableThread;
 class UCookOnTheFlyServer;
 namespace UE::Cook { class FCookWorkerServer; }
+namespace UE::Cook { struct FAssignPackageExtraData; }
 namespace UE::Cook { struct FCookWorkerProfileData; }
 namespace UE::Cook { struct FHeartbeatMessage; }
 namespace UE::Cook { struct FInitialConfigMessage; }
@@ -73,6 +74,13 @@ public:
 		TMap<FPackageData*, TArray<FPackageData*>>&& RequestGraph);
 	/** Notify the CookWorker that owns the cook of the package that the Director wants to take it back. */
 	void RemoveFromWorker(FPackageData& PackageData);
+
+	/**
+	 * Report EGeneratorEvent::QueuedGeneratedPackagesFencePassed to all CookWorkers so any of them
+	 * waiting on the fence to be passed before clearing their data can clear their data.
+	 */
+	void BroadcastGeneratorFencePassed(FGenerationHelper& GenerationHelper);
+
 	/** Periodic tick function. Sends/Receives messages to CookWorkers. */
 	void TickFromSchedulerThread();
 	/** Periodic display function, called from CookOnTheFlyServer.UpdateDisplay. */
@@ -212,7 +220,7 @@ private:
 #endif
 	void AssignRequests(TArray<FWorkerId>&& InWorkers, TArray<TRefCountPtr<FCookWorkerServer>>& InRemoteWorkers, 
 		TArrayView<FPackageData*> Requests, TArray<FWorkerId>& OutAssignments,
-		TMap<FPackageData*, TArray<FPackageData*>>&& RequestGraph);
+		TMap<FPackageData*, TArray<FPackageData*>>&& RequestGraph, bool bInitialAssignment);
 
 	TArray<TRefCountPtr<FCookWorkerServer>> CopyRemoteWorkers() const;
 	void DisplayRemainingPackages() const;
@@ -220,6 +228,10 @@ private:
 	FString GetDisplayName(const FCookWorkerServer& RemoteWorker, int32 PreferredWidth=-1) const;
 	const TRefCountPtr<FCookWorkerServer>* FindRemoteWorkerInLock(const FWorkerId& WorkerId) const;
 
+	TMap<FPackageData*, FAssignPackageExtraData> GetAssignPackageExtraDatas(
+		TConstArrayView<FPackageData*> Requests) const;
+	TArray<FPackageData*> GetInfoPackagesForRequests(TConstArrayView<FPackageData*> Requests) const;
+
 private:
 	// Synchronization primitives that can be used from any thread
 	mutable FCriticalSection CommunicationLock;
diff --git a/Engine/Source/Editor/UnrealEd/Private/Cooker/CookPackageData.cpp b/Engine/Source/Editor/UnrealEd/Private/Cooker/CookPackageData.cpp
--- a/Engine/Source/Editor/UnrealEd/Private/Cooker/CookPackageData.cpp
+++ b/Engine/Source/Editor/UnrealEd/Private/Cooker/CookPackageData.cpp
@@ -11,9 +11,11 @@
 #include "AssetRegistry/IAssetRegistry.h"
 #include "Async/ParallelFor.h"
 #include "CompactBinaryTCP.h"
+#include "Cooker/CookDirector.h"
 #include "Cooker/CookPlatformManager.h"
 #include "Cooker/CookRequestCluster.h"
 #include "Cooker/CookWorkerClient.h"
+#include "Cooker/IWorkerRequests.h"
 #include "Cooker/PackageTracker.h"
 #include "CookOnTheSide/CookOnTheFlyServer.h"
 #include "Containers/StringView.h"
@@ -88,7 +90,7 @@
 }
 
 FPackageData::FPackageData(FPackageDatas& PackageDatas, const FName& InPackageName, const FName& InFileName)
-	: GeneratedOwner(nullptr), PackageName(InPackageName), FileName(InFileName), PackageDatas(PackageDatas)
+	: ParentGenerationHelper(nullptr), PackageName(InPackageName), FileName(InFileName), PackageDatas(PackageDatas)
 	, Instigator(EInstigator::NotYetRequested), bIsUrgent(0), bIsCookLast(0)
 	, bIsVisited(0), bIsPreloadAttempted(0)
 	, bIsPreloaded(0), bHasSaveCache(0), bPrepareSaveFailed(0), bPrepareSaveRequiresGC(0)
@@ -113,7 +115,11 @@
 
 void FPackageData::ClearReferences()
 {
-	DestroyGenerationHelper();
+	if (GenerationHelper)
+	{
+		GenerationHelper->ClearSelfReferences();
+	}
+	ParentGenerationHelper.SafeRelease();
 }
 
 const FName& FPackageData::GetPackageName() const
@@ -317,6 +323,15 @@
 {
 	SetIsUrgent(false);
 	CompletionCallback = FCompletionCallback();
+	if (GenerationHelper)
+	{
+		// ClearKeepForGeneratorSave is called when finishing the save state, but not when demoting out of the save
+		// state after a garbage collect. Call it here in case we cancel the save of the packagedata after demotion.
+		// The other self-references (iterative, queued packages) should persist even when the packagedata is not
+		// in progress.
+		GenerationHelper->ClearKeepForGeneratorSave();
+	}
+	ParentGenerationHelper.SafeRelease();
 }
 
 void FPackageData::SetPlatformsCooked(
@@ -820,6 +835,14 @@
 
 void FPackageData::OnEnterAssignedToWorker()
 {
+	if (IsGenerated())
+	{
+		// Clear the referencecount that we added in OnEnterInProgress; we don't want to keep
+		// the GenerationHelper referenced for the entire duration of assigned packages
+		// running on other CookWorkers. If this package gets retracted and moved into
+		// LoadState locally, we will recreate the GenerationHelper if necessary.
+		ParentGenerationHelper.SafeRelease();
+	}
 }
 
 void FPackageData::SetWorkerAssignment(FWorkerId InWorkerAssignment, ESendFlags SendFlags)
@@ -888,6 +911,12 @@
 void FPackageData::OnEnterInProgress()
 {
 	PackageDatas.GetMonitor().OnInProgressChanged(*this, true);
+	if (IsGenerated())
+	{
+		// Keep a refcount to the ParentGenerationHelper until we are saved so that it does not destruct
+		// and waste time reconstructing when we reach the LoadQueue.
+		GetOrFindParentGenerationHelper();
+	}
 }
 
 void FPackageData::OnExitInProgress()
@@ -1354,11 +1383,13 @@
 	check(!GetCookedPlatformDataStarted());
 	check(!GetCookedPlatformDataCalled());
 	check(!GetCookedPlatformDataComplete());
-	check(!GetGenerationHelper() ||
-		GetGenerationHelper()->GetOwnerInfo().GetSaveState() <= FCookGenerationInfo::ESaveState::StartPopulate);
-	if (GetGeneratedOwner())
+	if (GetGenerationHelper())
 	{
-		FCookGenerationInfo* Info = GetGeneratedOwner()->FindInfo(*this);
+		check(GetGenerationHelper()->GetOwnerInfo().GetSaveState() <= FCookGenerationInfo::ESaveState::StartPopulate);
+	}
+	if (ParentGenerationHelper)
+	{
+		FCookGenerationInfo* Info = ParentGenerationHelper->FindInfo(*this);
 		check(!Info || Info->GetSaveState() <= FCookGenerationInfo::ESaveState::StartPopulate);
 	}
 }
@@ -1452,54 +1483,164 @@
 	}
 	else if (IsGenerated())
 	{
-		if (!GeneratedOwner)
+		if (!ParentGenerationHelper)
 		{
 			bOutDemote = true;
 		}
 		else
 		{
-			GeneratedOwner->UpdateSaveAfterGarbageCollect(*this, bOutDemote);
+			ParentGenerationHelper->UpdateSaveAfterGarbageCollect(*this, bOutDemote);
 		}
 	}
 }
 
-void FPackageData::SetGeneratedOwner(FGenerationHelper* InGeneratedOwner)
+void FPackageData::SetGenerated(FName InParentGenerator)
 {
-	check(IsGenerated());
-	check(!(GeneratedOwner && InGeneratedOwner));
-	GeneratedOwner = InGeneratedOwner;
+	bGenerated = true;
+	ParentGenerator = InParentGenerator;
 }
 
-UE::Cook::FGenerationHelper* FPackageData::GetGenerationHelper() const
+void FPackageData::SetParentGenerationHelper(FGenerationHelper* InGenerationHelper)
 {
-	UE::Cook::FGenerationHelper* Result = GenerationHelper.Get();
-	if (Result && Result->IsInitialized())
+	check(InGenerationHelper == nullptr || IsGenerated());
+	check(!(ParentGenerationHelper && InGenerationHelper) || ParentGenerationHelper == InGenerationHelper);
+	ParentGenerationHelper = InGenerationHelper;
+}
+
+TRefCountPtr<FGenerationHelper> FPackageData::GetOrFindParentGenerationHelper()
+{
+	if (ParentGenerationHelper)
 	{
+		return ParentGenerationHelper;
+	}
+	if (!IsGenerated())
+	{
+		return nullptr;
+	}
+
+	FPackageData* OwnerPackageData = PackageDatas.FindPackageDataByPackageName(GetParentGenerator());
+	if (!OwnerPackageData)
+	{
+		return nullptr;
+	}
+
+	ParentGenerationHelper = OwnerPackageData->GetGenerationHelper();
+	return ParentGenerationHelper;
+}
+
+TRefCountPtr<FGenerationHelper> FPackageData::TryCreateValidParentGenerationHelper()
+{
+	if (ParentGenerationHelper)
+	{
+		if (!ParentGenerationHelper->IsValid())
+		{
+			ParentGenerationHelper.SafeRelease();
+		}
+		return ParentGenerationHelper;
+	}
+	if (!IsGenerated())
+	{
+		return nullptr;
+	}
+
+	FPackageData* OwnerPackageData = PackageDatas.FindPackageDataByPackageName(GetParentGenerator());
+	if (!OwnerPackageData)
+	{
+		return nullptr;
+	}
+
+	ParentGenerationHelper = OwnerPackageData->TryCreateValidGenerationHelper();
+
+	return ParentGenerationHelper;
+}
+
+TRefCountPtr<FGenerationHelper> FPackageData::CreateUninitializedGenerationHelper()
+{
+	if (GenerationHelper)
+	{
+		return GenerationHelper;
+	}
+	TRefCountPtr<UE::Cook::FGenerationHelper> Result = new UE::Cook::FGenerationHelper(*this);
+	GenerationHelper = Result.GetReference();
+	return Result;
+}
+
+TRefCountPtr<UE::Cook::FGenerationHelper> FPackageData::TryCreateValidGenerationHelper()
+{
+	if (GenerationHelper && GenerationHelper->IsInitialized())
+	{
+		if (!GenerationHelper->IsValid())
+		{
+			// The GenerationHelper is not valid; we can get here if it was created from iterative cook data but this
+			// package is no longer a generator after syncing. If it has any self-references, clear them so that it
+			// will delete and this non-generator package will set the usual GenerationHelper=nullptr value.
+			GenerationHelper->ClearSelfReferences(); // Might set GenerationHelper=nullptr
+			// this->GenerationHelper might still be non-null, if there are some generated packages
+			// that still have a pointer to it. This will only happen in error-handling edge cases, but we
+			// need to check for invalid GenerationHelper everwhere we use them to cover this case.
+			// Our contract for TryCreateValidGenerationHelper this case is we return nullptr.
+			return nullptr;
+		}
+		return GenerationHelper;
+	}
+
+	UCookOnTheFlyServer& COTFS = PackageDatas.GetCookOnTheFlyServer();
+	UE::Cook::Private::FRegisteredCookPackageSplitter* RegisteredSplitterType = nullptr;
+	TUniquePtr<ICookPackageSplitter> CookPackageSplitterInstance;
+	UObject* SplitDataObject = nullptr;
+	UPackage* LocalPackage = GetPackage();
+	if (!LocalPackage)
+	{
+		LocalPackage = FGenerationHelper::FindOrLoadPackage(COTFS, *this);
+	}
+	if (LocalPackage)
+	{
+		TOptional<TConstArrayView<FCachedObjectInOuter>> LocalCachedObjectsInOuter;
+		if (GetHasSaveCache())
+		{
+			LocalCachedObjectsInOuter.Emplace(GetCachedObjectsInOuter());
+		}
+		FGenerationHelper::SearchForRegisteredSplitDataObject(COTFS, GetPackageName(),
+			LocalPackage, LocalCachedObjectsInOuter, SplitDataObject, RegisteredSplitterType,
+			CookPackageSplitterInstance);
+	}
+	TRefCountPtr<UE::Cook::FGenerationHelper> Result = GenerationHelper;
+	if (!SplitDataObject || !CookPackageSplitterInstance)
+	{
+		if (Result)
+		{
+			// Mark that GenerationHelper is invalid, and clear its references and return nullptr; see comment above.
+			Result->InitializeAsInvalid(); // cannot set GenerationHelper=nullptr because we have a local refcount.
+			Result->ClearSelfReferences();
+		}
+		return nullptr;
+	}
+	else
+	{
+		if (!Result)
+		{
+			Result = new UE::Cook::FGenerationHelper(*this);
+			GenerationHelper = Result.GetReference();
+		}
+
+		Result->Initialize(SplitDataObject, RegisteredSplitterType, MoveTemp(CookPackageSplitterInstance));
 		return Result;
 	}
+}
+
+TRefCountPtr<UE::Cook::FGenerationHelper> FPackageData::GetGenerationHelperIfValid()
+{
+	if (GenerationHelper && GenerationHelper->IsValid())
+	{
+		return GenerationHelper;
+	}
 	return nullptr;
 }
 
-UE::Cook::FGenerationHelper& FPackageData::CreateGenerationHelper(const UObject* InSplitDataObject,
-	ICookPackageSplitter* InCookPackageSplitterInstance)
+void FPackageData::OnGenerationHelperDestroyed(FGenerationHelper& InGenerationHelper)
 {
-	if (!GenerationHelper)
-	{
-		GenerationHelper.Reset(new UE::Cook::FGenerationHelper(*this, InSplitDataObject,
-			InCookPackageSplitterInstance));
-	}
-	else
-	{
-		GenerationHelper->InitializeSave(InSplitDataObject, InCookPackageSplitterInstance);
-		if (GenerationHelper->IsInitialized())
-		{
-			// The earlier exit from SaveState should have reset the progress back to StartGeneratorSave or earlier
-			check(GenerationHelper->GetOwnerInfo().GetSaveState() <= FCookGenerationInfo::ESaveState::StartPopulate);
-		}
-	}
-	FGenerationHelper* Result = GenerationHelper.Get();
-	check(Result);
-	return *Result;
+	check(GenerationHelper == &InGenerationHelper);
+	GenerationHelper = nullptr;
 }
 
 FConstructPackageData FPackageData::CreateConstructData()
@@ -1559,112 +1700,352 @@
 //////////////////////////////////////////////////////////////////////////
 // FGenerationHelper
 
-FGenerationHelper::FGenerationHelper(UE::Cook::FPackageData& InOwner, const UObject* InSplitDataObject,
-	ICookPackageSplitter* InCookPackageSplitterInstance)
-: OwnerInfo(InOwner, true /* bInGenerated */)
+FGenerationHelper::FGenerationHelper(UE::Cook::FPackageData& InOwner)
+: OwnerInfo(InOwner, true /* bInGenerator */)
 {
-	InitializeSave(InSplitDataObject, InCookPackageSplitterInstance);
 }
 
-void FGenerationHelper::InitializeSave(const UObject* InSplitDataObject,
-	ICookPackageSplitter* InCookPackageSplitterInstance)
-{
-	if (InCookPackageSplitterInstance)
-	{
-		// If we already have a splitter, keep the old and throw out the new. The old one
-		// still contains some state.
-		if (!CookPackageSplitterInstance)
-		{
-			CookPackageSplitterInstance.Reset(InCookPackageSplitterInstance);
-		}
-		else
-		{
-			delete InCookPackageSplitterInstance;
-		}
-
-		bInitialized = true;
-		FName InSplitDataObjectName = *InSplitDataObject->GetFullName();
-		check(SplitDataObjectName.IsNone() || SplitDataObjectName == InSplitDataObjectName);
-		SplitDataObjectName = InSplitDataObjectName;
-		bUseInternalReferenceToAvoidGarbageCollect =
-			CookPackageSplitterInstance->UseInternalReferenceToAvoidGarbageCollect();
-		SetOwnerPackage(GetOwner().GetPackage());
-	}
-}
-
 FGenerationHelper::~FGenerationHelper()
 {
-	if (!IsInitialized())
-	{
-		return;
-	}
-
-	ConditionalNotifyCompletion(ICookPackageSplitter::ETeardown::Canceled);
-	ClearGeneratedPackages();
+	NotifyCompletion(ICookPackageSplitter::ETeardown::Complete);
+	GetOwner().OnGenerationHelperDestroyed(*this);
 }
 
-void FGenerationHelper::ConditionalNotifyCompletion(ICookPackageSplitter::ETeardown Status)
+void FGenerationHelper::NotifyCompletion(ICookPackageSplitter::ETeardown Status)
 {
-	if (!bNotifiedCompletion)
+	if (IsInitialized() && IsValid() && CookPackageSplitterInstance)
 	{
-		bNotifiedCompletion = true;
 		CookPackageSplitterInstance->Teardown(Status);
 		CookPackageSplitterInstance.Reset();
 	}
 }
 
-void FGenerationHelper::ClearGeneratedPackages()
+void FGenerationHelper::Initialize()
 {
-	check(IsInitialized());
-	for (FCookGenerationInfo& Info: PackagesToGenerate)
+	if (InitializeStatus != EInitializeStatus::Uninitialized)
 	{
-		if (Info.PackageData)
+		return;
+	}
+
+	FPackageData& OwnerPackageData = GetOwner();
+	FName OwnerPackageName = OwnerPackageData.GetPackageName();
+	UCookOnTheFlyServer& COTFS = OwnerPackageData.GetPackageDatas().GetCookOnTheFlyServer();
+	UPackage* LocalOwnerPackage = FindOrLoadPackage(COTFS, OwnerPackageData);
+	if (!LocalOwnerPackage)
+	{
+		InitializeStatus = EInitializeStatus::Invalid;
+		return;
+	}
+
+	UObject* LocalSplitDataObject;
+	UE::Cook::Private::FRegisteredCookPackageSplitter* LocalRegisteredSplitterType = nullptr;
+	TUniquePtr<ICookPackageSplitter> LocalSplitter;
+	SearchForRegisteredSplitDataObject(COTFS, OwnerPackageName, LocalOwnerPackage,
+		TOptional<TConstArrayView<FCachedObjectInOuter>>(), LocalSplitDataObject, LocalRegisteredSplitterType,
+		LocalSplitter);
+	if (!LocalSplitDataObject || !LocalSplitter)
+	{
+		InitializeStatus = EInitializeStatus::Invalid;
+		return;
+	}
+
+	Initialize(LocalSplitDataObject, LocalRegisteredSplitterType, MoveTemp(LocalSplitter));
+}
+
+void FGenerationHelper::Initialize(const UObject* InSplitDataObject,
+	UE::Cook::Private::FRegisteredCookPackageSplitter* InRegisteredSplitterType,
+	TUniquePtr<ICookPackageSplitter>&& InCookPackageSplitterInstance)
+{
+	check(InSplitDataObject);
+	if (InitializeStatus != EInitializeStatus::Uninitialized)
+	{
+		// If we already have a splitter, keep the old and throw out the new. The old one
+		// still contains some state.
+		return;
+	}
+
+	RegisteredSplitterType = InRegisteredSplitterType;
+	CookPackageSplitterInstance = MoveTemp(InCookPackageSplitterInstance);
+	InitializeStatus = EInitializeStatus::Valid;
+
+	SplitDataObject = InSplitDataObject;
+	SplitDataObjectName = FName(FStringView(InSplitDataObject->GetFullName()));
+	bUseInternalReferenceToAvoidGarbageCollect =
+		CookPackageSplitterInstance->UseInternalReferenceToAvoidGarbageCollect();
+}
+
+void FGenerationHelper::InitializeAsInvalid()
+{
+	if (InitializeStatus != EInitializeStatus::Uninitialized)
+	{
+		return;
+	}
+	InitializeStatus = EInitializeStatus::Invalid;
+}
+
+UPackage* FGenerationHelper::FindOrLoadPackage(UCookOnTheFlyServer& COTFS, FPackageData& OwnerPackageData)
+{
+	// This is the static helper function on FGenerationHelper that loads the package for any FPackageData; for the
+	// member variable function that uses the cached pointer, see FindOrLoadOwnerPackage.
+	FName OwnerPackageName = OwnerPackageData.GetPackageName();
+	UPackage* Result = FindObjectFast<UPackage>(nullptr, OwnerPackageName);
+
+	if (!Result || !Result->IsFullyLoaded())
+	{
+		COTFS.LoadPackageForCooking(OwnerPackageData, Result);
+		if (!Result || !Result->IsFullyLoaded())
 		{
-			check(Info.PackageData->GetGeneratedOwner() == this);
-			Info.PackageData->SetGeneratedOwner(nullptr);
-			Info.PackageData = nullptr;
+			return nullptr;
 		}
 	}
+	return Result;
 }
 
-bool FGenerationHelper::TryGenerateList(UObject* OwnerObject, FPackageDatas& PackageDatas)
+void FGenerationHelper::SearchForRegisteredSplitDataObject(UCookOnTheFlyServer& COTFS,
+	FName PackageName, UPackage* Package,
+	TOptional<TConstArrayView<UE::Cook::FCachedObjectInOuter>> CachedObjectsInOuter,
+	UObject*& OutSplitDataObject, UE::Cook::Private::FRegisteredCookPackageSplitter*& OutRegisteredSplitter,
+	TUniquePtr<ICookPackageSplitter>& OutSplitterInstance)
 {
-	check(IsInitialized());
+	OutSplitDataObject = nullptr;
+	OutRegisteredSplitter = nullptr;
+	OutSplitterInstance = nullptr;
+	check(Package != nullptr || CachedObjectsInOuter.IsSet());
 
+	UObject* LocalSplitDataObject = nullptr;
+	Private::FRegisteredCookPackageSplitter* SplitterType = nullptr;
+	TArray<Private::FRegisteredCookPackageSplitter*> FoundRegisteredSplitters;
+	auto TryLookForSplitterOfObject =
+		[&COTFS, PackageName, &FoundRegisteredSplitters, &SplitterType, &LocalSplitDataObject](UObject* Obj)
+		{
+			FoundRegisteredSplitters.Reset();
+			COTFS.RegisteredSplitDataClasses.MultiFind(Obj->GetClass(), FoundRegisteredSplitters);
+
+			for (Private::FRegisteredCookPackageSplitter* SplitterForObject : FoundRegisteredSplitters)
+			{
+				if (SplitterForObject && SplitterForObject->ShouldSplitPackage(Obj))
+				{
+					if (!Obj->HasAnyFlags(RF_Public))
+					{
+						UE_LOG(LogCook, Error,
+							TEXT("SplitterData object %s must be publicly referenceable so we can keep them from being garbage collected"),
+							*Obj->GetFullName());
+						return false;
+					}
+
+					if (SplitterType)
+					{
+						UE_LOG(LogCook, Error,
+							TEXT("Found more than one registered Cook Package Splitter for package %s."),
+							*PackageName.ToString());
+						return false;
+					}
+
+					SplitterType = SplitterForObject;
+					LocalSplitDataObject = Obj;
+				}
+			}
+			return true;
+		};
+
+	if (CachedObjectsInOuter.IsSet())
+	{
+		// CachedObjectsInOuter might be set but empty for e.g. a generated package that has not been populated
+		for (const FCachedObjectInOuter& CachedObjectInOuter : *CachedObjectsInOuter)
+		{
+			UObject* Obj = CachedObjectInOuter.Object.Get();
+			if (!Obj)
+			{
+				continue;
+			}
+			if (!TryLookForSplitterOfObject(Obj))
+			{
+				return; // error condition, exit the entire search function
+			}
+		}
+	}
+	else
+	{
+		TArray<UObject*> ObjectsInPackage;
+		GetObjectsWithOuter(Package, ObjectsInPackage, true /* bIncludeNestedObjects */,
+			RF_NoFlags, EInternalObjectFlags::Garbage);
+		for (UObject* Obj : ObjectsInPackage)
+		{
+			if (!TryLookForSplitterOfObject(Obj))
+			{
+				return; // error condition, exit the entire search function
+			}
+		}
+	}
+
+	if (!SplitterType)
+	{
+		return;
+	}
+
+	// Create instance of CookPackageSplitter class
+	ICookPackageSplitter* SplitterInstance = SplitterType->CreateInstance(LocalSplitDataObject);
+	if (!SplitterInstance)
+	{
+		UE_LOG(LogCook, Error, TEXT("Error instantiating Cook Package Splitter %s for object %s."),
+			*SplitterType->GetSplitterDebugName(), *LocalSplitDataObject->GetFullName());
+		return;
+	}
+
+	OutSplitDataObject = LocalSplitDataObject;
+	OutRegisteredSplitter = SplitterType;
+	OutSplitterInstance.Reset(SplitterInstance);
+}
+
+void FGenerationHelper::ClearSelfReferences()
+{
+	// Any references we release might be the last reference and cause *this to be deleted,
+	// so create a local reference to keep it alive until the end of the function.
+	TRefCountPtr<FGenerationHelper> LocalRef(this);
+	ClearKeepForIterative();
+	ClearKeepForGeneratorSave();
+	ClearKeepForQueueResults();
+}
+
+FCookGenerationInfo* FGenerationHelper::FindInfo(const FPackageData& PackageData)
+{
+	ConditionalInitialize();
+
+	if (&PackageData == &GetOwner())
+	{
+		return &OwnerInfo;
+	}
+	for (FCookGenerationInfo& Info : PackagesToGenerate)
+	{
+		if (Info.PackageData == &PackageData)
+		{
+			return &Info;
+		}
+	}
+	return nullptr;
+}
+
+const FCookGenerationInfo* FGenerationHelper::FindInfo(const FPackageData& PackageData) const
+{
+	return const_cast<FGenerationHelper*>(this)->FindInfo(PackageData);
+}
+
+UObject* FGenerationHelper::FindOrLoadSplitDataObject()
+{
+	if (!IsValid())
+	{
+		return nullptr;
+	}
+	UObject* Result = SplitDataObject.Get();
+	if (Result)
+	{
+		return Result;
+	}
+
+	FString ObjectPath = GetSplitDataObjectName().ToString();
+	// SplitDataObjectName is a FullObjectPath; strip off the leading <ClassName> in
+	// "<ClassName> <Package>.<Object>:<SubObject>"
+	int32 ClassDelimiterIndex = -1;
+	if (ObjectPath.FindChar(' ', ClassDelimiterIndex))
+	{
+		ObjectPath.RightChopInline(ClassDelimiterIndex + 1);
+	}
+
+	Result = FindObject<UObject>(nullptr, *ObjectPath);
+	if (!Result)
+	{
+		FPackageData& OwnerPackageData = GetOwner();
+		FName OwnerPackageName = OwnerPackageData.GetPackageName();
+		UCookOnTheFlyServer& COTFS = OwnerPackageData.GetPackageDatas().GetCookOnTheFlyServer();
+		UPackage* LocalOwnerPackage;
+		COTFS.LoadPackageForCooking(OwnerPackageData, LocalOwnerPackage);
+
+		Result = FindObject<UObject>(nullptr, *ObjectPath);
+		if (!Result)
+		{
+			return nullptr;
+		}
+	}
+
+	SplitDataObject = Result;
+	return Result;
+}
+
+UPackage* FGenerationHelper::GetOwnerPackage()
+{
+	UPackage* Result = OwnerPackage.Get();
+	if (!Result && !OwnerPackage.GetEvenIfUnreachable())
+	{
+		OwnerPackage = FindObjectFast<UPackage>(nullptr, GetOwner().GetPackageName());
+		Result = OwnerPackage.Get();
+	}
+	return Result;
+}
+
+UPackage* FGenerationHelper::FindOrLoadOwnerPackage(UCookOnTheFlyServer& COTFS)
+{
+	UPackage* Result = GetOwnerPackage();
+	if (!Result)
+	{
+		Result = FindOrLoadPackage(COTFS, GetOwner());
+	}
+	return Result;
+}
+
+bool FGenerationHelper::TryGenerateList()
+{
+	if (bGeneratedList)
+	{
+		return true;
+	}
 	FPackageData& OwnerPackageData = GetOwner();
-	UPackage* LocalOwnerPackage = OwnerPackageData.GetPackage();
-	check(LocalOwnerPackage);
-	UCookOnTheFlyServer& COTFS = GetOwner().GetPackageDatas().GetCookOnTheFlyServer();
+	FName OwnerPackageName = OwnerPackageData.GetPackageName();
+	if (!IsValid())
+	{
+		// Unexpected, caller should not call in this case
+		UE_LOG(LogCook, Error, TEXT("TryGenerateList failed for package %s: Called on an invalid FGenerationHelper."),
+			*OwnerPackageName.ToString());
+		FDebug::DumpStackTraceToLog(ELogVerbosity::Warning);
+		return false;
+	}
 
+	FPackageDatas& PackageDatas = OwnerPackageData.GetPackageDatas();
+	UCookOnTheFlyServer& COTFS = PackageDatas.GetCookOnTheFlyServer();
+	UObject* OwnerObject = FindOrLoadSplitDataObject();
+	if (!OwnerObject)
+	{
+		// Unexpected, we found it earlier when we marked valid.
+		UE_LOG(LogCook, Error, TEXT("TryGenerateList failed for package %s: Valid GenerationHelper but could not find OwnerObject."),
+			*OwnerPackageName.ToString());
+		FDebug::DumpStackTraceToLog(ELogVerbosity::Warning);
+		return false;
+	}
+
+	UPackage* LocalOwnerPackage = OwnerObject->GetPackage();
+
 	TArray<ICookPackageSplitter::FGeneratedPackage> GeneratorDatas;
 	{
-		UCookOnTheFlyServer::FScopedActivePackage ScopedActivePackage(COTFS, OwnerPackageData.GetPackageName(),
+		UCookOnTheFlyServer::FScopedActivePackage ScopedActivePackage(COTFS, OwnerPackageName,
 			PackageAccessTrackingOps::NAME_CookerBuildObject);
 		GeneratorDatas = GetCookPackageSplitterInstance()->GetGenerateList(LocalOwnerPackage, OwnerObject);
 	}
 	PackagesToGenerate.Reset(GeneratorDatas.Num());
 	TArray<const ITargetPlatform*, TInlineAllocator<1>> PlatformsToCook;
 	OwnerPackageData.GetPlatformsNeedingCooking(PlatformsToCook);
-	bool bHybridIterativeEnabled = COTFS.bHybridIterativeEnabled;
 
-	int32 NumIterativeUnmodified = 0;
-	int32 NumIterativeModified = 0; 
-	int32 NumIterativeRemoved = 0;
-	int32 NumIterativePrevious = PreviousGeneratedPackages.Num();
-
 	for (ICookPackageSplitter::FGeneratedPackage& SplitterData : GeneratorDatas)
 	{
 		if (!SplitterData.GetCreateAsMap().IsSet())
 		{
 			UE_LOG(LogCook, Error,
 				TEXT("PackageSplitter did not specify whether CreateAsMap is true for generated package. Splitter=%s, Generated=%s."),
-				*this->GetSplitDataObjectName().ToString(), *OwnerPackageData.GetPackageName().ToString());
+				*this->GetSplitDataObjectName().ToString(), *OwnerPackageName.ToString());
 			return false;
 		}
 		bool bCreateAsMap = *SplitterData.GetCreateAsMap();
 
 		FString PackageRoot = SplitterData.GeneratedRootPath.IsEmpty() ?
-			OwnerPackageData.GetPackageName().ToString() : SplitterData.GeneratedRootPath;
+			OwnerPackageName.ToString() : SplitterData.GeneratedRootPath;
 
 		FString PackageName = FPaths::RemoveDuplicateSlashes(FString::Printf(TEXT("/%s/%s/%s"),
 			*PackageRoot, GeneratedPackageSubPath, *SplitterData.RelativePath));
@@ -1678,11 +2059,10 @@
 				*this->GetSplitDataObjectName().ToString(), *PackageName);
 			return false;
 		}
-		PackageData->SetGenerated(true);
-		// No package should be generated by two different splitters. If an earlier run of this splitter generated
-		// the package, the package's owner should have been reset to null when we called ClearGeneratedPackages
-		// between then and now
-		check(PackageData->GetGeneratedOwner() == nullptr);
+		// No package should be generated by two different splitters.
+		check(PackageData->GetParentGenerator().IsNone() ||
+			PackageData->GetParentGenerator() == OwnerPackageName);
+		PackageData->SetGenerated(OwnerPackageName);
 		if (IFileManager::Get().FileExists(*PackageData->GetFileName().ToString()))
 		{
 			UE_LOG(LogCook, Warning,
@@ -1709,47 +2089,285 @@
 			[](const FAssetDependency& A, const FAssetDependency& B) { return A.LexicalLess(B); });
 		GeneratedInfo.PackageDependencies.SetNum(Algo::Unique(GeneratedInfo.PackageDependencies));
 		GeneratedInfo.SetIsCreateAsMap(bCreateAsMap);
-		PackageData->SetGeneratedOwner(this);
-		PackageData->SetWorkerAssignmentConstraint(FWorkerId::Local());
+		if (COTFS.MPCookGeneratorSplit == UE::Cook::EMPCookGeneratorSplit::AllOnSameWorker)
+		{
+			PackageData->SetWorkerAssignmentConstraint(FWorkerId::Local());
+		}
 
 		// Create the Hash from the GenerationHash and Dependencies
 		GeneratedInfo.CreatePackageHash();
-
-		FIoHash PreviousHash;
-		if (PreviousGeneratedPackages.RemoveAndCopyValue(PackageFName, PreviousHash) && !bHybridIterativeEnabled)
-		{
-			bool bIterativelyUnmodified;
-			GeneratedInfo.IterativeCookValidateOrClear(*this, PlatformsToCook, PreviousHash, bIterativelyUnmodified);
-			++(bIterativelyUnmodified ? NumIterativeUnmodified : NumIterativeModified);
-		}
-	}
-	if (!PreviousGeneratedPackages.IsEmpty())
-	{
-		NumIterativeRemoved = PreviousGeneratedPackages.Num();
-		for (TPair<FName, FIoHash>& Pair : PreviousGeneratedPackages)
-		{
-			for (const ITargetPlatform* TargetPlatform : PlatformsToCook)
-			{
-				COTFS.DeleteOutputForPackage(Pair.Key, TargetPlatform);
-			}
-		}
-		PreviousGeneratedPackages.Empty();
-	}
-	if (NumIterativePrevious > 0 && !bHybridIterativeEnabled)
-	{
-		UE_LOG(LogCook, Display, TEXT("Found %d cooked package(s) in package store for generator package %s."),
-			NumIterativePrevious, *WriteToString<256>(GetOwner().GetPackageName()));
-		UE_LOG(LogCook, Display, TEXT("Keeping %d. Recooking %d. Removing %d."),
-			NumIterativeUnmodified, NumIterativeModified, NumIterativeRemoved);
 	}
 
-	RemainingToPopulate = GeneratorDatas.Num() + 1; // GeneratedPackaged plus one for the Generator
+	bGeneratedList = true;
 	return true;
 }
 
+void FGenerationHelper::StartOwnerSave()
+{
+	if (!IsValid())
+	{
+		return;
+	}
+	UE_LOG(LogCook, Display, TEXT("Splitting Package %s with splitter %s acting on object %s."),
+		*WriteToString<256>(GetOwner().GetPackageName()),
+		*GetRegisteredSplitterType()->GetSplitterDebugName(),
+		*WriteToString<256>(GetSplitDataObjectName()));
+	SetKeepForGeneratorSave();
+}
+
+void FGenerationHelper::StartQueueGeneratedPackages(UCookOnTheFlyServer& COTFS)
+{
+	if (!IsValid())
+	{
+		return;
+	}
+	NotifyStartQueueGeneratedPackages(COTFS, FWorkerId::Local());
+
+	bool bHybridIterativeEnabled = COTFS.bHybridIterativeEnabled;
+	if (!PreviousGeneratedPackages.IsEmpty())
+	{
+		FPackageData& OwnerPackageData = GetOwner();
+		TArray<const ITargetPlatform*, TInlineAllocator<1>> PlatformsToCook;
+		OwnerPackageData.GetPlatformsNeedingCooking(PlatformsToCook);
+
+		int32 NumIterativeUnmodified = 0;
+		int32 NumIterativeModified = 0;
+		int32 NumIterativeRemoved = 0;
+		int32 NumIterativePrevious = PreviousGeneratedPackages.Num();
+
+		for (FCookGenerationInfo& GeneratedInfo : PackagesToGenerate)
+		{
+			if (!GeneratedInfo.PackageData)
+			{
+				continue;
+			}
+			FIoHash PreviousHash;
+			if (PreviousGeneratedPackages.RemoveAndCopyValue(GeneratedInfo.PackageData->GetPackageName(), PreviousHash)
+				&& !bHybridIterativeEnabled)
+			{
+				bool bIterativelyUnmodified;
+				GeneratedInfo.IterativeCookValidateOrClear(*this, PlatformsToCook, PreviousHash, bIterativelyUnmodified);
+				++(bIterativelyUnmodified ? NumIterativeUnmodified : NumIterativeModified);
+			}
+		}
+		if (!PreviousGeneratedPackages.IsEmpty())
+		{
+			NumIterativeRemoved = PreviousGeneratedPackages.Num();
+			for (TPair<FName, FIoHash>& Pair : PreviousGeneratedPackages)
+			{
+				for (const ITargetPlatform* TargetPlatform : PlatformsToCook)
+				{
+					COTFS.DeleteOutputForPackage(Pair.Key, TargetPlatform);
+				}
+			}
+			PreviousGeneratedPackages.Empty();
+		}
+		ClearKeepForIterative();
+
+		if (NumIterativePrevious > 0 && !bHybridIterativeEnabled)
+		{
+			UE_LOG(LogCook, Display, TEXT("Found %d cooked package(s) in package store for generator package %s."),
+				NumIterativePrevious, *WriteToString<256>(GetOwner().GetPackageName()));
+			UE_LOG(LogCook, Display, TEXT("Keeping %d. Recooking %d. Removing %d."),
+				NumIterativeUnmodified, NumIterativeModified, NumIterativeRemoved);
+		}
+	}
+}
+
+void FGenerationHelper::NotifyStartQueueGeneratedPackages(UCookOnTheFlyServer& COTFS, FWorkerId SourceWorkerId)
+{
+	// Note this function can be called on an uninitialized Generator; the generator is only needed
+	// on the director so it can serve as the passer of messages. We have to keep ourselves referenced after
+	// this call, until after we send EGeneratorEvent::QueuedGeneratedPackagesFencePassed, so that we don't destruct
+	// and lose the WorkerIdThatSavedGenerator information.
+	if (COTFS.CookDirector)
+	{
+		COTFS.PackageDatas->GetRequestQueue().AddRequestFenceListener(GetOwner().GetPackageName());
+		WorkerIdThatSavedGenerator = SourceWorkerId;
+	}
+	SetKeepForQueueResults();
+}
+
+void FGenerationHelper::EndQueueGeneratedPackages(UCookOnTheFlyServer& COTFS)
+{
+	SetKeepForQueueResults();
+	COTFS.WorkerRequests->EndQueueGeneratedPackages(COTFS, *this);
+}
+
+void FGenerationHelper::EndQueueGeneratedPackagesOnDirector(UCookOnTheFlyServer& COTFS, FWorkerId SourceWorkerId)
+{
+	// Note this function can be called on an uninitialized Generator; the generator is only needed
+	// on the director so it can serve as the passer of messages.
+	COTFS.PackageDatas->GetRequestQueue().AddRequestFenceListener(GetOwner().GetPackageName());
+	SetKeepForQueueResults();
+
+	// Setting the WorkerIdThatSavedGenerator in response to this event is usually not needed because it is also set
+	// from the reported GeneratedPackages, but we set it anyway in case there is an edge condition that skips those
+	// notifications.
+	WorkerIdThatSavedGenerator = SourceWorkerId;
+}
+
+void FGenerationHelper::OnRequestFencePassedBroadcast(UCookOnTheFlyServer& COTFS)
+{
+	if (COTFS.CookDirector)
+	{
+		// Broadcast GenerationMessage API
+		COTFS.CookDirector->BroadcastGeneratorFencePassed(*this);
+	}
+	OnRequestFencePassed(COTFS);
+}
+
+void FGenerationHelper::OnRequestFencePassed(UCookOnTheFlyServer& COTFS)
+{
+	ClearKeepForQueueResults();
+	// We no longer need PreviousGeneratedPackages or KeepForIterative, because they are used (and cleared) in
+	// StartQueueGeneratedPackages. Clear them as well on the director and any CookWorkers that received it to
+	// free memory.
+	ClearKeepForIterative();
+	PreviousGeneratedPackages.Empty();
+}
+
+UPackage* FGenerationHelper::CreateGeneratedUPackage(FCookGenerationInfo& GeneratedInfo,
+	const UPackage* InOwnerPackage, const TCHAR* GeneratedPackageName)
+{
+	if (!IsValid())
+	{
+		return nullptr;
+	}
+
+#if ENABLE_COOK_STATS
+	++DetailedCookStats::NumRequestedLoads;
+#endif
+	UPackage* GeneratedPackage = CreatePackage(GeneratedPackageName);
+	GeneratedPackage->SetSavedHash(GeneratedInfo.PackageHash);
+	GeneratedPackage->SetPersistentGuid(InOwnerPackage->GetPersistentGuid());
+	GeneratedPackage->SetPackageFlags(PKG_CookGenerated);
+	GeneratedInfo.SetHasCreatedPackage(true);
+	if (!InOwnerPackage->IsLoadedByEditorPropertiesOnly())
+	{
+		GeneratedPackage->SetLoadedByEditorPropertiesOnly(false);
+	}
+
+	return GeneratedPackage;
+}
+
+void FGenerationHelper::ResetSaveState(FCookGenerationInfo& Info, UPackage* Package,
+	EStateChangeReason ReleaseSaveReason)
+{
+	ConditionalInitialize();
+
+	// We release references to *this in this function so keep a local reference to avoid deletion during the function.
+	TRefCountPtr<FGenerationHelper> LocalRefCount = this;
+
+	if (Info.GetSaveState() > FCookGenerationInfo::ESaveState::CallPopulate)
+	{
+		UObject* SplitObject = GetWeakSplitDataObject();
+		UPackage* LocalOwnerPackage = Info.IsGenerator() ? Package : GetOwnerPackage();
+		if (!SplitObject || !Package || !LocalOwnerPackage)
+		{
+			UE_LOG(LogCook, Warning,
+				TEXT("PackageSplitter: %s on %s was GarbageCollected before we finished saving it. This prevents us from calling PostSave and may corrupt other packages that it altered during Populate. Splitter=%s."),
+				(!Package ? TEXT("UPackage") :
+					(!LocalOwnerPackage ? TEXT("ParentGenerator UPackage") : TEXT("SplitDataObject"))),
+				Info.PackageData ? *Info.PackageData->GetPackageName().ToString() : *Info.RelativePath,
+				*GetSplitDataObjectName().ToString());
+		}
+		else
+		{
+			UCookOnTheFlyServer& COTFS = GetOwner().GetPackageDatas().GetCookOnTheFlyServer();
+			UCookOnTheFlyServer::FScopedActivePackage ScopedActivePackage(COTFS, GetOwner().GetPackageName(),
+				PackageAccessTrackingOps::NAME_CookerBuildObject);
+			if (Info.IsGenerator())
+			{
+				GetCookPackageSplitterInstance()->PostSaveGeneratorPackage(Package, SplitObject);
+			}
+			else
+			{
+				ICookPackageSplitter::FGeneratedPackageForPopulate PopulateInfo;
+				PopulateInfo.RelativePath = Info.RelativePath;
+				PopulateInfo.GeneratedRootPath = Info.GeneratedRootPath;
+				PopulateInfo.bCreatedAsMap = Info.IsCreateAsMap();
+				PopulateInfo.Package = Package;
+				GetCookPackageSplitterInstance()->PostSaveGeneratedPackage(LocalOwnerPackage, SplitObject,
+					PopulateInfo);
+			}
+		}
+	}
+
+	if (ReleaseSaveReason == EStateChangeReason::RecreateObjectCache ||
+		ReleaseSaveReason == EStateChangeReason::DoneForNow)
+	{
+		if (Info.IsGenerator())
+		{
+			if (Info.GetSaveState() >= FCookGenerationInfo::ESaveState::StartPopulate)
+			{
+				Info.SetSaveState(FCookGenerationInfo::ESaveState::StartPopulate);
+			}
+			else
+			{
+				// Redo all the steps since we didn't make it to the FinishCachePreObjectsToMove.
+				// Restarting in the middle of that flow is not robust
+				Info.SetSaveState(FCookGenerationInfo::ESaveState::StartSave);
+			}
+		}
+		else
+		{
+			Info.SetSaveState(FCookGenerationInfo::ESaveState::StartPopulate);
+		}
+	}
+	else
+	{
+		// The save is completed and we will not come back to it; set state back to initial
+		// state and drop our reference keeping this GenerationHelper in memory for the save.
+		if (ReleaseSaveReason == EStateChangeReason::Completed)
+		{
+			Info.SetHasSaved(true);
+		}
+
+		if (Info.IsGenerator())
+		{
+			ClearKeepForGeneratorSave();
+			Info.SetSaveState(FCookGenerationInfo::ESaveState::StartSave);
+		}
+		else
+		{
+			Info.SetSaveState(FCookGenerationInfo::ESaveState::StartPopulate);
+			if (Info.PackageData)
+			{
+				Info.PackageData->SetParentGenerationHelper(nullptr);
+			}
+		}
+	}
+
+	if (Info.HasTakenOverCachedCookedPlatformData())
+	{
+		if (Info.PackageData && Info.PackageData->GetCachedObjectsInOuter().Num() != 0 &&
+			IsUseInternalReferenceToAvoidGarbageCollect() &&
+			(ReleaseSaveReason != EStateChangeReason::Completed && ReleaseSaveReason != EStateChangeReason::DoneForNow
+				&& ReleaseSaveReason != EStateChangeReason::SaveError
+				 && ReleaseSaveReason != EStateChangeReason::CookerShutdown))
+		{
+			UE_LOG(LogCook, Error,
+				TEXT("CookPackageSplitter failure: We are demoting a %s package from save and removing our references that keep its objects loaded.\n")
+				TEXT("This will allow the objects to be garbage collected and cause failures in the splitter which expects them to remain loaded.\n")
+				TEXT("Package=%s, Splitter=%s, ReleaseSaveReason=%s"),
+				Info.IsGenerator() ? TEXT("generator") : TEXT("generated"),
+				Info.PackageData ? *Info.PackageData->GetPackageName().ToString() : *Info.RelativePath,
+				*GetSplitDataObjectName().ToString(), LexToString(ReleaseSaveReason));
+		}
+		Info.CachedObjectsInOuterInfo.Empty();
+		Info.SetHasTakenOverCachedCookedPlatformData(false);
+	}
+	Info.SetHasIssuedUndeclaredMovedObjectsWarning(false);
+	Info.KeepReferencedPackages.Reset();
+}
+
 void FGenerationHelper::FetchExternalActorDependencies()
 {
-	check(IsInitialized());
+	if (!IsValid())
+	{
+		return;
+	}
 
 	// The Generator package declares all its ExternalActor dependencies in its AssetRegistry dependencies
 	// The Generator's generated packages can also include ExternalActors from other maps due to level instancing,
@@ -1772,7 +2390,9 @@
 	FThreadSafeSet<FName>& NeverCookPackageList =
 		GetOwner().GetPackageDatas().GetCookOnTheFlyServer().PackageTracker->NeverCookPackageList;
 
-	// Keep only on-disk PackageDatas that are marked as NeverCook
+	// We are supposed to collect only ExternalActor dependencies, but we collected every dependency from the
+	// generated packages. Remove the packages that are not external actors, which we detect by being on-disk
+	// PackageDatas that are marked as NeverCook
 	ExternalActorDependencies.RemoveAll([&PackageDatas, &NeverCookPackageList](FName PackageName)
 		{
 			FPackageData* PackageData = PackageDatas.TryAddPackageDataByPackageName(PackageName);
@@ -1786,48 +2406,10 @@
 	ExternalActorDependencies.Shrink();
 }
 
-FCookGenerationInfo* FGenerationHelper::FindInfo(const FPackageData& PackageData)
+void FGenerationHelper::SetPreviousGeneratedPackages(TMap<FName, FIoHash>&& Packages)
 {
-	check(IsInitialized());
-
-	if (&PackageData == &GetOwner())
-	{
-		return &OwnerInfo;
-	}
-	for (FCookGenerationInfo& Info : PackagesToGenerate)
-	{
-		if (Info.PackageData == &PackageData)
-		{
-			return &Info;
-		}
-	}
-	return nullptr;
-}
-const FCookGenerationInfo* FGenerationHelper::FindInfo(const FPackageData& PackageData) const
-{
-	return const_cast<FGenerationHelper*>(this)->FindInfo(PackageData);
+	SetKeepForIterative();
+	PreviousGeneratedPackages = MoveTemp(Packages);
 }
 
-ICookPackageSplitter* FGenerationHelper::GetCookPackageSplitterInstance() const
-{
-	checkf(!bNotifiedCompletion,
-		TEXT("It is illegal for the cooker to try to access the CookPackageSplitterInstance after calling Teardown on it."));
-	return CookPackageSplitterInstance.Get();
-}
-
-UObject* FGenerationHelper::FindSplitDataObject() const
-{
-	check(IsInitialized());
-	FString ObjectPath = GetSplitDataObjectName().ToString();
-
-	// SplitDataObjectName is a FullObjectPath; strip off the leading <ClassName> in
-	// "<ClassName> <Package>.<Object>:<SubObject>"
-	int32 ClassDelimiterIndex = -1;
-	if (ObjectPath.FindChar(' ', ClassDelimiterIndex))
-	{
-		ObjectPath.RightChopInline(ClassDelimiterIndex + 1);
-	}
-	return FindObject<UObject>(nullptr, *ObjectPath);
-}
-
 void FGenerationHelper::PreGarbageCollect(FCookGenerationInfo& Info, TArray<TObjectPtr<UObject>>& GCKeepObjects,
@@ -1894,7 +2476,7 @@
 	{
 		// UCookOnTheFlyServer::PreGarbageCollect adds references for the Generator package and all its public
 		// objects, so it should still be loaded
-		if (!Owner.GetPackage() || !FindSplitDataObject())
+		if (!Owner.GetPackage() || !GetWeakSplitDataObject())
 		{
 			UE_LOG(LogCook, Error,
 				TEXT("PackageSplitter object was deleted by garbage collection while generation was still ongoing. This will break the generation.")
@@ -1910,8 +2492,7 @@
 		UPackage* LocalOwnerPackage = FindObject<UPackage>(nullptr, *Owner.GetPackageName().ToString());
 		if (LocalOwnerPackage)
 		{
-			if (RemainingToPopulate > 0 &&
-				!Owner.IsKeepReferencedDuringGC() &&
+			if (!Owner.IsInProgress() && !Owner.IsKeepReferencedDuringGC() &&
 				!IsUseInternalReferenceToAvoidGarbageCollect())
 			{
 				UE_LOG(LogCook, Error,
@@ -1956,132 +2537,6 @@
 	}
 }
 
-UPackage* FGenerationHelper::CreateGeneratedUPackage(FCookGenerationInfo& GeneratedInfo,
-	const UPackage* InOwnerPackage, const TCHAR* GeneratedPackageName)
-{
-	check(IsInitialized());
-#if ENABLE_COOK_STATS
-	++DetailedCookStats::NumRequestedLoads;
-#endif
-	UPackage* GeneratedPackage = CreatePackage(GeneratedPackageName);
-	GeneratedPackage->SetSavedHash(GeneratedInfo.PackageHash);
-	GeneratedPackage->SetPersistentGuid(InOwnerPackage->GetPersistentGuid());
-	GeneratedPackage->SetPackageFlags(PKG_CookGenerated);
-	GeneratedInfo.SetHasCreatedPackage(true);
-	if (!InOwnerPackage->IsLoadedByEditorPropertiesOnly())
-	{
-		GeneratedPackage->SetLoadedByEditorPropertiesOnly(false);
-	}
-
-	return GeneratedPackage;
-}
-
-void FGenerationHelper::SetPackageSaved(FCookGenerationInfo& Info, FPackageData& PackageData)
-{
-	check(IsInitialized());
-	if (Info.HasSaved())
-	{
-		return;
-	}
-	Info.SetHasSaved(true);
-	--RemainingToPopulate;
-	check(RemainingToPopulate >= 0);
-	if (IsComplete())
-	{
-		ConditionalNotifyCompletion(ICookPackageSplitter::ETeardown::Complete);
-	}
-}
-
-bool FGenerationHelper::IsComplete() const
-{
-	check(IsInitialized());
-	return RemainingToPopulate == 0;
-}
-
-void FGenerationHelper::ResetSaveState(FCookGenerationInfo& Info, UPackage* Package,
-	EStateChangeReason ReleaseSaveReason)
-{
-	check(IsInitialized());
-	if (Info.GetSaveState() > FCookGenerationInfo::ESaveState::CallPopulate)
-	{
-		UObject* SplitObject = FindSplitDataObject();
-		if (!SplitObject || !Package)
-		{
-			UE_LOG(LogCook, Warning,
-				TEXT("PackageSplitter: %s on %s was GarbageCollected before we finished saving it. This prevents us from calling PostSave and may corrupt other packages that it altered during Populate. Splitter=%s."),
-				(!Package ? TEXT("UPackage") : TEXT("SplitDataObject")),
-				Info.PackageData ? *Info.PackageData->GetPackageName().ToString() : *Info.RelativePath,
-				*GetSplitDataObjectName().ToString());
-		}
-		else
-		{
-			UCookOnTheFlyServer& COTFS = Info.PackageData->GetPackageDatas().GetCookOnTheFlyServer();
-			UCookOnTheFlyServer::FScopedActivePackage ScopedActivePackage(COTFS, GetOwner().GetPackageName(),
-				PackageAccessTrackingOps::NAME_CookerBuildObject);
-			if (Info.IsGenerator())
-			{
-				GetCookPackageSplitterInstance()->PostSaveGeneratorPackage(Package, SplitObject);
-			}
-			else
-			{
-				ICookPackageSplitter::FGeneratedPackageForPopulate PopulateInfo;
-				PopulateInfo.RelativePath = Info.RelativePath;
-				PopulateInfo.GeneratedRootPath = Info.GeneratedRootPath;
-				PopulateInfo.bCreatedAsMap = Info.IsCreateAsMap();
-				PopulateInfo.Package = Package;
-				GetCookPackageSplitterInstance()->PostSaveGeneratedPackage(GetOwnerPackage(), SplitObject,
-					PopulateInfo);
-			}
-		}
-	}
-
-	if (Info.IsGenerator())
-	{
-		if (ReleaseSaveReason == EStateChangeReason::RecreateObjectCache ||
-			ReleaseSaveReason == EStateChangeReason::DoneForNow)
-		{
-			if (Info.GetSaveState() >= FCookGenerationInfo::ESaveState::StartPopulate)
-			{
-				Info.SetSaveState(FCookGenerationInfo::ESaveState::StartPopulate);
-			}
-			else
-			{
-				// Redo all the steps since we didn't make it to the FinishCachePreObjectsToMove.
-				// Restarting in the middle of that flow is not robust
-				Info.SetSaveState(FCookGenerationInfo::ESaveState::StartGenerate);
-			}
-		}
-		else
-		{
-			Info.SetSaveState(FCookGenerationInfo::ESaveState::StartGenerate);
-		}
-	}
-	else
-	{
-		Info.SetSaveState(FCookGenerationInfo::ESaveState::StartPopulate);
-	}
-	if (Info.HasTakenOverCachedCookedPlatformData())
-	{
-		if (Info.PackageData && Info.PackageData->GetCachedObjectsInOuter().Num() != 0 &&
-			IsUseInternalReferenceToAvoidGarbageCollect() &&
-			(ReleaseSaveReason != EStateChangeReason::Completed && ReleaseSaveReason != EStateChangeReason::DoneForNow
-				&& ReleaseSaveReason != EStateChangeReason::SaveError
-				&& ReleaseSaveReason != EStateChangeReason::CookerShutdown))
-		{
-			UE_LOG(LogCook, Error,
-				TEXT("CookPackageSplitter failure: We are demoting a %s package from save and removing our references that keep its objects loaded.\n")
-				TEXT("This will allow the objects to be garbage collected and cause failures in the splitter which expects them to remain loaded.\n")
-				TEXT("Package=%s, Splitter=%s, ReleaseSaveReason=%s"),
-				Info.IsGenerator() ? TEXT("generator") : TEXT("generated"),
-				Info.PackageData ? *Info.PackageData->GetPackageName().ToString() : *Info.RelativePath,
-				*GetSplitDataObjectName().ToString(), LexToString(ReleaseSaveReason));
-		}
-		Info.SetHasTakenOverCachedCookedPlatformData(false);
-	}
-	Info.SetHasIssuedUndeclaredMovedObjectsWarning(false);
-	Info.KeepReferencedPackages.Reset();
-}
-
 void FGenerationHelper::UpdateSaveAfterGarbageCollect(const FPackageData& PackageData, bool& bInOutDemote)
 {
 	if (!IsInitialized())
@@ -2185,7 +2640,7 @@
 
 FCookGenerationInfo::FCookGenerationInfo(FPackageData& InPackageData, bool bInGenerator)
 	: PackageData(&InPackageData)
-	, GeneratorSaveState(bInGenerator ? ESaveState::StartGenerate : ESaveState::StartPopulate)
+	, GeneratorSaveState(bInGenerator ? ESaveState::StartSave : ESaveState::StartPopulate)
 	, bCreateAsMap(false), bHasCreatedPackage(false), bHasSaved(false), bTakenOverCachedCookedPlatformData(false)
 	, bIssuedUndeclaredMovedObjectsWarning(false), bGenerator(bInGenerator)
 {
@@ -3203,8 +3658,7 @@
 		[&PackageData](FThreadsafePackageData& ThreadsafeData, bool bNew)
 		{
 			ThreadsafeData.Instigator = PackageData.GetInstigator();
-			FGenerationHelper* GenerationHelper = PackageData.GetGeneratedOwner();
-			ThreadsafeData.Generator = GenerationHelper ? GenerationHelper->GetOwner().GetPackageName() : NAME_None;
+			ThreadsafeData.Generator = PackageData.GetParentGenerator();
 		});
 }
 
@@ -3242,7 +3696,7 @@
 {
 	FWriteScopeLock ExistenceWriteLock(ExistenceLock);
 	PendingCookedPlatformDataLists.Empty(); // These destructors will read/write PackageDatas
-	RequestQueue.Empty();
+	RequestQueue.Empty(*this);
 	SaveQueue.Empty();
 	PackageNameToPackageData.Empty();
 	FileNameToPackageData.Empty();
@@ -3479,8 +3933,12 @@
 	UpdateThreadsafePackageData(PackageData);
 }
 
-void FRequestQueue::Empty()
+void FRequestQueue::Empty(FPackageDatas& PackageDatas)
 {
+	RestartedRequests.Empty();
+	DiscoveryQueue.Empty();
+	RequestClusters.Empty();
+	NotifyRequestFencePassed(PackageDatas);
 	NormalRequests.Empty();
 	UrgentRequests.Empty();
 }
@@ -3601,6 +4059,28 @@
 	}
 }
 
+void FRequestQueue::AddRequestFenceListener(FName PackageName)
+{
+	RequestFencePackageListeners.Add(PackageName);
+}
+
+void FRequestQueue::NotifyRequestFencePassed(FPackageDatas& PackageDatas)
+{
+	for (FName PackageName : RequestFencePackageListeners)
+	{
+		FPackageData* PackageData = PackageDatas.FindPackageDataByPackageName(PackageName);
+		if (PackageData)
+		{
+			TRefCountPtr<FGenerationHelper> GenerationHelper = PackageData->GetGenerationHelper();
+			if (GenerationHelper)
+			{
+				GenerationHelper->OnRequestFencePassedBroadcast(PackageDatas.GetCookOnTheFlyServer());
+			}
+		}
+	}
+	RequestFencePackageListeners.Empty();
+}
+
 bool FLoadPrepareQueue::IsEmpty()
 {
 	return Num() == 0;
diff --git a/Engine/Source/Editor/UnrealEd/Private/Cooker/CookPackageData.h b/Engine/Source/Editor/UnrealEd/Private/Cooker/CookPackageData.h
--- a/Engine/Source/Editor/UnrealEd/Private/Cooker/CookPackageData.h
+++ b/Engine/Source/Editor/UnrealEd/Private/Cooker/CookPackageData.h
@@ -17,6 +17,7 @@
 #include "Misc/EnumClassFlags.h"
 #include "Misc/Optional.h"
 #include "Misc/ScopeRWLock.h"
+#include "Templates/RefCounting.h"
 #include "Templates/SharedPointer.h"
 #include "Templates/UniquePtr.h"
 #include "TypedBlockAllocator.h"
@@ -611,29 +612,59 @@ public:
 	/** Swap all ITargetPlatform* stored on this instance according to the mapping in @param Remap. */
 	void RemapTargetPlatforms(const TMap<ITargetPlatform*, ITargetPlatform*>& Remap);
 
-	/** Create a GenerationHelper struct for this PackageData and the given CookPackageSplitter and return it. */
-	FGenerationHelper& CreateGenerationHelper(const UObject* InSplitDataObject,
-		ICookPackageSplitter* InCookPackageSplitterInstance);
-	/** Destroy this PackageData's GenerationHelper struct. */
-	void DestroyGenerationHelper() { GenerationHelper.Reset(); }
-	/** Get the helper object containing data and functions for packages that need to generate other packages. */
-	FGenerationHelper* GetGenerationHelper() const;
+	// GenerationHelper is set on packages that are generator packages: they generate other packages
+	// during cook. The PackageData for the generator package has a pointer to the GenerationHelper to look
+	// it up, but does not keep it in memory. It is kept in memory by its internal state and by
+	// ParentGenerationHelper references from its generated packages.
+	/** Return the GenerationHelper if it already exists, otherwise return nullptr. */
+	TRefCountPtr<FGenerationHelper> GetGenerationHelper() const;
+	/**
+	 * Return the GenerationHelper if it already exists and is initialized and is valid. If not initialized,
+	 * load the package to initialize it. If it is not valid after initialization, return nullptr.
+	 */
+	TRefCountPtr<FGenerationHelper> GetGenerationHelperIfValid();
+	/** Return the GenerationHelper if it already exists, otherwise create it without calling Initialize. */
+	TRefCountPtr<FGenerationHelper> CreateUninitializedGenerationHelper();
+	/**
+	 * Return the GenerationHelper if it already exists, otherwise load the package and check all the objects
+	 * to see whether any of them have a registered CookPackageSplitter. If a split object exists, create the
+	 * GenerationHelper and return it. Otherwise return nullptr.
+	 */
+	TRefCountPtr<FGenerationHelper> TryCreateValidGenerationHelper();
+	/** Helper function for ~FGenerationHelper: clear the pointer this->GenerationHelper. */
+	void OnGenerationHelperDestroyed(FGenerationHelper& InGenerationHelper);
 	/** Get whether the PackageData has done any necessary Generator steps and is ready for BeginCache calls. */
-	bool HasCompletedGeneration() const { return static_cast<bool>(bCompletedGeneration); }
+	bool HasCompletedGeneration() const;
 	/** Set whether the PackageData has done any necessary Generator steps and is ready for BeginCache calls. */
-	void SetCompletedGeneration(bool Value) { bCompletedGeneration = Value != 0; }
+	void SetCompletedGeneration(bool Value);
 	/** Get whether the PackageData has completed the check for whether it is a Generator. */
-	bool HasInitializedGeneratorSave() const { return static_cast<bool>(bInitializedGeneratorSave); }
+	bool HasInitializedGeneratorSave() const;
 	/** Set whether the PackageData has completed the check for whether it is a Generator. */
-	void SetInitializedGeneratorSave(bool Value) { bInitializedGeneratorSave = Value != 0; }
+	void SetInitializedGeneratorSave(bool Value);
 	/** Return whether the PackageData is a generated package created by a owning generator PackageData. */
-	bool IsGenerated() const { return static_cast<bool>(bGenerated); }
-	/** Set whether the PackageData is a generated package created by a owning generator PackageData. */
-	void SetGenerated(bool Value) { bGenerated = Value != 0; }
-	/** Set the owning generator PackageData. Only valid to call if SetGenerated(true) has been called. */
-	void SetGeneratedOwner(FGenerationHelper* InGeneratedOwner);
-	/** Return the parent generator's GenerationHelper. Will be null if not a generated package or if orphaned. */
-	FGenerationHelper* GetGeneratedOwner() const { return GeneratedOwner; }
+	bool IsGenerated() const;
+	/** Mark that the PackageData is a generated package created by a ParentGenerator. */
+	void SetGenerated(FName InParentGenerator);
+	/** Return the name of the generator package that generates this package, or NAME_None if not IsGenerated. */
+	FName GetParentGenerator() const;
+	/**
+	 * Set the owning generator PackageData. Only valid to call with non-null if SetGenerated() has been called. Keeps
+	 * the GenerationHelper referenced until SetParentGenerationHelper(nullptr) is called.
+	 */
+	void SetParentGenerationHelper(FGenerationHelper* InGenerationHelper);
+	/** Return the ParentGenerator's GenerationHelper if the pointer to it has already been set on this. */
+	TRefCountPtr<FGenerationHelper> GetParentGenerationHelper() const;
+	/**
+	 * Return the ParentGenerator's GenerationHelper if the pointer to it has already been set on this, otherwise
+	 * look for it on the ParentGenerator by calling GetGenerationHelper, and if found, store a reference to it
+	 * and return it, otherwise return null.
+	 */
+	TRefCountPtr<FGenerationHelper> GetOrFindParentGenerationHelper();
+	/**
+	 * Return the ParentGenerator's GenerationHelper if the pointer to it has already been set on this, otherwise
+	 * try to find or create it by calling TryCreateValidGenerationHelper on the ParentGenerator.
+	 */
+	TRefCountPtr<FGenerationHelper> TryCreateValidParentGenerationHelper();
 
 	/**
 	 * Return the instigator for this package. The Instigator is the first code location or
@@ -755,8 +786,8 @@ private:
 
 	void OnPackageDataFirstMarkedReachable(FInstigator&& InInstigator);
 
-	TUniquePtr<FGenerationHelper> GenerationHelper;
-	FGenerationHelper* GeneratedOwner;
+	FGenerationHelper* GenerationHelper = nullptr;
+	TRefCountPtr<FGenerationHelper> ParentGenerationHelper;
 	/** Data for each platform that has been interacted with by *this. */
 	TSortedMap<const ITargetPlatform*, FPackagePlatformData, TInlineAllocator<1>> PlatformDatas;
 
@@ -765,6 +796,7 @@ private:
 	TUniquePtr<TMap<FPackageData*, EInstigator>> Unsolicited;
 	FName PackageName;
 	FName FileName;
+	FName ParentGenerator;
 
 	struct FAsyncRequest
 	{
@@ -827,10 +859,9 @@ public:
 	/** State variable for reentrant SplitPackage calls */
 	enum class ESaveState : uint8
 	{
-		StartGenerate = 0,
+		StartSave = 0,
 
-		GenerateList = StartGenerate,
-		ClearOldPackagesFirstAttempt,
+		ClearOldPackagesFirstAttempt = StartSave,
 		ClearOldPackagesLastAttempt,
 		QueueGeneratedPackages,
 
@@ -912,7 +943,7 @@ public:
 	TArray<UPackage*> KeepReferencedPackages;
 	TMap<UObject*, FCachedObjectInOuterGeneratorInfo> CachedObjectsInOuterInfo;
 private:
-	ESaveState GeneratorSaveState = ESaveState::StartGenerate;
+	ESaveState GeneratorSaveState = ESaveState::StartSave;
 	bool bCreateAsMap : 1;
 	bool bHasCreatedPackage : 1;
 	bool bHasSaved : 1;
@@ -922,112 +953,215 @@ private:
 };
 
 /**
- * Helper that wraps a ICookPackageSplitter, gets/caches packages to generate and provide
- * the necessary to iterate over all of them iteratively.
+ * Helper that wraps an ICookPackageSplitter, gets/caches packages to generate, and is a reference-counted
+ * collection of cached data and helper functions to save and list the generated packages.
  */
-struct FGenerationHelper
+struct FGenerationHelper : public FThreadSafeRefCountedObject
 {
 public:
 	/** Store the provided CookPackageSplitter and prepare the packages to generate. */
-	FGenerationHelper(UE::Cook::FPackageData& InOwner, const UObject* InSplitDataObject,
-		ICookPackageSplitter* InCookPackageSplitterInstance);
+	FGenerationHelper(UE::Cook::FPackageData& InOwner);
 	~FGenerationHelper();
-	void InitializeSave(const UObject* InSplitDataObject, ICookPackageSplitter* InCookPackageSplitterInstance);
-	bool IsInitialized() const { return bInitialized; }
-	/** Clear references to owned generated packages, and mark those packages as orphaned */
-	void ClearGeneratedPackages();
 
-	/** Call the Splitter's GetGenerateList and create the PackageDatas */
-	bool TryGenerateList(UObject* OwnerObject, FPackageDatas& PackageDatas);
 	/**
-	 * Record all dependencies from the generator package that are ExternalActor dependencies and
-	 * store them on this.
+	 * Early exits if already initialized. Otherwise, loads the package if not loaded and searches it for a
+	 * splitter and creates the splitter.
+	 * If load/search/creation fails, the GenerationHelper will be set to invalid @see IsValid.
+	 * 
+	 * Unless otherwise stated, all public functions call Initialize if not already initialized.
 	 */
-	void FetchExternalActorDependencies();
+	void Initialize();
+	/**
+	 * Version of Initialize that receives the splitterobject from the caller rather than needing to search
+	 * for it, as an optimization for callers that have it already. If already initialized, the input Splitter
+	 * will be destroyed and the function will early exit.
+	 */
+	void Initialize(const UObject* InSplitDataObject,
+		UE::Cook::Private::FRegisteredCookPackageSplitter* InRegisteredSplitterType,
+		TUniquePtr<ICookPackageSplitter>&& InSplitter);
+	/** Version of Initialize that sets IsValid=false. */
+	void InitializeAsInvalid();
+	/**
+	 * Clear all self references which keep this GenerationHelper from destructing. These could have been set by the
+	 * SetKeep... functions. The Owner PackageData will keep a pointer to this until all references from self, from
+	 * child generated PackageDatas, and other, are released. Note this means that we assert in the owner PackageData's
+	 * destructor if any references remain; PackageDatas can not be deleted until all such references are cleared.
+	 * Does not call Initialize.
+	 */
+	void ClearSelfReferences();
 
-	/** Accessor for the packages to generate */
-	TArrayView<UE::Cook::FCookGenerationInfo> GetPackagesToGenerate()
-	{
-		check(IsInitialized());
-		return PackagesToGenerate;
-	}
-	/** Return the GenerationInfo used to save the generator package's UPackage */
-	UE::Cook::FCookGenerationInfo& GetOwnerInfo() { check(IsInitialized()); return OwnerInfo; }
-	/** Return owner FPackageData. */
-	UE::Cook::FPackageData& GetOwner() { return *OwnerInfo.PackageData; }
+	/**
+	 * GenerationHelpers can be created uninitialized at the start of cook, for incremental cooks. We initialize
+	 * them on demand, which requires loading the package if not already loaded. 
+	 * IsInitialized reports whether initialization has been attempted. @see IsValid.
+	 * Does not call Initialize.
+	 */
+	bool IsInitialized() const;
+	/**
+	 * Since GenerationHelpers can be created before we load the package (@see IsInitialized) we might incorrectly
+	 * create one for a package that does not have a splitter. If so, when we load the package and discover that,
+	 * we set status to invalid. When in this state, all public functions are still valid to call, but will be
+	 * noops and will return empty data.
+	 */
+	bool IsValid();
+
+	/** Accessor for the packages to generate, will be empty if invalid or if TryGenerateList was not yet called. */
+	TArrayView<UE::Cook::FCookGenerationInfo> GetPackagesToGenerate();
+	/** Return the GenerationInfo used to save the generator package's UPackage. Does not call Initialize. */
+	UE::Cook::FCookGenerationInfo& GetOwnerInfo();
+	/** Return owner FPackageData. Does not call Initialize. */
+	UE::Cook::FPackageData& GetOwner();
 	/** Return the GenerationInfo for the given PackageData, or null if not found. */
 	UE::Cook::FCookGenerationInfo* FindInfo(const FPackageData& PackageData);
 	const UE::Cook::FCookGenerationInfo* FindInfo(const FPackageData& PackageData) const;
-
-	/** Return CookPackageSplitter. */
+	/** Return CookPackageSplitter. Will return null if !IsValid.*/
 	ICookPackageSplitter* GetCookPackageSplitterInstance() const;
-	/** Return the SplitDataObject's FullObjectPath. */
-	const FName GetSplitDataObjectName() const { check(IsInitialized()); return SplitDataObjectName; }
-	/** Return the Splitter's value for virtual bool UseInternalReferenceToAvoidGarbageCollect() */
-	bool IsUseInternalReferenceToAvoidGarbageCollect() const { return bUseInternalReferenceToAvoidGarbageCollect; }
-
+	/** Return RegisteredSplitterType. Will return null if !IsValid. */
+	UE::Cook::Private::FRegisteredCookPackageSplitter* GetRegisteredSplitterType() const;
+	/** Return the SplitDataObject's FullObjectPath. Will return NAME_None if !IsValid. */
+	const FName GetSplitDataObjectName() const;
+	/** Return the Splitter's value for virtual bool UseInternalReferenceToAvoidGarbageCollect(). */
+	bool IsUseInternalReferenceToAvoidGarbageCollect() const;
+	/** Return the cached pointer to the SplitDataObject. Returns null if no longer in memory or marked as garbage. */
+	UObject* GetWeakSplitDataObject() const;
 	/**
-	 * Find again the split object from its name, or return null if no longer in memory.
-	 * It may have been GC'd and reloaded since the last time we used it.
+	 * Return our cached pointer to the SplitDataObject if set. If not set, load the package and find it. Can still
+	 * return null if !IsValid or if not found even after loading the package.
 	 */
-	UObject* FindSplitDataObject() const;
+	UObject* FindOrLoadSplitDataObject();
+	
+	/** Find the OwnerPackage in memory, returns null if invalid or not already loaded. Does not call Initialize. */
+	UPackage* GetOwnerPackage();
+	/** Load the OwnerPackage if not already loaded. Can still return null if invalid or package fails to load. */
+	UPackage* FindOrLoadOwnerPackage(UCookOnTheFlyServer& COTFS);
+	/**
+	 * Return a reference to ExternalActors that were discovered during save and stored on this.
+	 * Does not call initialize.
+	 */
+	TConstArrayView<FName> GetExternalActorDependencies();
+	/**
+	 * Return the ExternalActors that were discovered during save and stored on this, and clear the values.
+	 * Does not call initialize.
+	 */
+	TArray<FName> ReleaseExternalActorDependencies();
 
+	/** Call the Splitter's GetGenerateList and create the PackageDatas. Logs errors and returns false on failure. */
+	bool TryGenerateList();
+	/**
+	 * Mark that the SavePackage of the Owner is starting. Keeps a reference to keep the generator alive until save
+	 * is finished.
+	 */
+	void StartOwnerSave();
+	/** Update state before we queue the generated packages, e.g. mark whether packages are iteratively skippable. */
+	void StartQueueGeneratedPackages(UCookOnTheFlyServer& COTFS);
+	/**
+	 * Mark that we have started queuing packages. Automatically called from StartQueueGeneratedPackages, but also
+	 * is called on the director in response to a discovered generated package from a CookWorker. It sets the
+	 * WorkerId so that it is available during assignment of the discovered generated packages.
+	 */
+	void NotifyStartQueueGeneratedPackages(UCookOnTheFlyServer& COTFS, FWorkerId SourceWorkerId);
+	/** Update state after we queue generated packages; e.g. register to stay alive until the packages are assigned. */
+	void EndQueueGeneratedPackages(UCookOnTheFlyServer& COTFS);
+	/** Does not call Initialize. */
+	void EndQueueGeneratedPackagesOnDirector(UCookOnTheFlyServer& COTFS, FWorkerId SourceWorkerId);
+	/**
+	 * Called from Director when the RequestFence added from EndQueueGeneratedPackages has passed. Calls
+	 * OnRequestFencePassed locally and on all Workers.
+	 */
+	void OnRequestFencePassedBroadcast(UCookOnTheFlyServer& COTFS);
+	/** Called from when the RequestFence added from EndQueueGeneratedPackages has passed. */
+	void OnRequestFencePassed(UCookOnTheFlyServer& COTFS);
+	/** Call CreatePackage and set the necessary data on the UPackage, and update status. */
+	UPackage* CreateGeneratedUPackage(FCookGenerationInfo& GenerationInfo, const UPackage* OwnerPackage,
+		const TCHAR* GeneratedPackageName);
+	/**
+	 * Clear any data that should only be held when an FPackageData is in the save state, for the given Info. The given
+	 * Info might specify the generator package or one of the generated packages.
+	 */
 	void ResetSaveState(FCookGenerationInfo& Info, UPackage* Package, UE::Cook::EStateChangeReason ReleaseSaveReason);
 
-	int32& GetNextPopulateIndex() { check(IsInitialized()); return NextPopulateIndex; }
+	/**
+	 * Record all dependencies from the generator package that are ExternalActor dependencies and
+	 * store them on this.
+	 */
+	void FetchExternalActorDependencies();
+	/** Iterative cook: store the list of generated packages from the last cook. Does not call Initialize. */
+	void SetPreviousGeneratedPackages(TMap<FName, FIoHash>&& Packages);
+	/** Return the information set by SetPreviousGeneratedPackages if not yet cleared. Does not call Initialize. */
+	const TMap<FName, FIoHash>& GetPreviousGeneratedPackages() const;
 
-	/** Callbacks during garbage collection */
+	/** Callback during garbage collection. Does not call initialize. */
 	void PreGarbageCollect(FCookGenerationInfo& Info, TArray<TObjectPtr<UObject>>& GCKeepObjects,
 		TArray<UPackage*>& GCKeepPackages, TArray<FPackageData*>& GCKeepPackageDatas, bool& bOutShouldDemote);
+	/** Callback during garbage collection. Does not call initialize. */
 	void PostGarbageCollect();
-
-	/** Call CreatePackage and set PackageData for deterministic generated packages, and update status. */
-	UPackage* CreateGeneratedUPackage(FCookGenerationInfo& GenerationInfo,
-		const UPackage* OwnerPackage, const TCHAR* GeneratedPackageName);
 	/**
-	 * Mark that the generator package or a generated package has saved, to keep track of when *this
-	 * is no longer needed.
+	 * Called from PackageData function of the same name to decide whether to demote the package out of save.
+	 * Does not call Initialize.
 	 */
-	void SetPackageSaved(FCookGenerationInfo& Info, FPackageData& PackageData);
-	/** Return whether list has been generated and all generated packages have been populated */
-	bool IsComplete() const;
-
 	void UpdateSaveAfterGarbageCollect(const FPackageData& PackageData, bool& bInOutDemote);
 
-	UPackage* GetOwnerPackage() const { return OwnerPackage.Get(); };
-	void SetOwnerPackage(UPackage* InPackage) { OwnerPackage = InPackage; }
-	void SetPreviousGeneratedPackages(TMap<FName, FIoHash>&& Packages)
-	{
-		PreviousGeneratedPackages = MoveTemp(Packages);
-	}
+	// Self-references that keep this GenerationHelper in memory and referenced from the packages that use it until
+	// the self-references are cleared. These Set/Clear functions do not call initialize.
+	void SetKeepForIterative();
+	void ClearKeepForIterative();
+	void SetKeepForQueueResults();
+	void ClearKeepForQueueResults();
+	void SetKeepForGeneratorSave();
+	void ClearKeepForGeneratorSave();
+
+	/**
+	 * Helper for assignment of generated packages in MPCook. Return the id of the CookWorker that saved the
+	 * generator, to decide where to assign the generated packages.
+	 */
+	FWorkerId GetWorkerIdThatSavedGenerator() const;
+	/**
+	 * Helper for assignment of generated packages in MPCook. A counter for assigned generated packages that is
+	 * needed by some assignment schemes.
+	 */
+	int32& GetMPCookNextAssignmentIndex();
+
+	/** Helper function for Initialize and for TryCreateValidGenerationHelper. */
+	static void SearchForRegisteredSplitDataObject(UCookOnTheFlyServer& COTFS, FName PackageName, UPackage* Package,
+		TOptional<TConstArrayView<FCachedObjectInOuter>> CachedObjectsInOuter, UObject*& OutSplitDataObject,
+		UE::Cook::Private::FRegisteredCookPackageSplitter*& OutRegisteredSplitterType,
+		TUniquePtr<ICookPackageSplitter>& OutSplitterInstance);
+	/** Helper function for Initialize and for TryCreateValidGenerationHelper. */
+	static UPackage* FindOrLoadPackage(UCookOnTheFlyServer& COTFS, FPackageData& OwnerPackageData);
 
-	TConstArrayView<FName> GetExternalActorDependencies() { check(IsInitialized()); return ExternalActorDependencies; }
-	TArray<FName> ReleaseExternalActorDependencies()
+private:
+	enum class EInitializeStatus : uint8
 	{
-		TArray<FName> Result = MoveTemp(ExternalActorDependencies); return Result;
-	}
+		Uninitialized,
+		Invalid,
+		Valid,
+	};
 
 private:
-	void ConditionalNotifyCompletion(ICookPackageSplitter::ETeardown Status);
+	void ConditionalInitialize() const;
+	void NotifyCompletion(ICookPackageSplitter::ETeardown Status);
 
+private:
 	/** PackageData for the package that is being split */
 	FCookGenerationInfo OwnerInfo;
+	FWeakObjectPtr SplitDataObject;
 	/** Name of the object that prompted the splitter creation */
 	FName SplitDataObjectName;
-	/** Cached CookPackageSplitter */
+	UE::Cook::Private::FRegisteredCookPackageSplitter* RegisteredSplitterType = nullptr;
 	TUniquePtr<ICookPackageSplitter> CookPackageSplitterInstance;
 	/** Recorded list of packages to generate from the splitter, and data we need about them */
 	TArray<FCookGenerationInfo> PackagesToGenerate;
 	TWeakObjectPtr<UPackage> OwnerPackage;
 	TMap<FName, FIoHash> PreviousGeneratedPackages;
 	TArray<FName> ExternalActorDependencies;
-
-	int32 NextPopulateIndex = 0;
-	int32 RemainingToPopulate = 0;
-
-	bool bInitialized = false;
-	bool bNotifiedCompletion = false;
+	TRefCountPtr<FGenerationHelper> ReferenceFromKeepForIterative;
+	TRefCountPtr<FGenerationHelper> ReferenceFromKeepForQueueResults;
+	TRefCountPtr<FGenerationHelper> ReferenceFromKeepForGeneratorSave;
+	int32 MPCookNextAssignmentIndex = 0;
+	FWorkerId WorkerIdThatSavedGenerator = FWorkerId::Invalid();
+	EInitializeStatus InitializeStatus = EInitializeStatus::Uninitialized;
 	bool bUseInternalReferenceToAvoidGarbageCollect = false;
+	bool bGeneratedList = false;
 };
 
 
@@ -1200,7 +1334,7 @@ public:
 	uint32 Remove(FPackageData* PackageData);
 	bool Contains(const FPackageData* PackageData) const;
 	bool DiscoveryQueueContains(FPackageData* PackageData) const;
-	void Empty();
+	void Empty(FPackageDatas& PackageDatas);
 
 	void AddRequest(FPackageData* PackageData, bool bForceUrgent=false);
 
@@ -1222,10 +1356,19 @@ public:
 	TRingBuffer<FRequestCluster>& GetRequestClusters() { return RequestClusters; }
 	FPackageDataSet& GetReadyRequestsUrgent() { return UrgentRequests; }
 	FPackageDataSet& GetReadyRequestsNormal() { return NormalRequests; }
+	/**
+	 * Add an FPackageData by name that will be notified when all packages that were present in the DiscoveryQueue
+	 * before the fence was created have been assigned or demoted to idle.
+	 */
+	void AddRequestFenceListener(FName PackageName);
+	/** Called when the DiscoveryQueue has been flushed and fence listeners can therefore be notified. */
+	void NotifyRequestFencePassed(FPackageDatas& PackageDatas);
+
 private:
 	FPackageDataSet RestartedRequests;
 	TRingBuffer<FDiscoveryQueueElement> DiscoveryQueue;
 	TRingBuffer<FRequestCluster> RequestClusters;
+	TSet<FName> RequestFencePackageListeners;
 	FPackageDataSet UrgentRequests;
 	FPackageDataSet NormalRequests;
 };
@@ -1633,4 +1776,166 @@ inline void FPackageData::GetReachablePlatforms(ArrayType& OutPlatforms) const
 	}
 }
 
+inline TRefCountPtr<UE::Cook::FGenerationHelper> FPackageData::GetGenerationHelper() const
+{
+	return GenerationHelper;
+}
+
+inline bool FPackageData::HasCompletedGeneration() const
+{
+	return static_cast<bool>(bCompletedGeneration);
+}
+
+inline void FPackageData::SetCompletedGeneration(bool Value)
+{
+	bCompletedGeneration = Value != 0;
+}
+
+inline bool FPackageData::HasInitializedGeneratorSave() const
+{
+	return static_cast<bool>(bInitializedGeneratorSave);
+}
+
+inline void FPackageData::SetInitializedGeneratorSave(bool Value)
+{
+	bInitializedGeneratorSave = Value != 0;
+}
+
+inline bool FPackageData::IsGenerated() const
+{
+	return static_cast<bool>(bGenerated);
+}
+
+inline FName FPackageData::GetParentGenerator() const
+{
+	return ParentGenerator;
+}
+
+inline TRefCountPtr<FGenerationHelper> FPackageData::GetParentGenerationHelper() const
+{
+	return ParentGenerationHelper;
+}
+
+inline bool FGenerationHelper::IsInitialized() const
+{
+	return InitializeStatus != EInitializeStatus::Uninitialized;
+}
+
+inline void FGenerationHelper::ConditionalInitialize() const
+{
+	if (InitializeStatus == EInitializeStatus::Uninitialized)
+	{
+		// Use a const cast so we can call from getter functions that are otherwise const.
+		const_cast<FGenerationHelper&>(*this).Initialize();
+	}
+}
+
+inline bool FGenerationHelper::IsValid()
+{
+	ConditionalInitialize();
+	return InitializeStatus == EInitializeStatus::Valid;
+}
+
+inline TArrayView<UE::Cook::FCookGenerationInfo> FGenerationHelper::GetPackagesToGenerate()
+{
+	ConditionalInitialize();
+	return PackagesToGenerate;
+}
+
+inline UE::Cook::FCookGenerationInfo& FGenerationHelper::GetOwnerInfo()
+{
+	return OwnerInfo;
+}
+
+inline UE::Cook::FPackageData& FGenerationHelper::GetOwner()
+{
+	return *OwnerInfo.PackageData;
+}
+
+inline ICookPackageSplitter* FGenerationHelper::GetCookPackageSplitterInstance() const
+{
+	ConditionalInitialize();
+	return CookPackageSplitterInstance.Get();
+}
+
+inline UE::Cook::Private::FRegisteredCookPackageSplitter* FGenerationHelper::GetRegisteredSplitterType() const
+{
+	ConditionalInitialize();
+	return RegisteredSplitterType;
+}
+
+inline const FName FGenerationHelper::GetSplitDataObjectName() const
+{
+	ConditionalInitialize();
+	return SplitDataObjectName;
+}
+
+inline bool FGenerationHelper::IsUseInternalReferenceToAvoidGarbageCollect() const
+{
+	ConditionalInitialize();
+	return bUseInternalReferenceToAvoidGarbageCollect;
+}
+
+inline UObject* FGenerationHelper::GetWeakSplitDataObject() const
+{
+	ConditionalInitialize();
+	return SplitDataObject.Get();
+}
+
+inline TConstArrayView<FName> FGenerationHelper::GetExternalActorDependencies()
+{
+	return ExternalActorDependencies;
+}
+
+inline TArray<FName> FGenerationHelper::ReleaseExternalActorDependencies()
+{
+	TArray<FName> Result = MoveTemp(ExternalActorDependencies);
+	return Result;
+}
+
+inline const TMap<FName, FIoHash>& FGenerationHelper::GetPreviousGeneratedPackages() const
+{
+	return PreviousGeneratedPackages;
+}
+
+inline void FGenerationHelper::SetKeepForIterative()
+{
+	ReferenceFromKeepForIterative = this;
+}
+
+inline void FGenerationHelper::ClearKeepForIterative()
+{
+	ReferenceFromKeepForIterative.SafeRelease();
+}
+
+inline void FGenerationHelper::SetKeepForQueueResults()
+{
+	ReferenceFromKeepForQueueResults = this;
+}
+
+inline void FGenerationHelper::ClearKeepForQueueResults()
+{
+	ReferenceFromKeepForQueueResults.SafeRelease();
+}
+
+inline void FGenerationHelper::SetKeepForGeneratorSave()
+{
+	ReferenceFromKeepForGeneratorSave = this;
+}
+
+inline void FGenerationHelper::ClearKeepForGeneratorSave()
+{
+	ReferenceFromKeepForGeneratorSave.SafeRelease();
+}
+
+inline FWorkerId FGenerationHelper::GetWorkerIdThatSavedGenerator() const
+{
+	return WorkerIdThatSavedGenerator;
+}
+
+inline int32& FGenerationHelper::GetMPCookNextAssignmentIndex()
+{
+	return MPCookNextAssignmentIndex;
+}
+
 } // namespace UE::Cook
\ No newline at end of file
diff --git a/Engine/Source/Editor/UnrealEd/Private/Cooker/CookWorkerClient.cpp b/Engine/Source/Editor/UnrealEd/Private/Cooker/CookWorkerClient.cpp
--- a/Engine/Source/Editor/UnrealEd/Private/Cooker/CookWorkerClient.cpp
+++ b/Engine/Source/Editor/UnrealEd/Private/Cooker/CookWorkerClient.cpp
@@ -194,7 +194,7 @@ void FCookWorkerClient::ReportPromoteToSaveComplete(FPackageData& PackageData)
 	Result->SetPackageName(PackageName);
 	Result->SetSuppressCookReason(ESuppressCookReason::NotSuppressed);
 	Result->SetPlatforms(OrderedSessionPlatforms);
-	if (FGenerationHelper* GenerationHelper = PackageData.GetGenerationHelper(); GenerationHelper)
+	if (TRefCountPtr<FGenerationHelper> GenerationHelper = PackageData.GetGenerationHelper(); GenerationHelper)
 	{
 		Result->SetExternalActorDependencies(GenerationHelper->ReleaseExternalActorDependencies());
 	}
@@ -313,11 +313,39 @@ void FCookWorkerClient::ReportDiscoveredPackage(const FPackageData& PackageData,
 	FDiscoveredPackageReplication& Discovered = PendingDiscoveredPackages.Emplace_GetRef();
 	Discovered.PackageName = PackageData.GetPackageName();
 	Discovered.NormalizedFileName = PackageData.GetFileName();
+	Discovered.ParentGenerator = PackageData.GetParentGenerator();
 	Discovered.Instigator = Instigator;
 	Discovered.Platforms = MoveTemp(ReachablePlatforms);
 	Discovered.Platforms.ConvertToBitfield(OrderedSessionAndSpecialPlatforms);
 }
 
+void FCookWorkerClient::ReportGeneratorQueuedGeneratedPackages(FGenerationHelper& GenerationHelper)
+{
+	PendingGeneratorEvents.Add(FGeneratorEventMessage(EGeneratorEvent::QueuedGeneratedPackages,
+		GenerationHelper.GetOwner().GetPackageName()));
+}
+
+void FCookWorkerClient::HandleGeneratorMessage(FGeneratorEventMessage&& GeneratorMessage)
+{
+	FPackageData* PackageData = COTFS.PackageDatas->FindPackageDataByPackageName(GeneratorMessage.PackageName);
+	if (PackageData)
+	{
+		TRefCountPtr<FGenerationHelper> GenerationHelper = PackageData->GetGenerationHelper();
+		if (GenerationHelper)
+		{
+			switch (GeneratorMessage.Event)
+			{
+			case EGeneratorEvent::QueuedGeneratedPackagesFencePassed:
+				GenerationHelper->OnRequestFencePassed(COTFS);
+				break;
+			default:
+				// We do not handle the remaining GeneratorEvents on clients
+				break;
+			}
+		}
+	}
+}
+
 EPollStatus FCookWorkerClient::PollTryConnect(const FDirectorConnectionInfo& ConnectInfo)
 {
 	for (;;)
@@ -600,6 +628,15 @@ void FCookWorkerClient::SendPendingResults()
 		SendMessage(DiscoveredMessage);
 		PendingDiscoveredPackages.Reset();
 	}
+
+	if (!PendingGeneratorEvents.IsEmpty())
+	{
+		for (FGeneratorEventMessage& GeneratorMessage : PendingGeneratorEvents)
+		{
+			SendMessage(GeneratorMessage);
+		}
+		PendingGeneratorEvents.Reset();
+	}
 }
 
 void FCookWorkerClient::PumpReceiveMessages()
@@ -672,6 +709,18 @@ void FCookWorkerClient::HandleReceiveMessages(TArray<UE::CompactBinaryTCP::FMars
 					AssignPackages(AssignPackagesMessage);
 				}
 			}
+			else if (Message.MessageType == FGeneratorEventMessage::MessageType)
+			{
+				FGeneratorEventMessage GeneratorMessage;
+				if (!GeneratorMessage.TryRead(Message.Object))
+				{
+					LogInvalidMessage(TEXT("FGeneratorEventMessage"));
+				}
+				else
+				{
+					HandleGeneratorMessage(MoveTemp(GeneratorMessage));
+				}
+			}
 			else
 			{
 				TRefCountPtr<IMPCollector>* Collector = Collectors.Find(Message.MessageType);
@@ -781,48 +830,67 @@ void FCookWorkerClient::LogInvalidMessage(const TCHAR* MessageTypeName)
 
 void FCookWorkerClient::AssignPackages(FAssignPackagesMessage& Message)
 {
-	if (Message.PackageDatas.IsEmpty())
+	if (!Message.ExistenceInfos.IsEmpty())
 	{
-		return;
+		for (FPackageDataExistenceInfo& ExistenceInfo : Message.ExistenceInfos)
+		{
+			FPackageData& PackageData = COTFS.PackageDatas->FindOrAddPackageData(ExistenceInfo.ConstructData.PackageName,
+				ExistenceInfo.ConstructData.NormalizedFileName);
+			if (!ExistenceInfo.ParentGenerator.IsNone())
+			{
+				PackageData.SetGenerated(ExistenceInfo.ParentGenerator);
+			}
+		}
 	}
-
-	TArray<const ITargetPlatform*, TInlineAllocator<ExpectedMaxNumPlatforms>> NeedCookPlatformsBuffer;
-	for (FAssignPackageData& AssignData: Message.PackageDatas)
+	if (!Message.PackageDatas.IsEmpty())
 	{
-		FPackageData& PackageData = COTFS.PackageDatas->FindOrAddPackageData(AssignData.ConstructData.PackageName,
-			AssignData.ConstructData.NormalizedFileName);
-		TConstArrayView<const ITargetPlatform*> NeedCookPlatforms = 
-			AssignData.NeedCookPlatforms.GetPlatforms(COTFS, nullptr, OrderedSessionPlatforms,
-				&NeedCookPlatformsBuffer);
-		if (PackageData.IsInProgress())
+		TArray<const ITargetPlatform*, TInlineAllocator<ExpectedMaxNumPlatforms>> NeedCookPlatformsBuffer;
+		for (FAssignPackageData& AssignData : Message.PackageDatas)
 		{
-			// If already in progress and no new platforms, ignore the duplicate
-			// If there are new platforms, we don't currently handle that; the director should have retracted it first
+			FPackageData& PackageData = COTFS.PackageDatas->FindOrAddPackageData(AssignData.ConstructData.PackageName,
+				AssignData.ConstructData.NormalizedFileName);
+			if (!AssignData.ParentGenerator.IsNone())
+			{
+				PackageData.SetGenerated(AssignData.ParentGenerator);
+			}
+			if (!AssignData.GeneratorPreviousGeneratedPackages.IsEmpty())
+			{
+				TRefCountPtr<FGenerationHelper> GenerationHelper = PackageData.CreateUninitializedGenerationHelper();
+				GenerationHelper->SetPreviousGeneratedPackages(MoveTemp(AssignData.GeneratorPreviousGeneratedPackages));
+			}
+			TConstArrayView<const ITargetPlatform*> NeedCookPlatforms =
+				AssignData.NeedCookPlatforms.GetPlatforms(COTFS, nullptr, OrderedSessionPlatforms,
+					&NeedCookPlatformsBuffer);
+			if (PackageData.IsInProgress())
+			{
+				// If already in progress and no new platforms, ignore the duplicate
+				// If there are new platforms, we don't currently handle that; the director should have retracted it first
+				for (const ITargetPlatform* TargetPlatform : NeedCookPlatforms)
+				{
+					check(TargetPlatform != CookerLoadingPlatformKey);
+					checkf(PackageData.FindOrAddPlatformData(TargetPlatform).IsReachable(),
+						TEXT("CookWorker received AssignPackage for package %s which is already in progress, with new platform %s. Adding new platforms to an inprogress package is not yet supported"),
+						*PackageData.GetPackageName().ToString(), *TargetPlatform->PlatformName());
+				}
+				continue;
+			}
+
+			// We do not want CookWorkers to explore dependencies in CookRequestCluster because the Director
+			// did it already. Mark the PackageDatas we get from the Director as already explored.
 			for (const ITargetPlatform* TargetPlatform : NeedCookPlatforms)
 			{
-				check(TargetPlatform != CookerLoadingPlatformKey);
-				checkf(PackageData.FindOrAddPlatformData(TargetPlatform).IsReachable(),
-					TEXT("CookWorker received AssignPackage for package %s which is already in progress, with new platform %s. Adding new platforms to an inprogress package is not yet supported"),
-					*PackageData.GetPackageName().ToString(), *TargetPlatform->PlatformName());
+				PackageData.FindOrAddPlatformData(TargetPlatform).MarkCookableForWorker(*this);
 			}
-			continue;
+			PackageData.FindOrAddPlatformData(CookerLoadingPlatformKey).MarkCookableForWorker(*this);
+			PackageData.SetInstigator(*this, FInstigator(AssignData.Instigator));
+			PackageData.SendToState(EPackageState::Request, ESendFlags::QueueAddAndRemove,
+				EStateChangeReason::DirectorRequest);
 		}
 
-		// We do not want CookWorkers to explore dependencies in CookRequestCluster because the Director did it
-		// already. Mark the PackageDatas we get from the Director as already explored.
-		for (const ITargetPlatform* TargetPlatform : NeedCookPlatforms)
-		{
-			PackageData.FindOrAddPlatformData(TargetPlatform).MarkCookableForWorker(*this);
-		}
-		PackageData.FindOrAddPlatformData(CookerLoadingPlatformKey).MarkCookableForWorker(*this);
-		PackageData.SetInstigator(*this, FInstigator(AssignData.Instigator));
-		PackageData.SendToState(EPackageState::Request, ESendFlags::QueueAddAndRemove,
-			EStateChangeReason::DirectorRequest);
+		// Clear the SoftGC diagnostic ExpectedNeverLoadPackages because we have new assigned packages
+		// that we didn't consider during SoftGC
+		COTFS.PackageTracker->ClearExpectedNeverLoadPackages();
 	}
-
-	// Clear the SoftGC diagnostic ExpectedNeverLoadPackages because we have new assigned packages
-	// that we didn't consider during SoftGC
-	COTFS.PackageTracker->ClearExpectedNeverLoadPackages();
 }
 
 void FCookWorkerClient::Register(IMPCollector* Collector)
diff --git a/Engine/Source/Editor/UnrealEd/Private/Cooker/CookWorkerClient.h b/Engine/Source/Editor/UnrealEd/Private/Cooker/CookWorkerClient.h
--- a/Engine/Source/Editor/UnrealEd/Private/Cooker/CookWorkerClient.h
+++ b/Engine/Source/Editor/UnrealEd/Private/Cooker/CookWorkerClient.h
@@ -21,6 +21,8 @@ namespace UE::Cook { struct FDirectorConnectionInfo; }
 namespace UE::Cook { struct FDiscoveredPackageReplication; }
 namespace UE::Cook { struct FHeartbeatMessage; }
 namespace UE::Cook { struct FInitialConfigMessage; }
+namespace UE::Cook { struct FGenerationHelper; }
+namespace UE::Cook { struct FGeneratorEventMessage; }
 namespace UE::Cook { struct FPackageRemoteResult; }
 namespace UE::Cook { struct FRetractionRequestMessage; }
 
@@ -67,6 +69,11 @@ public:
 	/** Queue a message to the server that a package was discovered as needed in the cook. Will be sent during Tick. */
 	void ReportDiscoveredPackage(const FPackageData& PackageData, const FInstigator& Instigator,
 		FDiscoveredPlatformSet&& ReachablePlatforms);
+	/**
+	 * Queue a message to the server that a PackageGenerator queued its generated packages for cooking, and will
+	 * keep itself in memory until it receives an EGeneratorEvent::QueuedGeneratedPackagesFencePassed.
+	 */
+	void ReportGeneratorQueuedGeneratedPackages(FGenerationHelper& GenerationHelper);
 
 	/** Register a Collector for periodic ticking that sends messages to the Director. */
 	void Register(IMPCollector* Collector);
@@ -127,6 +134,7 @@ private:
 		FRetractionRequestMessage&& Message);
 	void HandleHeartbeatMessage(FMPCollectorClientMessageContext& Context, bool bReadSuccessful,
 		FHeartbeatMessage&& Message);
+	void HandleGeneratorMessage(FGeneratorEventMessage&& GeneratorMessage);
 
 private:
 	/**
@@ -151,6 +159,7 @@ private:
 	TArray<ITargetPlatform*> OrderedSessionAndSpecialPlatforms;
 	TArray<FDiscoveredPackageReplication> PendingDiscoveredPackages;
 	TMap<FGuid, TRefCountPtr<IMPCollector>> Collectors;
+	TArray<FGeneratorEventMessage> PendingGeneratorEvents;
 	UE::CompactBinaryTCP::FSendBuffer SendBuffer;
 	UE::CompactBinaryTCP::FReceiveBuffer ReceiveBuffer;
 	FString DirectorURI;
diff --git a/Engine/Source/Editor/UnrealEd/Private/Cooker/CookWorkerServer.cpp b/Engine/Source/Editor/UnrealEd/Private/Cooker/CookWorkerServer.cpp
--- a/Engine/Source/Editor/UnrealEd/Private/Cooker/CookWorkerServer.cpp
+++ b/Engine/Source/Editor/UnrealEd/Private/Cooker/CookWorkerServer.cpp
@@ -4,7 +4,7 @@
 
 #include "Algo/Find.h"
 #include "Commandlets/AssetRegistryGenerator.h"
-#include "CompactBinaryTCP.h"
+#include "Cooker/CompactBinaryTCP.h"
 #include "Cooker/CookDirector.h"
 #include "Cooker/CookPackageData.h"
 #include "Cooker/CookPlatformManager.h"
@@ -233,11 +233,15 @@ void FCookWorkerServer::ShutdownRemoteProcess()
 	DetachFromRemoteProcess();
 }
 
-void FCookWorkerServer::AppendAssignments(TArrayView<FPackageData*> Assignments, ECookDirectorThread TickThread)
+void FCookWorkerServer::AppendAssignments(TArrayView<FPackageData*> Assignments,
+	TMap<FPackageData*, FAssignPackageExtraData>&& ExtraDatas, TArrayView<FPackageData*> InfoPackages,
+	ECookDirectorThread TickThread)
 {
 	FCommunicationScopeLock ScopeLock(this, TickThread, ETickAction::Queue);
 	++PackagesAssignedFenceMarker;
 	PackagesToAssign.Append(Assignments);
+	PackagesToAssignExtraDatas.Append(MoveTemp(ExtraDatas));
+	PackagesToAssignInfoPackages.Append(InfoPackages);
 }
 
 void FCookWorkerServer::AbortAllAssignments(TSet<FPackageData*>& OutPendingPackages, ECookDirectorThread TickThread)
@@ -265,6 +269,8 @@ void FCookWorkerServer::AbortAllAssignmentsInLock(TSet<FPackageData*>& OutPendin
 	}
 	OutPendingPackages.Append(PackagesToAssign);
 	PackagesToAssign.Empty();
+	PackagesToAssignExtraDatas.Empty();
+	PackagesToAssignInfoPackages.Empty();
 	++PackagesRetiredFenceMarker;
 }
 
@@ -293,6 +299,9 @@ void FCookWorkerServer::AbortAssignments(TConstArrayView<FPackageData*> PackageD
 		}
 
 		PackagesToAssign.Remove(PackageData);
+		PackagesToAssignExtraDatas.Remove(PackageData);
+		// We don't remove InfoPackages from PackagesToAssignInfoPackages because it would be too hard to calculate,
+		// and it's not a problem to send extra InfoPackages.
 	}
 	++PackagesRetiredFenceMarker;
 	if (!PackageNamesToMessage.IsEmpty())
@@ -460,7 +469,7 @@ void FCookWorkerServer::TickCommunication(ECookDirectorThread TickThread)
 			PumpReceiveMessages();
 			if (ConnectStatus == EConnectStatus::Connected)
 			{
-				SendPendingPackages();
+				SendPendingMessages();
 				PumpSendMessages();
 				return; // Tick duties complete; yield the tick
 			}
@@ -643,10 +652,22 @@ void FCookWorkerServer::PumpSendMessages()
 	}
 }
 
+void FCookWorkerServer::SendPendingMessages()
+{
+	SendPendingPackages();
+	for (UE::CompactBinaryTCP::FMarshalledMessage& MarshalledMessage : QueuedMessagesToSendAfterPackagesToAssign)
+	{
+		UE::CompactBinaryTCP::QueueMessage(SendBuffer, MoveTemp(MarshalledMessage));
+	}
+	QueuedMessagesToSendAfterPackagesToAssign.Empty();
+}
+
 void FCookWorkerServer::SendPendingPackages()
 {
 	if (PackagesToAssign.IsEmpty())
 	{
+		PackagesToAssignExtraDatas.Empty();
+		PackagesToAssignInfoPackages.Empty();
 		return;
 	}
 	LLM_SCOPE_BYTAG(Cooker_MPCook);
@@ -654,11 +675,14 @@ void FCookWorkerServer::SendPendingPackages()
 	TArray<FAssignPackageData> AssignDatas;
 	AssignDatas.Reserve(PackagesToAssign.Num());
 	TBitArray<> SessionPlatformNeedsCook;
+	TArray<FPackageDataExistenceInfo> ExistenceInfos;
+	ExistenceInfos.Reserve(PackagesToAssignInfoPackages.Num());
 
 	for (FPackageData* PackageData : PackagesToAssign)
 	{
 		FAssignPackageData& AssignData = AssignDatas.Emplace_GetRef();
 		AssignData.ConstructData = PackageData->CreateConstructData();
+		AssignData.ParentGenerator = PackageData->GetParentGenerator();
 		AssignData.Instigator = PackageData->GetInstigator();
 		SessionPlatformNeedsCook.Init(false, OrderedSessionPlatforms.Num());
 		int32 PlatformIndex = 0;
@@ -668,10 +692,23 @@ void FCookWorkerServer::SendPendingPackages()
 			SessionPlatformNeedsCook[PlatformIndex++] = PlatformData && PlatformData->NeedsCooking(SessionPlatform);
 		}
 		AssignData.NeedCookPlatforms = FDiscoveredPlatformSet(SessionPlatformNeedsCook);
+		FAssignPackageExtraData* ExtraData = PackagesToAssignExtraDatas.Find(PackageData);
+		if (ExtraData)
+		{
+			AssignData.GeneratorPreviousGeneratedPackages = ExtraData->GeneratorPreviousGeneratedPackages;
+		}
+	}
+	for (FPackageData* PackageData : PackagesToAssignInfoPackages)
+	{
+		FPackageDataExistenceInfo& ExistenceInfo = ExistenceInfos.Emplace_GetRef();
+		ExistenceInfo.ConstructData = PackageData->CreateConstructData();
+		ExistenceInfo.ParentGenerator = PackageData->GetParentGenerator();
 	}
 	PendingPackages.Append(PackagesToAssign);
 	PackagesToAssign.Empty();
-	FAssignPackagesMessage AssignPackagesMessage(MoveTemp(AssignDatas));
+	PackagesToAssignExtraDatas.Empty();
+	PackagesToAssignInfoPackages.Empty();
+	FAssignPackagesMessage AssignPackagesMessage(MoveTemp(AssignDatas), MoveTemp(ExistenceInfos));
 	AssignPackagesMessage.OrderedSessionPlatforms = OrderedSessionPlatforms;
 	SendMessageInLock(MoveTemp(AssignPackagesMessage));
 }
@@ -762,6 +799,18 @@ void FCookWorkerServer::HandleReceiveMessagesInternal()
 				}
 			}
 		}
+		else if (Message.MessageType == FGeneratorEventMessage::MessageType)
+		{
+			FGeneratorEventMessage GeneratorMessage;
+			if (!GeneratorMessage.TryRead(Message.Object))
+			{
+				LogInvalidMessage(TEXT("FGeneratorEventMessage"));
+			}
+			else
+			{
+				HandleGeneratorMessage(GeneratorMessage);
+			}
+		}
 		else
 		{
 			TRefCountPtr<IMPCollector>* Collector = Director.Collectors.Find(Message.MessageType);
@@ -825,6 +874,12 @@ void FCookWorkerServer::SendMessage(const IMPCollectorMessage& Message, ECookDir
 	SendMessageInLock(Message);
 }
 
+void FCookWorkerServer::AppendMessage(const IMPCollectorMessage& Message, ECookDirectorThread TickThread)
+{
+	FCommunicationScopeLock ScopeLock(this, TickThread, ETickAction::Queue);
+	QueuedMessagesToSendAfterPackagesToAssign.Add(MarshalToCompactBinaryTCP(Message));
+}
+
 void FCookWorkerServer::SendMessageInLock(const IMPCollectorMessage& Message)
 {
 	if (TickState.TickAction == ETickAction::Tick)
@@ -937,6 +992,18 @@ void FCookWorkerServer::QueueDiscoveredPackage(FDiscoveredPackageReplication&& D
 	FDiscoveredPlatformSet& Platforms = DiscoveredPackage.Platforms;
 	FPackageData& PackageData = PackageDatas.FindOrAddPackageData(DiscoveredPackage.PackageName,
 		DiscoveredPackage.NormalizedFileName);
+	if (!DiscoveredPackage.ParentGenerator.IsNone())
+	{
+		PackageData.SetGenerated(DiscoveredPackage.ParentGenerator);
+		FPackageData* GeneratorPackageData = PackageDatas.FindPackageDataByPackageName(
+			DiscoveredPackage.ParentGenerator);
+		if (GeneratorPackageData)
+		{
+			TRefCountPtr<FGenerationHelper> GenerationHelper =
+				GeneratorPackageData->CreateUninitializedGenerationHelper();
+			GenerationHelper->NotifyStartQueueGeneratedPackages(COTFS, WorkerId);
+		}
+	}
 
 	TArray<const ITargetPlatform*, TInlineAllocator<ExpectedMaxNumPlatforms>> BufferPlatforms;
 	TConstArrayView<const ITargetPlatform*> DiscoveredPlatforms;
@@ -973,9 +1040,9 @@ void FCookWorkerServer::QueueDiscoveredPackage(FDiscoveredPackageReplication&& D
 		return;
 	}
 
-	if (Instigator.Category == EInstigator::GeneratedPackage)
+	if (Instigator.Category == EInstigator::GeneratedPackage
+		&& COTFS.MPCookGeneratorSplit == EMPCookGeneratorSplit::AllOnSameWorker)
 	{
-		PackageData.SetGenerated(true);
 		PackageData.SetWorkerAssignmentConstraint(GetWorkerId());
 	}
 	Director.ResetFinalIdleHeartbeatFence();
@@ -984,6 +1051,35 @@ void FCookWorkerServer::QueueDiscoveredPackage(FDiscoveredPackageReplication&& D
 		false /* bUrgent */);
 }
 
+void FCookWorkerServer::HandleGeneratorMessage(FGeneratorEventMessage& GeneratorMessage)
+{
+	FPackageData* PackageData = COTFS.PackageDatas->FindPackageDataByPackageName(GeneratorMessage.PackageName);
+	TRefCountPtr<FGenerationHelper> GenerationHelper;
+	if (PackageData)
+	{
+		GenerationHelper = PackageData->CreateUninitializedGenerationHelper();
+	}
+	if (!GenerationHelper)
+	{
+		// This error should be impossible because GeneratorMessages are only sent in response to assignment from the server.
+		UE_LOG(LogCook, Error,
+			TEXT("CookWorkerServer received unexpected GeneratorMessage for package %s. The PackageData %s on the CookDirector. ")
+			TEXT("\n\tCook of this generator package and its generated packages will be invalid."),
+			*GeneratorMessage.PackageName.ToString(),
+			(!PackageData ? TEXT("does not exist") : TEXT("is not a valid generator")));
+		return;
+	}
+	switch (GeneratorMessage.Event)
+	{
+	case EGeneratorEvent::QueuedGeneratedPackages:
+		GenerationHelper->EndQueueGeneratedPackagesOnDirector(COTFS, GetWorkerId());
+		break;
+	default:
+		// We do not handle the remaining GeneratorEvents on clients
+		break;
+	}
+}
+
 FCookWorkerServer::FTickState::FTickState()
 {
 	TickThread = ECookDirectorThread::Invalid;
@@ -1021,8 +1117,10 @@ UE::CompactBinaryTCP::FMarshalledMessage MarshalToCompactBinaryTCP(const IMPColl
 	return Marshalled;
 }
 
-FAssignPackagesMessage::FAssignPackagesMessage(TArray<FAssignPackageData>&& InPackageDatas)
+FAssignPackagesMessage::FAssignPackagesMessage(TArray<FAssignPackageData>&& InPackageDatas,
+	TArray<FPackageDataExistenceInfo>&& InExistenceInfos)
 	: PackageDatas(MoveTemp(InPackageDatas))
+	, ExistenceInfos(MoveTemp(InExistenceInfos))
 {
 }
 
@@ -1034,6 +1132,12 @@ void FAssignPackagesMessage::Write(FCbWriter& Writer) const
 		WriteToCompactBinary(Writer, PackageData, OrderedSessionPlatforms);
 	}
 	Writer.EndArray();
+	Writer.BeginArray("I");
+	for (const FPackageDataExistenceInfo& ExistenceInfo : ExistenceInfos)
+	{
+		Writer << ExistenceInfo;
+	}
+	Writer.EndArray();
 }
 
 bool FAssignPackagesMessage::TryRead(FCbObjectView Object)
@@ -1049,29 +1153,59 @@ bool FAssignPackagesMessage::TryRead(FCbObjectView Object)
 			bOk = false;
 		}
 	}
+	ExistenceInfos.Reset();
+	for (FCbFieldView PackageField : Object["I"])
+	{
+		FPackageDataExistenceInfo& ExistenceInfo = ExistenceInfos.Emplace_GetRef();
+		if (!LoadFromCompactBinary(PackageField, ExistenceInfo))
+		{
+			ExistenceInfos.Pop();
+			bOk = false;
+		}
+	}
 	return bOk;
 }
 
 FGuid FAssignPackagesMessage::MessageType(TEXT("B7B1542B73254B679319D73F753DB6F8"));
 
-void WriteToCompactBinary(FCbWriter& Writer, const FAssignPackageData& AssignData, 
-	TConstArrayView<const ITargetPlatform*> OrderedSessionPlatforms)
+void FAssignPackageData::Write(FCbWriter& Writer,
+	TConstArrayView<const ITargetPlatform*> OrderedSessionPlatforms) const
+{
+	Writer.BeginArray();
+	Writer << ConstructData;
+	Writer << ParentGenerator;
+	Writer << Instigator;
+	WriteToCompactBinary(Writer, NeedCookPlatforms, OrderedSessionPlatforms);
+	Writer << GeneratorPreviousGeneratedPackages;
+	Writer.EndArray();
+}
+
+bool FAssignPackageData::TryRead(FCbFieldView Field, TConstArrayView<const ITargetPlatform*> OrderedSessionPlatforms)
+{
+	FCbFieldViewIterator It = Field.CreateViewIterator();
+	bool bOk = true;
+	bOk = LoadFromCompactBinary(*It++, ConstructData) & bOk;
+	bOk = LoadFromCompactBinary(*It++, ParentGenerator) & bOk;
+	bOk = LoadFromCompactBinary(*It++, Instigator) & bOk;
+	bOk = LoadFromCompactBinary(*It++, NeedCookPlatforms, OrderedSessionPlatforms) & bOk;
+	bOk = LoadFromCompactBinary(*It++, GeneratorPreviousGeneratedPackages) & bOk;
+	return bOk;
+}
+
+void FPackageDataExistenceInfo::Write(FCbWriter& Writer) const
 {
 	Writer.BeginArray();
-	Writer << AssignData.ConstructData;
-	Writer << AssignData.Instigator;
-	WriteToCompactBinary(Writer, AssignData.NeedCookPlatforms, OrderedSessionPlatforms);
+	Writer << ConstructData;
+	Writer << ParentGenerator;
 	Writer.EndArray();
 }
 
-bool LoadFromCompactBinary(FCbFieldView Field, FAssignPackageData& AssignData,
-	TConstArrayView<const ITargetPlatform*> OrderedSessionPlatforms)
+bool FPackageDataExistenceInfo::TryRead(FCbFieldView Field)
 {
 	FCbFieldViewIterator It = Field.CreateViewIterator();
 	bool bOk = true;
-	bOk = LoadFromCompactBinary(*It++, AssignData.ConstructData) & bOk;
-	bOk = LoadFromCompactBinary(*It++, AssignData.Instigator) & bOk;
-	bOk = LoadFromCompactBinary(*It++, AssignData.NeedCookPlatforms, OrderedSessionPlatforms) & bOk;
+	bOk = LoadFromCompactBinary(*It++, ConstructData) & bOk;
+	bOk = LoadFromCompactBinary(*It++, ParentGenerator) & bOk;
 	return bOk;
 }
 
@@ -1224,40 +1358,48 @@ bool FInitialConfigMessage::TryRead(FCbObjectView Object)
 
 FGuid FInitialConfigMessage::MessageType(TEXT("340CDCB927304CEB9C0A66B5F707FC2B"));
 
-void WriteToCompactBinary(FCbWriter& Writer, const FDiscoveredPackageReplication& Package,
-	TConstArrayView<const ITargetPlatform*> OrderedSessionAndSpecialPlatforms)
+void FDiscoveredPackageReplication::Write(FCbWriter& Writer,
+	TConstArrayView<const ITargetPlatform*> OrderedSessionAndSpecialPlatforms) const
 {
-	Writer.BeginObject();
-	Writer << "PackageName" << Package.PackageName;
-	Writer << "NormalizedFileName" << Package.NormalizedFileName;
-	Writer << "Instigator.Category" << static_cast<uint8>(Package.Instigator.Category);
-	Writer << "Instigator.Referencer" << Package.Instigator.Referencer;
-	Writer.SetName("Platforms");
-	WriteToCompactBinary(Writer, Package.Platforms, OrderedSessionAndSpecialPlatforms);
-	Writer.EndObject();
+	Writer.BeginArray();
+	Writer << PackageName;
+	Writer << NormalizedFileName;
+	Writer << ParentGenerator;
+	Writer << static_cast<uint8>(Instigator.Category);
+	Writer << Instigator.Referencer;
+	WriteToCompactBinary(Writer, Platforms, OrderedSessionAndSpecialPlatforms);
+	Writer.EndArray();
 }
 
-bool LoadFromCompactBinary(FCbFieldView Field, FDiscoveredPackageReplication& OutPackage,
+bool FDiscoveredPackageReplication::TryRead(FCbFieldView Field,
 	TConstArrayView<const ITargetPlatform*> OrderedSessionAndSpecialPlatforms)
 {
-	bool bOk = LoadFromCompactBinary(Field["PackageName"], OutPackage.PackageName);
-	bOk = LoadFromCompactBinary(Field["NormalizedFileName"], OutPackage.NormalizedFileName) & bOk;
+	FCbArrayView FieldList = Field.AsArrayView();
+	if (Field.HasError())
+	{
+		*this = FDiscoveredPackageReplication();
+		return false;
+	}
+	FCbFieldViewIterator Iter = FieldList.CreateViewIterator();
+
+	bool bOk = LoadFromCompactBinary(Iter++, PackageName);
+	bOk = LoadFromCompactBinary(Iter++, NormalizedFileName) & bOk;
+	bOk = LoadFromCompactBinary(Iter++, ParentGenerator) & bOk;
 	uint8 CategoryInt;
-	if (LoadFromCompactBinary(Field["Instigator.Category"], CategoryInt) &&
+	if (LoadFromCompactBinary(Iter++, CategoryInt) &&
 		CategoryInt < static_cast<uint8>(EInstigator::Count))
-
 	{
-		OutPackage.Instigator.Category = static_cast<EInstigator>(CategoryInt);
+		Instigator.Category = static_cast<EInstigator>(CategoryInt);
 	}
 	else
 	{
 		bOk = false;
 	}
-	bOk = LoadFromCompactBinary(Field["Instigator.Referencer"], OutPackage.Instigator.Referencer) & bOk;
-	bOk = LoadFromCompactBinary(Field["Platforms"], OutPackage.Platforms, OrderedSessionAndSpecialPlatforms) & bOk;
+	bOk = LoadFromCompactBinary(Iter++, Instigator.Referencer) & bOk;
+	bOk = LoadFromCompactBinary(Iter++, Platforms, OrderedSessionAndSpecialPlatforms) & bOk;
 	if (!bOk)
 	{
-		OutPackage = FDiscoveredPackageReplication();
+		*this = FDiscoveredPackageReplication();
 	}
 	return bOk;
 }
@@ -1290,6 +1432,38 @@ bool FDiscoveredPackagesMessage::TryRead(FCbObjectView Object)
 
 FGuid FDiscoveredPackagesMessage::MessageType(TEXT("C9F5BC5C11484B06B346B411F1ED3090"));
 
+FGeneratorEventMessage::FGeneratorEventMessage(EGeneratorEvent InEvent, FName InPackageName)
+	: PackageName(InPackageName)
+	, Event(InEvent)
+{
+}
+
+void FGeneratorEventMessage::Write(FCbWriter& Writer) const
+{
+	Writer << "E" << static_cast<uint8>(Event);
+	Writer << "P" << PackageName;
+}
+
+bool FGeneratorEventMessage::TryRead(FCbObjectView Object)
+{
+	bool bOk = true;
+	FCbFieldView EventField = Object["E"];
+	uint8 EventInt = EventField.AsUInt8();
+	if (!EventField.HasError() && EventInt < static_cast<uint8>(EGeneratorEvent::Num))
+	{
+		Event = static_cast<EGeneratorEvent>(EventInt);
+	}
+	else
+	{
+		Event = EGeneratorEvent::Invalid;
+		bOk = false;
+	}
+	bOk = LoadFromCompactBinary(Object["P"], PackageName) & bOk;
+	return bOk;
+}
+
+FGuid FGeneratorEventMessage::MessageType(TEXT("B6EE94CA70EC4F40B0D2214EDC11ED03"));
+
 FCbWriter& operator<<(FCbWriter& Writer, const FReplicatedLogData& Package)
 {
 	Writer.BeginArray();
diff --git a/Engine/Source/Editor/UnrealEd/Private/Cooker/CookWorkerServer.h b/Engine/Source/Editor/UnrealEd/Private/Cooker/CookWorkerServer.h
--- a/Engine/Source/Editor/UnrealEd/Private/Cooker/CookWorkerServer.h
+++ b/Engine/Source/Editor/UnrealEd/Private/Cooker/CookWorkerServer.h
@@ -5,6 +5,7 @@
 #include "CompactBinaryTCP.h"
 #include "Containers/Array.h"
 #include "Containers/ArrayView.h"
+#include "Containers/Map.h"
 #include "Containers/RingBuffer.h"
 #include "Containers/Set.h"
 #include "Cooker/CookPackageData.h"
@@ -13,6 +14,7 @@
 #include "Cooker/MPCollector.h"
 #include "CookOnTheSide/CookOnTheFlyServer.h"
 #include "HAL/CriticalSection.h"
+#include "IO/IoHash.h"
 #include "Logging/LogVerbosity.h"
 #include "Misc/Guid.h"
 #include "Misc/OutputDevice.h"
@@ -25,6 +27,7 @@ namespace UE::CompactBinaryTCP { struct FMarshalledMessage; }
 namespace UE::Cook { class FCookDirector; }
 namespace UE::Cook { struct FDiscoveredPackageReplication; }
 namespace UE::Cook { enum class ECookDirectorThread : uint8; };
+namespace UE::Cook { struct FGeneratorEventMessage; }
 namespace UE::Cook { struct FPackageData; }
 namespace UE::Cook { struct FPackageResultsMessage; }
 namespace UE::Cook { struct FWorkerConnectMessage; }
@@ -38,6 +41,15 @@ enum class ENotifyRemote
 	LocalOnly,
 };
 
+/**
+ * Extra non-threadsafe data about an FPackageData that we save from the scheduler thread for reading on the
+ * communication thread when we send the assignment of the PackageData out to the CookWorker.
+ */
+struct FAssignPackageExtraData
+{
+	TMap<FName, FIoHash> GeneratorPreviousGeneratedPackages;
+};
+
 /** Class in a Director process that communicates over a Socket with FCookWorkerClient in a CookWorker process. */
 class FCookWorkerServer : public FThreadSafeRefCountedObject
 {
@@ -49,7 +61,9 @@ public:
 	FWorkerId GetWorkerId() const { return WorkerId; }
 
 	/** Add the given assignments for the CookWorker. They will be sent during Tick */
-	void AppendAssignments(TArrayView<FPackageData*> Assignments, ECookDirectorThread TickThread);
+	void AppendAssignments(TArrayView<FPackageData*> Assignments,
+		TMap<FPackageData*, FAssignPackageExtraData>&& ExtraDatas, TArrayView<FPackageData*> InfoPackages,
+		ECookDirectorThread TickThread);
 	/** Remove assignment of the package from local state and optionally from the connected Client. */
 	void AbortAssignment(FPackageData& PackageData, ECookDirectorThread TickThread,
 		ENotifyRemote NotifyRemote = ENotifyRemote::NotifyRemote);
@@ -61,10 +75,7 @@ public:
 	 * Report all packages that were removed.
 	 */
 	void AbortAllAssignments(TSet<FPackageData*>& OutPendingPackages, ECookDirectorThread TickThread);
-	/**
-	 * AbortAllAssignments and tell the connected Client to gracefully terminate. Report all packages that were
-	 * unassigned.
-	 */
+	/** AbortAllAssignments and tell the Client to gracefully terminate. Report all packages that were unassigned. */
 	void AbortWorker(TSet<FPackageData*>& OutPendingPackages, ECookDirectorThread TickThread);
 	/** Take over the Socket for a CookWorker that has just connected. */
 	bool TryHandleConnectMessage(FWorkerConnectMessage& Message, FSocket* InSocket,
@@ -72,6 +83,11 @@ public:
 
 	/** Send the message immediately to the Socket. If cannot complete immediately, it will be finished during Tick. */
 	void SendMessage(const IMPCollectorMessage& Message, ECookDirectorThread TickThread);
+	/**
+	 * Queue the given message for sending; it will be sent during tick. Only needed for messages that have to be sent
+	 * after queued package assignments or other appended messages.
+	 */
+	void AppendMessage(const IMPCollectorMessage& Message, ECookDirectorThread TickThread);
 
 	/** Periodic Tick function to send and receive messages to the Client. */
 	void TickCommunication(ECookDirectorThread TickThread);
@@ -154,7 +170,9 @@ private:
 	void TickWaitForDisconnect();
 	/** Helper for Tick, pump send messages to a connected Client. */
 	void PumpSendMessages();
-	/** Helper for PumpSendMessages; send a message for any PackagesToAssign we have. */
+	/** Helper for PumpSendMessages; send messages that were queued up for sending during tick. */
+	void SendPendingMessages();
+	/** Helper for SendPendingMessages; send a message for any PackagesToAssign we have. */
 	void SendPendingPackages();
 	/** Helper for Tick, pump receive messages from a connected Client. */
 	void PumpReceiveMessages();
@@ -179,18 +197,22 @@ private:
 	void RecordResults(FPackageResultsMessage& Message);
 	void LogInvalidMessage(const TCHAR* MessageTypeName);
 	void QueueDiscoveredPackage(FDiscoveredPackageReplication&& DiscoveredPackage);
+	void HandleGeneratorMessage(FGeneratorEventMessage& GeneratorMessage);
 
 	// Lock guarding access to all data on *this
 	mutable FCriticalSection CommunicationLock;
 
 	// All data can only be read or written while the CommunicationLock is entered.
 	TArray<FPackageData*> PackagesToAssign;
+	TSet<FPackageData*> PackagesToAssignInfoPackages;
+	TMap<FPackageData*, FAssignPackageExtraData> PackagesToAssignExtraDatas;
 	TSet<FPackageData*> PendingPackages;
 	TArray<ITargetPlatform*> OrderedSessionPlatforms;
 	TArray<ITargetPlatform*> OrderedSessionAndSpecialPlatforms;
 	UE::CompactBinaryTCP::FSendBuffer SendBuffer;
 	UE::CompactBinaryTCP::FReceiveBuffer ReceiveBuffer;
 	TRingBuffer<UE::CompactBinaryTCP::FMarshalledMessage> ReceiveMessages;
+	TRingBuffer<UE::CompactBinaryTCP::FMarshalledMessage> QueuedMessagesToSendAfterPackagesToAssign;
 	FString CrashDiagnosticsError;
 	FCookDirector& Director;
 	UCookOnTheFlyServer& COTFS;
@@ -212,27 +234,65 @@ private:
 
 UE::CompactBinaryTCP::FMarshalledMessage MarshalToCompactBinaryTCP(const IMPCollectorMessage& Message);
 
-/** Information about a PackageData the director sends to cookworkers. */
+FCbWriter& operator<<(FCbWriter& Writer, const FInstigator& Instigator);
+bool LoadFromCompactBinary(FCbFieldView Field, FInstigator& Instigator);
+
+/** Information about a PackageData assigned to a CookWorker from the director. */
 struct FAssignPackageData
 {
 	FConstructPackageData ConstructData;
+	FName ParentGenerator;
 	FInstigator Instigator;
 	FDiscoveredPlatformSet NeedCookPlatforms;
+	TMap<FName, FIoHash> GeneratorPreviousGeneratedPackages;
+
+private:
+	void Write(FCbWriter& Writer, TConstArrayView<const ITargetPlatform*> OrderedSessionPlatforms) const;
+	bool TryRead(FCbFieldView Field, TConstArrayView<const ITargetPlatform*> OrderedSessionPlatforms);
+
+	friend FCbWriter& WriteToCompactBinary(FCbWriter& Writer, const FAssignPackageData& AssignData,
+		TConstArrayView<const ITargetPlatform*> OrderedSessionPlatforms)
+	{
+		AssignData.Write(Writer, OrderedSessionPlatforms);
+		return Writer;
+	}
+	friend bool LoadFromCompactBinary(FCbFieldView Field, FAssignPackageData& AssignData,
+		TConstArrayView<const ITargetPlatform*> OrderedSessionPlatforms)
+	{
+		return AssignData.TryRead(Field, OrderedSessionPlatforms);
+	}
 };
 
-void WriteToCompactBinary(FCbWriter& Writer, const FAssignPackageData& AssignData,
-	TConstArrayView<const ITargetPlatform*> OrderedSessionPlatforms);
-bool LoadFromCompactBinary(FCbFieldView Field, FAssignPackageData& AssignData, 
-	TConstArrayView<const ITargetPlatform*> OrderedSessionPlatforms);
-FCbWriter& operator<<(FCbWriter& Writer, const FInstigator& Instigator);
-bool LoadFromCompactBinary(FCbFieldView Field, FInstigator& Instigator);
+/**
+ * Information about a PackageData sent from the director that needs to be queryable on a CookWorker to support other
+ * assigned packages but is not being assigned to the CookWorker.
+ */
+struct FPackageDataExistenceInfo
+{
+	FConstructPackageData ConstructData;
+	FName ParentGenerator;
+private:
+	void Write(FCbWriter& Writer) const;
+	bool TryRead(FCbFieldView Field);
+
+	friend FCbWriter& operator<<(FCbWriter& Writer, const FPackageDataExistenceInfo& ExistenceInfo)
+	{
+		ExistenceInfo.Write(Writer);
+		return Writer;
+	}
+	friend bool LoadFromCompactBinary(FCbFieldView Field, FPackageDataExistenceInfo& ExistenceInfo)
+	{
+		return ExistenceInfo.TryRead(Field);
+	}
+};
 
 /** Message from Server to Client to cook the given packages. */
 struct FAssignPackagesMessage : public IMPCollectorMessage
 {
 public:
 	FAssignPackagesMessage() = default;
-	FAssignPackagesMessage(TArray<FAssignPackageData>&& InPackageDatas);
+	FAssignPackagesMessage(TArray<FAssignPackageData>&& InPackageDatas,
+		TArray<FPackageDataExistenceInfo>&& InExistenceInfos);
 
 	virtual void Write(FCbWriter& Writer) const override;
 	virtual bool TryRead(FCbObjectView Object) override;
@@ -241,6 +301,7 @@ public:
 
 public:
 	TArray<FAssignPackageData> PackageDatas;
+	TArray<FPackageDataExistenceInfo> ExistenceInfos;
 	TConstArrayView<const ITargetPlatform*> OrderedSessionPlatforms;
 	static FGuid MessageType;
 };
@@ -329,14 +390,25 @@ struct FDiscoveredPackageReplication
 {
 	FName PackageName;
 	FName NormalizedFileName;
+	FName ParentGenerator;
 	FInstigator Instigator;
 	FDiscoveredPlatformSet Platforms;
-};
 
-void WriteToCompactBinary(FCbWriter& Writer, const FDiscoveredPackageReplication& Package,
-	TConstArrayView<const ITargetPlatform*> OrderedSessionAndSpecialPlatforms);
-bool LoadFromCompactBinary(FCbFieldView Field, FDiscoveredPackageReplication& OutPackage,
-	TConstArrayView<const ITargetPlatform*> OrderedSessionAndSpecialPlatforms);
+private:
+	void Write(FCbWriter& Writer, TConstArrayView<const ITargetPlatform*> OrderedSessionAndSpecialPlatforms) const;
+	bool TryRead(FCbFieldView Field, TConstArrayView<const ITargetPlatform*> OrderedSessionAndSpecialPlatforms);
+
+	friend void WriteToCompactBinary(FCbWriter& Writer, const FDiscoveredPackageReplication& Package,
+		TConstArrayView<const ITargetPlatform*> OrderedSessionAndSpecialPlatforms)
+	{
+		Package.Write(Writer, OrderedSessionAndSpecialPlatforms);
+	}
+	friend bool LoadFromCompactBinary(FCbFieldView Field, FDiscoveredPackageReplication& OutPackage,
+		TConstArrayView<const ITargetPlatform*> OrderedSessionAndSpecialPlatforms)
+	{
+		return OutPackage.TryRead(Field, OrderedSessionAndSpecialPlatforms);
+	}
+};
 
 /**
  * Message from CookWorker to Director that reports dependency packages discovered during load/save of
@@ -356,6 +428,35 @@ public:
 	static FGuid MessageType;
 };
 
+/**
+ * Events on a PackageGenerator that occur on the Director or a CookWorker and need to be replicated to the other
+ * process(es).
+ */
+enum class EGeneratorEvent
+{
+	Invalid,
+	QueuedGeneratedPackages,
+	QueuedGeneratedPackagesFencePassed,
+	Num,
+};
+
+/** A message reporting an EGeneratorEvent that occurred on the Director or a CookWorker. */
+struct FGeneratorEventMessage : public IMPCollectorMessage
+{
+	FGeneratorEventMessage() = default;
+	FGeneratorEventMessage(EGeneratorEvent InEvent, FName InPackageName);
+
+	virtual void Write(FCbWriter& Writer) const override;
+	virtual bool TryRead(FCbObjectView Object) override;
+	virtual FGuid GetMessageType() const override { return MessageType; }
+	virtual const TCHAR* GetDebugName() const override { return TEXT("GeneratorEventMessage"); }
+
+public:
+	FName PackageName;
+	EGeneratorEvent Event = EGeneratorEvent::Invalid;
+	static FGuid MessageType;
+};
+
 /** Stores the data passed into FOutputDevice::Serialize, for replication to the CookDirector. */
 struct FReplicatedLogData
 {
diff --git a/Engine/Source/Editor/UnrealEd/Private/Cooker/IWorkerRequests.h b/Engine/Source/Editor/UnrealEd/Private/Cooker/IWorkerRequests.h
--- a/Engine/Source/Editor/UnrealEd/Private/Cooker/IWorkerRequests.h
+++ b/Engine/Source/Editor/UnrealEd/Private/Cooker/IWorkerRequests.h
@@ -9,6 +9,7 @@
 
 class FConfigFile;
 namespace UE::Cook { struct FPackageData; }
+namespace UE::Cook { struct FGenerationHelper; }
 
 namespace UE::Cook
 {
@@ -62,6 +63,7 @@ public:
 	// Writing Packages during Cook
 	virtual void QueueDiscoveredPackage(UCookOnTheFlyServer& COTFS, FPackageData& PackageData,
 		FInstigator&& Instigator, FDiscoveredPlatformSet&& ReachablePlatforms, bool bUrgent) = 0;
+	virtual void EndQueueGeneratedPackages(UCookOnTheFlyServer& COTFS, FGenerationHelper& GenerationHelper) = 0;
 
 	// Writing Packages from CookByTheBook
 	virtual void AddStartCookByTheBookRequest(FFilePlatformRequest&& Request) = 0;
diff --git a/Engine/Source/Editor/UnrealEd/Private/Cooker/WorkerRequestsLocal.cpp b/Engine/Source/Editor/UnrealEd/Private/Cooker/WorkerRequestsLocal.cpp
--- a/Engine/Source/Editor/UnrealEd/Private/Cooker/WorkerRequestsLocal.cpp
+++ b/Engine/Source/Editor/UnrealEd/Private/Cooker/WorkerRequestsLocal.cpp
@@ -3,6 +3,7 @@
 #include "WorkerRequestsLocal.h"
 
 #include "CookTypes.h"
+#include "Cooker/CookPackageData.h"
 #include "Cooker/CookRequests.h"
 #include "HAL/Event.h"
 #include "HAL/PlatformProcess.h"
@@ -15,7 +16,6 @@ class ITargetPlatform;
 namespace UE::Cook
 {
 struct FInstigator;
-struct FPackageData;
 
 bool FWorkerRequestsLocal::HasExternalRequests() const
 {
@@ -50,6 +50,12 @@ void FWorkerRequestsLocal::QueueDiscoveredPackage(UCookOnTheFlyServer& COTFS, FP
 	COTFS.QueueDiscoveredPackageOnDirector(PackageData, MoveTemp(Instigator), MoveTemp(ReachablePlatforms), bUrgent);
 }
 
+void FWorkerRequestsLocal::EndQueueGeneratedPackages(UCookOnTheFlyServer& COTFS,
+	FGenerationHelper& GenerationHelper)
+{
+	GenerationHelper.EndQueueGeneratedPackagesOnDirector(COTFS, FWorkerId::Local());
+}
+
 void FWorkerRequestsLocal::AddStartCookByTheBookRequest(FFilePlatformRequest&& Request)
 {
 	ExternalRequests.EnqueueUnique(MoveTemp(Request));
diff --git a/Engine/Source/Editor/UnrealEd/Private/Cooker/WorkerRequestsLocal.h b/Engine/Source/Editor/UnrealEd/Private/Cooker/WorkerRequestsLocal.h
--- a/Engine/Source/Editor/UnrealEd/Private/Cooker/WorkerRequestsLocal.h
+++ b/Engine/Source/Editor/UnrealEd/Private/Cooker/WorkerRequestsLocal.h
@@ -34,6 +34,7 @@ public:
 	virtual void OnRemoveSessionPlatform(const ITargetPlatform* TargetPlatform) override;
 	virtual void QueueDiscoveredPackage(UCookOnTheFlyServer& COTFS, FPackageData& PackageData,
 		FInstigator&& Instigator, FDiscoveredPlatformSet&& ReachablePlatforms, bool bUrgent) override;
+	virtual void EndQueueGeneratedPackages(UCookOnTheFlyServer& COTFS, FGenerationHelper& GenerationHelper) override;
 	virtual void AddStartCookByTheBookRequest(FFilePlatformRequest&& Request) override;
 	virtual void InitializeCookOnTheFly() override;
 	virtual void AddCookOnTheFlyRequest(FFilePlatformRequest&& Request) override;
diff --git a/Engine/Source/Editor/UnrealEd/Private/Cooker/WorkerRequestsRemote.cpp b/Engine/Source/Editor/UnrealEd/Private/Cooker/WorkerRequestsRemote.cpp
--- a/Engine/Source/Editor/UnrealEd/Private/Cooker/WorkerRequestsRemote.cpp
+++ b/Engine/Source/Editor/UnrealEd/Private/Cooker/WorkerRequestsRemote.cpp
@@ -61,6 +61,12 @@ void FWorkerRequestsRemote::QueueDiscoveredPackage(UCookOnTheFlyServer& COTFS, F
 	CookWorkerClient.ReportDiscoveredPackage(PackageData, MoveTemp(Instigator), MoveTemp(ReachablePlatforms));
 }
 
+void FWorkerRequestsRemote::EndQueueGeneratedPackages(UCookOnTheFlyServer& COTFS,
+	FGenerationHelper& GenerationHelper)
+{
+	CookWorkerClient.ReportGeneratorQueuedGeneratedPackages(GenerationHelper);
+}
+
 void FWorkerRequestsRemote::AddStartCookByTheBookRequest(FFilePlatformRequest&& Request)
 {
 	LogCalledCookByTheBookError(TEXT("AddStartCookByTheBookRequest"));
diff --git a/Engine/Source/Editor/UnrealEd/Private/Cooker/WorkerRequestsRemote.h b/Engine/Source/Editor/UnrealEd/Private/Cooker/WorkerRequestsRemote.h
--- a/Engine/Source/Editor/UnrealEd/Private/Cooker/WorkerRequestsRemote.h
+++ b/Engine/Source/Editor/UnrealEd/Private/Cooker/WorkerRequestsRemote.h
@@ -42,6 +42,7 @@ public:
 	virtual void OnRemoveSessionPlatform(const ITargetPlatform* TargetPlatform) override;
 	virtual void QueueDiscoveredPackage(UCookOnTheFlyServer& COTFS, FPackageData& PackageData,
 		FInstigator&& Instigator, FDiscoveredPlatformSet&& ReachablePlatforms, bool bUrgent) override;
+	virtual void EndQueueGeneratedPackages(UCookOnTheFlyServer& COTFS, FGenerationHelper& GenerationHelper) override;
 	virtual void AddStartCookByTheBookRequest(FFilePlatformRequest&& Request) override;
 	virtual void InitializeCookOnTheFly() override;
 	virtual void AddCookOnTheFlyRequest(FFilePlatformRequest&& Request) override;
diff --git a/Engine/Source/Runtime/Engine/Private/WorldPartition/WorldPartitionCookPackageSplitter.cpp b/Engine/Source/Runtime/Engine/Private/WorldPartition/WorldPartitionCookPackageSplitter.cpp
--- a/Engine/Source/Runtime/Engine/Private/WorldPartition/WorldPartitionCookPackageSplitter.cpp
+++ b/Engine/Source/Runtime/Engine/Private/WorldPartition/WorldPartitionCookPackageSplitter.cpp
@@ -36,21 +36,25 @@ void FWorldPartitionCookPackageSplitter::Teardown(ETeardown Status)
 
 	FWorldDelegates::OnWorldCleanup.RemoveAll(this);
 
-	// Assume that the world is partitioned as per FWorldPartitionCookPackageSplitter::ShouldSplit
-	UWorldPartition* WorldPartition = ReferencedWorld->PersistentLevel->GetWorldPartition();
-	check(WorldPartition);
+	if (ReferencedWorld != nullptr)
+	{
+		// We were destructed without GetGenerateList being called; nothing to teardown
+		// Assume that the world is partitioned as per FWorldPartitionCookPackageSplitter::ShouldSplit
+		UWorldPartition* WorldPartition = ReferencedWorld->PersistentLevel->GetWorldPartition();
+		check(WorldPartition);
 
-	WorldPartition->EndCook(CookContext);
-	WorldPartition->Uninitialize();
+		WorldPartition->EndCook(CookContext);
+		WorldPartition->Uninitialize();
 
-	if (bInitializedPhysicsSceneForSave)
-	{
-		GEditor->CleanupPhysicsSceneThatWasInitializedForSave(ReferencedWorld, bForceInitializedWorld);
-		bInitializedPhysicsSceneForSave = false;
-		bForceInitializedWorld = false;
-	}
+		if (bInitializedPhysicsSceneForSave)
+		{
+			GEditor->CleanupPhysicsSceneThatWasInitializedForSave(ReferencedWorld, bForceInitializedWorld);
+			bInitializedPhysicsSceneForSave = false;
+			bForceInitializedWorld = false;
+		}
 
-	ReferencedWorld = nullptr;
+		ReferencedWorld = nullptr;
+	}
 }
 
 void FWorldPartitionCookPackageSplitter::OnWorldCleanup(UWorld* InWorld, bool bSessionEnded, bool bCleanupResources)
