Cooker: CookPackageSplitter: Allow generator packages to transition from the middle of save state into idle without giving an error. The error about breaking the cookpackagesplitter's assumption is only a problem if the package remains in progress.
Improve the diagnostics for EStateChangeReasons that come from DemotePackageToIdle.
#rnx

diff --git a/Engine/Source/Editor/UnrealEd/Classes/CookOnTheSide/CookOnTheFlyServer.h b/Engine/Source/Editor/UnrealEd/Classes/CookOnTheSide/CookOnTheFlyServer.h
--- a/Engine/Source/Editor/UnrealEd/Classes/CookOnTheSide/CookOnTheFlyServer.h
+++ b/Engine/Source/Editor/UnrealEd/Classes/CookOnTheSide/CookOnTheFlyServer.h
@@ -173,6 +173,7 @@ namespace UE::Cook
 	class ICookOnTheFlyNetworkServer;
 	class IWorkerRequests;
 	enum class ECachedCookedPlatformDataEvent : uint8;
+	enum class EPackageState : uint8;
 	enum class EPollStatus : uint8;
 	enum class EStateChangeReason : uint8;
 	enum class ESuppressCookReason : uint8;
@@ -1175,7 +1176,8 @@ private:
 	 *        the save again. If true, all data will be wiped.
 	 * @param ReleaseSaveReason Why the save data is being released, allows specifying how much to tear down
 	 */
-	void ReleaseCookedPlatformData(UE::Cook::FPackageData& PackageData, UE::Cook::EStateChangeReason ReleaseSaveReason);
+	void ReleaseCookedPlatformData(UE::Cook::FPackageData& PackageData, UE::Cook::EStateChangeReason ReleaseSaveReason,
+		UE::Cook::EPackageState NewState);
 
 	/**
 	 * Poll the PendingCookedPlatformDatas and release their resources when they are complete.
diff --git a/Engine/Source/Editor/UnrealEd/Private/CookOnTheFlyServer.cpp b/Engine/Source/Editor/UnrealEd/Private/CookOnTheFlyServer.cpp
--- a/Engine/Source/Editor/UnrealEd/Private/CookOnTheFlyServer.cpp
+++ b/Engine/Source/Editor/UnrealEd/Private/CookOnTheFlyServer.cpp
@@ -2680,7 +2680,8 @@ void UCookOnTheFlyServer::NotifyRemovedFromWorker(UE::Cook::FPackageData& Packag
 	CookDirector->RemoveFromWorker(PackageData);
 }
 
-void UCookOnTheFlyServer::DemoteToIdle(UE::Cook::FPackageData& PackageData, UE::Cook::ESendFlags SendFlags, UE::Cook::ESuppressCookReason Reason)
+void UCookOnTheFlyServer::DemoteToIdle(UE::Cook::FPackageData& PackageData, UE::Cook::ESendFlags SendFlags,
+	UE::Cook::ESuppressCookReason Reason)
 {
 	using namespace UE::Cook;
 
@@ -2726,7 +2727,7 @@ void UCookOnTheFlyServer::DemoteToIdle(UE::Cook::FPackageData& PackageData, UE::
 			}
 		}
 	}
-	PackageData.SendToState(UE::Cook::EPackageState::Idle, SendFlags, EStateChangeReason::CookSuppressed);
+	PackageData.SendToState(UE::Cook::EPackageState::Idle, SendFlags, ConvertToStateChangeReason(Reason));
 }
 
 void UCookOnTheFlyServer::PromoteToSaveComplete(UE::Cook::FPackageData& PackageData, UE::Cook::ESendFlags SendFlags)
@@ -4126,7 +4127,8 @@ UE::Cook::EPollStatus UCookOnTheFlyServer::CallBeginCacheOnObjects(UE::Cook::FPa
 	return EPollStatus::Success;
 }
 
-void UCookOnTheFlyServer::ReleaseCookedPlatformData(UE::Cook::FPackageData& PackageData, UE::Cook::EStateChangeReason ReleaseSaveReason)
+void UCookOnTheFlyServer::ReleaseCookedPlatformData(UE::Cook::FPackageData& PackageData,
+	UE::Cook::EStateChangeReason ReleaseSaveReason, UE::Cook::EPackageState NewState)
 {
 	using namespace UE::Cook;
 
@@ -4222,7 +4224,7 @@ void UCookOnTheFlyServer::ReleaseCookedPlatformData(UE::Cook::FPackageData& Pack
 
 	if (GenerationInfo)
 	{
-		GenerationHelper->ResetSaveState(*GenerationInfo, PackageData.GetPackage(), ReleaseSaveReason);
+		GenerationHelper->ResetSaveState(*GenerationInfo, PackageData.GetPackage(), ReleaseSaveReason, NewState);
 		if (GenerationInfo->IsGenerator())
 		{
 			PackageData.SetInitializedGeneratorSave(false);
@@ -4662,7 +4664,7 @@ void UCookOnTheFlyServer::PumpSaves(UE::Cook::FTickStackData& StackData, uint32
 			if (PrepareSaveStatus == EPollStatus::Error)
 			{
 				check(PackageData.HasPrepareSaveFailed()); // Should have been set by PrepareSave; we rely on this for cleanup
-				ReleaseCookedPlatformData(PackageData, EStateChangeReason::SaveError);
+				ReleaseCookedPlatformData(PackageData, EStateChangeReason::SaveError, EPackageState::Idle);
 				PackageData.SetPlatformsCooked(PlatformsForPackage, ECookResult::Failed);
 				DemoteToIdle(PackageData, ESendFlags::QueueAdd, ESuppressCookReason::SaveError);
 				++OutNumPushed;
@@ -4752,7 +4754,7 @@ void UCookOnTheFlyServer::PumpSaves(UE::Cook::FTickStackData& StackData, uint32
 		{
 			// Timeouts can occur because of new objects created during the save, so we need to update our object cache,
 			// so we call ReleaseCookedPlatformData and ClearObjectCache to clear it and recache on next attempt.
-			ReleaseCookedPlatformData(PackageData, EStateChangeReason::RecreateObjectCache);
+			ReleaseCookedPlatformData(PackageData, EStateChangeReason::RecreateObjectCache, EPackageState::Save);
 			PackageData.ClearObjectCache();
 			if (PackageData.GetIsUrgent())
 			{
@@ -4765,7 +4767,9 @@ void UCookOnTheFlyServer::PumpSaves(UE::Cook::FTickStackData& StackData, uint32
 			continue;
 		}
 
-		ReleaseCookedPlatformData(PackageData, !Context.bHasRetryErrorCode ? EStateChangeReason::Completed : EStateChangeReason::DoneForNow);
+		ReleaseCookedPlatformData(PackageData,
+			!Context.bHasRetryErrorCode ? EStateChangeReason::Completed : EStateChangeReason::DoneForNow,
+			EPackageState::Idle);
 		PromoteToSaveComplete(PackageData, ESendFlags::QueueAdd);
 		++OutNumPushed;
 #if ENABLE_LOW_LEVEL_MEM_TRACKER
@@ -5325,7 +5329,8 @@ void UCookOnTheFlyServer::PreGarbageCollect()
 				GCKeepPackageDatas, bShouldDemote);
 			if (bShouldDemote)
 			{
-				ReleaseCookedPlatformData(*PackageData, UE::Cook::EStateChangeReason::GeneratorPreGarbageCollected);
+				ReleaseCookedPlatformData(*PackageData, UE::Cook::EStateChangeReason::GeneratorPreGarbageCollected,
+					EPackageState::Request);
 			}
 		}
 		if (PackageData->GetIsCookLast())
diff --git a/Engine/Source/Editor/UnrealEd/Private/Cooker/CookPackageData.cpp b/Engine/Source/Editor/UnrealEd/Private/Cooker/CookPackageData.cpp
--- a/Engine/Source/Editor/UnrealEd/Private/Cooker/CookPackageData.cpp
+++ b/Engine/Source/Editor/UnrealEd/Private/Cooker/CookPackageData.cpp
@@ -630,7 +630,7 @@ void FPackageData::SendToState(EPackageState NextState, ESendFlags SendFlags, ES
 		{
 			ensure(PackageDatas.GetSaveQueue().Remove(this) == 1);
 		}
-		OnExitSave(ReleaseSaveReason);
+		OnExitSave(ReleaseSaveReason, NextState);
 		break;
 	default:
 		check(false);
@@ -877,9 +877,9 @@ void FPackageData::OnEnterSave()
 	CheckCookedPlatformDataEmpty();
 }
 
-void FPackageData::OnExitSave(EStateChangeReason ReleaseSaveReason)
+void FPackageData::OnExitSave(EStateChangeReason ReleaseSaveReason, EPackageState NewState)
 {
-	PackageDatas.GetCookOnTheFlyServer().ReleaseCookedPlatformData(*this, ReleaseSaveReason);
+	PackageDatas.GetCookOnTheFlyServer().ReleaseCookedPlatformData(*this, ReleaseSaveReason, NewState);
 	ClearObjectCache();
 	SetHasPrepareSaveFailed(false);
 	SetIsPrepareSaveRequiresGC(false);
@@ -1999,7 +1999,7 @@ bool FGenerationHelper::IsComplete() const
 }
 
 void FGenerationHelper::ResetSaveState(FCookGenerationInfo& Info, UPackage* Package,
-	EStateChangeReason ReleaseSaveReason)
+	EStateChangeReason ReleaseSaveReason, EPackageState NewState)
 {
 	check(IsInitialized());
 	if (Info.GetSaveState() > FCookGenerationInfo::ESaveState::CallPopulate)
@@ -2062,8 +2062,8 @@ void FGenerationHelper::ResetSaveState(FCookGenerationInfo& Info, UPackage* Pack
 	}
 	if (Info.HasTakenOverCachedCookedPlatformData())
 	{
-		if (Info.PackageData && Info.PackageData->GetCachedObjectsInOuter().Num() != 0 &&
-			IsUseInternalReferenceToAvoidGarbageCollect() &&
+		if (NewState != EPackageState::Idle && Info.PackageData &&
+			Info.PackageData->GetCachedObjectsInOuter().Num() != 0 && IsUseInternalReferenceToAvoidGarbageCollect() &&
 			(ReleaseSaveReason != EStateChangeReason::Completed && ReleaseSaveReason != EStateChangeReason::DoneForNow
 				&& ReleaseSaveReason != EStateChangeReason::SaveError
 				&& ReleaseSaveReason != EStateChangeReason::CookerShutdown))
diff --git a/Engine/Source/Editor/UnrealEd/Private/Cooker/CookPackageData.h b/Engine/Source/Editor/UnrealEd/Private/Cooker/CookPackageData.h
--- a/Engine/Source/Editor/UnrealEd/Private/Cooker/CookPackageData.h
+++ b/Engine/Source/Editor/UnrealEd/Private/Cooker/CookPackageData.h
@@ -744,7 +744,7 @@ private:
 	void OnEnterLoadReady();
 	void OnExitLoadReady();
 	void OnEnterSave();
-	void OnExitSave(EStateChangeReason ReleaseSaveReason);
+	void OnExitSave(EStateChangeReason ReleaseSaveReason, EPackageState NewState);
 	/* Entry/Exit gates for Properties shared between multiple states */
 	void OnExitInProgress();
 	void OnEnterInProgress();
@@ -972,7 +972,8 @@ public:
 	 */
 	UObject* FindSplitDataObject() const;
 
-	void ResetSaveState(FCookGenerationInfo& Info, UPackage* Package, UE::Cook::EStateChangeReason ReleaseSaveReason);
+	void ResetSaveState(FCookGenerationInfo& Info, UPackage* Package, UE::Cook::EStateChangeReason ReleaseSaveReason,
+		EPackageState NewState);
 
 	int32& GetNextPopulateIndex() { check(IsInitialized()); return NextPopulateIndex; }
 
diff --git a/Engine/Source/Editor/UnrealEd/Private/Cooker/CookTypes.cpp b/Engine/Source/Editor/UnrealEd/Private/Cooker/CookTypes.cpp
--- a/Engine/Source/Editor/UnrealEd/Private/Cooker/CookTypes.cpp
+++ b/Engine/Source/Editor/UnrealEd/Private/Cooker/CookTypes.cpp
@@ -74,6 +74,20 @@ const TCHAR* LexToString(ESuppressCookReason Reason)
 	}
 }
 
+EStateChangeReason ConvertToStateChangeReason(ESuppressCookReason Reason)
+{
+	switch (Reason)
+	{
+	case ESuppressCookReason::OrphanedGenerated: return EStateChangeReason::SaveError;
+	case ESuppressCookReason::LoadError: return EStateChangeReason::SaveError;
+	case ESuppressCookReason::ValidationError: return EStateChangeReason::SaveError;
+	case ESuppressCookReason::SaveError: return EStateChangeReason::SaveError;
+	case ESuppressCookReason::CookCanceled: return EStateChangeReason::CookerShutdown;
+	case ESuppressCookReason::MultiprocessAssignmentError: return EStateChangeReason::ReassignAbortedPackages;
+	case ESuppressCookReason::RetractedByCookDirector: return EStateChangeReason::Retraction;
+	default: return EStateChangeReason::CookSuppressed;
+	}
+}
 
 FCookerTimer::FCookerTimer(float InTimeSlice)
 	: TickStartTime(FPlatformTime::Seconds()), ActionStartTime(TickStartTime)
diff --git a/Engine/Source/Editor/UnrealEd/Private/Cooker/CookTypes.h b/Engine/Source/Editor/UnrealEd/Private/Cooker/CookTypes.h
--- a/Engine/Source/Editor/UnrealEd/Private/Cooker/CookTypes.h
+++ b/Engine/Source/Editor/UnrealEd/Private/Cooker/CookTypes.h
@@ -165,12 +165,13 @@ namespace UE::Cook
 		CookFilter,
 	};
 	const TCHAR* LexToString(UE::Cook::ESuppressCookReason Reason);
+	EStateChangeReason ConvertToStateChangeReason(ESuppressCookReason Reason);
 
 	/** The type of callback for External Requests that needs to be executed within the Scheduler's lock. */
 	typedef TUniqueFunction<void()> FSchedulerCallback;
 
 	/** Which phase of cooking a Package is in.  */
-	enum class EPackageState
+	enum class EPackageState : uint8
 	{
 		/**
 		 * The Package is not being operated on by the cooker, and is not in any queues. This is the state both for
