Fbx scene importer: Crash fix when re-importing a scene. Do not check and simply return if the override options are empty.

#rnx

diff --git a/Engine/Source/Editor/UnrealEd/Private/Fbx/FbxSceneImportData.cpp b/Engine/Source/Editor/UnrealEd/Private/Fbx/FbxSceneImportData.cpp
--- a/Engine/Source/Editor/UnrealEd/Private/Fbx/FbxSceneImportData.cpp
+++ b/Engine/Source/Editor/UnrealEd/Private/Fbx/FbxSceneImportData.cpp
@@ -205,7 +205,7 @@ FString FbxOptionToJSON(FString OptionName, UnFbx::FBXImportOptions *Option)
 		Option->bImportStaticMeshLODs ? 1 : 0
 		);
 
-	JsonString += FString::Printf(TEXT("\"bUpdateSkeletonReferencePose\" : \"%d\", \"bUseT0AsRefPose\" : \"%d\", \"bPreserveSmoothingGroups\" : \"%d\", \"bKeepSectionsSeparate\" : \"%d\", \"bImportMeshesInBoneHierarchy\" : \"%d\", \"bImportMorphTargets\" : \"%d\", , \"bImportVertexAttributes\" : \"%d\", \"OverlappingThresholds\" : {\"ThresholdPosition\" : \"%f\", \"ThresholdTangentNormal\" : \"%f\", \"ThresholdUV\" : \"%f\", \"MorphThresholdPosition\" : \"%f\"},"),
+	JsonString += FString::Printf(TEXT("\"bUpdateSkeletonReferencePose\" : \"%d\", \"bUseT0AsRefPose\" : \"%d\", \"bPreserveSmoothingGroups\" : \"%d\", \"bKeepSectionsSeparate\" : \"%d\", \"bImportMeshesInBoneHierarchy\" : \"%d\", \"bImportMorphTargets\" : \"%d\", \"bImportVertexAttributes\" : \"%d\", \"OverlappingThresholds\" : {\"ThresholdPosition\" : \"%f\", \"ThresholdTangentNormal\" : \"%f\", \"ThresholdUV\" : \"%f\", \"MorphThresholdPosition\" : \"%f\"},"),
 		Option->bUpdateSkeletonReferencePose ? 1 : 0,
 		Option->bUseT0AsRefPose ? 1 : 0,
 		Option->bPreserveSmoothingGroups ? 1 : 0,
