Fixed iwyu compile errors

diff --git a/Engine/Source/Editor/UnrealEd/Private/Settings/ProjectPackagingSettingsCustomization.h b/Engine/Source/Editor/UnrealEd/Private/Settings/ProjectPackagingSettingsCustomization.h
--- a/Engine/Source/Editor/UnrealEd/Private/Settings/ProjectPackagingSettingsCustomization.h
+++ b/Engine/Source/Editor/UnrealEd/Private/Settings/ProjectPackagingSettingsCustomization.h
@@ -3,9 +3,11 @@
 #pragma once
 
 #include "CoreMinimal.h"
+#include "Framework/MultiBox/MultiBoxBuilder.h"
 #include "Misc/Attribute.h"
 #include "Layout/Margin.h"
 #include "Layout/Visibility.h"
+#include "Misc/NotifyHook.h"
 #include "Widgets/DeclarativeSyntaxSupport.h"
 #include "Widgets/SWidget.h"
 #include "Widgets/SCompoundWidget.h"
diff --git a/Engine/Source/Runtime/Core/Private/AutoRTFM/TaskArray.h b/Engine/Source/Runtime/Core/Private/AutoRTFM/TaskArray.h
--- a/Engine/Source/Runtime/Core/Private/AutoRTFM/TaskArray.h
+++ b/Engine/Source/Runtime/Core/Private/AutoRTFM/TaskArray.h
@@ -4,6 +4,7 @@
 
 #include "Containers/Array.h"
 #include "Templates/Function.h"
+#include "Utils.h"
 
 namespace AutoRTFM
 {
diff --git a/Engine/Source/Runtime/Engine/Public/PrimitiveSceneDesc.h b/Engine/Source/Runtime/Engine/Public/PrimitiveSceneDesc.h
--- a/Engine/Source/Runtime/Engine/Public/PrimitiveSceneDesc.h
+++ b/Engine/Source/Runtime/Engine/Public/PrimitiveSceneDesc.h
@@ -2,6 +2,7 @@
 
 #pragma once
 
+#include "Math/BoxSphereBounds.h"
 #include "PrimitiveComponentId.h"
 #include "PrimitiveSceneInfoData.h"
 #include "Engine/EngineTypes.h"
