#UE Do not preparesave if you are going to skip saving

diff --git a/Engine/Source/Editor/UnrealEd/Private/CookOnTheFlyServer.cpp b/Engine/Source/Editor/UnrealEd/Private/CookOnTheFlyServer.cpp
--- a/Engine/Source/Editor/UnrealEd/Private/CookOnTheFlyServer.cpp
+++ b/Engine/Source/Editor/UnrealEd/Private/CookOnTheFlyServer.cpp
@@ -3742,7 +3742,11 @@ UE::Cook::EPollStatus UCookOnTheFlyServer::PrepareSave(UE::Cook::FPackageData& P
 	using namespace UE::Cook;
 
 	EPollStatus Result = EPollStatus::Incomplete;
-	if (PackageData.GetCookedPlatformDataComplete())
+	if (bSkipSave)
+	{
+		Result = EPollStatus::Success;
+	}
+	else if (PackageData.GetCookedPlatformDataComplete())
 	{
 		Result = EPollStatus::Success;
 	}
