Support for FilterActorDesc when loading a umap with WorldPartition.

diff --git a/Engine/Source/Editor/UnrealEd/Classes/Commandlets/BaseIteratePackagesCommandlet.h b/Engine/Source/Editor/UnrealEd/Classes/Commandlets/BaseIteratePackagesCommandlet.h
--- a/Engine/Source/Editor/UnrealEd/Classes/Commandlets/BaseIteratePackagesCommandlet.h
+++ b/Engine/Source/Editor/UnrealEd/Classes/Commandlets/BaseIteratePackagesCommandlet.h
@@ -200,6 +200,9 @@ protected:
 	/** Loads and saves a single package */
 	UNREALED_API virtual void LoadAndSaveOnePackage(const FString& Filename);
 
+	/** Checks if an actor should be loading during WorldPartition loading */
+	virtual bool FilterActorDesc(const FWorldPartitionActorDesc*) const { return true; };
+
 	/** Checks to see if a package should be skipped */
 	UNREALED_API virtual bool ShouldSkipPackage(const FString& Filename);
 
diff --git a/Engine/Source/Editor/UnrealEd/Private/Commandlets/BaseIteratePackagesCommandlet.cpp b/Engine/Source/Editor/UnrealEd/Private/Commandlets/BaseIteratePackagesCommandlet.cpp
--- a/Engine/Source/Editor/UnrealEd/Private/Commandlets/BaseIteratePackagesCommandlet.cpp
+++ b/Engine/Source/Editor/UnrealEd/Private/Commandlets/BaseIteratePackagesCommandlet.cpp
@@ -694,6 +694,8 @@ void UBaseIteratePackagesCommandlet::LoadAndSaveOnePackage(const FString& Filena
 						PackagesToSave.Empty();
 					};
 
+				ForEachActorWithLoadingParams.FilterActorDesc = [this](const FWorldPartitionActorDesc* ActorDesc) -> bool{ return this->FilterActorDesc(ActorDesc); };
+
 				UWorldPartition* WorldPartition = World->GetWorldPartition();
 
 				FWorldPartitionHelpers::ForEachActorWithLoading(WorldPartition, [&PackagesToSave, this](const FWorldPartitionActorDescInstance* ActorDescInstance)
