MPCook And CookPackageSplitter: Allow each splitter class to opt out of the new behavor allowing generated packages to be retracted. Some splitters need to run BeginCacheForCookedPlatformData before they can decide which packages to generate.
[RN] Minor, cooking

diff --git a/Engine/Plugins/Experimental/Mutable/Source/CustomizableObjectEditor/Private/MuCOE/CustomizableObjectCookPackageSplitter.h b/Engine/Plugins/Experimental/Mutable/Source/CustomizableObjectEditor/Private/MuCOE/CustomizableObjectCookPackageSplitter.h
--- a/Engine/Plugins/Experimental/Mutable/Source/CustomizableObjectEditor/Private/MuCOE/CustomizableObjectCookPackageSplitter.h
+++ b/Engine/Plugins/Experimental/Mutable/Source/CustomizableObjectEditor/Private/MuCOE/CustomizableObjectCookPackageSplitter.h
@@ -46,6 +46,11 @@ public:
 	 * should expect them to still be in memory after a garbage collect so long as the splitter is alive.
 	 */
 	virtual bool UseInternalReferenceToAvoidGarbageCollect() override { return true; }
+	/**
+	 * If true, the cooker will not call any splitter work functions (GetGenerateList, Populate*, PreSave*, PostSave*)
+	 * until BeginCacheForCookedPlatformData has been called for all objects in the generator package.
+	 */
+	virtual bool NeedCachedPlatformDataBeforeSplit() override { return true; }
 
 private:
 
diff --git a/Engine/Source/Editor/UnrealEd/Private/CookOnTheFlyServer.cpp b/Engine/Source/Editor/UnrealEd/Private/CookOnTheFlyServer.cpp
--- a/Engine/Source/Editor/UnrealEd/Private/CookOnTheFlyServer.cpp
+++ b/Engine/Source/Editor/UnrealEd/Private/CookOnTheFlyServer.cpp
@@ -11483,22 +11483,30 @@ void UCookOnTheFlyServer::GetPackagesToRetract(int32 NumToRetract, TArray<FName>
 			// on WorldPartition packages, if we abort them and then try to restart them later.
 			return false;
 		}
-		if (MPCookGeneratorSplit == UE::Cook::EMPCookGeneratorSplit::AllOnSameWorker && PackageData->IsGenerated())
+		if (PackageData->IsGenerated())
 		{
-			// With the legacy AllOnSameWorker setting, GeneratedPackages are automatically constrained to this worker.
-			return false;
+			if (PackageData->IsGeneratedNeedCachedPlatformDataBeforeSplit()
+				|| MPCookGeneratorSplit == UE::Cook::EMPCookGeneratorSplit::AllOnSameWorker)
+			{
+				// With IsNeedCachedPlatformDataBeforeSplit or the AllOnSameWorker setting, GeneratedPackages are
+				// constrained to this worker.
+				return false;
+			}
 		}
-		if (MPCookGeneratorSplit != UE::Cook::EMPCookGeneratorSplit::AnyWorker)
+		if (FGenerationHelper* GenerationHelper = PackageData->GetGenerationHelper())
 		{
-			if (FGenerationHelper* GenerationHelper = PackageData->GetGenerationHelper())
+			if (GenerationHelper->IsInitialized() &&
+				GenerationHelper->GetOwnerInfo().GetSaveState()
+					>= FCookGenerationInfo::ESaveState::QueueGeneratedPackages)
 			{
-				if (GenerationHelper->GetOwnerInfo().GetSaveState() >=
-					FCookGenerationInfo::ESaveState::QueueGeneratedPackages)
+				if (GenerationHelper->IsNeedCachedPlatformDataBeforeSplit() ||
+					MPCookGeneratorSplit != UE::Cook::EMPCookGeneratorSplit::AnyWorker)
 				{
-					// With any EMPCookGeneratorSplit setting other than AnyWorker, we make assignment decisions based
-					// on the worker that saved and queued the generator package. We do not track queuing separately;
-					// we assume it happened on the worker that saved the package. Therefore, do not allow retraction
-					// of a generator package if it has already entered the QueueGeneratedPackages state.
+					// With IsNeedCachedPlatformDataBeforeSplit or with any MPCookGeneratorSplit setting other than
+					// AnyWorker, we make assignment decisions based on the worker that saved and queued the generator
+					// package. We do not track queuing separately; we assume it happened on the worker that saved the
+					// package. Therefore, do not allow retraction of a generator package if it has already entered
+					// the QueueGeneratedPackages state.
 					return false;
 				}
 			}
diff --git a/Engine/Source/Editor/UnrealEd/Private/Cooker/CookPackageData.cpp b/Engine/Source/Editor/UnrealEd/Private/Cooker/CookPackageData.cpp
--- a/Engine/Source/Editor/UnrealEd/Private/Cooker/CookPackageData.cpp
+++ b/Engine/Source/Editor/UnrealEd/Private/Cooker/CookPackageData.cpp
@@ -97,7 +97,7 @@ FPackageData::FPackageData(FPackageDatas& PackageDatas, const FName& InPackageNa
 	, bCookedPlatformDataStarted(0), bCookedPlatformDataCalled(0), bCookedPlatformDataComplete(0)
 	, MonitorCookResult((uint8)ECookResult::NotAttempted)
 	, bInitializedGeneratorSave(0), bCompletedGeneration(0), bGenerated(0), bKeepReferencedDuringGC(0)
-	, bWasCookedThisSession(0)
+	, bWasCookedThisSession(0), bGeneratedNeedCachedPlatformDataBeforeSplit(0)
 {
 	SetState(EPackageState::Idle);
 	SendToState(EPackageState::Idle, ESendFlags::QueueAdd, EStateChangeReason::Discovered);
@@ -1772,6 +1772,8 @@ void FGenerationHelper::Initialize(const UObject* InSplitDataObject,
 	SplitDataObjectName = FName(FStringView(InSplitDataObject->GetFullName()));
 	bUseInternalReferenceToAvoidGarbageCollect =
 		CookPackageSplitterInstance->UseInternalReferenceToAvoidGarbageCollect();
+	bNeedCachedPlatformDataBeforeSplit =
+		CookPackageSplitterInstance->NeedCachedPlatformDataBeforeSplit();
 }
 
 void FGenerationHelper::InitializeAsInvalid()
@@ -2063,6 +2065,7 @@ bool FGenerationHelper::TryGenerateList()
 		check(PackageData->GetParentGenerator().IsNone() ||
 			PackageData->GetParentGenerator() == OwnerPackageName);
 		PackageData->SetGenerated(OwnerPackageName);
+		PackageData->SetGeneratedNeedCachedPlatformDataBeforeSplit(bNeedCachedPlatformDataBeforeSplit);
 		if (IFileManager::Get().FileExists(*PackageData->GetFileName().ToString()))
 		{
 			UE_LOG(LogCook, Warning,
@@ -2089,7 +2092,8 @@ bool FGenerationHelper::TryGenerateList()
 			[](const FAssetDependency& A, const FAssetDependency& B) { return A.LexicalLess(B); });
 		GeneratedInfo.PackageDependencies.SetNum(Algo::Unique(GeneratedInfo.PackageDependencies));
 		GeneratedInfo.SetIsCreateAsMap(bCreateAsMap);
-		if (COTFS.MPCookGeneratorSplit == UE::Cook::EMPCookGeneratorSplit::AllOnSameWorker)
+		if (bNeedCachedPlatformDataBeforeSplit ||
+			COTFS.MPCookGeneratorSplit == UE::Cook::EMPCookGeneratorSplit::AllOnSameWorker)
 		{
 			PackageData->SetWorkerAssignmentConstraint(FWorkerId::Local());
 		}
diff --git a/Engine/Source/Editor/UnrealEd/Private/Cooker/CookPackageData.h b/Engine/Source/Editor/UnrealEd/Private/Cooker/CookPackageData.h
--- a/Engine/Source/Editor/UnrealEd/Private/Cooker/CookPackageData.h
+++ b/Engine/Source/Editor/UnrealEd/Private/Cooker/CookPackageData.h
@@ -666,6 +666,10 @@ public:
 	 */
 	TRefCountPtr<FGenerationHelper> TryCreateValidParentGenerationHelper();
 
+	/** Get/Set true if this is a generated package and its splitter reports NeedCachedPlatformDataBeforeSplit. */
+	bool IsGeneratedNeedCachedPlatformDataBeforeSplit() const;
+	void SetGeneratedNeedCachedPlatformDataBeforeSplit(bool bValue);
+
 	/**
 	 * Return the instigator for this package. The Instigator is the first code location or
 	 * referencing package that causes the package to enter the requested state.
@@ -847,6 +851,7 @@ private:
 	uint32 bGenerated : 1;
 	uint32 bKeepReferencedDuringGC : 1;
 	uint32 bWasCookedThisSession : 1;
+	uint32 bGeneratedNeedCachedPlatformDataBeforeSplit : 1;
 };
 
 /**
@@ -1022,6 +1027,8 @@ public:
 	const FName GetSplitDataObjectName() const;
 	/** Return the Splitter's value for virtual bool UseInternalReferenceToAvoidGarbageCollect(). */
 	bool IsUseInternalReferenceToAvoidGarbageCollect() const;
+	/** Return the Splitter's value for virtual bool NeedCachedPlatformDataBeforeSplit(). */
+	bool IsNeedCachedPlatformDataBeforeSplit() const;
 	/** Return the cached pointer to the SplitDataObject. Returns null if no longer in memory or marked as garbage. */
 	UObject* GetWeakSplitDataObject() const;
 	/**
@@ -1162,6 +1169,7 @@ private:
 	FWorkerId WorkerIdThatSavedGenerator = FWorkerId::Invalid();
 	EInitializeStatus InitializeStatus = EInitializeStatus::Uninitialized;
 	bool bUseInternalReferenceToAvoidGarbageCollect = false;
+	bool bNeedCachedPlatformDataBeforeSplit = false;
 	bool bGeneratedList = false;
 };
 
@@ -1817,6 +1825,16 @@ inline TRefCountPtr<FGenerationHelper> FPackageData::GetParentGenerationHelper()
 	return ParentGenerationHelper;
 }
 
+inline bool FPackageData::IsGeneratedNeedCachedPlatformDataBeforeSplit() const
+{
+	return bGeneratedNeedCachedPlatformDataBeforeSplit != 0;
+}
+
+inline void FPackageData::SetGeneratedNeedCachedPlatformDataBeforeSplit(bool bValue)
+{
+	bGeneratedNeedCachedPlatformDataBeforeSplit = (uint32)bValue;
+}
+
 inline bool FGenerationHelper::IsInitialized() const
 {
 	return InitializeStatus != EInitializeStatus::Uninitialized;
@@ -1877,6 +1895,12 @@ inline bool FGenerationHelper::IsUseInternalReferenceToAvoidGarbageCollect() con
 	return bUseInternalReferenceToAvoidGarbageCollect;
 }
 
+inline bool FGenerationHelper::IsNeedCachedPlatformDataBeforeSplit() const
+{
+	ConditionalInitialize();
+	return bNeedCachedPlatformDataBeforeSplit;
+}
+
 inline UObject* FGenerationHelper::GetWeakSplitDataObject() const
 {
 	ConditionalInitialize();
diff --git a/Engine/Source/Editor/UnrealEd/Private/Cooker/CookWorkerClient.cpp b/Engine/Source/Editor/UnrealEd/Private/Cooker/CookWorkerClient.cpp
--- a/Engine/Source/Editor/UnrealEd/Private/Cooker/CookWorkerClient.cpp
+++ b/Engine/Source/Editor/UnrealEd/Private/Cooker/CookWorkerClient.cpp
@@ -317,6 +317,7 @@ void FCookWorkerClient::ReportDiscoveredPackage(const FPackageData& PackageData,
 	Discovered.Instigator = Instigator;
 	Discovered.Platforms = MoveTemp(ReachablePlatforms);
 	Discovered.Platforms.ConvertToBitfield(OrderedSessionAndSpecialPlatforms);
+	Discovered.bNeedCachedPlatformDataBeforeSplit = PackageData.IsGeneratedNeedCachedPlatformDataBeforeSplit();
 }
 
 void FCookWorkerClient::ReportGeneratorQueuedGeneratedPackages(FGenerationHelper& GenerationHelper)
diff --git a/Engine/Source/Editor/UnrealEd/Private/Cooker/CookWorkerServer.cpp b/Engine/Source/Editor/UnrealEd/Private/Cooker/CookWorkerServer.cpp
--- a/Engine/Source/Editor/UnrealEd/Private/Cooker/CookWorkerServer.cpp
+++ b/Engine/Source/Editor/UnrealEd/Private/Cooker/CookWorkerServer.cpp
@@ -683,6 +683,9 @@ void FCookWorkerServer::SendPendingPackages()
 		FAssignPackageData& AssignData = AssignDatas.Emplace_GetRef();
 		AssignData.ConstructData = PackageData->CreateConstructData();
 		AssignData.ParentGenerator = PackageData->GetParentGenerator();
+		// We do not record PackageData->IsNeedCachedPlatformDataBeforeSplit(). The CookWorker that created
+		// it knows it, and the Director knows it from FDiscoveredPackageReplication, and the package will not
+		// be sent to other CookWorkers if it is true.
 		AssignData.Instigator = PackageData->GetInstigator();
 		SessionPlatformNeedsCook.Init(false, OrderedSessionPlatforms.Num());
 		int32 PlatformIndex = 0;
@@ -995,6 +998,7 @@ void FCookWorkerServer::QueueDiscoveredPackage(FDiscoveredPackageReplication&& D
 	if (!DiscoveredPackage.ParentGenerator.IsNone())
 	{
 		PackageData.SetGenerated(DiscoveredPackage.ParentGenerator);
+		PackageData.SetGeneratedNeedCachedPlatformDataBeforeSplit(DiscoveredPackage.bNeedCachedPlatformDataBeforeSplit);
 		FPackageData* GeneratorPackageData = PackageDatas.FindPackageDataByPackageName(
 			DiscoveredPackage.ParentGenerator);
 		if (GeneratorPackageData)
@@ -1040,8 +1044,9 @@ void FCookWorkerServer::QueueDiscoveredPackage(FDiscoveredPackageReplication&& D
 		return;
 	}
 
-	if (Instigator.Category == EInstigator::GeneratedPackage
-		&& COTFS.MPCookGeneratorSplit == EMPCookGeneratorSplit::AllOnSameWorker)
+	if (PackageData.IsGenerated()
+		&& (PackageData.IsGeneratedNeedCachedPlatformDataBeforeSplit()
+				|| COTFS.MPCookGeneratorSplit == EMPCookGeneratorSplit::AllOnSameWorker))
 	{
 		PackageData.SetWorkerAssignmentConstraint(GetWorkerId());
 	}
@@ -1367,6 +1372,7 @@ void FDiscoveredPackageReplication::Write(FCbWriter& Writer,
 	Writer << ParentGenerator;
 	Writer << static_cast<uint8>(Instigator.Category);
 	Writer << Instigator.Referencer;
+	Writer << bNeedCachedPlatformDataBeforeSplit;
 	WriteToCompactBinary(Writer, Platforms, OrderedSessionAndSpecialPlatforms);
 	Writer.EndArray();
 }
@@ -1396,6 +1402,8 @@ bool FDiscoveredPackageReplication::TryRead(FCbFieldView Field,
 		bOk = false;
 	}
 	bOk = LoadFromCompactBinary(Iter++, Instigator.Referencer) & bOk;
+	bNeedCachedPlatformDataBeforeSplit = Iter->AsBool();
+	bOk = (!(Iter++)->HasError()) & bOk;
 	bOk = LoadFromCompactBinary(Iter++, Platforms, OrderedSessionAndSpecialPlatforms) & bOk;
 	if (!bOk)
 	{
diff --git a/Engine/Source/Editor/UnrealEd/Private/Cooker/CookWorkerServer.h b/Engine/Source/Editor/UnrealEd/Private/Cooker/CookWorkerServer.h
--- a/Engine/Source/Editor/UnrealEd/Private/Cooker/CookWorkerServer.h
+++ b/Engine/Source/Editor/UnrealEd/Private/Cooker/CookWorkerServer.h
@@ -393,6 +393,7 @@ struct FDiscoveredPackageReplication
 	FName ParentGenerator;
 	FInstigator Instigator;
 	FDiscoveredPlatformSet Platforms;
+	bool bNeedCachedPlatformDataBeforeSplit = false;
 
 private:
 	void Write(FCbWriter& Writer, TConstArrayView<const ITargetPlatform*> OrderedSessionAndSpecialPlatforms) const;
diff --git a/Engine/Source/Editor/UnrealEd/Public/CookPackageSplitter.h b/Engine/Source/Editor/UnrealEd/Public/CookPackageSplitter.h
--- a/Engine/Source/Editor/UnrealEd/Public/CookPackageSplitter.h
+++ b/Engine/Source/Editor/UnrealEd/Public/CookPackageSplitter.h
@@ -39,7 +39,9 @@ public:
 
 
 	// Virtual API functions - functions called from the cooker after creating the splitter.
-	virtual ~ICookPackageSplitter() {}
+	virtual ~ICookPackageSplitter()
+	{
+	}
 
 	enum class ETeardown
 	{
@@ -47,13 +49,28 @@ public:
 		Canceled,
 	};
 	/** Do teardown actions after all packages have saved, or when the cook is cancelled. Always called before destruction. */
-	virtual void Teardown(ETeardown Status) {}
+	virtual void Teardown(ETeardown Status)
+	{
+	}
 
 	/**
 	 * If true, this splitter forces the Generator package objects it needs to remain referenced, and the cooker
 	 * should expect them to still be in memory after a garbage collect so long as the splitter is alive.
 	 */
-	virtual bool UseInternalReferenceToAvoidGarbageCollect() { return false; }
+	virtual bool UseInternalReferenceToAvoidGarbageCollect()
+	{
+		return false;
+	}
+	/**
+	 * If true, the cooker will not call any splitter work functions (GetGenerateList, Populate*, PreSave*, PostSave*)
+	 * until BeginCacheForCookedPlatformData has been called for all objects in the generator package. This reduces
+	 * MPCook performance - generated packages have to be saved on the same Worker that saved the generator - but is
+	 * otherwise not a problem.
+	 */
+	virtual bool NeedCachedPlatformDataBeforeSplit() 
+	{
+		 return false;
+	 }
 
 	/** Data sent to the cooker to describe each desired generated package */
 	struct FGeneratedPackage
