Interchange: Set the toolkit editor to EVisibility::Hidden instead of EVisibility::Collapsed, to avoid a dx12 crash. The jira is not fix, this change simply avoid using feature that crash with dx12.

#rnx

diff --git a/Engine/Source/Editor/UnrealEd/Private/Toolkits/AssetEditorToolkit.cpp b/Engine/Source/Editor/UnrealEd/Private/Toolkits/AssetEditorToolkit.cpp
--- a/Engine/Source/Editor/UnrealEd/Private/Toolkits/AssetEditorToolkit.cpp
+++ b/Engine/Source/Editor/UnrealEd/Private/Toolkits/AssetEditorToolkit.cpp
@@ -231,11 +231,11 @@ void FAssetEditorToolkit::InitAssetEditor( const EToolkitMode::Type Mode, const
 						{
 							if (AsyncAsset->IsCompiling())
 							{
-								return EVisibility::Collapsed;
+								return EVisibility::Hidden;
 							}
 						}
 					}
-					return EVisibility::All;
+					return EVisibility::Visible;
 				})
 			.OnRequestClose(this, &FAssetEditorToolkit::OnRequestClose, EAssetEditorCloseReason::AssetEditorHostClosed)
 			.OnClose(this, &FAssetEditorToolkit::OnClose)
