Sequencer: FBXMainExport.cpp fix for bogus pointer used static check.

diff --git a/Engine/Source/Editor/UnrealEd/Private/Fbx/FbxMainExport.cpp b/Engine/Source/Editor/UnrealEd/Private/Fbx/FbxMainExport.cpp
--- a/Engine/Source/Editor/UnrealEd/Private/Fbx/FbxMainExport.cpp
+++ b/Engine/Source/Editor/UnrealEd/Private/Fbx/FbxMainExport.cpp
@@ -1717,16 +1717,15 @@ FLevelSequenceAnimTrackAdapter::FLevelSequenceAnimTrackAdapter(IMovieScenePlayer
 	MovieScene = MovieSceneSequence->GetMovieScene();
 	RootToLocalTransform = InRootToLocalTransform;
 	AnimTrack = InAnimTrack;
+	MovieSceneSequence = nullptr;
+	RootMovieSceneSequence = nullptr;
 	if (MovieScene)
 	{
-		UMovieSceneSequence* OwnerSceneSequence = MovieScene->GetTypedOuter<UMovieSceneSequence>();
-		MovieSceneSequence = OwnerSceneSequence;
-		RootMovieSceneSequence = OwnerSceneSequence;
-	}
-	else
-	{
-		MovieSceneSequence = nullptr;
-		RootMovieSceneSequence = nullptr;
+		if (UMovieSceneSequence* OwnerSceneSequence = MovieScene->GetTypedOuter<UMovieSceneSequence>())
+		{
+			MovieSceneSequence = OwnerSceneSequence;
+			RootMovieSceneSequence = OwnerSceneSequence;
+		}
 	}
 }
 
