Sequencer: FBXMainExport.cpp fix for bogus pointer used static check.

diff --git a/Engine/Source/Editor/UnrealEd/Private/Fbx/FbxMainExport.cpp b/Engine/Source/Editor/UnrealEd/Private/Fbx/FbxMainExport.cpp
--- a/Engine/Source/Editor/UnrealEd/Private/Fbx/FbxMainExport.cpp
+++ b/Engine/Source/Editor/UnrealEd/Private/Fbx/FbxMainExport.cpp
@@ -1714,7 +1714,7 @@ FLevelSequenceAnimTrackAdapter::FLevelSequenceAnimTrackAdapter(IMovieScenePlayer
 FLevelSequenceAnimTrackAdapter::FLevelSequenceAnimTrackAdapter(IMovieScenePlayer* InMovieScenePlayer, UMovieScene* InMovieScene, const FMovieSceneSequenceTransform& InRootToLocalTransform, UMovieSceneSkeletalAnimationTrack* InAnimTrack)
 {
 	MovieScenePlayer = InMovieScenePlayer;
-	MovieScene = MovieSceneSequence->GetMovieScene();
+	MovieScene = InMovieScene;
 	RootToLocalTransform = InRootToLocalTransform;
 	AnimTrack = InAnimTrack;
 	MovieSceneSequence = nullptr;
