Fixed old comment in FBX Importer

diff --git a/Engine/Source/Editor/UnrealEd/Private/SkeletalMeshEdit.cpp b/Engine/Source/Editor/UnrealEd/Private/SkeletalMeshEdit.cpp
--- a/Engine/Source/Editor/UnrealEd/Private/SkeletalMeshEdit.cpp
+++ b/Engine/Source/Editor/UnrealEd/Private/SkeletalMeshEdit.cpp
@@ -2212,7 +2212,7 @@ void UnFbx::FFbxImporter::ImportBoneTracks(USkeleton* Skeleton, FAnimCurveImport
 				continue;
 			}
 
-			// Add data to RawTracks, TimeKeys, NewDebugDatas
+			// Add data to RawTracks, TimeKeys
 
 			RawTracks[SourceTrackIdx].ScaleKeys.Add(FVector3f(LocalTransform.GetScale3D()));
 			RawTracks[SourceTrackIdx].PosKeys.Add(FVector3f(LocalTransform.GetTranslation()));
