Cooker: Move FGenerationHelper from CookPackageData.* into CookGenerationHelper.* to improve readability and reduce header compile time.
#rnx

diff --git a/Engine/Source/Editor/UnrealEd/Private/CookOnTheFlyServer.cpp b/Engine/Source/Editor/UnrealEd/Private/CookOnTheFlyServer.cpp
--- a/Engine/Source/Editor/UnrealEd/Private/CookOnTheFlyServer.cpp
+++ b/Engine/Source/Editor/UnrealEd/Private/CookOnTheFlyServer.cpp
@@ -23,6 +23,7 @@
 #include "Cooker/CookConfigAccessTracker.h"
 #include "Cooker/CookDiagnostics.h"
 #include "Cooker/CookDirector.h"
+#include "Cooker/CookGenerationHelper.h"
 #include "Cooker/CookOnTheFlyServerInterface.h"
 #include "Cooker/CookPackageData.h"
 #include "Cooker/CookPlatformManager.h"
diff --git a/Engine/Source/Editor/UnrealEd/Private/Cooker/CookDiagnostics.cpp b/Engine/Source/Editor/UnrealEd/Private/Cooker/CookDiagnostics.cpp
--- a/Engine/Source/Editor/UnrealEd/Private/Cooker/CookDiagnostics.cpp
+++ b/Engine/Source/Editor/UnrealEd/Private/Cooker/CookDiagnostics.cpp
@@ -5,6 +5,7 @@
 #include "AssetRegistry/IAssetRegistry.h"
 #include "Containers/Array.h"
 #include "Containers/RingBuffer.h"
+#include "Cooker/CookGenerationHelper.h"
 #include "Cooker/CookPackageData.h"
 #include "Cooker/CookRequestCluster.h"
 #include "Cooker/CookTypes.h"
diff --git a/Engine/Source/Editor/UnrealEd/Private/Cooker/CookDirector.cpp b/Engine/Source/Editor/UnrealEd/Private/Cooker/CookDirector.cpp
--- a/Engine/Source/Editor/UnrealEd/Private/Cooker/CookDirector.cpp
+++ b/Engine/Source/Editor/UnrealEd/Private/Cooker/CookDirector.cpp
@@ -4,6 +4,7 @@
 
 #include "Async/Fundamental/Scheduler.h"
 #include "CompactBinaryTCP.h"
+#include "Cooker/CookGenerationHelper.h"
 #include "Cooker/CookPackageData.h"
 #include "Cooker/CookPlatformManager.h"
 #include "Cooker/CookWorkerServer.h"
diff --git a/Engine/Source/Editor/UnrealEd/Private/Cooker/CookDirector.h b/Engine/Source/Editor/UnrealEd/Private/Cooker/CookDirector.h
--- a/Engine/Source/Editor/UnrealEd/Private/Cooker/CookDirector.h
+++ b/Engine/Source/Editor/UnrealEd/Private/Cooker/CookDirector.h
@@ -30,6 +30,7 @@ class UCookOnTheFlyServer;
 namespace UE::Cook { class FCookWorkerServer; }
 namespace UE::Cook { struct FAssignPackageExtraData; }
 namespace UE::Cook { struct FCookWorkerProfileData; }
+namespace UE::Cook { struct FGenerationHelper; }
 namespace UE::Cook { struct FHeartbeatMessage; }
 namespace UE::Cook { struct FInitialConfigMessage; }
 namespace UE::Cook { struct FPackageData; }
diff --git a/Engine/Source/Editor/UnrealEd/Private/Cooker/CookGenerationHelper.cpp b/Engine/Source/Editor/UnrealEd/Private/Cooker/CookGenerationHelper.cpp
new file mode 100644
--- /dev/null
+++ b/Engine/Source/Editor/UnrealEd/Private/Cooker/CookGenerationHelper.cpp
@@ -0,0 +1,1188 @@
+// Copyright Epic Games, Inc. All Rights Reserved.
+
+#include "Cooker/CookGenerationHelper.h"
+
+#include "Algo/Unique.h"
+#include "AssetRegistry/AssetData.h"
+#include "Cooker/CookDirector.h"
+#include "Cooker/IWorkerRequests.h"
+#include "Cooker/PackageTracker.h"
+#include "Misc/PackageAccessTrackingOps.h"
+#include "UObject/ReferenceChainSearch.h"
+
+namespace UE::Cook
+{
+
+//////////////////////////////////////////////////////////////////////////
+// FGenerationHelper
+
+FGenerationHelper::FGenerationHelper(UE::Cook::FPackageData& InOwner)
+: OwnerInfo(InOwner, true /* bInGenerator */)
+{
+}
+
+FGenerationHelper::~FGenerationHelper()
+{
+	NotifyCompletion(ICookPackageSplitter::ETeardown::Complete);
+	GetOwner().OnGenerationHelperDestroyed(*this);
+}
+
+void FGenerationHelper::NotifyCompletion(ICookPackageSplitter::ETeardown Status)
+{
+	if (IsInitialized() && IsValid() && CookPackageSplitterInstance)
+	{
+		CookPackageSplitterInstance->Teardown(Status);
+		CookPackageSplitterInstance.Reset();
+	}
+}
+
+void FGenerationHelper::Initialize()
+{
+	if (InitializeStatus != EInitializeStatus::Uninitialized)
+	{
+		return;
+	}
+
+	FPackageData& OwnerPackageData = GetOwner();
+	FName OwnerPackageName = OwnerPackageData.GetPackageName();
+	UCookOnTheFlyServer& COTFS = OwnerPackageData.GetPackageDatas().GetCookOnTheFlyServer();
+	UPackage* LocalOwnerPackage = FindOrLoadPackage(COTFS, OwnerPackageData);
+	if (!LocalOwnerPackage)
+	{
+		InitializeStatus = EInitializeStatus::Invalid;
+		return;
+	}
+
+	UObject* LocalSplitDataObject;
+	UE::Cook::Private::FRegisteredCookPackageSplitter* LocalRegisteredSplitterType = nullptr;
+	TUniquePtr<ICookPackageSplitter> LocalSplitter;
+	SearchForRegisteredSplitDataObject(COTFS, OwnerPackageName, LocalOwnerPackage,
+		TOptional<TConstArrayView<FCachedObjectInOuter>>(), LocalSplitDataObject, LocalRegisteredSplitterType,
+		LocalSplitter);
+	if (!LocalSplitDataObject || !LocalSplitter)
+	{
+		InitializeStatus = EInitializeStatus::Invalid;
+		return;
+	}
+
+	Initialize(LocalSplitDataObject, LocalRegisteredSplitterType, MoveTemp(LocalSplitter));
+}
+
+void FGenerationHelper::Initialize(const UObject* InSplitDataObject,
+	UE::Cook::Private::FRegisteredCookPackageSplitter* InRegisteredSplitterType,
+	TUniquePtr<ICookPackageSplitter>&& InCookPackageSplitterInstance)
+{
+	check(InSplitDataObject);
+	if (InitializeStatus != EInitializeStatus::Uninitialized)
+	{
+		// If we already have a splitter, keep the old and throw out the new. The old one
+		// still contains some state.
+		return;
+	}
+
+	RegisteredSplitterType = InRegisteredSplitterType;
+	CookPackageSplitterInstance = MoveTemp(InCookPackageSplitterInstance);
+	InitializeStatus = EInitializeStatus::Valid;
+
+	SplitDataObject = InSplitDataObject;
+	SplitDataObjectName = FName(FStringView(InSplitDataObject->GetFullName()));
+	bUseInternalReferenceToAvoidGarbageCollect =
+		CookPackageSplitterInstance->UseInternalReferenceToAvoidGarbageCollect();
+	bNeedCachedPlatformDataBeforeSplit =
+		CookPackageSplitterInstance->NeedCachedPlatformDataBeforeSplit();
+}
+
+void FGenerationHelper::InitializeAsInvalid()
+{
+	if (InitializeStatus != EInitializeStatus::Uninitialized)
+	{
+		return;
+	}
+	InitializeStatus = EInitializeStatus::Invalid;
+}
+
+UPackage* FGenerationHelper::FindOrLoadPackage(UCookOnTheFlyServer& COTFS, FPackageData& OwnerPackageData)
+{
+	// This is the static helper function on FGenerationHelper that loads the package for any FPackageData; for the
+	// member variable function that uses the cached pointer, see FindOrLoadOwnerPackage.
+	FName OwnerPackageName = OwnerPackageData.GetPackageName();
+	UPackage* Result = FindObjectFast<UPackage>(nullptr, OwnerPackageName);
+
+	if (!Result || !Result->IsFullyLoaded())
+	{
+		COTFS.LoadPackageForCooking(OwnerPackageData, Result);
+		if (!Result || !Result->IsFullyLoaded())
+		{
+			return nullptr;
+		}
+	}
+	return Result;
+}
+
+void FGenerationHelper::SearchForRegisteredSplitDataObject(UCookOnTheFlyServer& COTFS,
+	FName PackageName, UPackage* Package,
+	TOptional<TConstArrayView<UE::Cook::FCachedObjectInOuter>> CachedObjectsInOuter,
+	UObject*& OutSplitDataObject, UE::Cook::Private::FRegisteredCookPackageSplitter*& OutRegisteredSplitter,
+	TUniquePtr<ICookPackageSplitter>& OutSplitterInstance)
+{
+	OutSplitDataObject = nullptr;
+	OutRegisteredSplitter = nullptr;
+	OutSplitterInstance = nullptr;
+	check(Package != nullptr || CachedObjectsInOuter.IsSet());
+
+	UObject* LocalSplitDataObject = nullptr;
+	Private::FRegisteredCookPackageSplitter* SplitterType = nullptr;
+	TArray<Private::FRegisteredCookPackageSplitter*> FoundRegisteredSplitters;
+	auto TryLookForSplitterOfObject =
+		[&COTFS, PackageName, &FoundRegisteredSplitters, &SplitterType, &LocalSplitDataObject](UObject* Obj)
+		{
+			FoundRegisteredSplitters.Reset();
+			COTFS.RegisteredSplitDataClasses.MultiFind(Obj->GetClass(), FoundRegisteredSplitters);
+
+			for (Private::FRegisteredCookPackageSplitter* SplitterForObject : FoundRegisteredSplitters)
+			{
+				if (SplitterForObject && SplitterForObject->ShouldSplitPackage(Obj))
+				{
+					if (!Obj->HasAnyFlags(RF_Public))
+					{
+						UE_LOG(LogCook, Error,
+							TEXT("SplitterData object %s must be publicly referenceable so we can keep them from being garbage collected"),
+							*Obj->GetFullName());
+						return false;
+					}
+
+					if (SplitterType)
+					{
+						UE_LOG(LogCook, Error,
+							TEXT("Found more than one registered Cook Package Splitter for package %s."),
+							*PackageName.ToString());
+						return false;
+					}
+
+					SplitterType = SplitterForObject;
+					LocalSplitDataObject = Obj;
+				}
+			}
+			return true;
+		};
+
+	if (CachedObjectsInOuter.IsSet())
+	{
+		// CachedObjectsInOuter might be set but empty for e.g. a generated package that has not been populated
+		for (const FCachedObjectInOuter& CachedObjectInOuter : *CachedObjectsInOuter)
+		{
+			UObject* Obj = CachedObjectInOuter.Object.Get();
+			if (!Obj)
+			{
+				continue;
+			}
+			if (!TryLookForSplitterOfObject(Obj))
+			{
+				return; // error condition, exit the entire search function
+			}
+		}
+	}
+	else
+	{
+		TArray<UObject*> ObjectsInPackage;
+		GetObjectsWithOuter(Package, ObjectsInPackage, true /* bIncludeNestedObjects */,
+			RF_NoFlags, EInternalObjectFlags::Garbage);
+		for (UObject* Obj : ObjectsInPackage)
+		{
+			if (!TryLookForSplitterOfObject(Obj))
+			{
+				return; // error condition, exit the entire search function
+			}
+		}
+	}
+
+	if (!SplitterType)
+	{
+		return;
+	}
+
+	// Create instance of CookPackageSplitter class
+	ICookPackageSplitter* SplitterInstance = SplitterType->CreateInstance(LocalSplitDataObject);
+	if (!SplitterInstance)
+	{
+		UE_LOG(LogCook, Error, TEXT("Error instantiating Cook Package Splitter %s for object %s."),
+			*SplitterType->GetSplitterDebugName(), *LocalSplitDataObject->GetFullName());
+		return;
+	}
+
+	OutSplitDataObject = LocalSplitDataObject;
+	OutRegisteredSplitter = SplitterType;
+	OutSplitterInstance.Reset(SplitterInstance);
+}
+
+void FGenerationHelper::ClearSelfReferences()
+{
+	// Any references we release might be the last reference and cause *this to be deleted,
+	// so create a local reference to keep it alive until the end of the function.
+	TRefCountPtr<FGenerationHelper> LocalRef(this);
+	ClearKeepForIterative();
+	ClearKeepForGeneratorSave();
+	ClearKeepForQueueResults();
+}
+
+FCookGenerationInfo* FGenerationHelper::FindInfo(const FPackageData& PackageData)
+{
+	ConditionalInitialize();
+
+	if (&PackageData == &GetOwner())
+	{
+		return &OwnerInfo;
+	}
+	for (FCookGenerationInfo& Info : PackagesToGenerate)
+	{
+		if (Info.PackageData == &PackageData)
+		{
+			return &Info;
+		}
+	}
+	return nullptr;
+}
+
+const FCookGenerationInfo* FGenerationHelper::FindInfo(const FPackageData& PackageData) const
+{
+	return const_cast<FGenerationHelper*>(this)->FindInfo(PackageData);
+}
+
+UObject* FGenerationHelper::FindOrLoadSplitDataObject()
+{
+	if (!IsValid())
+	{
+		return nullptr;
+	}
+	UObject* Result = SplitDataObject.Get();
+	if (Result)
+	{
+		return Result;
+	}
+
+	FString ObjectPath = GetSplitDataObjectName().ToString();
+	// SplitDataObjectName is a FullObjectPath; strip off the leading <ClassName> in
+	// "<ClassName> <Package>.<Object>:<SubObject>"
+	int32 ClassDelimiterIndex = -1;
+	if (ObjectPath.FindChar(' ', ClassDelimiterIndex))
+	{
+		ObjectPath.RightChopInline(ClassDelimiterIndex + 1);
+	}
+
+	Result = FindObject<UObject>(nullptr, *ObjectPath);
+	if (!Result)
+	{
+		FPackageData& OwnerPackageData = GetOwner();
+		FName OwnerPackageName = OwnerPackageData.GetPackageName();
+		UCookOnTheFlyServer& COTFS = OwnerPackageData.GetPackageDatas().GetCookOnTheFlyServer();
+		UPackage* LocalOwnerPackage;
+		COTFS.LoadPackageForCooking(OwnerPackageData, LocalOwnerPackage);
+
+		Result = FindObject<UObject>(nullptr, *ObjectPath);
+		if (!Result)
+		{
+			return nullptr;
+		}
+	}
+
+	SplitDataObject = Result;
+	return Result;
+}
+
+UPackage* FGenerationHelper::GetOwnerPackage()
+{
+	UPackage* Result = OwnerPackage.Get();
+	if (!Result && !OwnerPackage.GetEvenIfUnreachable())
+	{
+		OwnerPackage = FindObjectFast<UPackage>(nullptr, GetOwner().GetPackageName());
+		Result = OwnerPackage.Get();
+	}
+	return Result;
+}
+
+UPackage* FGenerationHelper::FindOrLoadOwnerPackage(UCookOnTheFlyServer& COTFS)
+{
+	UPackage* Result = GetOwnerPackage();
+	if (!Result)
+	{
+		Result = FindOrLoadPackage(COTFS, GetOwner());
+	}
+	return Result;
+}
+
+bool FGenerationHelper::TryGenerateList()
+{
+	if (bGeneratedList)
+	{
+		return true;
+	}
+	FPackageData& OwnerPackageData = GetOwner();
+	FName OwnerPackageName = OwnerPackageData.GetPackageName();
+	if (!IsValid())
+	{
+		// Unexpected, caller should not call in this case
+		UE_LOG(LogCook, Error, TEXT("TryGenerateList failed for package %s: Called on an invalid FGenerationHelper."),
+			*OwnerPackageName.ToString());
+		FDebug::DumpStackTraceToLog(ELogVerbosity::Warning);
+		return false;
+	}
+
+	FPackageDatas& PackageDatas = OwnerPackageData.GetPackageDatas();
+	UCookOnTheFlyServer& COTFS = PackageDatas.GetCookOnTheFlyServer();
+	UObject* OwnerObject = FindOrLoadSplitDataObject();
+	if (!OwnerObject)
+	{
+		// Unexpected, we found it earlier when we marked valid.
+		UE_LOG(LogCook, Error, TEXT("TryGenerateList failed for package %s: Valid GenerationHelper but could not find OwnerObject."),
+			*OwnerPackageName.ToString());
+		FDebug::DumpStackTraceToLog(ELogVerbosity::Warning);
+		return false;
+	}
+
+	UPackage* LocalOwnerPackage = OwnerObject->GetPackage();
+
+	TArray<ICookPackageSplitter::FGeneratedPackage> GeneratorDatas;
+	{
+		UCookOnTheFlyServer::FScopedActivePackage ScopedActivePackage(COTFS, OwnerPackageName,
+			PackageAccessTrackingOps::NAME_CookerBuildObject);
+		GeneratorDatas = GetCookPackageSplitterInstance()->GetGenerateList(LocalOwnerPackage, OwnerObject);
+	}
+	PackagesToGenerate.Reset(GeneratorDatas.Num());
+	TArray<const ITargetPlatform*, TInlineAllocator<1>> PlatformsToCook;
+	OwnerPackageData.GetPlatformsNeedingCooking(PlatformsToCook);
+
+	for (ICookPackageSplitter::FGeneratedPackage& SplitterData : GeneratorDatas)
+	{
+		if (!SplitterData.GetCreateAsMap().IsSet())
+		{
+			UE_LOG(LogCook, Error,
+				TEXT("PackageSplitter did not specify whether CreateAsMap is true for generated package. Splitter=%s, Generated=%s."),
+				*this->GetSplitDataObjectName().ToString(), *OwnerPackageName.ToString());
+			return false;
+		}
+		bool bCreateAsMap = *SplitterData.GetCreateAsMap();
+
+		FString PackageRoot = SplitterData.GeneratedRootPath.IsEmpty() ?
+			OwnerPackageName.ToString() : SplitterData.GeneratedRootPath;
+
+		FString PackageName = FPaths::RemoveDuplicateSlashes(FString::Printf(TEXT("/%s/%s/%s"),
+			*PackageRoot, GeneratedPackageSubPath, *SplitterData.RelativePath));
+		const FName PackageFName(*PackageName);
+		UE::Cook::FPackageData* PackageData = PackageDatas.TryAddPackageDataByPackageName(PackageFName,
+			false /* bRequireExists */, bCreateAsMap);
+		if (!PackageData)
+		{
+			UE_LOG(LogCook, Error,
+				TEXT("PackageSplitter could not find mounted filename for generated packagepath. Splitter=%s, Generated=%s."),
+				*this->GetSplitDataObjectName().ToString(), *PackageName);
+			return false;
+		}
+		// No package should be generated by two different splitters.
+		check(PackageData->GetParentGenerator().IsNone() ||
+			PackageData->GetParentGenerator() == OwnerPackageName);
+		PackageData->SetGenerated(OwnerPackageName);
+		PackageData->SetGeneratedNeedCachedPlatformDataBeforeSplit(bNeedCachedPlatformDataBeforeSplit);
+		if (IFileManager::Get().FileExists(*PackageData->GetFileName().ToString()))
+		{
+			UE_LOG(LogCook, Warning,
+				TEXT("PackageSplitter specified a generated package that already exists in the workspace domain. Splitter=%s, Generated=%s."),
+				*this->GetSplitDataObjectName().ToString(), *PackageName);
+			return false;
+		}
+
+		FCookGenerationInfo& GeneratedInfo = PackagesToGenerate.Emplace_GetRef(*PackageData, false /* bInGenerator */);
+		GeneratedInfo.RelativePath = MoveTemp(SplitterData.RelativePath);
+		GeneratedInfo.GeneratedRootPath = MoveTemp(SplitterData.GeneratedRootPath);
+		GeneratedInfo.PackageDependencies = MoveTemp(SplitterData.PackageDependencies);
+		for (TArray<FAssetDependency>::TIterator Iter(GeneratedInfo.PackageDependencies); Iter; ++Iter)
+		{
+			if (Iter->Category != UE::AssetRegistry::EDependencyCategory::Package)
+			{
+				UE_LOG(LogCook, Error,
+					TEXT("PackageSplitter specified a dependency with category %d rather than category Package. Dependency will be ignored. Splitter=%s, Generated=%s."),
+					(int32)Iter->Category, *this->GetSplitDataObjectName().ToString(), *PackageName);
+				Iter.RemoveCurrent();
+			}
+		}
+		Algo::Sort(GeneratedInfo.PackageDependencies,
+			[](const FAssetDependency& A, const FAssetDependency& B) { return A.LexicalLess(B); });
+		GeneratedInfo.PackageDependencies.SetNum(Algo::Unique(GeneratedInfo.PackageDependencies));
+		GeneratedInfo.SetIsCreateAsMap(bCreateAsMap);
+		if (bNeedCachedPlatformDataBeforeSplit ||
+			COTFS.MPCookGeneratorSplit == UE::Cook::EMPCookGeneratorSplit::AllOnSameWorker)
+		{
+			PackageData->SetWorkerAssignmentConstraint(FWorkerId::Local());
+		}
+
+		// Create the Hash from the GenerationHash and Dependencies
+		GeneratedInfo.CreatePackageHash();
+	}
+
+	bGeneratedList = true;
+	return true;
+}
+
+void FGenerationHelper::StartOwnerSave()
+{
+	if (!IsValid())
+	{
+		return;
+	}
+	UE_LOG(LogCook, Display, TEXT("Splitting Package %s with splitter %s acting on object %s."),
+		*WriteToString<256>(GetOwner().GetPackageName()),
+		*GetRegisteredSplitterType()->GetSplitterDebugName(),
+		*WriteToString<256>(GetSplitDataObjectName()));
+	SetKeepForGeneratorSave();
+}
+
+void FGenerationHelper::StartQueueGeneratedPackages(UCookOnTheFlyServer& COTFS)
+{
+	if (!IsValid())
+	{
+		return;
+	}
+	NotifyStartQueueGeneratedPackages(COTFS, FWorkerId::Local());
+
+	bool bHybridIterativeEnabled = COTFS.bHybridIterativeEnabled;
+	if (!PreviousGeneratedPackages.IsEmpty())
+	{
+		FPackageData& OwnerPackageData = GetOwner();
+		TArray<const ITargetPlatform*, TInlineAllocator<1>> PlatformsToCook;
+		OwnerPackageData.GetPlatformsNeedingCooking(PlatformsToCook);
+
+		int32 NumIterativeUnmodified = 0;
+		int32 NumIterativeModified = 0;
+		int32 NumIterativeRemoved = 0;
+		int32 NumIterativePrevious = PreviousGeneratedPackages.Num();
+
+		for (FCookGenerationInfo& GeneratedInfo : PackagesToGenerate)
+		{
+			if (!GeneratedInfo.PackageData)
+			{
+				continue;
+			}
+			FIoHash PreviousHash;
+			if (PreviousGeneratedPackages.RemoveAndCopyValue(GeneratedInfo.PackageData->GetPackageName(), PreviousHash)
+				&& !bHybridIterativeEnabled)
+			{
+				bool bIterativelyUnmodified;
+				GeneratedInfo.IterativeCookValidateOrClear(*this, PlatformsToCook, PreviousHash, bIterativelyUnmodified);
+				++(bIterativelyUnmodified ? NumIterativeUnmodified : NumIterativeModified);
+			}
+		}
+		if (!PreviousGeneratedPackages.IsEmpty())
+		{
+			NumIterativeRemoved = PreviousGeneratedPackages.Num();
+			for (TPair<FName, FIoHash>& Pair : PreviousGeneratedPackages)
+			{
+				for (const ITargetPlatform* TargetPlatform : PlatformsToCook)
+				{
+					COTFS.DeleteOutputForPackage(Pair.Key, TargetPlatform);
+				}
+			}
+			PreviousGeneratedPackages.Empty();
+		}
+		ClearKeepForIterative();
+
+		if (NumIterativePrevious > 0 && !bHybridIterativeEnabled)
+		{
+			UE_LOG(LogCook, Display, TEXT("Found %d cooked package(s) in package store for generator package %s."),
+				NumIterativePrevious, *WriteToString<256>(GetOwner().GetPackageName()));
+			UE_LOG(LogCook, Display, TEXT("Keeping %d. Recooking %d. Removing %d."),
+				NumIterativeUnmodified, NumIterativeModified, NumIterativeRemoved);
+		}
+	}
+}
+
+void FGenerationHelper::NotifyStartQueueGeneratedPackages(UCookOnTheFlyServer& COTFS, FWorkerId SourceWorkerId)
+{
+	// Note this function can be called on an uninitialized Generator; the generator is only needed
+	// on the director so it can serve as the passer of messages. We have to keep ourselves referenced after
+	// this call, until after we send EGeneratorEvent::QueuedGeneratedPackagesFencePassed, so that we don't destruct
+	// and lose the WorkerIdThatSavedGenerator information.
+	if (COTFS.CookDirector)
+	{
+		COTFS.PackageDatas->GetRequestQueue().AddRequestFenceListener(GetOwner().GetPackageName());
+		WorkerIdThatSavedGenerator = SourceWorkerId;
+	}
+	SetKeepForQueueResults();
+}
+
+void FGenerationHelper::EndQueueGeneratedPackages(UCookOnTheFlyServer& COTFS)
+{
+	SetKeepForQueueResults();
+	COTFS.WorkerRequests->EndQueueGeneratedPackages(COTFS, *this);
+}
+
+void FGenerationHelper::EndQueueGeneratedPackagesOnDirector(UCookOnTheFlyServer& COTFS, FWorkerId SourceWorkerId)
+{
+	// Note this function can be called on an uninitialized Generator; the generator is only needed
+	// on the director so it can serve as the passer of messages.
+	COTFS.PackageDatas->GetRequestQueue().AddRequestFenceListener(GetOwner().GetPackageName());
+	SetKeepForQueueResults();
+
+	// Setting the WorkerIdThatSavedGenerator in response to this event is usually not needed because it is also set
+	// from the reported GeneratedPackages, but we set it anyway in case there is an edge condition that skips those
+	// notifications.
+	WorkerIdThatSavedGenerator = SourceWorkerId;
+}
+
+void FGenerationHelper::OnRequestFencePassedBroadcast(UCookOnTheFlyServer& COTFS)
+{
+	if (COTFS.CookDirector)
+	{
+		// Broadcast GenerationMessage API
+		COTFS.CookDirector->BroadcastGeneratorFencePassed(*this);
+	}
+	OnRequestFencePassed(COTFS);
+}
+
+void FGenerationHelper::OnRequestFencePassed(UCookOnTheFlyServer& COTFS)
+{
+	ClearKeepForQueueResults();
+	// We no longer need PreviousGeneratedPackages or KeepForIterative, because they are used (and cleared) in
+	// StartQueueGeneratedPackages. Clear them as well on the director and any CookWorkers that received it to
+	// free memory.
+	ClearKeepForIterative();
+	PreviousGeneratedPackages.Empty();
+}
+
+UPackage* FGenerationHelper::CreateGeneratedUPackage(FCookGenerationInfo& GeneratedInfo,
+	const UPackage* InOwnerPackage, const TCHAR* GeneratedPackageName)
+{
+	if (!IsValid())
+	{
+		return nullptr;
+	}
+
+#if ENABLE_COOK_STATS
+	++DetailedCookStats::NumRequestedLoads;
+#endif
+	UPackage* GeneratedPackage = CreatePackage(GeneratedPackageName);
+	GeneratedPackage->SetSavedHash(GeneratedInfo.PackageHash);
+	GeneratedPackage->SetPersistentGuid(InOwnerPackage->GetPersistentGuid());
+	GeneratedPackage->SetPackageFlags(PKG_CookGenerated);
+	GeneratedInfo.SetHasCreatedPackage(true);
+	if (!InOwnerPackage->IsLoadedByEditorPropertiesOnly())
+	{
+		GeneratedPackage->SetLoadedByEditorPropertiesOnly(false);
+	}
+
+	return GeneratedPackage;
+}
+
+void FGenerationHelper::ResetSaveState(FCookGenerationInfo& Info, UPackage* Package,
+	EStateChangeReason ReleaseSaveReason, EPackageState NewState)
+{
+	ConditionalInitialize();
+
+	// We release references to *this in this function so keep a local reference to avoid deletion during the function.
+	TRefCountPtr<FGenerationHelper> LocalRefCount = this;
+
+	if (Info.GetSaveState() > FCookGenerationInfo::ESaveState::CallPopulate)
+	{
+		UObject* SplitObject = GetWeakSplitDataObject();
+		UPackage* LocalOwnerPackage = Info.IsGenerator() ? Package : GetOwnerPackage();
+		if (!SplitObject || !Package || !LocalOwnerPackage)
+		{
+			UE_LOG(LogCook, Warning,
+				TEXT("PackageSplitter: %s on %s was GarbageCollected before we finished saving it. This prevents us from calling PostSave and may corrupt other packages that it altered during Populate. Splitter=%s."),
+				(!Package ? TEXT("UPackage") :
+					(!LocalOwnerPackage ? TEXT("ParentGenerator UPackage") : TEXT("SplitDataObject"))),
+				Info.PackageData ? *Info.PackageData->GetPackageName().ToString() : *Info.RelativePath,
+				*GetSplitDataObjectName().ToString());
+		}
+		else
+		{
+			UCookOnTheFlyServer& COTFS = GetOwner().GetPackageDatas().GetCookOnTheFlyServer();
+			UCookOnTheFlyServer::FScopedActivePackage ScopedActivePackage(COTFS, GetOwner().GetPackageName(),
+				PackageAccessTrackingOps::NAME_CookerBuildObject);
+			if (Info.IsGenerator())
+			{
+				GetCookPackageSplitterInstance()->PostSaveGeneratorPackage(Package, SplitObject);
+			}
+			else
+			{
+				ICookPackageSplitter::FGeneratedPackageForPopulate PopulateInfo;
+				PopulateInfo.RelativePath = Info.RelativePath;
+				PopulateInfo.GeneratedRootPath = Info.GeneratedRootPath;
+				PopulateInfo.bCreatedAsMap = Info.IsCreateAsMap();
+				PopulateInfo.Package = Package;
+				GetCookPackageSplitterInstance()->PostSaveGeneratedPackage(LocalOwnerPackage, SplitObject,
+					PopulateInfo);
+			}
+		}
+	}
+
+	if (ReleaseSaveReason == EStateChangeReason::RecreateObjectCache ||
+		ReleaseSaveReason == EStateChangeReason::DoneForNow)
+	{
+		if (Info.IsGenerator())
+		{
+			if (Info.GetSaveState() >= FCookGenerationInfo::ESaveState::StartPopulate)
+			{
+				Info.SetSaveState(FCookGenerationInfo::ESaveState::StartPopulate);
+			}
+			else
+			{
+				// Redo all the steps since we didn't make it to the FinishCachePreObjectsToMove.
+				// Restarting in the middle of that flow is not robust
+				Info.SetSaveState(FCookGenerationInfo::ESaveState::StartSave);
+			}
+		}
+		else
+		{
+			Info.SetSaveState(FCookGenerationInfo::ESaveState::StartPopulate);
+		}
+	}
+	else
+	{
+		// The save is completed and we will not come back to it; set state back to initial
+		// state and drop our reference keeping this GenerationHelper in memory for the save.
+		if (ReleaseSaveReason == EStateChangeReason::Completed)
+		{
+			Info.SetHasSaved(true);
+		}
+
+		if (Info.IsGenerator())
+		{
+			ClearKeepForGeneratorSave();
+			Info.SetSaveState(FCookGenerationInfo::ESaveState::StartSave);
+		}
+		else
+		{
+			Info.SetSaveState(FCookGenerationInfo::ESaveState::StartPopulate);
+			if (Info.PackageData)
+			{
+				Info.PackageData->SetParentGenerationHelper(nullptr);
+			}
+		}
+	}
+
+	if (Info.HasTakenOverCachedCookedPlatformData())
+	{
+		if (NewState != EPackageState::Idle && Info.PackageData &&
+			Info.PackageData->GetCachedObjectsInOuter().Num() != 0 && IsUseInternalReferenceToAvoidGarbageCollect() &&
+			(ReleaseSaveReason != EStateChangeReason::Completed && ReleaseSaveReason != EStateChangeReason::DoneForNow
+				&& ReleaseSaveReason != EStateChangeReason::SaveError
+				&& ReleaseSaveReason != EStateChangeReason::CookerShutdown))
+		{
+			UE_LOG(LogCook, Error,
+				TEXT("CookPackageSplitter failure: We are demoting a %s package from save and removing our references that keep its objects loaded.\n")
+				TEXT("This will allow the objects to be garbage collected and cause failures in the splitter which expects them to remain loaded.\n")
+				TEXT("Package=%s, Splitter=%s, ReleaseSaveReason=%s"),
+				Info.IsGenerator() ? TEXT("generator") : TEXT("generated"),
+				Info.PackageData ? *Info.PackageData->GetPackageName().ToString() : *Info.RelativePath,
+				*GetSplitDataObjectName().ToString(), LexToString(ReleaseSaveReason));
+		}
+		Info.CachedObjectsInOuterInfo.Empty();
+		Info.SetHasTakenOverCachedCookedPlatformData(false);
+	}
+	Info.SetHasIssuedUndeclaredMovedObjectsWarning(false);
+	Info.KeepReferencedPackages.Reset();
+}
+
+void FGenerationHelper::FetchExternalActorDependencies()
+{
+	if (!IsValid())
+	{
+		return;
+	}
+
+	// The Generator package declares all its ExternalActor dependencies in its AssetRegistry dependencies
+	// The Generator's generated packages can also include ExternalActors from other maps due to level instancing,
+	// these are included in the dependencies reported by the Generator for each GeneratedPackage in the data
+	// returned from GetGenerateList. These sets will overlap; take the union.
+	ExternalActorDependencies.Reset();
+	IAssetRegistry::GetChecked().GetDependencies(GetOwner().GetPackageName(), ExternalActorDependencies,
+		UE::AssetRegistry::EDependencyCategory::Package);
+	for (const FCookGenerationInfo& Info : PackagesToGenerate)
+	{
+		ExternalActorDependencies.Reserve(Info.GetDependencies().Num() + ExternalActorDependencies.Num());
+		for (const FAssetDependency& Dependency : Info.GetDependencies())
+		{
+			ExternalActorDependencies.Add(Dependency.AssetId.PackageName);
+		}
+	}
+	Algo::Sort(ExternalActorDependencies, FNameFastLess());
+	ExternalActorDependencies.SetNum(Algo::Unique(ExternalActorDependencies));
+	FPackageDatas& PackageDatas = this->GetOwner().GetPackageDatas();
+	FThreadSafeSet<FName>& NeverCookPackageList =
+		GetOwner().GetPackageDatas().GetCookOnTheFlyServer().PackageTracker->NeverCookPackageList;
+
+	// We are supposed to collect only ExternalActor dependencies, but we collected every dependency from the
+	// generated packages. Remove the packages that are not external actors, which we detect by being on-disk
+	// PackageDatas that are marked as NeverCook
+	ExternalActorDependencies.RemoveAll([&PackageDatas, &NeverCookPackageList](FName PackageName)
+		{
+			FPackageData* PackageData = PackageDatas.TryAddPackageDataByPackageName(PackageName);
+			if (!PackageData)
+			{
+				return true;
+			}
+			bool bIsNeverCook = NeverCookPackageList.Contains(PackageData->GetFileName());
+			return !bIsNeverCook;
+		});
+	ExternalActorDependencies.Shrink();
+}
+
+void FGenerationHelper::SetPreviousGeneratedPackages(TMap<FName, FIoHash>&& Packages)
+{
+	SetKeepForIterative();
+	PreviousGeneratedPackages = MoveTemp(Packages);
+}
+
+void FGenerationHelper::PreGarbageCollect(FCookGenerationInfo& Info, TArray<TObjectPtr<UObject>>& GCKeepObjects,
+	TArray<UPackage*>& GCKeepPackages, TArray<FPackageData*>& GCKeepPackageDatas, bool& bOutShouldDemote)
+{
+	if (!IsInitialized())
+	{
+		return;
+	}
+
+	bOutShouldDemote = false;
+	check(Info.PackageData); // Caller validates this is non-null
+	if (Info.GetSaveState() > FCookGenerationInfo::ESaveState::CallPopulate)
+	{
+		if (IsUseInternalReferenceToAvoidGarbageCollect() || Info.PackageData->GetIsCookLast())
+		{
+			UPackage* Package = Info.PackageData->GetPackage();
+			if (Package)
+			{
+				GCKeepPackages.Add(Package);
+				GCKeepPackageDatas.Add(Info.PackageData);
+			}
+		}
+		else
+		{
+			bOutShouldDemote = true;
+		}
+	}
+	if (Info.HasTakenOverCachedCookedPlatformData())
+	{
+		if (IsUseInternalReferenceToAvoidGarbageCollect())
+		{
+			// For the UseInternalReferenceToAvoidGarbageCollect case, part of the CookPackageSplitter contract is that
+			// the Cooker will keep referenced the package and all objects returned from GetObjectsToMove* functions
+			// until the PostSave function is called
+			UPackage* Package = Info.PackageData->GetPackage();
+			if (Package)
+			{
+				GCKeepPackages.Add(Package);
+				GCKeepPackageDatas.Add(Info.PackageData);
+			}
+			GCKeepPackages.Append(Info.KeepReferencedPackages);
+			for (FCachedObjectInOuter& CachedObjectInOuter : Info.PackageData->GetCachedObjectsInOuter())
+			{
+				UObject* Object = CachedObjectInOuter.Object.Get();
+				if (Object)
+				{
+					GCKeepObjects.Add(Object);
+				}
+			}
+		}
+	}
+}
+
+void FGenerationHelper::PostGarbageCollect()
+{
+	if (!IsInitialized())
+	{
+		return;
+	}
+
+	FPackageData& Owner = GetOwner();
+	if (Owner.GetState() == EPackageState::Save)
+	{
+		// UCookOnTheFlyServer::PreGarbageCollect adds references for the Generator package and all its public
+		// objects, so it should still be loaded
+		if (!Owner.GetPackage() || !GetWeakSplitDataObject())
+		{
+			UE_LOG(LogCook, Error,
+				TEXT("PackageSplitter object was deleted by garbage collection while generation was still ongoing. This will break the generation.")
+				TEXT("\n\tSplitter=%s."), *GetSplitDataObjectName().ToString());
+		}
+	}
+	else
+	{
+		// After the Generator Package is saved, we drop our references to it and it can be garbage collected
+		// If we have any packages left to populate, our splitter contract requires that it be garbage collected
+		// because we promise that the package is not partially GC'd during calls to TryPopulateGeneratedPackage
+		// The splitter can opt-out of this contract and keep it referenced itself if it desires.
+		UPackage* LocalOwnerPackage = FindObject<UPackage>(nullptr, *Owner.GetPackageName().ToString());
+		if (LocalOwnerPackage)
+		{
+			if (!Owner.IsInProgress() && !Owner.IsKeepReferencedDuringGC() &&
+				!IsUseInternalReferenceToAvoidGarbageCollect())
+			{
+				UE_LOG(LogCook, Error,
+					TEXT("PackageSplitter found the Generator package still in memory after it should have been deleted by GC.")
+					TEXT("\n\tThis is unexpected since garbage has been collected and the package should have been unreferenced so it should have been collected, and will break population of Generated packages.")
+					TEXT("\n\tSplitter=%s"), *GetSplitDataObjectName().ToString());
+				EReferenceChainSearchMode SearchMode = EReferenceChainSearchMode::Shortest
+					| EReferenceChainSearchMode::PrintAllResults
+					| EReferenceChainSearchMode::FullChain;
+				FReferenceChainSearch RefChainSearch(LocalOwnerPackage, SearchMode);
+			}
+		}
+	}
+
+	bool bHasIssuedWarning = false;
+	for (FCookGenerationInfo& Info : PackagesToGenerate)
+	{
+		if (FindObject<UPackage>(nullptr, *Info.PackageData->GetPackageName().ToString()))
+		{
+			if (!Info.PackageData->IsKeepReferencedDuringGC() && !Info.HasSaved() && !bHasIssuedWarning)
+			{
+				UE_LOG(LogCook, Warning,
+					TEXT("PackageSplitter found a package it generated that was not removed from memory during garbage collection. This will cause errors later during population.")
+					TEXT("\n\tSplitter=%s, Generated=%s."), *GetSplitDataObjectName().ToString(),
+					*Info.PackageData->GetPackageName().ToString());
+				
+				{
+					// Compute UCookOnTheFlyServer's references so they are gathered by OBJ REFS below 
+					UCookOnTheFlyServer::FScopeFindCookReferences(Owner.GetPackageDatas().GetCookOnTheFlyServer());
+
+					StaticExec(nullptr, *FString::Printf(TEXT("OBJ REFS NAME=%s"),
+						*Info.PackageData->GetPackageName().ToString()));
+				}
+				
+				bHasIssuedWarning = true; // Only issue the warning once per GC
+			}
+		}
+		else
+		{
+			Info.SetHasCreatedPackage(false);
+		}
+	}
+}
+
+void FGenerationHelper::UpdateSaveAfterGarbageCollect(const FPackageData& PackageData, bool& bInOutDemote)
+{
+	if (!IsInitialized())
+	{
+		return;
+	}
+	FCookGenerationInfo* Info = FindInfo(PackageData);
+	if (!Info)
+	{
+		bInOutDemote = true;
+		return;
+	}
+
+	if (!Info->IsGenerator())
+	{
+		UPackage* LocalPackage = OwnerPackage.Get();
+		if (!LocalPackage || !LocalPackage->IsFullyLoaded())
+		{
+			bInOutDemote = true;
+			return;
+		}
+	}
+
+	if (bInOutDemote && IsUseInternalReferenceToAvoidGarbageCollect() && Info->HasTakenOverCachedCookedPlatformData())
+	{
+		// No public objects should have been deleted; we are supposed to keep them referenced by keeping the package
+		// referenced in UCookOnTheFlyServer::PreGarbageCollect, and the package keeping its public objects referenced
+		// by UPackage::AddReferencedObjects. Since no public objects were deleted, our caller should not have
+		// set bInOutDemote=true.
+		// Allowing demotion after the splitter has started moving objects breaks our contract with the splitter
+		// and can cause a crash. So log this as an error.
+		// For better feedback, look in our extra data to identify the name of the public UObject that was deleted.
+		FString DeletedObject;
+		if (!PackageData.GetPackage())
+		{
+			DeletedObject = FString::Printf(TEXT("UPackage %s"), *PackageData.GetPackageName().ToString());
+		}
+		else
+		{
+			TSet<UObject*> ExistingObjectsAfterSave;
+			for (const FCachedObjectInOuter& CachedObjectInOuter : PackageData.GetCachedObjectsInOuter())
+			{
+				UObject* Ptr = CachedObjectInOuter.Object.Get();
+				if (Ptr)
+				{
+					ExistingObjectsAfterSave.Add(Ptr);
+				}
+			}
+
+			for (const TPair<UObject*, FCachedObjectInOuterGeneratorInfo>& Pair : Info->CachedObjectsInOuterInfo)
+			{
+				if (Pair.Value.bPublic && !ExistingObjectsAfterSave.Contains(Pair.Key))
+				{
+					DeletedObject = Pair.Value.FullName;
+					break;
+				}
+			}
+			if (DeletedObject.IsEmpty())
+			{
+				if (!PackageData.GetPackage()->IsFullyLoaded())
+				{
+					DeletedObject = FString::Printf(TEXT("UPackage %s is no longer FullyLoaded"),
+						*PackageData.GetPackageName().ToString());
+				}
+				else
+				{
+					DeletedObject = TEXT("<Unknown>");
+				}
+			}
+		}
+		UE_LOG(LogCook, Error,
+			TEXT("A %s package had some of its UObjects deleted during garbage collection after it started generating. This will cause errors during save of the package.")
+			TEXT("\n\tDeleted object: %s")
+			TEXT("\n\tSplitter=%s%s"),
+			Info->IsGenerator() ? TEXT("Generator") : TEXT("Generated"),
+			*DeletedObject,
+			*GetSplitDataObjectName().ToString(),
+			Info->IsGenerator() ? TEXT(".") : *FString::Printf(TEXT(", Generated=%s."),
+				*Info->PackageData->GetPackageName().ToString()));
+	}
+
+	// Remove raw pointers from RootMovedObjects if they no longer exist in the weakpointers in CachedObjectsInOuter
+	TSet<UObject*> CachedObjectsInOuterSet;
+	for (FCachedObjectInOuter& CachedObjectInOuter : Info->PackageData->GetCachedObjectsInOuter())
+	{
+		UObject* Object = CachedObjectInOuter.Object.Get();
+		if (Object)
+		{
+			CachedObjectsInOuterSet.Add(Object);
+		}
+	}
+	for (TMap<UObject*, FCachedObjectInOuterGeneratorInfo>::TIterator Iter(Info->CachedObjectsInOuterInfo);
+		Iter; ++Iter)
+	{
+		if (!CachedObjectsInOuterSet.Contains(Iter->Key))
+		{
+			Iter.RemoveCurrent();
+		}
+	}
+}
+
+FCookGenerationInfo::FCookGenerationInfo(FPackageData& InPackageData, bool bInGenerator)
+	: PackageData(&InPackageData)
+	, GeneratorSaveState(bInGenerator ? ESaveState::StartSave : ESaveState::StartPopulate)
+	, bCreateAsMap(false), bHasCreatedPackage(false), bHasSaved(false), bTakenOverCachedCookedPlatformData(false)
+	, bIssuedUndeclaredMovedObjectsWarning(false), bGenerator(bInGenerator)
+{
+}
+
+void FCookGenerationInfo::SetSaveStateComplete(ESaveState CompletedState)
+{
+	GeneratorSaveState = CompletedState;
+	if (GeneratorSaveState < ESaveState::Last)
+	{
+		GeneratorSaveState = static_cast<ESaveState>(static_cast<uint8>(GeneratorSaveState) + 1);
+	}
+}
+
+void FCachedObjectInOuterGeneratorInfo::Initialize(UObject* Object)
+{
+	if (Object)
+	{
+		FullName = Object->GetFullName();
+		bPublic = Object->HasAnyFlags(RF_Public);
+	}
+	else
+	{
+		FullName.Empty();
+		bPublic = false;
+	}
+
+	bInitialized = true;
+}
+
+void FCookGenerationInfo::TakeOverCachedObjectsAndAddMoved(FGenerationHelper& GenerationHelper,
+	TArray<FCachedObjectInOuter>& CachedObjectsInOuter, TArray<UObject*>& MovedObjects)
+{
+	CachedObjectsInOuterInfo.Reset();
+
+	for (FCachedObjectInOuter& ObjectInOuter : CachedObjectsInOuter)
+	{
+		UObject* Object = ObjectInOuter.Object.Get();
+		if (Object)
+		{
+			CachedObjectsInOuterInfo.FindOrAdd(Object).Initialize(Object);
+		}
+	}
+
+	TArray<UObject*> ChildrenOfMovedObjects;
+	for (UObject* Object : MovedObjects)
+	{
+		if (!IsValid(Object))
+		{
+			UE_LOG(LogCook, Warning,
+				TEXT("CookPackageSplitter found non-valid object %s returned from %s on Splitter %s%s. Ignoring it."),
+				Object ? *Object->GetFullName() : TEXT("<null>"),
+				IsGenerator() ? TEXT("PopulateGeneratorPackage") : TEXT("PopulateGeneratedPackage"),
+				*GenerationHelper.GetSplitDataObjectName().ToString(),
+				IsGenerator() ? TEXT("") : *FString::Printf(TEXT(", Package %s"),
+					*PackageData->GetPackageName().ToString()));
+			continue;
+		}
+		FCachedObjectInOuterGeneratorInfo& Info = CachedObjectsInOuterInfo.FindOrAdd(Object);
+		if (!Info.bInitialized)
+		{
+			Info.Initialize(Object);
+			Info.bMoved = true;
+			Info.bMovedRoot = true;
+			CachedObjectsInOuter.Emplace(Object);
+			GetObjectsWithOuter(Object, ChildrenOfMovedObjects, true /* bIncludeNestedObjects */, RF_NoFlags,
+				EInternalObjectFlags::Garbage);
+		}
+	}
+
+	for (UObject* Object : ChildrenOfMovedObjects)
+	{
+		check(IsValid(Object));
+		FCachedObjectInOuterGeneratorInfo& Info = CachedObjectsInOuterInfo.FindOrAdd(Object);
+		if (!Info.bInitialized)
+		{
+			Info.Initialize(Object);
+			Info.bMoved = true;
+			CachedObjectsInOuter.Emplace(Object);
+		}
+	}
+
+	SetHasTakenOverCachedCookedPlatformData(true);
+}
+
+EPollStatus FCookGenerationInfo::RefreshPackageObjects(FGenerationHelper& GenerationHelper, UPackage* Package,
+	bool& bOutFoundNewObjects, ESaveState DemotionState)
+{
+	bOutFoundNewObjects = false;
+	TArray<UObject*> CurrentObjectsInOuter;
+	GetObjectsWithOuter(Package, CurrentObjectsInOuter, true /* bIncludeNestedObjects */, RF_NoFlags,
+		EInternalObjectFlags::Garbage);
+
+	check(PackageData); // RefreshPackageObjects is only called when there is a PackageData
+	TArray<FCachedObjectInOuter>& CachedObjectsInOuter = PackageData->GetCachedObjectsInOuter();
+	UObject* FirstNewObject = nullptr;
+	for (UObject* Object : CurrentObjectsInOuter)
+	{
+		FCachedObjectInOuterGeneratorInfo& Info = CachedObjectsInOuterInfo.FindOrAdd(Object);
+		if (!Info.bInitialized)
+		{
+			Info.Initialize(Object);
+			CachedObjectsInOuter.Emplace(Object);
+			if (!FirstNewObject)
+			{
+				FirstNewObject = Object;
+			}
+		}
+	}
+	bOutFoundNewObjects = FirstNewObject != nullptr;
+
+	if (FirstNewObject != nullptr && DemotionState != ESaveState::Last)
+	{
+		SetSaveState(DemotionState);
+		if (++PackageData->GetNumRetriesBeginCacheOnObjects() > FPackageData::GetMaxNumRetriesBeginCacheOnObjects())
+		{
+			UE_LOG(LogCook, Error,
+				TEXT("Cooker has repeatedly tried to call BeginCacheForCookedPlatformData on all objects in a generated package, but keeps finding new objects.\n")
+				TEXT("Aborting the save of the package; programmer needs to debug why objects keep getting added to the package.\n")
+				TEXT("Splitter: %s%s. Most recent created object: %s."),
+				*GenerationHelper.GetSplitDataObjectName().ToString(),
+				IsGenerator() ? TEXT("") : *FString::Printf(TEXT(", Package: %s"),
+					*PackageData->GetPackageName().ToString()),
+				*FirstNewObject->GetFullName());
+			return EPollStatus::Error;
+		}
+	}
+	return EPollStatus::Success;
+}
+
+void FCookGenerationInfo::CreatePackageHash()
+{
+	FBlake3 Blake3;
+	Blake3.Update(&GenerationHash, sizeof(GenerationHash));
+	IAssetRegistry& AssetRegistry = IAssetRegistry::GetChecked();
+	for (const FAssetDependency& Dependency : PackageDependencies)
+	{
+		TOptional<FAssetPackageData> DependencyData =
+			AssetRegistry.GetAssetPackageDataCopy(Dependency.AssetId.PackageName);
+		if (DependencyData)
+		{
+			Blake3.Update(&DependencyData->GetPackageSavedHash().GetBytes(),
+				sizeof(decltype(DependencyData->GetPackageSavedHash().GetBytes())));
+		}
+	}
+	PackageHash = FIoHash(Blake3.Finalize());
+	// We store the PackageHash as a FIoHash, but UPackage and FAssetPackageData store it as a FGuid, which is smaller,
+	// so we have to remove any data which doesn't fit into FGuid. This can be removed when we remove the deprecated
+	// Guid storage on UPackage.
+	PRAGMA_DISABLE_DEPRECATION_WARNINGS;
+	constexpr int SizeDifference = sizeof(PackageHash) - sizeof(decltype(DeclVal<UPackage>().GetGuid()));
+	if (SizeDifference > 0)
+	{
+		FMemory::Memset(((uint8*)&PackageHash.GetBytes())
+			+ (sizeof(decltype(PackageHash.GetBytes())) - SizeDifference),
+			0, SizeDifference);
+	}
+	PRAGMA_ENABLE_DEPRECATION_WARNINGS;
+}
+
+void FCookGenerationInfo::IterativeCookValidateOrClear(FGenerationHelper& GenerationHelper,
+	TConstArrayView<const ITargetPlatform*> RequestedPlatforms, const FIoHash& PreviousPackageHash,
+	bool& bOutIterativelyUnmodified)
+{
+	UCookOnTheFlyServer& COTFS = GenerationHelper.GetOwner().GetPackageDatas().GetCookOnTheFlyServer();
+	bOutIterativelyUnmodified = PreviousPackageHash == this->PackageHash;
+	if (bOutIterativelyUnmodified)
+	{
+		// If not directly modified, mark it as indirectly modified if any of its dependencies
+		// were detected as modified during PopulateCookedPackages.
+		for (const FAssetDependency& Dependency : this->PackageDependencies)
+		{
+			FPackageData* DependencyData =
+				COTFS.PackageDatas->FindPackageDataByPackageName(Dependency.AssetId.PackageName);
+			if (!DependencyData)
+			{
+				bOutIterativelyUnmodified = false;
+				break;
+			}
+			for (const ITargetPlatform* TargetPlatform : RequestedPlatforms)
+			{
+				FPackagePlatformData* DependencyPlatformData = DependencyData->FindPlatformData(TargetPlatform);
+				if (!DependencyPlatformData || !DependencyPlatformData->IsIterativelyUnmodified())
+				{
+					bOutIterativelyUnmodified = false;
+					break;
+				}
+			}
+			if (!bOutIterativelyUnmodified)
+			{
+				break;
+			}
+		}
+	}
+
+	bool bFirstPlatform = true;
+	for (const ITargetPlatform* TargetPlatform : RequestedPlatforms)
+	{
+		if (bOutIterativelyUnmodified)
+		{
+			PackageData->FindOrAddPlatformData(TargetPlatform).SetIterativelyUnmodified(true);
+		}
+		bool bShouldIterativelySkip = bOutIterativelyUnmodified;
+		ICookedPackageWriter& PackageWriter = COTFS.FindOrCreatePackageWriter(TargetPlatform);
+		PackageWriter.UpdatePackageModificationStatus(PackageData->GetPackageName(), bOutIterativelyUnmodified,
+			bShouldIterativelySkip);
+		if (bShouldIterativelySkip)
+		{
+			PackageData->SetPlatformCooked(TargetPlatform, ECookResult::Succeeded);
+			if (bFirstPlatform)
+			{
+				COOK_STAT(++DetailedCookStats::NumPackagesIterativelySkipped);
+			}
+			// Declare the package to the EDLCookInfo verification so we don't warn about missing exports from it
+			UE::SavePackageUtilities::EDLCookInfoAddIterativelySkippedPackage(PackageData->GetPackageName());
+		}
+		else
+		{
+			COTFS.DeleteOutputForPackage(PackageData->GetPackageName(), TargetPlatform);
+		}
+		bFirstPlatform = false;
+	}
+}
+
+}
\ No newline at end of file
diff --git a/Engine/Source/Editor/UnrealEd/Private/Cooker/CookGenerationHelper.h b/Engine/Source/Editor/UnrealEd/Private/Cooker/CookGenerationHelper.h
new file mode 100644
--- /dev/null
+++ b/Engine/Source/Editor/UnrealEd/Private/Cooker/CookGenerationHelper.h
@@ -0,0 +1,583 @@
+// Copyright Epic Games, Inc. All Rights Reserved.
+
+#pragma once
+
+#include "AssetRegistry/IAssetRegistry.h"
+#include "Containers/Array.h"
+#include "Containers/ArrayView.h"
+#include "Containers/Map.h"
+#include "Containers/UnrealString.h"
+#include "Cooker/CookPackageData.h"
+#include "Cooker/CookTypes.h"
+#include "Cooker/MPCollector.h"
+#include "CookPackageSplitter.h"
+#include "Hash/Blake3.h"
+#include "IO/IoHash.h"
+#include "Misc/Optional.h"
+#include "Templates/RefCounting.h"
+#include "Templates/UniquePtr.h"
+#include "UObject/NameTypes.h"
+#include "UObject/Object.h"
+#include "UObject/ObjectPtr.h"
+#include "UObject/Package.h"
+#include "UObject/WeakObjectPtr.h"
+
+class ITargetPlatform;
+class UCookOnTheFlyServer;
+
+namespace UE::Cook
+{
+
+/**
+ * Extra information about CachedObjectsInOuter that a GenerationHelper needs to know for diagnostics.
+ * The GenerationHelper constructs an associated array (aka TMap) of these structures when it takes
+ * over the CachedObjectsInOuter.
+ */
+struct FCachedObjectInOuterGeneratorInfo
+{
+public:
+	/** Object->GetFullName() before it was deleted. */
+	FString FullName;
+	/** Has Initialize been called on *this. */
+	bool bInitialized = false;
+	/** Object->GetFlags() had RF_Public when the info was initialized. */
+	bool bPublic = false;
+	/** bMovedRoot is true, or this is a child object of such a moved object. */
+	bool bMoved = false;
+	/** Splitter informed us that object was moved into this package from another package. */
+	bool bMovedRoot = false;
+
+public:
+	void Initialize(UObject* Object);
+};
+
+/**
+ * Struct used with ICookPackageSplitter, which contains the state information for the save of either the
+ * Generator package or one of its Generated packages.
+ */
+struct FCookGenerationInfo
+{
+public:
+	/** State variable for reentrant SplitPackage calls */
+	enum class ESaveState : uint8
+	{
+		StartSave = 0,
+
+		ClearOldPackagesFirstAttempt = StartSave,
+		ClearOldPackagesLastAttempt,
+		QueueGeneratedPackages,
+
+		StartPopulate,
+		FinishCachePreMove = StartPopulate,
+		CallObjectsToMove,
+		BeginCacheObjectsToMove,
+		FinishCacheObjectsToMove,
+		CallPopulate,
+		CallGetPostMoveObjects,
+		BeginCachePostMove,
+		FinishCachePostMove,
+
+		ReadyForSave,
+		Last = ReadyForSave,
+	};
+
+	FCookGenerationInfo(FPackageData& InPackageData, bool bInGenerator);
+
+	ESaveState GetSaveState() const;
+	void SetSaveState(ESaveState InValue);
+	void SetSaveStateComplete(ESaveState CompletedState);
+
+	bool IsCreateAsMap() const;
+	void SetIsCreateAsMap(bool bValue);
+	bool HasCreatedPackage() const;
+	void SetHasCreatedPackage(bool bValue);
+	bool HasSaved() const;
+	void SetHasSaved(bool bValue);
+	bool HasTakenOverCachedCookedPlatformData() const;
+	void SetHasTakenOverCachedCookedPlatformData(bool bValue);
+	bool HasIssuedUndeclaredMovedObjectsWarning() const;
+	void SetHasIssuedUndeclaredMovedObjectsWarning(bool bValue);
+	bool IsGenerator() const;
+	void SetIsGenerator(bool bValue);
+
+	/**
+	 * Steal the list of cached objects to call BeginCacheForCookedPlatformData on from the PackageData,
+	 * and add them to this. Also add NewObjects reported by the splitter that will be moved into the package.
+	 */
+	void TakeOverCachedObjectsAndAddMoved(FGenerationHelper& GenerationHelper,
+		TArray<FCachedObjectInOuter>& CachedObjectsInOuter, TArray<UObject*>& MovedObjects);
+	/**
+	 * Fetch all the objects currently in the package and add them to the list of objects that need
+	 * BeginCacheForCookedPlatformData.
+	 * Reports whether new objects were found. If DemotionState is not ESaveState::Last,
+	 * will SetState back to DemotionState
+	 * if new objects were found, and will error exit if this demotion has happened too many times.
+	 */
+	EPollStatus RefreshPackageObjects(FGenerationHelper& GenerationHelper, UPackage* Package,
+		bool& bOutFoundNewObjects, ESaveState DemotionState);
+
+	void AddKeepReferencedPackages(TArray<UPackage*>& InKeepReferencedPackages);
+
+	/** Create the hash for this generated package, based on dependencies and GenerationHash. */
+	void CreatePackageHash();
+
+	/**
+	 * If the package has iterative results from a previous cook that were not invalidated by dependency changes,
+	 * test whether they need to be invalidated now and clear the results if so. Only called in legacy iterative;
+	 * incremental cooks handle invalidation by querying the TargetDomainDigest during the RequestCluster.
+	 */
+	void IterativeCookValidateOrClear(FGenerationHelper& GeneratorHelper,
+		TConstArrayView<const ITargetPlatform*> RequestedPlatforms, const FIoHash& PreviousHash,
+		bool& bOutIterativelyUnmodified);
+
+	TConstArrayView<FAssetDependency> GetDependencies() const;
+
+public:
+	FIoHash PackageHash;
+	FString RelativePath;
+	FString GeneratedRootPath;
+	FBlake3Hash GenerationHash;
+	TArray<FAssetDependency> PackageDependencies;
+	FPackageData* PackageData = nullptr;
+	TArray<UPackage*> KeepReferencedPackages;
+	TMap<UObject*, FCachedObjectInOuterGeneratorInfo> CachedObjectsInOuterInfo;
+private:
+	ESaveState GeneratorSaveState = ESaveState::StartSave;
+	bool bCreateAsMap : 1;
+	bool bHasCreatedPackage : 1;
+	bool bHasSaved : 1;
+	bool bTakenOverCachedCookedPlatformData : 1;
+	bool bIssuedUndeclaredMovedObjectsWarning : 1;
+	bool bGenerator : 1;
+};
+
+/**
+ * Helper that wraps an ICookPackageSplitter, gets/caches packages to generate, and is a reference-counted
+ * collection of cached data and helper functions to save and list the generated packages.
+ */
+struct FGenerationHelper : public FThreadSafeRefCountedObject
+{
+public:
+	/** Store the provided CookPackageSplitter and prepare the packages to generate. */
+	FGenerationHelper(UE::Cook::FPackageData& InOwner);
+	~FGenerationHelper();
+
+	/**
+	 * Early exits if already initialized. Otherwise, loads the package if not loaded and searches it for a
+	 * splitter and creates the splitter.
+	 * If load/search/creation fails, the GenerationHelper will be set to invalid @see IsValid.
+	 *
+	 * Unless otherwise stated, all public functions call Initialize if not already initialized.
+	 */
+	void Initialize();
+	/**
+	 * Version of Initialize that receives the splitterobject from the caller rather than needing to search
+	 * for it, as an optimization for callers that have it already. If already initialized, the input Splitter
+	 * will be destroyed and the function will early exit.
+	 */
+	void Initialize(const UObject* InSplitDataObject,
+		UE::Cook::Private::FRegisteredCookPackageSplitter* InRegisteredSplitterType,
+		TUniquePtr<ICookPackageSplitter>&& InSplitter);
+	/** Version of Initialize that sets IsValid=false. */
+	void InitializeAsInvalid();
+	/**
+	 * Clear all self references which keep this GenerationHelper from destructing. These could have been set by the
+	 * SetKeep... functions. The Owner PackageData will keep a pointer to this until all references from self, from
+	 * child generated PackageDatas, and other, are released. Note this means that we assert in the owner PackageData's
+	 * destructor if any references remain; PackageDatas can not be deleted until all such references are cleared.
+	 * Does not call Initialize.
+	 */
+	void ClearSelfReferences();
+
+	/**
+	 * GenerationHelpers can be created uninitialized at the start of cook, for incremental cooks. We initialize
+	 * them on demand, which requires loading the package if not already loaded.
+	 * IsInitialized reports whether initialization has been attempted. @see IsValid.
+	 * Does not call Initialize.
+	 */
+	bool IsInitialized() const;
+	/**
+	 * Since GenerationHelpers can be created before we load the package (@see IsInitialized) we might incorrectly
+	 * create one for a package that does not have a splitter. If so, when we load the package and discover that,
+	 * we set status to invalid. When in this state, all public functions are still valid to call, but will be
+	 * noops and will return empty data.
+	 */
+	bool IsValid();
+
+	/** Accessor for the packages to generate, will be empty if invalid or if TryGenerateList was not yet called. */
+	TArrayView<FCookGenerationInfo> GetPackagesToGenerate();
+	/** Return the GenerationInfo used to save the generator package's UPackage. Does not call Initialize. */
+	FCookGenerationInfo& GetOwnerInfo();
+	/** Return owner FPackageData. Does not call Initialize. */
+	FPackageData& GetOwner();
+	/** Return the GenerationInfo for the given PackageData, or null if not found. */
+	FCookGenerationInfo* FindInfo(const FPackageData& PackageData);
+	const FCookGenerationInfo* FindInfo(const FPackageData& PackageData) const;
+	/** Return CookPackageSplitter. Will return null if !IsValid.*/
+	ICookPackageSplitter* GetCookPackageSplitterInstance() const;
+	/** Return RegisteredSplitterType. Will return null if !IsValid. */
+	UE::Cook::Private::FRegisteredCookPackageSplitter* GetRegisteredSplitterType() const;
+	/** Return the SplitDataObject's FullObjectPath. Will return NAME_None if !IsValid. */
+	const FName GetSplitDataObjectName() const;
+	/** Return the Splitter's value for virtual bool UseInternalReferenceToAvoidGarbageCollect(). */
+	bool IsUseInternalReferenceToAvoidGarbageCollect() const;
+	/** Return the Splitter's value for virtual bool NeedCachedPlatformDataBeforeSplit(). */
+	bool IsNeedCachedPlatformDataBeforeSplit() const;
+	/** Return the cached pointer to the SplitDataObject. Returns null if no longer in memory or marked as garbage. */
+	UObject* GetWeakSplitDataObject() const;
+	/**
+	 * Return our cached pointer to the SplitDataObject if set. If not set, load the package and find it. Can still
+	 * return null if !IsValid or if not found even after loading the package.
+	 */
+	UObject* FindOrLoadSplitDataObject();
+
+	/** Find the OwnerPackage in memory, returns null if invalid or not already loaded. Does not call Initialize. */
+	UPackage* GetOwnerPackage();
+	/** Load the OwnerPackage if not already loaded. Can still return null if invalid or package fails to load. */
+	UPackage* FindOrLoadOwnerPackage(UCookOnTheFlyServer& COTFS);
+	/**
+	 * Return a reference to ExternalActors that were discovered during save and stored on this.
+	 * Does not call initialize.
+	 */
+	TConstArrayView<FName> GetExternalActorDependencies();
+	/**
+	 * Return the ExternalActors that were discovered during save and stored on this, and clear the values.
+	 * Does not call initialize.
+	 */
+	TArray<FName> ReleaseExternalActorDependencies();
+
+	/** Call the Splitter's GetGenerateList and create the PackageDatas. Logs errors and returns false on failure. */
+	bool TryGenerateList();
+	/**
+	 * Mark that the SavePackage of the Owner is starting. Keeps a reference to keep the generator alive until save
+	 * is finished.
+	 */
+	void StartOwnerSave();
+	/** Update state before we queue the generated packages, e.g. mark whether packages are iteratively skippable. */
+	void StartQueueGeneratedPackages(UCookOnTheFlyServer& COTFS);
+	/**
+	 * Mark that we have started queuing packages. Automatically called from StartQueueGeneratedPackages, but also
+	 * is called on the director in response to a discovered generated package from a CookWorker. It sets the
+	 * WorkerId so that it is available during assignment of the discovered generated packages.
+	 */
+	void NotifyStartQueueGeneratedPackages(UCookOnTheFlyServer& COTFS, FWorkerId SourceWorkerId);
+	/** Update state after we queue generated packages; e.g. register to stay alive until the packages are assigned. */
+	void EndQueueGeneratedPackages(UCookOnTheFlyServer& COTFS);
+	/** Does not call Initialize. */
+	void EndQueueGeneratedPackagesOnDirector(UCookOnTheFlyServer& COTFS, FWorkerId SourceWorkerId);
+	/**
+	 * Called from Director when the RequestFence added from EndQueueGeneratedPackages has passed. Calls
+	 * OnRequestFencePassed locally and on all Workers.
+	 */
+	void OnRequestFencePassedBroadcast(UCookOnTheFlyServer& COTFS);
+	/** Called from when the RequestFence added from EndQueueGeneratedPackages has passed. */
+	void OnRequestFencePassed(UCookOnTheFlyServer& COTFS);
+	/** Call CreatePackage and set the necessary data on the UPackage, and update status. */
+	UPackage* CreateGeneratedUPackage(FCookGenerationInfo& GenerationInfo, const UPackage* OwnerPackage,
+		const TCHAR* GeneratedPackageName);
+	/**
+	 * Clear any data that should only be held when an FPackageData is in the save state, for the given Info. The given
+	 * Info might specify the generator package or one of the generated packages.
+	 */
+	void ResetSaveState(FCookGenerationInfo& Info, UPackage* Package, EStateChangeReason ReleaseSaveReason,
+		EPackageState NewState);
+
+	/**
+	 * Record all dependencies from the generator package that are ExternalActor dependencies and
+	 * store them on this.
+	 */
+	void FetchExternalActorDependencies();
+	/** Iterative cook: store the list of generated packages from the last cook. Does not call Initialize. */
+	void SetPreviousGeneratedPackages(TMap<FName, FIoHash>&& Packages);
+	/** Return the information set by SetPreviousGeneratedPackages if not yet cleared. Does not call Initialize. */
+	const TMap<FName, FIoHash>& GetPreviousGeneratedPackages() const;
+
+	/** Callback during garbage collection. Does not call initialize. */
+	void PreGarbageCollect(FCookGenerationInfo& Info, TArray<TObjectPtr<UObject>>& GCKeepObjects,
+		TArray<UPackage*>& GCKeepPackages, TArray<FPackageData*>& GCKeepPackageDatas, bool& bOutShouldDemote);
+	/** Callback during garbage collection. Does not call initialize. */
+	void PostGarbageCollect();
+	/**
+	 * Called from PackageData function of the same name to decide whether to demote the package out of save.
+	 * Does not call Initialize.
+	 */
+	void UpdateSaveAfterGarbageCollect(const FPackageData& PackageData, bool& bInOutDemote);
+
+	// Self-references that keep this GenerationHelper in memory and referenced from the packages that use it until
+	// the self-references are cleared. These Set/Clear functions do not call initialize.
+	void SetKeepForIterative();
+	void ClearKeepForIterative();
+	void SetKeepForQueueResults();
+	void ClearKeepForQueueResults();
+	void SetKeepForGeneratorSave();
+	void ClearKeepForGeneratorSave();
+
+	/**
+	 * Helper for assignment of generated packages in MPCook. Return the id of the CookWorker that saved the
+	 * generator, to decide where to assign the generated packages.
+	 */
+	FWorkerId GetWorkerIdThatSavedGenerator() const;
+	/**
+	 * Helper for assignment of generated packages in MPCook. A counter for assigned generated packages that is
+	 * needed by some assignment schemes.
+	 */
+	int32& GetMPCookNextAssignmentIndex();
+
+	/** Helper function for Initialize and for TryCreateValidGenerationHelper. */
+	static void SearchForRegisteredSplitDataObject(UCookOnTheFlyServer& COTFS, FName PackageName, UPackage* Package,
+		TOptional<TConstArrayView<FCachedObjectInOuter>> CachedObjectsInOuter, UObject*& OutSplitDataObject,
+		UE::Cook::Private::FRegisteredCookPackageSplitter*& OutRegisteredSplitterType,
+		TUniquePtr<ICookPackageSplitter>& OutSplitterInstance);
+	/** Helper function for Initialize and for TryCreateValidGenerationHelper. */
+	static UPackage* FindOrLoadPackage(UCookOnTheFlyServer& COTFS, FPackageData& OwnerPackageData);
+
+private:
+	enum class EInitializeStatus : uint8
+	{
+		Uninitialized,
+		Invalid,
+		Valid,
+	};
+
+private:
+	void ConditionalInitialize() const;
+	void NotifyCompletion(ICookPackageSplitter::ETeardown Status);
+
+private:
+	/** PackageData for the package that is being split */
+	FCookGenerationInfo OwnerInfo;
+	FWeakObjectPtr SplitDataObject;
+	/** Name of the object that prompted the splitter creation */
+	FName SplitDataObjectName;
+	UE::Cook::Private::FRegisteredCookPackageSplitter* RegisteredSplitterType = nullptr;
+	TUniquePtr<ICookPackageSplitter> CookPackageSplitterInstance;
+	/** Recorded list of packages to generate from the splitter, and data we need about them */
+	TArray<FCookGenerationInfo> PackagesToGenerate;
+	TWeakObjectPtr<UPackage> OwnerPackage;
+	TMap<FName, FIoHash> PreviousGeneratedPackages;
+	TArray<FName> ExternalActorDependencies;
+	TRefCountPtr<FGenerationHelper> ReferenceFromKeepForIterative;
+	TRefCountPtr<FGenerationHelper> ReferenceFromKeepForQueueResults;
+	TRefCountPtr<FGenerationHelper> ReferenceFromKeepForGeneratorSave;
+	int32 MPCookNextAssignmentIndex = 0;
+	FWorkerId WorkerIdThatSavedGenerator = FWorkerId::Invalid();
+	EInitializeStatus InitializeStatus = EInitializeStatus::Uninitialized;
+	bool bUseInternalReferenceToAvoidGarbageCollect = false;
+	bool bNeedCachedPlatformDataBeforeSplit = false;
+	bool bGeneratedList = false;
+};
+
+
+///////////////////////////////////////////////////////
+// Inline implementations
+///////////////////////////////////////////////////////
+
+inline FCookGenerationInfo::ESaveState FCookGenerationInfo::GetSaveState() const
+{
+	return GeneratorSaveState;
+}
+
+inline void FCookGenerationInfo::SetSaveState(ESaveState InValue)
+{
+	GeneratorSaveState = InValue;
+}
+
+inline bool FCookGenerationInfo::IsCreateAsMap() const
+{
+	return bCreateAsMap;
+}
+
+inline void FCookGenerationInfo::SetIsCreateAsMap(bool bValue)
+{
+	bCreateAsMap = bValue;
+}
+
+inline bool FCookGenerationInfo::HasCreatedPackage() const
+{
+	return bHasCreatedPackage;
+
+}
+inline void FCookGenerationInfo::SetHasCreatedPackage(bool bValue)
+{
+	bHasCreatedPackage = bValue;
+}
+
+inline bool FCookGenerationInfo::HasSaved() const
+{
+	return bHasSaved;
+}
+
+inline void FCookGenerationInfo::SetHasSaved(bool bValue)
+{
+	bHasSaved = bValue;
+}
+
+inline bool FCookGenerationInfo::HasTakenOverCachedCookedPlatformData() const
+{
+	return bTakenOverCachedCookedPlatformData;
+}
+
+inline void FCookGenerationInfo::SetHasTakenOverCachedCookedPlatformData(bool bValue)
+{
+	bTakenOverCachedCookedPlatformData = bValue;
+}
+
+inline bool FCookGenerationInfo::HasIssuedUndeclaredMovedObjectsWarning() const
+{
+	return bIssuedUndeclaredMovedObjectsWarning;
+}
+
+inline void FCookGenerationInfo::SetHasIssuedUndeclaredMovedObjectsWarning(bool bValue)
+{
+	bIssuedUndeclaredMovedObjectsWarning = bValue;
+}
+
+inline bool FCookGenerationInfo::IsGenerator() const
+{
+	return bGenerator;
+}
+
+inline void FCookGenerationInfo::SetIsGenerator(bool bValue)
+{
+	bGenerator = bValue;
+}
+
+inline void FCookGenerationInfo::AddKeepReferencedPackages(TArray<UPackage*>& InKeepReferencedPackages)
+{
+	KeepReferencedPackages.Append(InKeepReferencedPackages);
+}
+
+inline TConstArrayView<FAssetDependency> FCookGenerationInfo::GetDependencies() const
+{
+	return PackageDependencies;
+}
+
+inline bool FGenerationHelper::IsInitialized() const
+{
+	return InitializeStatus != EInitializeStatus::Uninitialized;
+}
+
+inline void FGenerationHelper::ConditionalInitialize() const
+{
+	if (InitializeStatus == EInitializeStatus::Uninitialized)
+	{
+		// Use a const cast so we can call from getter functions that are otherwise const.
+		const_cast<FGenerationHelper&>(*this).Initialize();
+	}
+}
+
+inline bool FGenerationHelper::IsValid()
+{
+	ConditionalInitialize();
+	return InitializeStatus == EInitializeStatus::Valid;
+}
+
+inline TArrayView<UE::Cook::FCookGenerationInfo> FGenerationHelper::GetPackagesToGenerate()
+{
+	ConditionalInitialize();
+	return PackagesToGenerate;
+}
+
+inline UE::Cook::FCookGenerationInfo& FGenerationHelper::GetOwnerInfo()
+{
+	return OwnerInfo;
+}
+
+inline UE::Cook::FPackageData& FGenerationHelper::GetOwner()
+{
+	return *OwnerInfo.PackageData;
+}
+
+inline ICookPackageSplitter* FGenerationHelper::GetCookPackageSplitterInstance() const
+{
+	ConditionalInitialize();
+	return CookPackageSplitterInstance.Get();
+}
+
+inline UE::Cook::Private::FRegisteredCookPackageSplitter* FGenerationHelper::GetRegisteredSplitterType() const
+{
+	ConditionalInitialize();
+	return RegisteredSplitterType;
+}
+
+inline const FName FGenerationHelper::GetSplitDataObjectName() const
+{
+	ConditionalInitialize();
+	return SplitDataObjectName;
+}
+
+inline bool FGenerationHelper::IsUseInternalReferenceToAvoidGarbageCollect() const
+{
+	ConditionalInitialize();
+	return bUseInternalReferenceToAvoidGarbageCollect;
+}
+
+inline bool FGenerationHelper::IsNeedCachedPlatformDataBeforeSplit() const
+{
+	ConditionalInitialize();
+	return bNeedCachedPlatformDataBeforeSplit;
+}
+
+inline UObject* FGenerationHelper::GetWeakSplitDataObject() const
+{
+	ConditionalInitialize();
+	return SplitDataObject.Get();
+}
+
+inline TConstArrayView<FName> FGenerationHelper::GetExternalActorDependencies()
+{
+	return ExternalActorDependencies;
+}
+
+inline TArray<FName> FGenerationHelper::ReleaseExternalActorDependencies()
+{
+	TArray<FName> Result = MoveTemp(ExternalActorDependencies);
+	return Result;
+}
+
+inline const TMap<FName, FIoHash>& FGenerationHelper::GetPreviousGeneratedPackages() const
+{
+	return PreviousGeneratedPackages;
+}
+
+inline void FGenerationHelper::SetKeepForIterative()
+{
+	ReferenceFromKeepForIterative = this;
+}
+
+inline void FGenerationHelper::ClearKeepForIterative()
+{
+	ReferenceFromKeepForIterative.SafeRelease();
+}
+
+inline void FGenerationHelper::SetKeepForQueueResults()
+{
+	ReferenceFromKeepForQueueResults = this;
+}
+
+inline void FGenerationHelper::ClearKeepForQueueResults()
+{
+	ReferenceFromKeepForQueueResults.SafeRelease();
+}
+
+inline void FGenerationHelper::SetKeepForGeneratorSave()
+{
+	ReferenceFromKeepForGeneratorSave = this;
+}
+
+inline void FGenerationHelper::ClearKeepForGeneratorSave()
+{
+	ReferenceFromKeepForGeneratorSave.SafeRelease();
+}
+
+inline FWorkerId FGenerationHelper::GetWorkerIdThatSavedGenerator() const
+{
+	return WorkerIdThatSavedGenerator;
+}
+
+inline int32& FGenerationHelper::GetMPCookNextAssignmentIndex()
+{
+	return MPCookNextAssignmentIndex;
+}
+
+} // namespace UE::Cook
\ No newline at end of file
diff --git a/Engine/Source/Editor/UnrealEd/Private/Cooker/CookPackageData.cpp b/Engine/Source/Editor/UnrealEd/Private/Cooker/CookPackageData.cpp
--- a/Engine/Source/Editor/UnrealEd/Private/Cooker/CookPackageData.cpp
+++ b/Engine/Source/Editor/UnrealEd/Private/Cooker/CookPackageData.cpp
@@ -12,6 +12,7 @@
 #include "Async/ParallelFor.h"
 #include "CompactBinaryTCP.h"
 #include "Cooker/CookDirector.h"
+#include "Cooker/CookGenerationHelper.h"
 #include "Cooker/CookPlatformManager.h"
 #include "Cooker/CookRequestCluster.h"
 #include "Cooker/CookWorkerClient.h"
@@ -1494,12 +1495,22 @@ void FPackageData::UpdateSaveAfterGarbageCollect(bool& bOutDemote)
 	}
 }
 
+TRefCountPtr<UE::Cook::FGenerationHelper> FPackageData::GetGenerationHelper() const
+{
+	return GenerationHelper;
+}
+
 void FPackageData::SetGenerated(FName InParentGenerator)
 {
 	bGenerated = true;
 	ParentGenerator = InParentGenerator;
 }
 
+TRefCountPtr<FGenerationHelper> FPackageData::GetParentGenerationHelper() const
+{
+	return ParentGenerationHelper;
+}
+
 void FPackageData::SetParentGenerationHelper(FGenerationHelper* InGenerationHelper)
 {
 	check(InGenerationHelper == nullptr || IsGenerated());
@@ -1697,1178 +1708,6 @@ bool LoadFromCompactBinary(FCbFieldView Field, UE::Cook::FConstructPackageData&
 namespace UE::Cook
 {
 
-//////////////////////////////////////////////////////////////////////////
-// FGenerationHelper
-
-FGenerationHelper::FGenerationHelper(UE::Cook::FPackageData& InOwner)
-: OwnerInfo(InOwner, true /* bInGenerator */)
-{
-}
-
-FGenerationHelper::~FGenerationHelper()
-{
-	NotifyCompletion(ICookPackageSplitter::ETeardown::Complete);
-	GetOwner().OnGenerationHelperDestroyed(*this);
-}
-
-void FGenerationHelper::NotifyCompletion(ICookPackageSplitter::ETeardown Status)
-{
-	if (IsInitialized() && IsValid() && CookPackageSplitterInstance)
-	{
-		CookPackageSplitterInstance->Teardown(Status);
-		CookPackageSplitterInstance.Reset();
-	}
-}
-
-void FGenerationHelper::Initialize()
-{
-	if (InitializeStatus != EInitializeStatus::Uninitialized)
-	{
-		return;
-	}
-
-	FPackageData& OwnerPackageData = GetOwner();
-	FName OwnerPackageName = OwnerPackageData.GetPackageName();
-	UCookOnTheFlyServer& COTFS = OwnerPackageData.GetPackageDatas().GetCookOnTheFlyServer();
-	UPackage* LocalOwnerPackage = FindOrLoadPackage(COTFS, OwnerPackageData);
-	if (!LocalOwnerPackage)
-	{
-		InitializeStatus = EInitializeStatus::Invalid;
-		return;
-	}
-
-	UObject* LocalSplitDataObject;
-	UE::Cook::Private::FRegisteredCookPackageSplitter* LocalRegisteredSplitterType = nullptr;
-	TUniquePtr<ICookPackageSplitter> LocalSplitter;
-	SearchForRegisteredSplitDataObject(COTFS, OwnerPackageName, LocalOwnerPackage,
-		TOptional<TConstArrayView<FCachedObjectInOuter>>(), LocalSplitDataObject, LocalRegisteredSplitterType,
-		LocalSplitter);
-	if (!LocalSplitDataObject || !LocalSplitter)
-	{
-		InitializeStatus = EInitializeStatus::Invalid;
-		return;
-	}
-
-	Initialize(LocalSplitDataObject, LocalRegisteredSplitterType, MoveTemp(LocalSplitter));
-}
-
-void FGenerationHelper::Initialize(const UObject* InSplitDataObject,
-	UE::Cook::Private::FRegisteredCookPackageSplitter* InRegisteredSplitterType,
-	TUniquePtr<ICookPackageSplitter>&& InCookPackageSplitterInstance)
-{
-	check(InSplitDataObject);
-	if (InitializeStatus != EInitializeStatus::Uninitialized)
-	{
-		// If we already have a splitter, keep the old and throw out the new. The old one
-		// still contains some state.
-		return;
-	}
-
-	RegisteredSplitterType = InRegisteredSplitterType;
-	CookPackageSplitterInstance = MoveTemp(InCookPackageSplitterInstance);
-	InitializeStatus = EInitializeStatus::Valid;
-
-	SplitDataObject = InSplitDataObject;
-	SplitDataObjectName = FName(FStringView(InSplitDataObject->GetFullName()));
-	bUseInternalReferenceToAvoidGarbageCollect =
-		CookPackageSplitterInstance->UseInternalReferenceToAvoidGarbageCollect();
-	bNeedCachedPlatformDataBeforeSplit =
-		CookPackageSplitterInstance->NeedCachedPlatformDataBeforeSplit();
-}
-
-void FGenerationHelper::InitializeAsInvalid()
-{
-	if (InitializeStatus != EInitializeStatus::Uninitialized)
-	{
-		return;
-	}
-	InitializeStatus = EInitializeStatus::Invalid;
-}
-
-UPackage* FGenerationHelper::FindOrLoadPackage(UCookOnTheFlyServer& COTFS, FPackageData& OwnerPackageData)
-{
-	// This is the static helper function on FGenerationHelper that loads the package for any FPackageData; for the
-	// member variable function that uses the cached pointer, see FindOrLoadOwnerPackage.
-	FName OwnerPackageName = OwnerPackageData.GetPackageName();
-	UPackage* Result = FindObjectFast<UPackage>(nullptr, OwnerPackageName);
-
-	if (!Result || !Result->IsFullyLoaded())
-	{
-		COTFS.LoadPackageForCooking(OwnerPackageData, Result);
-		if (!Result || !Result->IsFullyLoaded())
-		{
-			return nullptr;
-		}
-	}
-	return Result;
-}
-
-void FGenerationHelper::SearchForRegisteredSplitDataObject(UCookOnTheFlyServer& COTFS,
-	FName PackageName, UPackage* Package,
-	TOptional<TConstArrayView<UE::Cook::FCachedObjectInOuter>> CachedObjectsInOuter,
-	UObject*& OutSplitDataObject, UE::Cook::Private::FRegisteredCookPackageSplitter*& OutRegisteredSplitter,
-	TUniquePtr<ICookPackageSplitter>& OutSplitterInstance)
-{
-	OutSplitDataObject = nullptr;
-	OutRegisteredSplitter = nullptr;
-	OutSplitterInstance = nullptr;
-	check(Package != nullptr || CachedObjectsInOuter.IsSet());
-
-	UObject* LocalSplitDataObject = nullptr;
-	Private::FRegisteredCookPackageSplitter* SplitterType = nullptr;
-	TArray<Private::FRegisteredCookPackageSplitter*> FoundRegisteredSplitters;
-	auto TryLookForSplitterOfObject =
-		[&COTFS, PackageName, &FoundRegisteredSplitters, &SplitterType, &LocalSplitDataObject](UObject* Obj)
-		{
-			FoundRegisteredSplitters.Reset();
-			COTFS.RegisteredSplitDataClasses.MultiFind(Obj->GetClass(), FoundRegisteredSplitters);
-
-			for (Private::FRegisteredCookPackageSplitter* SplitterForObject : FoundRegisteredSplitters)
-			{
-				if (SplitterForObject && SplitterForObject->ShouldSplitPackage(Obj))
-				{
-					if (!Obj->HasAnyFlags(RF_Public))
-					{
-						UE_LOG(LogCook, Error,
-							TEXT("SplitterData object %s must be publicly referenceable so we can keep them from being garbage collected"),
-							*Obj->GetFullName());
-						return false;
-					}
-
-					if (SplitterType)
-					{
-						UE_LOG(LogCook, Error,
-							TEXT("Found more than one registered Cook Package Splitter for package %s."),
-							*PackageName.ToString());
-						return false;
-					}
-
-					SplitterType = SplitterForObject;
-					LocalSplitDataObject = Obj;
-				}
-			}
-			return true;
-		};
-
-	if (CachedObjectsInOuter.IsSet())
-	{
-		// CachedObjectsInOuter might be set but empty for e.g. a generated package that has not been populated
-		for (const FCachedObjectInOuter& CachedObjectInOuter : *CachedObjectsInOuter)
-		{
-			UObject* Obj = CachedObjectInOuter.Object.Get();
-			if (!Obj)
-			{
-				continue;
-			}
-			if (!TryLookForSplitterOfObject(Obj))
-			{
-				return; // error condition, exit the entire search function
-			}
-		}
-	}
-	else
-	{
-		TArray<UObject*> ObjectsInPackage;
-		GetObjectsWithOuter(Package, ObjectsInPackage, true /* bIncludeNestedObjects */,
-			RF_NoFlags, EInternalObjectFlags::Garbage);
-		for (UObject* Obj : ObjectsInPackage)
-		{
-			if (!TryLookForSplitterOfObject(Obj))
-			{
-				return; // error condition, exit the entire search function
-			}
-		}
-	}
-
-	if (!SplitterType)
-	{
-		return;
-	}
-
-	// Create instance of CookPackageSplitter class
-	ICookPackageSplitter* SplitterInstance = SplitterType->CreateInstance(LocalSplitDataObject);
-	if (!SplitterInstance)
-	{
-		UE_LOG(LogCook, Error, TEXT("Error instantiating Cook Package Splitter %s for object %s."),
-			*SplitterType->GetSplitterDebugName(), *LocalSplitDataObject->GetFullName());
-		return;
-	}
-
-	OutSplitDataObject = LocalSplitDataObject;
-	OutRegisteredSplitter = SplitterType;
-	OutSplitterInstance.Reset(SplitterInstance);
-}
-
-void FGenerationHelper::ClearSelfReferences()
-{
-	// Any references we release might be the last reference and cause *this to be deleted,
-	// so create a local reference to keep it alive until the end of the function.
-	TRefCountPtr<FGenerationHelper> LocalRef(this);
-	ClearKeepForIterative();
-	ClearKeepForGeneratorSave();
-	ClearKeepForQueueResults();
-}
-
-FCookGenerationInfo* FGenerationHelper::FindInfo(const FPackageData& PackageData)
-{
-	ConditionalInitialize();
-
-	if (&PackageData == &GetOwner())
-	{
-		return &OwnerInfo;
-	}
-	for (FCookGenerationInfo& Info : PackagesToGenerate)
-	{
-		if (Info.PackageData == &PackageData)
-		{
-			return &Info;
-		}
-	}
-	return nullptr;
-}
-
-const FCookGenerationInfo* FGenerationHelper::FindInfo(const FPackageData& PackageData) const
-{
-	return const_cast<FGenerationHelper*>(this)->FindInfo(PackageData);
-}
-
-UObject* FGenerationHelper::FindOrLoadSplitDataObject()
-{
-	if (!IsValid())
-	{
-		return nullptr;
-	}
-	UObject* Result = SplitDataObject.Get();
-	if (Result)
-	{
-		return Result;
-	}
-
-	FString ObjectPath = GetSplitDataObjectName().ToString();
-	// SplitDataObjectName is a FullObjectPath; strip off the leading <ClassName> in
-	// "<ClassName> <Package>.<Object>:<SubObject>"
-	int32 ClassDelimiterIndex = -1;
-	if (ObjectPath.FindChar(' ', ClassDelimiterIndex))
-	{
-		ObjectPath.RightChopInline(ClassDelimiterIndex + 1);
-	}
-
-	Result = FindObject<UObject>(nullptr, *ObjectPath);
-	if (!Result)
-	{
-		FPackageData& OwnerPackageData = GetOwner();
-		FName OwnerPackageName = OwnerPackageData.GetPackageName();
-		UCookOnTheFlyServer& COTFS = OwnerPackageData.GetPackageDatas().GetCookOnTheFlyServer();
-		UPackage* LocalOwnerPackage;
-		COTFS.LoadPackageForCooking(OwnerPackageData, LocalOwnerPackage);
-
-		Result = FindObject<UObject>(nullptr, *ObjectPath);
-		if (!Result)
-		{
-			return nullptr;
-		}
-	}
-
-	SplitDataObject = Result;
-	return Result;
-}
-
-UPackage* FGenerationHelper::GetOwnerPackage()
-{
-	UPackage* Result = OwnerPackage.Get();
-	if (!Result && !OwnerPackage.GetEvenIfUnreachable())
-	{
-		OwnerPackage = FindObjectFast<UPackage>(nullptr, GetOwner().GetPackageName());
-		Result = OwnerPackage.Get();
-	}
-	return Result;
-}
-
-UPackage* FGenerationHelper::FindOrLoadOwnerPackage(UCookOnTheFlyServer& COTFS)
-{
-	UPackage* Result = GetOwnerPackage();
-	if (!Result)
-	{
-		Result = FindOrLoadPackage(COTFS, GetOwner());
-	}
-	return Result;
-}
-
-bool FGenerationHelper::TryGenerateList()
-{
-	if (bGeneratedList)
-	{
-		return true;
-	}
-	FPackageData& OwnerPackageData = GetOwner();
-	FName OwnerPackageName = OwnerPackageData.GetPackageName();
-	if (!IsValid())
-	{
-		// Unexpected, caller should not call in this case
-		UE_LOG(LogCook, Error, TEXT("TryGenerateList failed for package %s: Called on an invalid FGenerationHelper."),
-			*OwnerPackageName.ToString());
-		FDebug::DumpStackTraceToLog(ELogVerbosity::Warning);
-		return false;
-	}
-
-	FPackageDatas& PackageDatas = OwnerPackageData.GetPackageDatas();
-	UCookOnTheFlyServer& COTFS = PackageDatas.GetCookOnTheFlyServer();
-	UObject* OwnerObject = FindOrLoadSplitDataObject();
-	if (!OwnerObject)
-	{
-		// Unexpected, we found it earlier when we marked valid.
-		UE_LOG(LogCook, Error, TEXT("TryGenerateList failed for package %s: Valid GenerationHelper but could not find OwnerObject."),
-			*OwnerPackageName.ToString());
-		FDebug::DumpStackTraceToLog(ELogVerbosity::Warning);
-		return false;
-	}
-
-	UPackage* LocalOwnerPackage = OwnerObject->GetPackage();
-
-	TArray<ICookPackageSplitter::FGeneratedPackage> GeneratorDatas;
-	{
-		UCookOnTheFlyServer::FScopedActivePackage ScopedActivePackage(COTFS, OwnerPackageName,
-			PackageAccessTrackingOps::NAME_CookerBuildObject);
-		GeneratorDatas = GetCookPackageSplitterInstance()->GetGenerateList(LocalOwnerPackage, OwnerObject);
-	}
-	PackagesToGenerate.Reset(GeneratorDatas.Num());
-	TArray<const ITargetPlatform*, TInlineAllocator<1>> PlatformsToCook;
-	OwnerPackageData.GetPlatformsNeedingCooking(PlatformsToCook);
-
-	for (ICookPackageSplitter::FGeneratedPackage& SplitterData : GeneratorDatas)
-	{
-		if (!SplitterData.GetCreateAsMap().IsSet())
-		{
-			UE_LOG(LogCook, Error,
-				TEXT("PackageSplitter did not specify whether CreateAsMap is true for generated package. Splitter=%s, Generated=%s."),
-				*this->GetSplitDataObjectName().ToString(), *OwnerPackageName.ToString());
-			return false;
-		}
-		bool bCreateAsMap = *SplitterData.GetCreateAsMap();
-
-		FString PackageRoot = SplitterData.GeneratedRootPath.IsEmpty() ?
-			OwnerPackageName.ToString() : SplitterData.GeneratedRootPath;
-
-		FString PackageName = FPaths::RemoveDuplicateSlashes(FString::Printf(TEXT("/%s/%s/%s"),
-			*PackageRoot, GeneratedPackageSubPath, *SplitterData.RelativePath));
-		const FName PackageFName(*PackageName);
-		UE::Cook::FPackageData* PackageData = PackageDatas.TryAddPackageDataByPackageName(PackageFName,
-			false /* bRequireExists */, bCreateAsMap);
-		if (!PackageData)
-		{
-			UE_LOG(LogCook, Error,
-				TEXT("PackageSplitter could not find mounted filename for generated packagepath. Splitter=%s, Generated=%s."),
-				*this->GetSplitDataObjectName().ToString(), *PackageName);
-			return false;
-		}
-		// No package should be generated by two different splitters.
-		check(PackageData->GetParentGenerator().IsNone() ||
-			PackageData->GetParentGenerator() == OwnerPackageName);
-		PackageData->SetGenerated(OwnerPackageName);
-		PackageData->SetGeneratedNeedCachedPlatformDataBeforeSplit(bNeedCachedPlatformDataBeforeSplit);
-		if (IFileManager::Get().FileExists(*PackageData->GetFileName().ToString()))
-		{
-			UE_LOG(LogCook, Warning,
-				TEXT("PackageSplitter specified a generated package that already exists in the workspace domain. Splitter=%s, Generated=%s."),
-				*this->GetSplitDataObjectName().ToString(), *PackageName);
-			return false;
-		}
-
-		FCookGenerationInfo& GeneratedInfo = PackagesToGenerate.Emplace_GetRef(*PackageData, false /* bInGenerator */);
-		GeneratedInfo.RelativePath = MoveTemp(SplitterData.RelativePath);
-		GeneratedInfo.GeneratedRootPath = MoveTemp(SplitterData.GeneratedRootPath);
-		GeneratedInfo.PackageDependencies = MoveTemp(SplitterData.PackageDependencies);
-		for (TArray<FAssetDependency>::TIterator Iter(GeneratedInfo.PackageDependencies); Iter; ++Iter)
-		{
-			if (Iter->Category != UE::AssetRegistry::EDependencyCategory::Package)
-			{
-				UE_LOG(LogCook, Error,
-					TEXT("PackageSplitter specified a dependency with category %d rather than category Package. Dependency will be ignored. Splitter=%s, Generated=%s."),
-					(int32)Iter->Category, *this->GetSplitDataObjectName().ToString(), *PackageName);
-				Iter.RemoveCurrent();
-			}
-		}
-		Algo::Sort(GeneratedInfo.PackageDependencies,
-			[](const FAssetDependency& A, const FAssetDependency& B) { return A.LexicalLess(B); });
-		GeneratedInfo.PackageDependencies.SetNum(Algo::Unique(GeneratedInfo.PackageDependencies));
-		GeneratedInfo.SetIsCreateAsMap(bCreateAsMap);
-		if (bNeedCachedPlatformDataBeforeSplit ||
-			COTFS.MPCookGeneratorSplit == UE::Cook::EMPCookGeneratorSplit::AllOnSameWorker)
-		{
-			PackageData->SetWorkerAssignmentConstraint(FWorkerId::Local());
-		}
-
-		// Create the Hash from the GenerationHash and Dependencies
-		GeneratedInfo.CreatePackageHash();
-	}
-
-	bGeneratedList = true;
-	return true;
-}
-
-void FGenerationHelper::StartOwnerSave()
-{
-	if (!IsValid())
-	{
-		return;
-	}
-	UE_LOG(LogCook, Display, TEXT("Splitting Package %s with splitter %s acting on object %s."),
-		*WriteToString<256>(GetOwner().GetPackageName()),
-		*GetRegisteredSplitterType()->GetSplitterDebugName(),
-		*WriteToString<256>(GetSplitDataObjectName()));
-	SetKeepForGeneratorSave();
-}
-
-void FGenerationHelper::StartQueueGeneratedPackages(UCookOnTheFlyServer& COTFS)
-{
-	if (!IsValid())
-	{
-		return;
-	}
-	NotifyStartQueueGeneratedPackages(COTFS, FWorkerId::Local());
-
-	bool bHybridIterativeEnabled = COTFS.bHybridIterativeEnabled;
-	if (!PreviousGeneratedPackages.IsEmpty())
-	{
-		FPackageData& OwnerPackageData = GetOwner();
-		TArray<const ITargetPlatform*, TInlineAllocator<1>> PlatformsToCook;
-		OwnerPackageData.GetPlatformsNeedingCooking(PlatformsToCook);
-
-		int32 NumIterativeUnmodified = 0;
-		int32 NumIterativeModified = 0;
-		int32 NumIterativeRemoved = 0;
-		int32 NumIterativePrevious = PreviousGeneratedPackages.Num();
-
-		for (FCookGenerationInfo& GeneratedInfo : PackagesToGenerate)
-		{
-			if (!GeneratedInfo.PackageData)
-			{
-				continue;
-			}
-			FIoHash PreviousHash;
-			if (PreviousGeneratedPackages.RemoveAndCopyValue(GeneratedInfo.PackageData->GetPackageName(), PreviousHash)
-				&& !bHybridIterativeEnabled)
-			{
-				bool bIterativelyUnmodified;
-				GeneratedInfo.IterativeCookValidateOrClear(*this, PlatformsToCook, PreviousHash, bIterativelyUnmodified);
-				++(bIterativelyUnmodified ? NumIterativeUnmodified : NumIterativeModified);
-			}
-		}
-		if (!PreviousGeneratedPackages.IsEmpty())
-		{
-			NumIterativeRemoved = PreviousGeneratedPackages.Num();
-			for (TPair<FName, FIoHash>& Pair : PreviousGeneratedPackages)
-			{
-				for (const ITargetPlatform* TargetPlatform : PlatformsToCook)
-				{
-					COTFS.DeleteOutputForPackage(Pair.Key, TargetPlatform);
-				}
-			}
-			PreviousGeneratedPackages.Empty();
-		}
-		ClearKeepForIterative();
-
-		if (NumIterativePrevious > 0 && !bHybridIterativeEnabled)
-		{
-			UE_LOG(LogCook, Display, TEXT("Found %d cooked package(s) in package store for generator package %s."),
-				NumIterativePrevious, *WriteToString<256>(GetOwner().GetPackageName()));
-			UE_LOG(LogCook, Display, TEXT("Keeping %d. Recooking %d. Removing %d."),
-				NumIterativeUnmodified, NumIterativeModified, NumIterativeRemoved);
-		}
-	}
-}
-
-void FGenerationHelper::NotifyStartQueueGeneratedPackages(UCookOnTheFlyServer& COTFS, FWorkerId SourceWorkerId)
-{
-	// Note this function can be called on an uninitialized Generator; the generator is only needed
-	// on the director so it can serve as the passer of messages. We have to keep ourselves referenced after
-	// this call, until after we send EGeneratorEvent::QueuedGeneratedPackagesFencePassed, so that we don't destruct
-	// and lose the WorkerIdThatSavedGenerator information.
-	if (COTFS.CookDirector)
-	{
-		COTFS.PackageDatas->GetRequestQueue().AddRequestFenceListener(GetOwner().GetPackageName());
-		WorkerIdThatSavedGenerator = SourceWorkerId;
-	}
-	SetKeepForQueueResults();
-}
-
-void FGenerationHelper::EndQueueGeneratedPackages(UCookOnTheFlyServer& COTFS)
-{
-	SetKeepForQueueResults();
-	COTFS.WorkerRequests->EndQueueGeneratedPackages(COTFS, *this);
-}
-
-void FGenerationHelper::EndQueueGeneratedPackagesOnDirector(UCookOnTheFlyServer& COTFS, FWorkerId SourceWorkerId)
-{
-	// Note this function can be called on an uninitialized Generator; the generator is only needed
-	// on the director so it can serve as the passer of messages.
-	COTFS.PackageDatas->GetRequestQueue().AddRequestFenceListener(GetOwner().GetPackageName());
-	SetKeepForQueueResults();
-
-	// Setting the WorkerIdThatSavedGenerator in response to this event is usually not needed because it is also set
-	// from the reported GeneratedPackages, but we set it anyway in case there is an edge condition that skips those
-	// notifications.
-	WorkerIdThatSavedGenerator = SourceWorkerId;
-}
-
-void FGenerationHelper::OnRequestFencePassedBroadcast(UCookOnTheFlyServer& COTFS)
-{
-	if (COTFS.CookDirector)
-	{
-		// Broadcast GenerationMessage API
-		COTFS.CookDirector->BroadcastGeneratorFencePassed(*this);
-	}
-	OnRequestFencePassed(COTFS);
-}
-
-void FGenerationHelper::OnRequestFencePassed(UCookOnTheFlyServer& COTFS)
-{
-	ClearKeepForQueueResults();
-	// We no longer need PreviousGeneratedPackages or KeepForIterative, because they are used (and cleared) in
-	// StartQueueGeneratedPackages. Clear them as well on the director and any CookWorkers that received it to
-	// free memory.
-	ClearKeepForIterative();
-	PreviousGeneratedPackages.Empty();
-}
-
-UPackage* FGenerationHelper::CreateGeneratedUPackage(FCookGenerationInfo& GeneratedInfo,
-	const UPackage* InOwnerPackage, const TCHAR* GeneratedPackageName)
-{
-	if (!IsValid())
-	{
-		return nullptr;
-	}
-
-#if ENABLE_COOK_STATS
-	++DetailedCookStats::NumRequestedLoads;
-#endif
-	UPackage* GeneratedPackage = CreatePackage(GeneratedPackageName);
-	GeneratedPackage->SetSavedHash(GeneratedInfo.PackageHash);
-	GeneratedPackage->SetPersistentGuid(InOwnerPackage->GetPersistentGuid());
-	GeneratedPackage->SetPackageFlags(PKG_CookGenerated);
-	GeneratedInfo.SetHasCreatedPackage(true);
-	if (!InOwnerPackage->IsLoadedByEditorPropertiesOnly())
-	{
-		GeneratedPackage->SetLoadedByEditorPropertiesOnly(false);
-	}
-
-	return GeneratedPackage;
-}
-
-void FGenerationHelper::ResetSaveState(FCookGenerationInfo& Info, UPackage* Package,
-	EStateChangeReason ReleaseSaveReason, EPackageState NewState)
-{
-	ConditionalInitialize();
-
-	// We release references to *this in this function so keep a local reference to avoid deletion during the function.
-	TRefCountPtr<FGenerationHelper> LocalRefCount = this;
-
-	if (Info.GetSaveState() > FCookGenerationInfo::ESaveState::CallPopulate)
-	{
-		UObject* SplitObject = GetWeakSplitDataObject();
-		UPackage* LocalOwnerPackage = Info.IsGenerator() ? Package : GetOwnerPackage();
-		if (!SplitObject || !Package || !LocalOwnerPackage)
-		{
-			UE_LOG(LogCook, Warning,
-				TEXT("PackageSplitter: %s on %s was GarbageCollected before we finished saving it. This prevents us from calling PostSave and may corrupt other packages that it altered during Populate. Splitter=%s."),
-				(!Package ? TEXT("UPackage") :
-					(!LocalOwnerPackage ? TEXT("ParentGenerator UPackage") : TEXT("SplitDataObject"))),
-				Info.PackageData ? *Info.PackageData->GetPackageName().ToString() : *Info.RelativePath,
-				*GetSplitDataObjectName().ToString());
-		}
-		else
-		{
-			UCookOnTheFlyServer& COTFS = GetOwner().GetPackageDatas().GetCookOnTheFlyServer();
-			UCookOnTheFlyServer::FScopedActivePackage ScopedActivePackage(COTFS, GetOwner().GetPackageName(),
-				PackageAccessTrackingOps::NAME_CookerBuildObject);
-			if (Info.IsGenerator())
-			{
-				GetCookPackageSplitterInstance()->PostSaveGeneratorPackage(Package, SplitObject);
-			}
-			else
-			{
-				ICookPackageSplitter::FGeneratedPackageForPopulate PopulateInfo;
-				PopulateInfo.RelativePath = Info.RelativePath;
-				PopulateInfo.GeneratedRootPath = Info.GeneratedRootPath;
-				PopulateInfo.bCreatedAsMap = Info.IsCreateAsMap();
-				PopulateInfo.Package = Package;
-				GetCookPackageSplitterInstance()->PostSaveGeneratedPackage(LocalOwnerPackage, SplitObject,
-					PopulateInfo);
-			}
-		}
-	}
-
-	if (ReleaseSaveReason == EStateChangeReason::RecreateObjectCache ||
-		ReleaseSaveReason == EStateChangeReason::DoneForNow)
-	{
-		if (Info.IsGenerator())
-		{
-			if (Info.GetSaveState() >= FCookGenerationInfo::ESaveState::StartPopulate)
-			{
-				Info.SetSaveState(FCookGenerationInfo::ESaveState::StartPopulate);
-			}
-			else
-			{
-				// Redo all the steps since we didn't make it to the FinishCachePreObjectsToMove.
-				// Restarting in the middle of that flow is not robust
-				Info.SetSaveState(FCookGenerationInfo::ESaveState::StartSave);
-			}
-		}
-		else
-		{
-			Info.SetSaveState(FCookGenerationInfo::ESaveState::StartPopulate);
-		}
-	}
-	else
-	{
-		// The save is completed and we will not come back to it; set state back to initial
-		// state and drop our reference keeping this GenerationHelper in memory for the save.
-		if (ReleaseSaveReason == EStateChangeReason::Completed)
-		{
-			Info.SetHasSaved(true);
-		}
-
-		if (Info.IsGenerator())
-		{
-			ClearKeepForGeneratorSave();
-			Info.SetSaveState(FCookGenerationInfo::ESaveState::StartSave);
-		}
-		else
-		{
-			Info.SetSaveState(FCookGenerationInfo::ESaveState::StartPopulate);
-			if (Info.PackageData)
-			{
-				Info.PackageData->SetParentGenerationHelper(nullptr);
-			}
-		}
-	}
-
-	if (Info.HasTakenOverCachedCookedPlatformData())
-	{
-		if (NewState != EPackageState::Idle && Info.PackageData &&
-			Info.PackageData->GetCachedObjectsInOuter().Num() != 0 && IsUseInternalReferenceToAvoidGarbageCollect() &&
-			(ReleaseSaveReason != EStateChangeReason::Completed && ReleaseSaveReason != EStateChangeReason::DoneForNow
-				&& ReleaseSaveReason != EStateChangeReason::SaveError
-				&& ReleaseSaveReason != EStateChangeReason::CookerShutdown))
-		{
-			UE_LOG(LogCook, Error,
-				TEXT("CookPackageSplitter failure: We are demoting a %s package from save and removing our references that keep its objects loaded.\n")
-				TEXT("This will allow the objects to be garbage collected and cause failures in the splitter which expects them to remain loaded.\n")
-				TEXT("Package=%s, Splitter=%s, ReleaseSaveReason=%s"),
-				Info.IsGenerator() ? TEXT("generator") : TEXT("generated"),
-				Info.PackageData ? *Info.PackageData->GetPackageName().ToString() : *Info.RelativePath,
-				*GetSplitDataObjectName().ToString(), LexToString(ReleaseSaveReason));
-		}
-		Info.CachedObjectsInOuterInfo.Empty();
-		Info.SetHasTakenOverCachedCookedPlatformData(false);
-	}
-	Info.SetHasIssuedUndeclaredMovedObjectsWarning(false);
-	Info.KeepReferencedPackages.Reset();
-}
-
-void FGenerationHelper::FetchExternalActorDependencies()
-{
-	if (!IsValid())
-	{
-		return;
-	}
-
-	// The Generator package declares all its ExternalActor dependencies in its AssetRegistry dependencies
-	// The Generator's generated packages can also include ExternalActors from other maps due to level instancing,
-	// these are included in the dependencies reported by the Generator for each GeneratedPackage in the data
-	// returned from GetGenerateList. These sets will overlap; take the union.
-	ExternalActorDependencies.Reset();
-	IAssetRegistry::GetChecked().GetDependencies(GetOwner().GetPackageName(), ExternalActorDependencies,
-		UE::AssetRegistry::EDependencyCategory::Package);
-	for (const FCookGenerationInfo& Info : PackagesToGenerate)
-	{
-		ExternalActorDependencies.Reserve(Info.GetDependencies().Num() + ExternalActorDependencies.Num());
-		for (const FAssetDependency& Dependency : Info.GetDependencies())
-		{
-			ExternalActorDependencies.Add(Dependency.AssetId.PackageName);
-		}
-	}
-	Algo::Sort(ExternalActorDependencies, FNameFastLess());
-	ExternalActorDependencies.SetNum(Algo::Unique(ExternalActorDependencies));
-	FPackageDatas& PackageDatas = this->GetOwner().GetPackageDatas();
-	FThreadSafeSet<FName>& NeverCookPackageList =
-		GetOwner().GetPackageDatas().GetCookOnTheFlyServer().PackageTracker->NeverCookPackageList;
-
-	// We are supposed to collect only ExternalActor dependencies, but we collected every dependency from the
-	// generated packages. Remove the packages that are not external actors, which we detect by being on-disk
-	// PackageDatas that are marked as NeverCook
-	ExternalActorDependencies.RemoveAll([&PackageDatas, &NeverCookPackageList](FName PackageName)
-		{
-			FPackageData* PackageData = PackageDatas.TryAddPackageDataByPackageName(PackageName);
-			if (!PackageData)
-			{
-				return true;
-			}
-			bool bIsNeverCook = NeverCookPackageList.Contains(PackageData->GetFileName());
-			return !bIsNeverCook;
-		});
-	ExternalActorDependencies.Shrink();
-}
-
-void FGenerationHelper::SetPreviousGeneratedPackages(TMap<FName, FIoHash>&& Packages)
-{
-	SetKeepForIterative();
-	PreviousGeneratedPackages = MoveTemp(Packages);
-}
-
-void FGenerationHelper::PreGarbageCollect(FCookGenerationInfo& Info, TArray<TObjectPtr<UObject>>& GCKeepObjects,
-	TArray<UPackage*>& GCKeepPackages, TArray<FPackageData*>& GCKeepPackageDatas, bool& bOutShouldDemote)
-{
-	if (!IsInitialized())
-	{
-		return;
-	}
-
-	bOutShouldDemote = false;
-	check(Info.PackageData); // Caller validates this is non-null
-	if (Info.GetSaveState() > FCookGenerationInfo::ESaveState::CallPopulate)
-	{
-		if (IsUseInternalReferenceToAvoidGarbageCollect() || Info.PackageData->GetIsCookLast())
-		{
-			UPackage* Package = Info.PackageData->GetPackage();
-			if (Package)
-			{
-				GCKeepPackages.Add(Package);
-				GCKeepPackageDatas.Add(Info.PackageData);
-			}
-		}
-		else
-		{
-			bOutShouldDemote = true;
-		}
-	}
-	if (Info.HasTakenOverCachedCookedPlatformData())
-	{
-		if (IsUseInternalReferenceToAvoidGarbageCollect())
-		{
-			// For the UseInternalReferenceToAvoidGarbageCollect case, part of the CookPackageSplitter contract is that
-			// the Cooker will keep referenced the package and all objects returned from GetObjectsToMove* functions
-			// until the PostSave function is called
-			UPackage* Package = Info.PackageData->GetPackage();
-			if (Package)
-			{
-				GCKeepPackages.Add(Package);
-				GCKeepPackageDatas.Add(Info.PackageData);
-			}
-			GCKeepPackages.Append(Info.KeepReferencedPackages);
-			for (FCachedObjectInOuter& CachedObjectInOuter : Info.PackageData->GetCachedObjectsInOuter())
-			{
-				UObject* Object = CachedObjectInOuter.Object.Get();
-				if (Object)
-				{
-					GCKeepObjects.Add(Object);
-				}
-			}
-		}
-	}
-}
-
-void FGenerationHelper::PostGarbageCollect()
-{
-	if (!IsInitialized())
-	{
-		return;
-	}
-
-	FPackageData& Owner = GetOwner();
-	if (Owner.GetState() == EPackageState::Save)
-	{
-		// UCookOnTheFlyServer::PreGarbageCollect adds references for the Generator package and all its public
-		// objects, so it should still be loaded
-		if (!Owner.GetPackage() || !GetWeakSplitDataObject())
-		{
-			UE_LOG(LogCook, Error,
-				TEXT("PackageSplitter object was deleted by garbage collection while generation was still ongoing. This will break the generation.")
-				TEXT("\n\tSplitter=%s."), *GetSplitDataObjectName().ToString());
-		}
-	}
-	else
-	{
-		// After the Generator Package is saved, we drop our references to it and it can be garbage collected
-		// If we have any packages left to populate, our splitter contract requires that it be garbage collected
-		// because we promise that the package is not partially GC'd during calls to TryPopulateGeneratedPackage
-		// The splitter can opt-out of this contract and keep it referenced itself if it desires.
-		UPackage* LocalOwnerPackage = FindObject<UPackage>(nullptr, *Owner.GetPackageName().ToString());
-		if (LocalOwnerPackage)
-		{
-			if (!Owner.IsInProgress() && !Owner.IsKeepReferencedDuringGC() &&
-				!IsUseInternalReferenceToAvoidGarbageCollect())
-			{
-				UE_LOG(LogCook, Error,
-					TEXT("PackageSplitter found the Generator package still in memory after it should have been deleted by GC.")
-					TEXT("\n\tThis is unexpected since garbage has been collected and the package should have been unreferenced so it should have been collected, and will break population of Generated packages.")
-					TEXT("\n\tSplitter=%s"), *GetSplitDataObjectName().ToString());
-				EReferenceChainSearchMode SearchMode = EReferenceChainSearchMode::Shortest
-					| EReferenceChainSearchMode::PrintAllResults
-					| EReferenceChainSearchMode::FullChain;
-				FReferenceChainSearch RefChainSearch(LocalOwnerPackage, SearchMode);
-			}
-		}
-	}
-
-	bool bHasIssuedWarning = false;
-	for (FCookGenerationInfo& Info : PackagesToGenerate)
-	{
-		if (FindObject<UPackage>(nullptr, *Info.PackageData->GetPackageName().ToString()))
-		{
-			if (!Info.PackageData->IsKeepReferencedDuringGC() && !Info.HasSaved() && !bHasIssuedWarning)
-			{
-				UE_LOG(LogCook, Warning,
-					TEXT("PackageSplitter found a package it generated that was not removed from memory during garbage collection. This will cause errors later during population.")
-					TEXT("\n\tSplitter=%s, Generated=%s."), *GetSplitDataObjectName().ToString(),
-					*Info.PackageData->GetPackageName().ToString());
-				
-				{
-					// Compute UCookOnTheFlyServer's references so they are gathered by OBJ REFS below 
-					UCookOnTheFlyServer::FScopeFindCookReferences(Owner.GetPackageDatas().GetCookOnTheFlyServer());
-
-					StaticExec(nullptr, *FString::Printf(TEXT("OBJ REFS NAME=%s"),
-						*Info.PackageData->GetPackageName().ToString()));
-				}
-				
-				bHasIssuedWarning = true; // Only issue the warning once per GC
-			}
-		}
-		else
-		{
-			Info.SetHasCreatedPackage(false);
-		}
-	}
-}
-
-void FGenerationHelper::UpdateSaveAfterGarbageCollect(const FPackageData& PackageData, bool& bInOutDemote)
-{
-	if (!IsInitialized())
-	{
-		return;
-	}
-	FCookGenerationInfo* Info = FindInfo(PackageData);
-	if (!Info)
-	{
-		bInOutDemote = true;
-		return;
-	}
-
-	if (!Info->IsGenerator())
-	{
-		UPackage* LocalPackage = OwnerPackage.Get();
-		if (!LocalPackage || !LocalPackage->IsFullyLoaded())
-		{
-			bInOutDemote = true;
-			return;
-		}
-	}
-
-	if (bInOutDemote && IsUseInternalReferenceToAvoidGarbageCollect() && Info->HasTakenOverCachedCookedPlatformData())
-	{
-		// No public objects should have been deleted; we are supposed to keep them referenced by keeping the package
-		// referenced in UCookOnTheFlyServer::PreGarbageCollect, and the package keeping its public objects referenced
-		// by UPackage::AddReferencedObjects. Since no public objects were deleted, our caller should not have
-		// set bInOutDemote=true.
-		// Allowing demotion after the splitter has started moving objects breaks our contract with the splitter
-		// and can cause a crash. So log this as an error.
-		// For better feedback, look in our extra data to identify the name of the public UObject that was deleted.
-		FString DeletedObject;
-		if (!PackageData.GetPackage())
-		{
-			DeletedObject = FString::Printf(TEXT("UPackage %s"), *PackageData.GetPackageName().ToString());
-		}
-		else
-		{
-			TSet<UObject*> ExistingObjectsAfterSave;
-			for (const FCachedObjectInOuter& CachedObjectInOuter : PackageData.GetCachedObjectsInOuter())
-			{
-				UObject* Ptr = CachedObjectInOuter.Object.Get();
-				if (Ptr)
-				{
-					ExistingObjectsAfterSave.Add(Ptr);
-				}
-			}
-
-			for (const TPair<UObject*, FCachedObjectInOuterGeneratorInfo>& Pair : Info->CachedObjectsInOuterInfo)
-			{
-				if (Pair.Value.bPublic && !ExistingObjectsAfterSave.Contains(Pair.Key))
-				{
-					DeletedObject = Pair.Value.FullName;
-					break;
-				}
-			}
-			if (DeletedObject.IsEmpty())
-			{
-				if (!PackageData.GetPackage()->IsFullyLoaded())
-				{
-					DeletedObject = FString::Printf(TEXT("UPackage %s is no longer FullyLoaded"),
-						*PackageData.GetPackageName().ToString());
-				}
-				else
-				{
-					DeletedObject = TEXT("<Unknown>");
-				}
-			}
-		}
-		UE_LOG(LogCook, Error,
-			TEXT("A %s package had some of its UObjects deleted during garbage collection after it started generating. This will cause errors during save of the package.")
-			TEXT("\n\tDeleted object: %s")
-			TEXT("\n\tSplitter=%s%s"),
-			Info->IsGenerator() ? TEXT("Generator") : TEXT("Generated"),
-			*DeletedObject,
-			*GetSplitDataObjectName().ToString(),
-			Info->IsGenerator() ? TEXT(".") : *FString::Printf(TEXT(", Generated=%s."),
-				*Info->PackageData->GetPackageName().ToString()));
-	}
-
-	// Remove raw pointers from RootMovedObjects if they no longer exist in the weakpointers in CachedObjectsInOuter
-	TSet<UObject*> CachedObjectsInOuterSet;
-	for (FCachedObjectInOuter& CachedObjectInOuter : Info->PackageData->GetCachedObjectsInOuter())
-	{
-		UObject* Object = CachedObjectInOuter.Object.Get();
-		if (Object)
-		{
-			CachedObjectsInOuterSet.Add(Object);
-		}
-	}
-	for (TMap<UObject*, FCachedObjectInOuterGeneratorInfo>::TIterator Iter(Info->CachedObjectsInOuterInfo);
-		Iter; ++Iter)
-	{
-		if (!CachedObjectsInOuterSet.Contains(Iter->Key))
-		{
-			Iter.RemoveCurrent();
-		}
-	}
-}
-
-FCookGenerationInfo::FCookGenerationInfo(FPackageData& InPackageData, bool bInGenerator)
-	: PackageData(&InPackageData)
-	, GeneratorSaveState(bInGenerator ? ESaveState::StartSave : ESaveState::StartPopulate)
-	, bCreateAsMap(false), bHasCreatedPackage(false), bHasSaved(false), bTakenOverCachedCookedPlatformData(false)
-	, bIssuedUndeclaredMovedObjectsWarning(false), bGenerator(bInGenerator)
-{
-}
-
-void FCookGenerationInfo::SetSaveStateComplete(ESaveState CompletedState)
-{
-	GeneratorSaveState = CompletedState;
-	if (GeneratorSaveState < ESaveState::Last)
-	{
-		GeneratorSaveState = static_cast<ESaveState>(static_cast<uint8>(GeneratorSaveState) + 1);
-	}
-}
-
-void FCachedObjectInOuterGeneratorInfo::Initialize(UObject* Object)
-{
-	if (Object)
-	{
-		FullName = Object->GetFullName();
-		bPublic = Object->HasAnyFlags(RF_Public);
-	}
-	else
-	{
-		FullName.Empty();
-		bPublic = false;
-	}
-
-	bInitialized = true;
-}
-
-void FCookGenerationInfo::TakeOverCachedObjectsAndAddMoved(FGenerationHelper& GenerationHelper,
-	TArray<FCachedObjectInOuter>& CachedObjectsInOuter, TArray<UObject*>& MovedObjects)
-{
-	CachedObjectsInOuterInfo.Reset();
-
-	for (FCachedObjectInOuter& ObjectInOuter : CachedObjectsInOuter)
-	{
-		UObject* Object = ObjectInOuter.Object.Get();
-		if (Object)
-		{
-			CachedObjectsInOuterInfo.FindOrAdd(Object).Initialize(Object);
-		}
-	}
-
-	TArray<UObject*> ChildrenOfMovedObjects;
-	for (UObject* Object : MovedObjects)
-	{
-		if (!IsValid(Object))
-		{
-			UE_LOG(LogCook, Warning,
-				TEXT("CookPackageSplitter found non-valid object %s returned from %s on Splitter %s%s. Ignoring it."),
-				Object ? *Object->GetFullName() : TEXT("<null>"),
-				IsGenerator() ? TEXT("PopulateGeneratorPackage") : TEXT("PopulateGeneratedPackage"),
-				*GenerationHelper.GetSplitDataObjectName().ToString(),
-				IsGenerator() ? TEXT("") : *FString::Printf(TEXT(", Package %s"),
-					*PackageData->GetPackageName().ToString()));
-			continue;
-		}
-		FCachedObjectInOuterGeneratorInfo& Info = CachedObjectsInOuterInfo.FindOrAdd(Object);
-		if (!Info.bInitialized)
-		{
-			Info.Initialize(Object);
-			Info.bMoved = true;
-			Info.bMovedRoot = true;
-			CachedObjectsInOuter.Emplace(Object);
-			GetObjectsWithOuter(Object, ChildrenOfMovedObjects, true /* bIncludeNestedObjects */, RF_NoFlags,
-				EInternalObjectFlags::Garbage);
-		}
-	}
-
-	for (UObject* Object : ChildrenOfMovedObjects)
-	{
-		check(IsValid(Object));
-		FCachedObjectInOuterGeneratorInfo& Info = CachedObjectsInOuterInfo.FindOrAdd(Object);
-		if (!Info.bInitialized)
-		{
-			Info.Initialize(Object);
-			Info.bMoved = true;
-			CachedObjectsInOuter.Emplace(Object);
-		}
-	}
-
-	SetHasTakenOverCachedCookedPlatformData(true);
-}
-
-EPollStatus FCookGenerationInfo::RefreshPackageObjects(FGenerationHelper& GenerationHelper, UPackage* Package,
-	bool& bOutFoundNewObjects, ESaveState DemotionState)
-{
-	bOutFoundNewObjects = false;
-	TArray<UObject*> CurrentObjectsInOuter;
-	GetObjectsWithOuter(Package, CurrentObjectsInOuter, true /* bIncludeNestedObjects */, RF_NoFlags,
-		EInternalObjectFlags::Garbage);
-
-	check(PackageData); // RefreshPackageObjects is only called when there is a PackageData
-	TArray<FCachedObjectInOuter>& CachedObjectsInOuter = PackageData->GetCachedObjectsInOuter();
-	UObject* FirstNewObject = nullptr;
-	for (UObject* Object : CurrentObjectsInOuter)
-	{
-		FCachedObjectInOuterGeneratorInfo& Info = CachedObjectsInOuterInfo.FindOrAdd(Object);
-		if (!Info.bInitialized)
-		{
-			Info.Initialize(Object);
-			CachedObjectsInOuter.Emplace(Object);
-			if (!FirstNewObject)
-			{
-				FirstNewObject = Object;
-			}
-		}
-	}
-	bOutFoundNewObjects = FirstNewObject != nullptr;
-
-	if (FirstNewObject != nullptr && DemotionState != ESaveState::Last)
-	{
-		SetSaveState(DemotionState);
-		if (++PackageData->GetNumRetriesBeginCacheOnObjects() > FPackageData::GetMaxNumRetriesBeginCacheOnObjects())
-		{
-			UE_LOG(LogCook, Error,
-				TEXT("Cooker has repeatedly tried to call BeginCacheForCookedPlatformData on all objects in a generated package, but keeps finding new objects.\n")
-				TEXT("Aborting the save of the package; programmer needs to debug why objects keep getting added to the package.\n")
-				TEXT("Splitter: %s%s. Most recent created object: %s."),
-				*GenerationHelper.GetSplitDataObjectName().ToString(),
-				IsGenerator() ? TEXT("") : *FString::Printf(TEXT(", Package: %s"),
-					*PackageData->GetPackageName().ToString()),
-				*FirstNewObject->GetFullName());
-			return EPollStatus::Error;
-		}
-	}
-	return EPollStatus::Success;
-}
-
-void FCookGenerationInfo::CreatePackageHash()
-{
-	FBlake3 Blake3;
-	Blake3.Update(&GenerationHash, sizeof(GenerationHash));
-	IAssetRegistry& AssetRegistry = IAssetRegistry::GetChecked();
-	for (const FAssetDependency& Dependency : PackageDependencies)
-	{
-		TOptional<FAssetPackageData> DependencyData =
-			AssetRegistry.GetAssetPackageDataCopy(Dependency.AssetId.PackageName);
-		if (DependencyData)
-		{
-			Blake3.Update(&DependencyData->GetPackageSavedHash().GetBytes(),
-				sizeof(decltype(DependencyData->GetPackageSavedHash().GetBytes())));
-		}
-	}
-	PackageHash = FIoHash(Blake3.Finalize());
-	// We store the PackageHash as a FIoHash, but UPackage and FAssetPackageData store it as a FGuid, which is smaller,
-	// so we have to remove any data which doesn't fit into FGuid. This can be removed when we remove the deprecated
-	// Guid storage on UPackage.
-	PRAGMA_DISABLE_DEPRECATION_WARNINGS;
-	constexpr int SizeDifference = sizeof(PackageHash) - sizeof(decltype(DeclVal<UPackage>().GetGuid()));
-	if (SizeDifference > 0)
-	{
-		FMemory::Memset(((uint8*)&PackageHash.GetBytes())
-			+ (sizeof(decltype(PackageHash.GetBytes())) - SizeDifference),
-			0, SizeDifference);
-	}
-	PRAGMA_ENABLE_DEPRECATION_WARNINGS;
-}
-
-void FCookGenerationInfo::IterativeCookValidateOrClear(FGenerationHelper& GenerationHelper,
-	TConstArrayView<const ITargetPlatform*> RequestedPlatforms, const FIoHash& PreviousPackageHash,
-	bool& bOutIterativelyUnmodified)
-{
-	UCookOnTheFlyServer& COTFS = GenerationHelper.GetOwner().GetPackageDatas().GetCookOnTheFlyServer();
-	bOutIterativelyUnmodified = PreviousPackageHash == this->PackageHash;
-	if (bOutIterativelyUnmodified)
-	{
-		// If not directly modified, mark it as indirectly modified if any of its dependencies
-		// were detected as modified during PopulateCookedPackages.
-		for (const FAssetDependency& Dependency : this->PackageDependencies)
-		{
-			FPackageData* DependencyData =
-				COTFS.PackageDatas->FindPackageDataByPackageName(Dependency.AssetId.PackageName);
-			if (!DependencyData)
-			{
-				bOutIterativelyUnmodified = false;
-				break;
-			}
-			for (const ITargetPlatform* TargetPlatform : RequestedPlatforms)
-			{
-				FPackagePlatformData* DependencyPlatformData = DependencyData->FindPlatformData(TargetPlatform);
-				if (!DependencyPlatformData || !DependencyPlatformData->IsIterativelyUnmodified())
-				{
-					bOutIterativelyUnmodified = false;
-					break;
-				}
-			}
-			if (!bOutIterativelyUnmodified)
-			{
-				break;
-			}
-		}
-	}
-
-	bool bFirstPlatform = true;
-	for (const ITargetPlatform* TargetPlatform : RequestedPlatforms)
-	{
-		if (bOutIterativelyUnmodified)
-		{
-			PackageData->FindOrAddPlatformData(TargetPlatform).SetIterativelyUnmodified(true);
-		}
-		bool bShouldIterativelySkip = bOutIterativelyUnmodified;
-		ICookedPackageWriter& PackageWriter = COTFS.FindOrCreatePackageWriter(TargetPlatform);
-		PackageWriter.UpdatePackageModificationStatus(PackageData->GetPackageName(), bOutIterativelyUnmodified,
-			bShouldIterativelySkip);
-		if (bShouldIterativelySkip)
-		{
-			PackageData->SetPlatformCooked(TargetPlatform, ECookResult::Succeeded);
-			if (bFirstPlatform)
-			{
-				COOK_STAT(++DetailedCookStats::NumPackagesIterativelySkipped);
-			}
-			// Declare the package to the EDLCookInfo verification so we don't warn about missing exports from it
-			UE::SavePackageUtilities::EDLCookInfoAddIterativelySkippedPackage(PackageData->GetPackageName());
-		}
-		else
-		{
-			COTFS.DeleteOutputForPackage(PackageData->GetPackageName(), TargetPlatform);
-		}
-		bFirstPlatform = false;
-	}
-}
-
 //////////////////////////////////////////////////////////////////////////
 // FPendingCookedPlatformData
 
diff --git a/Engine/Source/Editor/UnrealEd/Private/Cooker/CookPackageData.h b/Engine/Source/Editor/UnrealEd/Private/Cooker/CookPackageData.h
--- a/Engine/Source/Editor/UnrealEd/Private/Cooker/CookPackageData.h
+++ b/Engine/Source/Editor/UnrealEd/Private/Cooker/CookPackageData.h
@@ -3,7 +3,6 @@
 #pragma once
 
 #include "CookTypes.h"
-#include "CookPackageSplitter.h"
 #include "Containers/Array.h"
 #include "Containers/ArrayView.h"
 #include "Containers/Map.h"
@@ -121,29 +120,6 @@ struct FCachedObjectInOuter
 	}
 };
 
-/**
- * Extra information about CachedObjectsInOuter that a GenerationHelper needs to know for diagnostics.
- * The GenerationHelper constructs an associated array (aka TMap) of these structures when it takes
- * over the CachedObjectsInOuter.
- */
-struct FCachedObjectInOuterGeneratorInfo
-{
-public:
-	/** Object->GetFullName() before it was deleted. */
-	FString FullName;
-	/** Has Initialize been called on *this. */
-	bool bInitialized = false;
-	/** Object->GetFlags() had RF_Public when the info was initialized. */
-	bool bPublic = false;
-	/** bMovedRoot is true, or this is a child object of such a moved object. */
-	bool bMoved = false;
-	/** Splitter informed us that object was moved into this package from another package. */
-	bool bMovedRoot = false;
-
-public:
-	void Initialize(UObject* Object);
-};
-
 /** Flags specifying the behavior of FPackageData::SendToState */
 enum class ESendFlags : uint8
 {
@@ -854,326 +830,6 @@ private:
 	uint32 bGeneratedNeedCachedPlatformDataBeforeSplit : 1;
 };
 
-/**
- * Struct used with ICookPackageSplitter, which contains the state information for the save of either the
- * Generator package or one of its Generated packages.
- */
-struct FCookGenerationInfo
-{
-public:
-	/** State variable for reentrant SplitPackage calls */
-	enum class ESaveState : uint8
-	{
-		StartSave = 0,
-
-		ClearOldPackagesFirstAttempt = StartSave,
-		ClearOldPackagesLastAttempt,
-		QueueGeneratedPackages,
-
-		StartPopulate,
-		FinishCachePreMove = StartPopulate,
-		CallObjectsToMove,
-		BeginCacheObjectsToMove,
-		FinishCacheObjectsToMove,
-		CallPopulate,
-		CallGetPostMoveObjects,
-		BeginCachePostMove,
-		FinishCachePostMove,
-
-		ReadyForSave,
-		Last = ReadyForSave,
-	};
-
-	FCookGenerationInfo(FPackageData& InPackageData, bool bInGenerator);
-
-	ESaveState GetSaveState() const { return GeneratorSaveState; }
-	void SetSaveState(ESaveState InValue) { GeneratorSaveState = InValue; }
-	void SetSaveStateComplete(ESaveState CompletedState);
-
-	bool IsCreateAsMap() const { return bCreateAsMap; }
-	void SetIsCreateAsMap(bool bValue) { bCreateAsMap = bValue; }
-	bool HasCreatedPackage() const { return bHasCreatedPackage; }
-	void SetHasCreatedPackage(bool bValue) { bHasCreatedPackage = bValue; }
-	bool HasSaved() const { return bHasSaved; }
-	void SetHasSaved(bool bValue) { bHasSaved = bValue; }
-	bool HasTakenOverCachedCookedPlatformData() const { return bTakenOverCachedCookedPlatformData; }
-	void SetHasTakenOverCachedCookedPlatformData(bool bValue) { bTakenOverCachedCookedPlatformData = bValue; }
-	bool HasIssuedUndeclaredMovedObjectsWarning() const { return bIssuedUndeclaredMovedObjectsWarning; }
-	void SetHasIssuedUndeclaredMovedObjectsWarning(bool bValue) { bIssuedUndeclaredMovedObjectsWarning = bValue; }
-	bool IsGenerator() const { return bGenerator; }
-	void SetIsGenerator(bool bValue) { bGenerator = bValue; }
-
-	/**
-	 * Steal the list of cached objects to call BeginCacheForCookedPlatformData on from the PackageData,
-	 * and add them to this. Also add NewObjects reported by the splitter that will be moved into the package.
-	 */
-	void TakeOverCachedObjectsAndAddMoved(FGenerationHelper& GenerationHelper,
-		TArray<FCachedObjectInOuter>& CachedObjectsInOuter, TArray<UObject*>& MovedObjects);
-	/**
-	 * Fetch all the objects currently in the package and add them to the list of objects that need
-	 * BeginCacheForCookedPlatformData.
-	 * Reports whether new objects were found. If DemotionState is not ESaveState::Last,
-	 * will SetState back to DemotionState
-	 * if new objects were found, and will error exit if this demotion has happened too many times.
-	 */
-	EPollStatus RefreshPackageObjects(FGenerationHelper& GenerationHelper, UPackage* Package,
-		bool& bOutFoundNewObjects, ESaveState DemotionState);
-
-	void AddKeepReferencedPackages(TArray<UPackage*>& InKeepReferencedPackages)
-	{
-		KeepReferencedPackages.Append(InKeepReferencedPackages);
-	}
-
-	/** Create the hash for this generated package, based on dependencies and GenerationHash. */
-	void CreatePackageHash();
-
-	/**
-	 * If the package has iterative results from a previous cook that were not invalidated by dependency changes,
-	 * test whether they need to be invalidated now and clear the results if so. Only called in legacy iterative;
-	 * incremental cooks handle invalidation by querying the TargetDomainDigest during the RequestCluster.
-	 */
-	void IterativeCookValidateOrClear(FGenerationHelper& GeneratorHelper,
-		TConstArrayView<const ITargetPlatform*> RequestedPlatforms, const FIoHash& PreviousHash,
-		bool& bOutIterativelyUnmodified);
-
-	TConstArrayView<FAssetDependency> GetDependencies() const { return PackageDependencies; }
-
-public:
-	FIoHash PackageHash;
-	FString RelativePath;
-	FString GeneratedRootPath;
-	FBlake3Hash GenerationHash;
-	TArray<FAssetDependency> PackageDependencies;
-	FPackageData* PackageData = nullptr;
-	TArray<UPackage*> KeepReferencedPackages;
-	TMap<UObject*, FCachedObjectInOuterGeneratorInfo> CachedObjectsInOuterInfo;
-private:
-	ESaveState GeneratorSaveState = ESaveState::StartSave;
-	bool bCreateAsMap : 1;
-	bool bHasCreatedPackage : 1;
-	bool bHasSaved : 1;
-	bool bTakenOverCachedCookedPlatformData : 1;
-	bool bIssuedUndeclaredMovedObjectsWarning : 1;
-	bool bGenerator : 1;
-};
-
-/**
- * Helper that wraps an ICookPackageSplitter, gets/caches packages to generate, and is a reference-counted
- * collection of cached data and helper functions to save and list the generated packages.
- */
-struct FGenerationHelper : public FThreadSafeRefCountedObject
-{
-public:
-	/** Store the provided CookPackageSplitter and prepare the packages to generate. */
-	FGenerationHelper(UE::Cook::FPackageData& InOwner);
-	~FGenerationHelper();
-
-	/**
-	 * Early exits if already initialized. Otherwise, loads the package if not loaded and searches it for a
-	 * splitter and creates the splitter.
-	 * If load/search/creation fails, the GenerationHelper will be set to invalid @see IsValid.
-	 * 
-	 * Unless otherwise stated, all public functions call Initialize if not already initialized.
-	 */
-	void Initialize();
-	/**
-	 * Version of Initialize that receives the splitterobject from the caller rather than needing to search
-	 * for it, as an optimization for callers that have it already. If already initialized, the input Splitter
-	 * will be destroyed and the function will early exit.
-	 */
-	void Initialize(const UObject* InSplitDataObject,
-		UE::Cook::Private::FRegisteredCookPackageSplitter* InRegisteredSplitterType,
-		TUniquePtr<ICookPackageSplitter>&& InSplitter);
-	/** Version of Initialize that sets IsValid=false. */
-	void InitializeAsInvalid();
-	/**
-	 * Clear all self references which keep this GenerationHelper from destructing. These could have been set by the
-	 * SetKeep... functions. The Owner PackageData will keep a pointer to this until all references from self, from
-	 * child generated PackageDatas, and other, are released. Note this means that we assert in the owner PackageData's
-	 * destructor if any references remain; PackageDatas can not be deleted until all such references are cleared.
-	 * Does not call Initialize.
-	 */
-	void ClearSelfReferences();
-
-	/**
-	 * GenerationHelpers can be created uninitialized at the start of cook, for incremental cooks. We initialize
-	 * them on demand, which requires loading the package if not already loaded. 
-	 * IsInitialized reports whether initialization has been attempted. @see IsValid.
-	 * Does not call Initialize.
-	 */
-	bool IsInitialized() const;
-	/**
-	 * Since GenerationHelpers can be created before we load the package (@see IsInitialized) we might incorrectly
-	 * create one for a package that does not have a splitter. If so, when we load the package and discover that,
-	 * we set status to invalid. When in this state, all public functions are still valid to call, but will be
-	 * noops and will return empty data.
-	 */
-	bool IsValid();
-
-	/** Accessor for the packages to generate, will be empty if invalid or if TryGenerateList was not yet called. */
-	TArrayView<UE::Cook::FCookGenerationInfo> GetPackagesToGenerate();
-	/** Return the GenerationInfo used to save the generator package's UPackage. Does not call Initialize. */
-	UE::Cook::FCookGenerationInfo& GetOwnerInfo();
-	/** Return owner FPackageData. Does not call Initialize. */
-	UE::Cook::FPackageData& GetOwner();
-	/** Return the GenerationInfo for the given PackageData, or null if not found. */
-	UE::Cook::FCookGenerationInfo* FindInfo(const FPackageData& PackageData);
-	const UE::Cook::FCookGenerationInfo* FindInfo(const FPackageData& PackageData) const;
-	/** Return CookPackageSplitter. Will return null if !IsValid.*/
-	ICookPackageSplitter* GetCookPackageSplitterInstance() const;
-	/** Return RegisteredSplitterType. Will return null if !IsValid. */
-	UE::Cook::Private::FRegisteredCookPackageSplitter* GetRegisteredSplitterType() const;
-	/** Return the SplitDataObject's FullObjectPath. Will return NAME_None if !IsValid. */
-	const FName GetSplitDataObjectName() const;
-	/** Return the Splitter's value for virtual bool UseInternalReferenceToAvoidGarbageCollect(). */
-	bool IsUseInternalReferenceToAvoidGarbageCollect() const;
-	/** Return the Splitter's value for virtual bool NeedCachedPlatformDataBeforeSplit(). */
-	bool IsNeedCachedPlatformDataBeforeSplit() const;
-	/** Return the cached pointer to the SplitDataObject. Returns null if no longer in memory or marked as garbage. */
-	UObject* GetWeakSplitDataObject() const;
-	/**
-	 * Return our cached pointer to the SplitDataObject if set. If not set, load the package and find it. Can still
-	 * return null if !IsValid or if not found even after loading the package.
-	 */
-	UObject* FindOrLoadSplitDataObject();
-	
-	/** Find the OwnerPackage in memory, returns null if invalid or not already loaded. Does not call Initialize. */
-	UPackage* GetOwnerPackage();
-	/** Load the OwnerPackage if not already loaded. Can still return null if invalid or package fails to load. */
-	UPackage* FindOrLoadOwnerPackage(UCookOnTheFlyServer& COTFS);
-	/**
-	 * Return a reference to ExternalActors that were discovered during save and stored on this.
-	 * Does not call initialize.
-	 */
-	TConstArrayView<FName> GetExternalActorDependencies();
-	/**
-	 * Return the ExternalActors that were discovered during save and stored on this, and clear the values.
-	 * Does not call initialize.
-	 */
-	TArray<FName> ReleaseExternalActorDependencies();
-
-	/** Call the Splitter's GetGenerateList and create the PackageDatas. Logs errors and returns false on failure. */
-	bool TryGenerateList();
-	/**
-	 * Mark that the SavePackage of the Owner is starting. Keeps a reference to keep the generator alive until save
-	 * is finished.
-	 */
-	void StartOwnerSave();
-	/** Update state before we queue the generated packages, e.g. mark whether packages are iteratively skippable. */
-	void StartQueueGeneratedPackages(UCookOnTheFlyServer& COTFS);
-	/**
-	 * Mark that we have started queuing packages. Automatically called from StartQueueGeneratedPackages, but also
-	 * is called on the director in response to a discovered generated package from a CookWorker. It sets the
-	 * WorkerId so that it is available during assignment of the discovered generated packages.
-	 */
-	void NotifyStartQueueGeneratedPackages(UCookOnTheFlyServer& COTFS, FWorkerId SourceWorkerId);
-	/** Update state after we queue generated packages; e.g. register to stay alive until the packages are assigned. */
-	void EndQueueGeneratedPackages(UCookOnTheFlyServer& COTFS);
-	/** Does not call Initialize. */
-	void EndQueueGeneratedPackagesOnDirector(UCookOnTheFlyServer& COTFS, FWorkerId SourceWorkerId);
-	/**
-	 * Called from Director when the RequestFence added from EndQueueGeneratedPackages has passed. Calls
-	 * OnRequestFencePassed locally and on all Workers.
-	 */
-	void OnRequestFencePassedBroadcast(UCookOnTheFlyServer& COTFS);
-	/** Called from when the RequestFence added from EndQueueGeneratedPackages has passed. */
-	void OnRequestFencePassed(UCookOnTheFlyServer& COTFS);
-	/** Call CreatePackage and set the necessary data on the UPackage, and update status. */
-	UPackage* CreateGeneratedUPackage(FCookGenerationInfo& GenerationInfo, const UPackage* OwnerPackage,
-		const TCHAR* GeneratedPackageName);
-	/**
-	 * Clear any data that should only be held when an FPackageData is in the save state, for the given Info. The given
-	 * Info might specify the generator package or one of the generated packages.
-	 */
-	void ResetSaveState(FCookGenerationInfo& Info, UPackage* Package, UE::Cook::EStateChangeReason ReleaseSaveReason,
-		EPackageState NewState);
-
-	/**
-	 * Record all dependencies from the generator package that are ExternalActor dependencies and
-	 * store them on this.
-	 */
-	void FetchExternalActorDependencies();
-	/** Iterative cook: store the list of generated packages from the last cook. Does not call Initialize. */
-	void SetPreviousGeneratedPackages(TMap<FName, FIoHash>&& Packages);
-	/** Return the information set by SetPreviousGeneratedPackages if not yet cleared. Does not call Initialize. */
-	const TMap<FName, FIoHash>& GetPreviousGeneratedPackages() const;
-
-	/** Callback during garbage collection. Does not call initialize. */
-	void PreGarbageCollect(FCookGenerationInfo& Info, TArray<TObjectPtr<UObject>>& GCKeepObjects,
-		TArray<UPackage*>& GCKeepPackages, TArray<FPackageData*>& GCKeepPackageDatas, bool& bOutShouldDemote);
-	/** Callback during garbage collection. Does not call initialize. */
-	void PostGarbageCollect();
-	/**
-	 * Called from PackageData function of the same name to decide whether to demote the package out of save.
-	 * Does not call Initialize.
-	 */
-	void UpdateSaveAfterGarbageCollect(const FPackageData& PackageData, bool& bInOutDemote);
-
-	// Self-references that keep this GenerationHelper in memory and referenced from the packages that use it until
-	// the self-references are cleared. These Set/Clear functions do not call initialize.
-	void SetKeepForIterative();
-	void ClearKeepForIterative();
-	void SetKeepForQueueResults();
-	void ClearKeepForQueueResults();
-	void SetKeepForGeneratorSave();
-	void ClearKeepForGeneratorSave();
-
-	/**
-	 * Helper for assignment of generated packages in MPCook. Return the id of the CookWorker that saved the
-	 * generator, to decide where to assign the generated packages.
-	 */
-	FWorkerId GetWorkerIdThatSavedGenerator() const;
-	/**
-	 * Helper for assignment of generated packages in MPCook. A counter for assigned generated packages that is
-	 * needed by some assignment schemes.
-	 */
-	int32& GetMPCookNextAssignmentIndex();
-
-	/** Helper function for Initialize and for TryCreateValidGenerationHelper. */
-	static void SearchForRegisteredSplitDataObject(UCookOnTheFlyServer& COTFS, FName PackageName, UPackage* Package,
-		TOptional<TConstArrayView<FCachedObjectInOuter>> CachedObjectsInOuter, UObject*& OutSplitDataObject,
-		UE::Cook::Private::FRegisteredCookPackageSplitter*& OutRegisteredSplitterType,
-		TUniquePtr<ICookPackageSplitter>& OutSplitterInstance);
-	/** Helper function for Initialize and for TryCreateValidGenerationHelper. */
-	static UPackage* FindOrLoadPackage(UCookOnTheFlyServer& COTFS, FPackageData& OwnerPackageData);
-
-private:
-	enum class EInitializeStatus : uint8
-	{
-		Uninitialized,
-		Invalid,
-		Valid,
-	};
-
-private:
-	void ConditionalInitialize() const;
-	void NotifyCompletion(ICookPackageSplitter::ETeardown Status);
-
-private:
-	/** PackageData for the package that is being split */
-	FCookGenerationInfo OwnerInfo;
-	FWeakObjectPtr SplitDataObject;
-	/** Name of the object that prompted the splitter creation */
-	FName SplitDataObjectName;
-	UE::Cook::Private::FRegisteredCookPackageSplitter* RegisteredSplitterType = nullptr;
-	TUniquePtr<ICookPackageSplitter> CookPackageSplitterInstance;
-	/** Recorded list of packages to generate from the splitter, and data we need about them */
-	TArray<FCookGenerationInfo> PackagesToGenerate;
-	TWeakObjectPtr<UPackage> OwnerPackage;
-	TMap<FName, FIoHash> PreviousGeneratedPackages;
-	TArray<FName> ExternalActorDependencies;
-	TRefCountPtr<FGenerationHelper> ReferenceFromKeepForIterative;
-	TRefCountPtr<FGenerationHelper> ReferenceFromKeepForQueueResults;
-	TRefCountPtr<FGenerationHelper> ReferenceFromKeepForGeneratorSave;
-	int32 MPCookNextAssignmentIndex = 0;
-	FWorkerId WorkerIdThatSavedGenerator = FWorkerId::Invalid();
-	EInitializeStatus InitializeStatus = EInitializeStatus::Uninitialized;
-	bool bUseInternalReferenceToAvoidGarbageCollect = false;
-	bool bNeedCachedPlatformDataBeforeSplit = false;
-	bool bGeneratedList = false;
-};
-
-
 /**
  * Stores information about the pending action in response to a single call to BeginCacheForCookedPlatformData that
  * was made on a given object for the given platform, when saving the given PackageData.
@@ -1785,11 +1441,6 @@ inline void FPackageData::GetReachablePlatforms(ArrayType& OutPlatforms) const
 	}
 }
 
-inline TRefCountPtr<UE::Cook::FGenerationHelper> FPackageData::GetGenerationHelper() const
-{
-	return GenerationHelper;
-}
-
 inline bool FPackageData::HasCompletedGeneration() const
 {
 	return static_cast<bool>(bCompletedGeneration);
@@ -1820,11 +1471,6 @@ inline FName FPackageData::GetParentGenerator() const
 	return ParentGenerator;
 }
 
-inline TRefCountPtr<FGenerationHelper> FPackageData::GetParentGenerationHelper() const
-{
-	return ParentGenerationHelper;
-}
-
 inline bool FPackageData::IsGeneratedNeedCachedPlatformDataBeforeSplit() const
 {
 	return bGeneratedNeedCachedPlatformDataBeforeSplit != 0;
@@ -1835,132 +1481,4 @@ inline void FPackageData::SetGeneratedNeedCachedPlatformDataBeforeSplit(bool bVa
 	bGeneratedNeedCachedPlatformDataBeforeSplit = (uint32)bValue;
 }
 
-inline bool FGenerationHelper::IsInitialized() const
-{
-	return InitializeStatus != EInitializeStatus::Uninitialized;
-}
-
-inline void FGenerationHelper::ConditionalInitialize() const
-{
-	if (InitializeStatus == EInitializeStatus::Uninitialized)
-	{
-		// Use a const cast so we can call from getter functions that are otherwise const.
-		const_cast<FGenerationHelper&>(*this).Initialize();
-	}
-}
-
-inline bool FGenerationHelper::IsValid()
-{
-	ConditionalInitialize();
-	return InitializeStatus == EInitializeStatus::Valid;
-}
-
-inline TArrayView<UE::Cook::FCookGenerationInfo> FGenerationHelper::GetPackagesToGenerate()
-{
-	ConditionalInitialize();
-	return PackagesToGenerate;
-}
-
-inline UE::Cook::FCookGenerationInfo& FGenerationHelper::GetOwnerInfo()
-{
-	return OwnerInfo;
-}
-
-inline UE::Cook::FPackageData& FGenerationHelper::GetOwner()
-{
-	return *OwnerInfo.PackageData;
-}
-
-inline ICookPackageSplitter* FGenerationHelper::GetCookPackageSplitterInstance() const
-{
-	ConditionalInitialize();
-	return CookPackageSplitterInstance.Get();
-}
-
-inline UE::Cook::Private::FRegisteredCookPackageSplitter* FGenerationHelper::GetRegisteredSplitterType() const
-{
-	ConditionalInitialize();
-	return RegisteredSplitterType;
-}
-
-inline const FName FGenerationHelper::GetSplitDataObjectName() const
-{
-	ConditionalInitialize();
-	return SplitDataObjectName;
-}
-
-inline bool FGenerationHelper::IsUseInternalReferenceToAvoidGarbageCollect() const
-{
-	ConditionalInitialize();
-	return bUseInternalReferenceToAvoidGarbageCollect;
-}
-
-inline bool FGenerationHelper::IsNeedCachedPlatformDataBeforeSplit() const
-{
-	ConditionalInitialize();
-	return bNeedCachedPlatformDataBeforeSplit;
-}
-
-inline UObject* FGenerationHelper::GetWeakSplitDataObject() const
-{
-	ConditionalInitialize();
-	return SplitDataObject.Get();
-}
-
-inline TConstArrayView<FName> FGenerationHelper::GetExternalActorDependencies()
-{
-	return ExternalActorDependencies;
-}
-
-inline TArray<FName> FGenerationHelper::ReleaseExternalActorDependencies()
-{
-	TArray<FName> Result = MoveTemp(ExternalActorDependencies);
-	return Result;
-}
-
-inline const TMap<FName, FIoHash>& FGenerationHelper::GetPreviousGeneratedPackages() const
-{
-	return PreviousGeneratedPackages;
-}
-
-inline void FGenerationHelper::SetKeepForIterative()
-{
-	ReferenceFromKeepForIterative = this;
-}
-
-inline void FGenerationHelper::ClearKeepForIterative()
-{
-	ReferenceFromKeepForIterative.SafeRelease();
-}
-
-inline void FGenerationHelper::SetKeepForQueueResults()
-{
-	ReferenceFromKeepForQueueResults = this;
-}
-
-inline void FGenerationHelper::ClearKeepForQueueResults()
-{
-	ReferenceFromKeepForQueueResults.SafeRelease();
-}
-
-inline void FGenerationHelper::SetKeepForGeneratorSave()
-{
-	ReferenceFromKeepForGeneratorSave = this;
-}
-
-inline void FGenerationHelper::ClearKeepForGeneratorSave()
-{
-	ReferenceFromKeepForGeneratorSave.SafeRelease();
-}
-
-inline FWorkerId FGenerationHelper::GetWorkerIdThatSavedGenerator() const
-{
-	return WorkerIdThatSavedGenerator;
-}
-
-inline int32& FGenerationHelper::GetMPCookNextAssignmentIndex()
-{
-	return MPCookNextAssignmentIndex;
-}
-
 } // namespace UE::Cook
\ No newline at end of file
diff --git a/Engine/Source/Editor/UnrealEd/Private/Cooker/CookWorkerClient.cpp b/Engine/Source/Editor/UnrealEd/Private/Cooker/CookWorkerClient.cpp
--- a/Engine/Source/Editor/UnrealEd/Private/Cooker/CookWorkerClient.cpp
+++ b/Engine/Source/Editor/UnrealEd/Private/Cooker/CookWorkerClient.cpp
@@ -4,6 +4,7 @@
 
 #include "Cooker/CompactBinaryTCP.h"
 #include "Cooker/CookDirector.h"
+#include "Cooker/CookGenerationHelper.h"
 #include "Cooker/CookPackageData.h"
 #include "Cooker/CookPlatformManager.h"
 #include "Cooker/CookTypes.h"
diff --git a/Engine/Source/Editor/UnrealEd/Private/Cooker/CookWorkerServer.cpp b/Engine/Source/Editor/UnrealEd/Private/Cooker/CookWorkerServer.cpp
--- a/Engine/Source/Editor/UnrealEd/Private/Cooker/CookWorkerServer.cpp
+++ b/Engine/Source/Editor/UnrealEd/Private/Cooker/CookWorkerServer.cpp
@@ -6,6 +6,7 @@
 #include "Commandlets/AssetRegistryGenerator.h"
 #include "Cooker/CompactBinaryTCP.h"
 #include "Cooker/CookDirector.h"
+#include "Cooker/CookGenerationHelper.h"
 #include "Cooker/CookPackageData.h"
 #include "Cooker/CookPlatformManager.h"
 #include "HAL/PlatformProcess.h"
diff --git a/Engine/Source/Editor/UnrealEd/Private/Cooker/WorkerRequestsLocal.cpp b/Engine/Source/Editor/UnrealEd/Private/Cooker/WorkerRequestsLocal.cpp
--- a/Engine/Source/Editor/UnrealEd/Private/Cooker/WorkerRequestsLocal.cpp
+++ b/Engine/Source/Editor/UnrealEd/Private/Cooker/WorkerRequestsLocal.cpp
@@ -3,6 +3,7 @@
 #include "WorkerRequestsLocal.h"
 
 #include "CookTypes.h"
+#include "Cooker/CookGenerationHelper.h"
 #include "Cooker/CookPackageData.h"
 #include "Cooker/CookRequests.h"
 #include "HAL/Event.h"
