Fix cases where render resources accumulate despite having a visible realtime viewport.

UpdateSingleViewportClient doesn't guarantee that a viewport will be updated, even if it's realtime (e.g. when FSlateThrottleManager::Get().IsAllowingExpensiveTasks() returns false, which happens if a popup menu is shown). Instead of trying to guess what it did at the call site, this CL makes it return that decision to the caller. (Note that the return value does not mean a viewport was updated, it returns true only if it updated a non-realtime viewport. That should have been an output param as well, but it's too late to change now because it's a public function.)

#rnx

diff --git a/Engine/Source/Editor/UnrealEd/Classes/Editor/EditorEngine.h b/Engine/Source/Editor/UnrealEd/Classes/Editor/EditorEngine.h
--- a/Engine/Source/Editor/UnrealEd/Classes/Editor/EditorEngine.h
+++ b/Engine/Source/Editor/UnrealEd/Classes/Editor/EditorEngine.h
@@ -873,9 +873,10 @@ public:
 	 * @param Viewport - the viewport that we're trying to draw
 	 * @param bInAllowNonRealtimeViewportToDraw - whether or not to allow non-realtime viewports to update
 	 * @param bLinkedOrthoMovement	True if orthographic viewport movement is linked
+	 * @param bOutViewportDrawn If non-null, will be set to true if the viewport was drawn in this call.
 	 * @return - Whether a NON-realtime viewport has updated in this call.  Used to help time-slice canvas redraws
 	 */
-	UNREALED_API bool UpdateSingleViewportClient(FEditorViewportClient* InViewportClient, const bool bInAllowNonRealtimeViewportToDraw, bool bLinkedOrthoMovement );
+	UNREALED_API bool UpdateSingleViewportClient(FEditorViewportClient* InViewportClient, const bool bInAllowNonRealtimeViewportToDraw, bool bLinkedOrthoMovement, bool* bOutViewportDrawn = nullptr );
 
 	/** Used for generating status bar text */
 	enum EMousePositionType
diff --git a/Engine/Source/Editor/UnrealEd/Private/EditorEngine.cpp b/Engine/Source/Editor/UnrealEd/Private/EditorEngine.cpp
--- a/Engine/Source/Editor/UnrealEd/Private/EditorEngine.cpp
+++ b/Engine/Source/Editor/UnrealEd/Private/EditorEngine.cpp
@@ -2202,11 +2202,12 @@ void UEditorEngine::Tick( float DeltaSeconds, bool bIdleMode )
 			{
 				bool bAllowNonRealtimeViewports = true;
 				GCurrentLevelEditingViewportClient->SetIsCurrentLevelEditingFocus(true);
-				bool bWasNonRealtimeViewportDrawn = UpdateSingleViewportClient(GCurrentLevelEditingViewportClient, bAllowNonRealtimeViewports, bUpdateLinkedOrthoViewports);
+				bool bViewportDrawn;
+				bool bWasNonRealtimeViewportDrawn = UpdateSingleViewportClient(GCurrentLevelEditingViewportClient, bAllowNonRealtimeViewports, bUpdateLinkedOrthoViewports, &bViewportDrawn);
 				if (GCurrentLevelEditingViewportClient->IsLevelEditorClient())
 				{
 					bEditorFrameNonRealtimeViewportDrawn |= bWasNonRealtimeViewportDrawn;
-					bAnyLevelEditorsDrawn |= GCurrentLevelEditingViewportClient->IsRealtime() || bWasNonRealtimeViewportDrawn;
+					bAnyLevelEditorsDrawn |= bViewportDrawn;
 				}
 			}
 		}
@@ -2230,11 +2231,12 @@ void UEditorEngine::Tick( float DeltaSeconds, bool bIdleMode )
 						//if we haven't drawn a non-realtime viewport OR not one of the main viewports
 						bool bAllowNonRealtimeViewports = (!bEditorFrameNonRealtimeViewportDrawn) || !(ViewportClient->IsLevelEditorClient());
 						ViewportClient->SetIsCurrentLevelEditingFocus(true);
-						bool bWasNonRealtimeViewportDrawn = UpdateSingleViewportClient(ViewportClient, bAllowNonRealtimeViewports, bUpdateLinkedOrthoViewports);
+						bool bViewportDrawn;
+						bool bWasNonRealtimeViewportDrawn = UpdateSingleViewportClient(ViewportClient, bAllowNonRealtimeViewports, bUpdateLinkedOrthoViewports, &bViewportDrawn);
 						if (ViewportClient->IsLevelEditorClient())
 						{
 							bEditorFrameNonRealtimeViewportDrawn |= bWasNonRealtimeViewportDrawn;
-							bAnyLevelEditorsDrawn |= ViewportClient->IsRealtime() || bWasNonRealtimeViewportDrawn;
+							bAnyLevelEditorsDrawn |= bViewportDrawn;
 						}
 					}
 				}
@@ -2444,9 +2446,10 @@ void UEditorEngine::SetRealTimeAudioVolume(float VolumeLevel)
 	LevelEditorMiscSettings->PostEditChange();
 }
 
-bool UEditorEngine::UpdateSingleViewportClient(FEditorViewportClient* InViewportClient, const bool bInAllowNonRealtimeViewportToDraw, bool bLinkedOrthoMovement )
+bool UEditorEngine::UpdateSingleViewportClient(FEditorViewportClient* InViewportClient, const bool bInAllowNonRealtimeViewportToDraw, bool bLinkedOrthoMovement, bool* bOutViewportDrawn /*= nullptr*/)
 {
 	bool bUpdatedNonRealtimeViewport = false;
+	bool bViewportDrawn = false;
 
 	if (InViewportClient->Viewport->IsSlateViewport())
 	{
@@ -2506,6 +2509,7 @@ bool UEditorEngine::UpdateSingleViewportClient(FEditorViewportClient* InViewport
 			InViewportClient->Viewport->Draw();
 			InViewportClient->bNeedsRedraw = false;
 			InViewportClient->bNeedsLinkedRedraw = false;
+			bViewportDrawn = true;
 		}
 		// Redraw any linked ortho viewports that need to be updated this frame.
 		else if( InViewportClient->IsOrtho() && bLinkedOrthoMovement && InViewportClient->IsVisible() )
@@ -2516,6 +2520,7 @@ bool UEditorEngine::UpdateSingleViewportClient(FEditorViewportClient* InViewport
 				InViewportClient->Viewport->Draw();
 				InViewportClient->bNeedsLinkedRedraw = false;
 				InViewportClient->bNeedsRedraw = false;
+				bViewportDrawn = true;
 			}
 			else
 			{
@@ -2528,6 +2533,7 @@ bool UEditorEngine::UpdateSingleViewportClient(FEditorViewportClient* InViewport
 		{
 			InViewportClient->Viewport->Draw();
 			InViewportClient->bNeedsRedraw = false;
+			bViewportDrawn = true;
 			bUpdatedNonRealtimeViewport = true;
 		}
 		else if(UWorld* World = GetWorld())
@@ -2551,6 +2557,11 @@ bool UEditorEngine::UpdateSingleViewportClient(FEditorViewportClient* InViewport
 		}
 	}
 
+	if (bOutViewportDrawn)
+	{
+		*bOutViewportDrawn = bViewportDrawn;
+	}
+
 	return bUpdatedNonRealtimeViewport;
 }
 
