Cooker: CookPackageSplitter: Fix deadlock during cooker shutdown if RequestFencePackageListeners was non-empty.
1) FPackageDatas::Clear was processing notifications while inside a writelock on ExistenceLock, but the notifications need a readlock on the ExistenceLock. Remove the processing of notifications during Clear; it is not valid anyway because notification handlers assume the cook is still in progress.
2) RequestFencePackageListeners were not being processed as soon as possible unless there were other requests to explore, which caused the RequestFencePackageListeners to still be non-empty at end of cook and also caused FGenerationHelper memory to not be freed until end of cook. Fixed by adding RequestFencePackageListeners to HasRequestsToExplore.
3) A generator retracted during save, after queue packages, would queuepackages again on another cookworker, but would lose its iterativecook information since that was torn down as soon as QueueGeneratedPackages was done. Fixed by tearing down FGenerationHelper iterative cook at end of save rather than end of QueueGeneratedPackages.
#rnx

diff --git a/Engine/Source/Editor/UnrealEd/Private/Cooker/CookGenerationHelper.cpp b/Engine/Source/Editor/UnrealEd/Private/Cooker/CookGenerationHelper.cpp
--- a/Engine/Source/Editor/UnrealEd/Private/Cooker/CookGenerationHelper.cpp
+++ b/Engine/Source/Editor/UnrealEd/Private/Cooker/CookGenerationHelper.cpp
@@ -446,6 +446,8 @@ void FGenerationHelper::StartQueueGeneratedPackages(UCookOnTheFlyServer& COTFS)
 	bool bHybridIterativeEnabled = COTFS.bHybridIterativeEnabled;
 	if (!PreviousGeneratedPackages.IsEmpty())
 	{
+		TMap<FName, FIoHash> RemainingPreviousPackages = PreviousGeneratedPackages;
+
 		FPackageData& OwnerPackageData = GetOwner();
 		TArray<const ITargetPlatform*, TInlineAllocator<1>> PlatformsToCook;
 		OwnerPackageData.GetPlatformsNeedingCooking(PlatformsToCook);
@@ -453,7 +455,7 @@ void FGenerationHelper::StartQueueGeneratedPackages(UCookOnTheFlyServer& COTFS)
 		int32 NumIterativeUnmodified = 0;
 		int32 NumIterativeModified = 0;
 		int32 NumIterativeRemoved = 0;
-		int32 NumIterativePrevious = PreviousGeneratedPackages.Num();
+		int32 NumIterativePrevious = RemainingPreviousPackages.Num();
 
 		for (FCookGenerationInfo& GeneratedInfo : PackagesToGenerate)
 		{
@@ -462,7 +464,7 @@ void FGenerationHelper::StartQueueGeneratedPackages(UCookOnTheFlyServer& COTFS)
 				continue;
 			}
 			FIoHash PreviousHash;
-			if (PreviousGeneratedPackages.RemoveAndCopyValue(GeneratedInfo.PackageData->GetPackageName(), PreviousHash)
+			if (RemainingPreviousPackages.RemoveAndCopyValue(GeneratedInfo.PackageData->GetPackageName(), PreviousHash)
 				&& !bHybridIterativeEnabled)
 			{
 				bool bIterativelyUnmodified;
@@ -470,19 +472,17 @@ void FGenerationHelper::StartQueueGeneratedPackages(UCookOnTheFlyServer& COTFS)
 				++(bIterativelyUnmodified ? NumIterativeUnmodified : NumIterativeModified);
 			}
 		}
-		if (!PreviousGeneratedPackages.IsEmpty())
+		if (!RemainingPreviousPackages.IsEmpty())
 		{
-			NumIterativeRemoved = PreviousGeneratedPackages.Num();
-			for (TPair<FName, FIoHash>& Pair : PreviousGeneratedPackages)
+			NumIterativeRemoved = RemainingPreviousPackages.Num();
+			for (TPair<FName, FIoHash>& Pair : RemainingPreviousPackages)
 			{
 				for (const ITargetPlatform* TargetPlatform : PlatformsToCook)
 				{
 					COTFS.DeleteOutputForPackage(Pair.Key, TargetPlatform);
 				}
 			}
-			PreviousGeneratedPackages.Empty();
 		}
-		ClearKeepForIterative();
 
 		if (NumIterativePrevious > 0 && !bHybridIterativeEnabled)
 		{
@@ -646,6 +646,12 @@ void FGenerationHelper::ResetSaveState(FCookGenerationInfo& Info, UPackage* Pack
 
 		if (Info.IsGenerator())
 		{
+			// Now that we've finished saving, we know that we will not call QueueGeneratedPackages again, so we can
+			// teardown iterative results as well
+			ClearKeepForIterative();
+			PreviousGeneratedPackages.Empty();
+
+			// And also teardown data needed during save
 			ClearKeepForGeneratorSave();
 			Info.SetSaveState(FCookGenerationInfo::ESaveState::StartSave);
 		}
diff --git a/Engine/Source/Editor/UnrealEd/Private/Cooker/CookPackageData.cpp b/Engine/Source/Editor/UnrealEd/Private/Cooker/CookPackageData.cpp
--- a/Engine/Source/Editor/UnrealEd/Private/Cooker/CookPackageData.cpp
+++ b/Engine/Source/Editor/UnrealEd/Private/Cooker/CookPackageData.cpp
@@ -108,6 +108,10 @@ FPackageData::~FPackageData()
 {
 	// ClearReferences should have been called earlier, but call it here in case it was missed
 	ClearReferences();
+	// FPackageDatas guarantees that all references to GenerationHelper are removed before any PackageDatas are deleted.
+	// We rely on that so that we can be sure that when this PackageData is being deleted, its GenerationHelper - which
+	// assumes the FPackageData lifetime exceeds its own - has already been deleted.
+	check(GenerationHelper == nullptr);
 	// We need to send OnLastCookedPlatformRemoved message to the monitor, so call SetPlatformsNotCooked
 	ClearCookResults();
 	// Update the monitor's counters and call exit functions
@@ -2539,7 +2543,7 @@ void FPackageDatas::Clear()
 {
 	FWriteScopeLock ExistenceWriteLock(ExistenceLock);
 	PendingCookedPlatformDataLists.Empty(); // These destructors will read/write PackageDatas
-	RequestQueue.Empty(*this);
+	RequestQueue.Empty();
 	SaveQueue.Empty();
 	PackageNameToPackageData.Empty();
 	FileNameToPackageData.Empty();
@@ -2776,12 +2780,12 @@ void FPackageDatas::DebugInstigator(FPackageData& PackageData)
 	UpdateThreadsafePackageData(PackageData);
 }
 
-void FRequestQueue::Empty(FPackageDatas& PackageDatas)
+void FRequestQueue::Empty()
 {
 	RestartedRequests.Empty();
 	DiscoveryQueue.Empty();
 	RequestClusters.Empty();
-	NotifyRequestFencePassed(PackageDatas);
+	RequestFencePackageListeners.Empty();
 	NormalRequests.Empty();
 	UrgentRequests.Empty();
 }
@@ -2860,7 +2864,8 @@ bool FRequestQueue::IsReadyRequestsEmpty() const
 
 bool FRequestQueue::HasRequestsToExplore() const
 {
-	return !RequestClusters.IsEmpty() | !RestartedRequests.IsEmpty() | !DiscoveryQueue.IsEmpty();
+	return !RequestClusters.IsEmpty() | !RestartedRequests.IsEmpty() | !DiscoveryQueue.IsEmpty()
+		| !RequestFencePackageListeners.IsEmpty();
 }
 
 uint32 FRequestQueue::ReadyRequestsNum() const
diff --git a/Engine/Source/Editor/UnrealEd/Private/Cooker/CookPackageData.h b/Engine/Source/Editor/UnrealEd/Private/Cooker/CookPackageData.h
--- a/Engine/Source/Editor/UnrealEd/Private/Cooker/CookPackageData.h
+++ b/Engine/Source/Editor/UnrealEd/Private/Cooker/CookPackageData.h
@@ -999,7 +999,7 @@ public:
 	uint32 Remove(FPackageData* PackageData);
 	bool Contains(const FPackageData* PackageData) const;
 	bool DiscoveryQueueContains(FPackageData* PackageData) const;
-	void Empty(FPackageDatas& PackageDatas);
+	void Empty();
 
 	void AddRequest(FPackageData* PackageData, bool bForceUrgent=false);
 
