Don't have the cook director close the worker socket before forcible worker process termination and detach because otherwise DetachFromRemoteProcess won't call OnMultiprocessWorkerDetached for the worker leading to issues like DerivedDataReplay emitting in error because a worker is not detached and it can't merge replay files.  Also have the multi-process detach context inform the delegate of whether the detach was abnormal so that delegate methods can make context-appropriate choices about how to handle abnormal detach scenarios.  The derived data replay mechanism will ignore replays from abnormally detached processes when doing replay merging and warn about the fact that the replay is incomplete.

diff --git a/Engine/Source/Developer/DerivedDataCache/Private/DerivedDataCacheReplay.cpp b/Engine/Source/Developer/DerivedDataCache/Private/DerivedDataCacheReplay.cpp
--- a/Engine/Source/Developer/DerivedDataCache/Private/DerivedDataCacheReplay.cpp
+++ b/Engine/Source/Developer/DerivedDataCache/Private/DerivedDataCacheReplay.cpp
@@ -306,7 +306,10 @@ void FCacheStoreReplay::WorkerCreated(const FMultiprocessCreatedContext& Context
 
 void FCacheStoreReplay::WorkerDetached(const FMultiprocessDetachedContext& Context)
 {
-	WorkerIdToState.FindOrAdd(Context.Id) = false;
+	if (!Context.bAbnormalDetach)
+	{
+		WorkerIdToState.FindOrAdd(Context.Id) = false;
+	}
 }
 
 void FCacheStoreReplay::MergeWorkerReplays()
@@ -322,7 +325,7 @@ void FCacheStoreReplay::MergeWorkerReplays()
 		WorkerReplayPath.Appendf(TEXT(".worker%d.tmp"), WorkerId);
 		if (Worker.Value)
 		{
-			UE_LOG(LogDerivedDataCache, Error,
+			UE_LOG(LogDerivedDataCache, Warning,
 				TEXT("Replay: Skipped replay file '%s' because its worker has not detached."), *WorkerReplayPath);
 			continue;
 		}
diff --git a/Engine/Source/Editor/UnrealEd/Private/Cooker/CookWorkerServer.cpp b/Engine/Source/Editor/UnrealEd/Private/Cooker/CookWorkerServer.cpp
--- a/Engine/Source/Editor/UnrealEd/Private/Cooker/CookWorkerServer.cpp
+++ b/Engine/Source/Editor/UnrealEd/Private/Cooker/CookWorkerServer.cpp
@@ -47,14 +47,14 @@ FCookWorkerServer::~FCookWorkerServer()
 			TEXT("CookWorkerServer %d was destroyed before it finished Disconnect. The remote process may linger and may interfere with writes of future packages."),
 			ProfileId);
 	}
-	DetachFromRemoteProcess();
+	DetachFromRemoteProcess(EWorkerDetachType::StillRunning);
 }
 
-void FCookWorkerServer::DetachFromRemoteProcess()
+void FCookWorkerServer::DetachFromRemoteProcess(EWorkerDetachType DetachType)
 {
 	if (Socket != nullptr)
 	{
-		FCoreDelegates::OnMultiprocessWorkerDetached.Broadcast({WorkerId.GetMultiprocessId()});
+		FCoreDelegates::OnMultiprocessWorkerDetached.Broadcast({WorkerId.GetMultiprocessId(), DetachType != EWorkerDetachType::Dismissed});
 	}
 	Sockets::CloseSocket(Socket);
 	CookWorkerHandle = FProcHandle();
@@ -226,12 +226,13 @@ void FCookWorkerServer::SendCrashDiagnostics()
 
 void FCookWorkerServer::ShutdownRemoteProcess()
 {
-	Sockets::CloseSocket(Socket);
+	EWorkerDetachType DetachType = EWorkerDetachType::Dismissed;
 	if (CookWorkerHandle.IsValid())
 	{
 		FPlatformProcess::TerminateProc(CookWorkerHandle, /* bKillTree */true);
+		DetachType = EWorkerDetachType::ForceTerminated;
 	}
-	DetachFromRemoteProcess();
+	DetachFromRemoteProcess(DetachType);
 }
 
 void FCookWorkerServer::AppendAssignments(TArrayView<FPackageData*> Assignments,
@@ -351,7 +352,7 @@ void FCookWorkerServer::SendToState(EConnectStatus TargetStatus)
 		ConnectTestStartTimeSeconds = ConnectStartTimeSeconds;
 		break;
 	case EConnectStatus::LostConnection:
-		DetachFromRemoteProcess();
+		DetachFromRemoteProcess(bNeedCrashDiagnostics ? EWorkerDetachType::Crashed : EWorkerDetachType::Dismissed);
 		break;
 	default:
 		break;
diff --git a/Engine/Source/Editor/UnrealEd/Private/Cooker/CookWorkerServer.h b/Engine/Source/Editor/UnrealEd/Private/Cooker/CookWorkerServer.h
--- a/Engine/Source/Editor/UnrealEd/Private/Cooker/CookWorkerServer.h
+++ b/Engine/Source/Editor/UnrealEd/Private/Cooker/CookWorkerServer.h
@@ -135,6 +135,13 @@ private:
 		WaitForDisconnect,
 		LostConnection,
 	};
+	enum class EWorkerDetachType
+	{
+		Dismissed,
+		StillRunning,
+		ForceTerminated,
+		Crashed,
+	};
 	enum class ETickAction
 	{
 		Tick,
@@ -183,7 +190,7 @@ private:
 	/** Send this into the given state. Update any state-dependent variables. */
 	void SendToState(EConnectStatus TargetStatus);
 	/** Close the connection and connection resources to the remote process. Does not kill the process. */
-	void DetachFromRemoteProcess();
+	void DetachFromRemoteProcess(EWorkerDetachType DetachType);
 	/** Report from the log of the crashed CookWorker. */
 	void SendCrashDiagnostics();
 	/** Kill the Client process (non-graceful termination), and close the connection resources. */
diff --git a/Engine/Source/Runtime/Core/Public/Misc/CoreDelegates.h b/Engine/Source/Runtime/Core/Public/Misc/CoreDelegates.h
--- a/Engine/Source/Runtime/Core/Public/Misc/CoreDelegates.h
+++ b/Engine/Source/Runtime/Core/Public/Misc/CoreDelegates.h
@@ -80,6 +80,7 @@ struct FMultiprocessCreatedContext
 struct FMultiprocessDetachedContext
 {
 	int32 Id;
+	bool bAbnormalDetach;
 };
 }
 
