Suppressible Modal Dialog: Ensure the result of "Don't show this again" are saved in the ini file.

diff --git a/Engine/Source/Editor/UnrealEd/Private/Dialogs/Dialogs.cpp b/Engine/Source/Editor/UnrealEd/Private/Dialogs/Dialogs.cpp
--- a/Engine/Source/Editor/UnrealEd/Private/Dialogs/Dialogs.cpp
+++ b/Engine/Source/Editor/UnrealEd/Private/Dialogs/Dialogs.cpp
@@ -918,6 +918,7 @@ FSuppressableWarningDialog::EResult FSuppressableWarningDialog::ShowModal() cons
 				{
 					GConfig->SetBool(*ConfigSection, *ResponseIniSettingName, true, IniSettingFileName);
 				}
+				GConfig->Flush(false, IniSettingFileName);
 			}
 		}
 		else
@@ -926,6 +927,7 @@ FSuppressableWarningDialog::EResult FSuppressableWarningDialog::ShowModal() cons
 			{
 				GConfig->SetBool(*ConfigSection, *IniSettingName, bShouldSuppressDialog, IniSettingFileName);
 				GConfig->SetBool(*ConfigSection, *ResponseIniSettingName, false, IniSettingFileName);
+				GConfig->Flush(false, IniSettingFileName);
 			}
 		}
 	}
