Skein / ConflictUE: Remove 'check out failed dialog', opting to always save locally instead.

diff --git a/Engine/Source/Editor/UnrealEd/Private/FileHelpers.cpp b/Engine/Source/Editor/UnrealEd/Private/FileHelpers.cpp
--- a/Engine/Source/Editor/UnrealEd/Private/FileHelpers.cpp
+++ b/Engine/Source/Editor/UnrealEd/Private/FileHelpers.cpp
@@ -4786,60 +4786,13 @@ bool FEditorFileUtils::AutomaticCheckoutOrPromptToRevertPackages(const TArray<UP
 		if (PackagesCheckOutImpossible.Num() > 0 || PackagesCheckOutFailure.Num() > 0)
 		{
 			// No.
-			// Show dialog with assets that weren't checked out.
-
-			const FText DialogTitle = NSLOCTEXT("PackagesDialogModule", "CheckoutPackagesFailedDialogTitle", "Check-out failed");
-			const FText DialogHeading = NSLOCTEXT("PackagesDialogModule", "CheckoutPackagesFailedDialogMessage",
-				"While saving, check-out failed for the following assets. Revert your changes to these assets and sync to the latest snapshot to avoid conflicts with your teammates.\r\n\r\n"
-				"If necessary, you may also proceed by saving locally only, but you will likely run into conflicts later when trying to check in these changes.\r\n\r\n"
-				"Tip: Turn on automatic checkout and automatic undo in your Unreal Revision Control settings to avoid future conflicts with your teammates and conflict warnings."
-			);
-
-			FPackagesDialogModule& CheckoutPackagesDialogModule = FModuleManager::LoadModuleChecked<FPackagesDialogModule>(TEXT("PackagesDialog"));
-			CheckoutPackagesDialogModule.CreatePackagesDialog(
-				DialogTitle,
-				DialogHeading,
-				/*InReadOnly=*/true,
-				/*InAllowSourceControlConnection*/true
-			);
+			// Make the packages writable and proceed to save.
 
 			TArray<UPackage*> PackagesNotCheckedOut;
 			PackagesNotCheckedOut.Append(PackagesCheckOutFailure);
 			PackagesNotCheckedOut.Append(PackagesCheckOutImpossible);
 
-			for (UPackage* Package : PackagesNotCheckedOut)
-			{
-				FSourceControlStatePtr State = SourceControlProvider.GetState(Package, EStateCacheUsage::Use);
-				if (!State->IsCurrent())
-				{
-					CheckoutPackagesDialogModule.AddPackageItem(Package, ECheckBoxState::Unchecked, true, TEXT("SavePackages.SCC_DlgNotCurrent"), State->GetDisplayTooltip().ToString());
-				}
-				else if (State->IsCheckedOutOther())
-				{
-					CheckoutPackagesDialogModule.AddPackageItem(Package, ECheckBoxState::Unchecked, true, TEXT("SavePackages.SCC_DlgCheckedOutOther"), State->GetDisplayTooltip().ToString());
-				}
-				else
-				{
-					CheckoutPackagesDialogModule.AddPackageItem(Package, ECheckBoxState::Unchecked, true, TEXT("SavePackages.SCC_DlgNoIcon"), State->GetDisplayTooltip().ToString());
-				}
-			}
-
-			// The Revert button will allow the user to undo changes to those assets.
-			CheckoutPackagesDialogModule.AddButton(DRT_Revert, DBS_Primary, NSLOCTEXT("PackagesDialogModule", "Dlg_RevertButton", "Revert My Changes"), NSLOCTEXT("PackagesDialogModule", "Dlg_RevertButtonTooltip", "Revert changes to files that could not be checked out (recommended)."));
-
-			// The Save button will allow the user to proceed with saving those assets anyway, thereby risking conflicts.
-			CheckoutPackagesDialogModule.AddButton(DRT_Save, DBS_Normal, NSLOCTEXT("PackagesDialogModule", "Dlg_SaveButton", "Save Locally Only"), NSLOCTEXT("PackagesDialogModule", "Dlg_SaveButtonTooltip", "Save changes to files that could not be checked out anyway. You will likely be unable to check-in these changes."));
-
-			EDialogReturnType UserResponse = CheckoutPackagesDialogModule.ShowPackagesDialog();
-			if (UserResponse == DRT_Revert)
-			{
-				PackagesToRevert = PackagesNotCheckedOut;
-			}
-			if (UserResponse == DRT_Save)
-			{
-				// Make the packages writable and proceed to save.
-				MakePackagesWritable(PackagesNotCheckedOut, &PackagesWritableSuccess, &PackagesWritableFailure);
-			}
+			MakePackagesWritable(PackagesNotCheckedOut, &PackagesWritableSuccess, &PackagesWritableFailure);
 		}
 		else
 		{
