PR #11566:  Added support for custom thumbnail renderers for blueprints based on non-actor classes

#pf 6606f1f8177e09365fda85c6

diff --git a/Engine/Source/Editor/UnrealEd/Classes/ThumbnailRendering/ThumbnailManager.h b/Engine/Source/Editor/UnrealEd/Classes/ThumbnailRendering/ThumbnailManager.h
--- a/Engine/Source/Editor/UnrealEd/Classes/ThumbnailRendering/ThumbnailManager.h
+++ b/Engine/Source/Editor/UnrealEd/Classes/ThumbnailRendering/ThumbnailManager.h
@@ -71,6 +71,12 @@ struct FThumbnailRenderingInfo
 	UPROPERTY()
 	TObjectPtr<class UThumbnailRenderer> Renderer;
 
+	/**
+	 * Should we get the default object from the specified object before passing it to the renderer?  (only used for blueprints)
+	 */
+	UPROPERTY()
+	bool bUseClassDefaultObject = false;
+
 public:
 	FThumbnailRenderingInfo()
 		: ClassNeedingThumbnail(NULL)
@@ -237,4 +243,3 @@ private:
 };
 
 
-
diff --git a/Engine/Source/Editor/UnrealEd/Private/ObjectTools.cpp b/Engine/Source/Editor/UnrealEd/Private/ObjectTools.cpp
--- a/Engine/Source/Editor/UnrealEd/Private/ObjectTools.cpp
+++ b/Engine/Source/Editor/UnrealEd/Private/ObjectTools.cpp
@@ -5075,6 +5075,18 @@ namespace ThumbnailTools
 		// Get the rendering info for this object
 		FThumbnailRenderingInfo* RenderInfo = GUnrealEd ? GUnrealEd->GetThumbnailManager()->GetRenderingInfo( InObject ) : nullptr;
 
+		UObject* ObjectToRender = InObject;
+		if ((RenderInfo != nullptr) && RenderInfo->bUseClassDefaultObject)
+		{
+			if (UBlueprint* Blueprint = Cast<UBlueprint>(InObject))
+			{
+				if (Blueprint->GeneratedClass != nullptr)
+				{
+					ObjectToRender = Blueprint->GeneratedClass->ClassDefaultObject;
+				}
+			}
+		}
+
 		if( InFlushMode == EThumbnailTextureFlushMode::AlwaysFlush )
 		{
 			// Wait for pending load requests.
@@ -5115,7 +5127,7 @@ namespace ThumbnailTools
 					if( RenderInfo->Renderer->IsA( UTextureThumbnailRenderer::StaticClass() ) )
 					{
 						RenderInfo->Renderer->GetThumbnailSize(
-							InObject,
+							ObjectToRender,
 							ZoomFactor,
 							DesiredWidth,		// Out
 							DesiredHeight );	// Out
@@ -5149,7 +5161,7 @@ namespace ThumbnailTools
 			const int32 YPos = 0;
 			const bool bAdditionalViewFamily = false;
 			RenderInfo->Renderer->Draw(
-				InObject,
+				ObjectToRender,
 				XPos,
 				YPos,
 				DrawWidth,
@@ -5814,4 +5826,3 @@ namespace ThumbnailTools
 		return false;
 	}
 		}
-
diff --git a/Engine/Source/Editor/UnrealEd/Private/ThumbnailManager.cpp b/Engine/Source/Editor/UnrealEd/Private/ThumbnailManager.cpp
--- a/Engine/Source/Editor/UnrealEd/Private/ThumbnailManager.cpp
+++ b/Engine/Source/Editor/UnrealEd/Private/ThumbnailManager.cpp
@@ -13,6 +13,7 @@
 #include "SourceControlOperations.h"
 #include "ISourceControlProvider.h"
 #include "ISourceControlModule.h"
+#include "Engine/Blueprint.h"
 #include "Engine/StaticMesh.h"
 #include "UnrealClient.h"
 #include "Engine/TextureCube.h"
@@ -134,34 +135,55 @@ FThumbnailRenderingInfo* UThumbnailManager::GetRenderingInfo(UObject* Object)
 
 	TArray<FThumbnailRenderingInfo>& ThumbnailTypes = RenderableThumbnailTypes;
 
-	// Get the class to check against.
-	UClass *ClassToCheck = Object->GetClass();
-	
-	// Search for the cached entry and do the slower if not found
-	FThumbnailRenderingInfo* RenderInfo = RenderInfoMap.FindRef(ClassToCheck);
-	if (RenderInfo == nullptr)
+	auto FindThumbnailRendererForClass = [&](UObject* Object, UClass* ClassToCheck) -> FThumbnailRenderingInfo*
 	{
-		// Loop through searching for the right thumbnail entry
-		for (int32 Index = ThumbnailTypes.Num() - 1; (Index >= 0) && (RenderInfo == nullptr); Index--)
+		// Search for the cached entry and do the slower if not found
+		FThumbnailRenderingInfo* RenderInfo = RenderInfoMap.FindRef(ClassToCheck);
+		if (RenderInfo == nullptr)
 		{
-			RenderInfo = &ThumbnailTypes[Index];
-
-			// See if this thumbnail renderer will work for the specified class or
-			// if there is some data reason not to render the thumbnail
-			if ((ClassToCheck->IsChildOf(RenderInfo->ClassNeedingThumbnail) == false) || (RenderInfo->Renderer == nullptr))
+			// Loop through searching for the right thumbnail entry
+			for (int32 Index = ThumbnailTypes.Num() - 1; (Index >= 0) && (RenderInfo == nullptr); Index--)
 			{
-				RenderInfo = nullptr;
+				RenderInfo = &ThumbnailTypes[Index];
+
+				// See if this thumbnail renderer will work for the specified class or
+				// if there is some data reason not to render the thumbnail
+				if ((ClassToCheck->IsChildOf(RenderInfo->ClassNeedingThumbnail) == false) || (RenderInfo->Renderer == nullptr))
+				{
+					RenderInfo = nullptr;
+				}
 			}
+
+			// Make sure to add it to the cache if it is missing
+			RenderInfoMap.Add(ClassToCheck, (RenderInfo != nullptr) ? RenderInfo : &NotSupported);
 		}
 
-		// Make sure to add it to the cache if it is missing
-		RenderInfoMap.Add(ClassToCheck, (RenderInfo != nullptr) ? RenderInfo : &NotSupported);
-	}
+		if (RenderInfo && RenderInfo->Renderer && !RenderInfo->Renderer->CanVisualizeAsset(Object))
+		{
+			// This is an asset with a thumbnail renderer, but it can't visualized (i.e it is something like a blueprint that doesn't contain any visible primitive components)
+			RenderInfo = nullptr;
+		}
+
+		return RenderInfo;
+	};
 
-	if (RenderInfo && RenderInfo->Renderer && !RenderInfo->Renderer->CanVisualizeAsset(Object))
+	FThumbnailRenderingInfo* RenderInfo = FindThumbnailRendererForClass(Object, Object->GetClass());
+
+	// For blueprints that the UBlueprint subclass renderer didn't want to deal with, ask the class itself
+	if (RenderInfo == nullptr)
 	{
-		// This is an asset with a thumbnail renderer, but it can't visualized (i.e it is something like a blueprint that doesn't contain any visible primitive components)
-		RenderInfo = nullptr;
+		if (UBlueprint* Blueprint = Cast<UBlueprint>(Object))
+		{
+			if (Blueprint->GeneratedClass != nullptr)
+			{
+				UObject* ClassCDO = Blueprint->GeneratedClass->ClassDefaultObject;
+				RenderInfo = FindThumbnailRendererForClass(ClassCDO, Blueprint->GeneratedClass);
+				if (RenderInfo != nullptr)
+				{
+					RenderInfo->bUseClassDefaultObject = true;
+				}
+			}
+		}
 	}
 
 	// Check to see if this object is the "not supported" type or not
@@ -482,4 +504,3 @@ void UThumbnailManager::DirtyThumbnailForObject(UObject* ObjectBeingModified)
 		OnThumbnailDirtied.Broadcast(FSoftObjectPath(ObjectBeingModified));
 	}
 }
-
