Adding tests for HighestSharedChunk
Fixing bug for finding highest shared chunk where recursion was not undone.

diff --git a/Engine/Source/Editor/UnrealEd/Private/Commandlets/ChunkDependencyInfo.cpp b/Engine/Source/Editor/UnrealEd/Private/Commandlets/ChunkDependencyInfo.cpp
--- a/Engine/Source/Editor/UnrealEd/Private/Commandlets/ChunkDependencyInfo.cpp
+++ b/Engine/Source/Editor/UnrealEd/Private/Commandlets/ChunkDependencyInfo.cpp
@@ -1,10 +1,45 @@
 // Copyright Epic Games, Inc. All Rights Reserved.
 
 #include "Commandlets/ChunkDependencyInfo.h"
+#include "Algo/Sort.h"
 #include "Algo/Unique.h"
+#include "Containers/Queue.h"
 
 DEFINE_LOG_CATEGORY_STATIC(LogChunkDependencyInfo, Log, All);
 
+namespace ChunkDependencyInfo
+{
+	// breadth first traversal
+	static TArray<int32> BuildTopologicallySortedArray(const FChunkDependencyTreeNode* RootNode)
+	{
+		TArray<int32> OutArray;
+		if (!RootNode)
+		{
+			return OutArray;
+		}
+		TQueue<const FChunkDependencyTreeNode*> Nodes;
+		Nodes.Enqueue(RootNode);
+
+		TSet<const FChunkDependencyTreeNode*> ProcessedNodes;
+		const FChunkDependencyTreeNode* CurrentNode = nullptr;
+		while (!Nodes.IsEmpty())
+		{
+			Nodes.Dequeue(CurrentNode);
+			if (ProcessedNodes.Contains(CurrentNode))
+			{
+				continue;
+			}
+			ProcessedNodes.Add(CurrentNode);
+			OutArray.Add(CurrentNode->ChunkID);
+			for (const FChunkDependencyTreeNode& ChildNode : CurrentNode->ChildNodes)
+			{
+				Nodes.Enqueue(&ChildNode);
+			}
+		}
+
+		return OutArray;
+	}
+}
 
 UChunkDependencyInfo::UChunkDependencyInfo(const FObjectInitializer& ObjectInitializer)
 	: Super(ObjectInitializer)
@@ -32,6 +67,7 @@ const FChunkDependencyTreeNode* UChunkDependencyInfo::BuildChunkDependencyGraph(
 	RootTreeNode.ChildNodes.Reset(0);
 
 	ChildToParentMap.Reset();
+	ChildToParentMap.Add(0);
 	CachedHighestChunk = HighestChunk;
 
 	// Ensure the DependencyArray is OK to work with.
@@ -78,6 +114,7 @@ const FChunkDependencyTreeNode* UChunkDependencyInfo::BuildChunkDependencyGraph(
 	}
 
 	AddChildrenRecursive(RootTreeNode, DependencyArray, TSet<int32>());
+	TopologicallySortedChunks = ChunkDependencyInfo::BuildTopologicallySortedArray(&RootTreeNode);
 	return &RootTreeNode;
 }
 
@@ -124,55 +161,51 @@ void UChunkDependencyInfo::RemoveRedundantChunks(TArray<int32>& ChunkIDs) const
 	}
 }
 
-const FChunkDependencyTreeNode* LowestCommonAncestor(const FChunkDependencyTreeNode* RootNode, TSet<int32>& FoundSet, const TSet<int32>& FindSet)
+int32 UChunkDependencyInfo::FindHighestSharedChunk(const TArray<int32>& ChunkIDs) const
 {
-	if (RootNode)
-	{
-		for (const FChunkDependencyTreeNode& ChildNode : RootNode->ChildNodes)
-		{
-			if (const FChunkDependencyTreeNode* ReturnNode = LowestCommonAncestor(&ChildNode, FoundSet, FindSet))
-			{
-				return ReturnNode;
-			}
-		}
-
-		FoundSet.Add(RootNode->ChunkID);
+	TArray<int32> TestChunkIds;
+	TestChunkIds.Append(ChunkIDs);
+	Algo::Sort(TestChunkIds);
+	TestChunkIds.SetNum(Algo::Unique(TestChunkIds));
 
-		for (int32 Item : FindSet)
+	for (int32 ChunkId : TestChunkIds)
+	{
+		if (!ChildToParentMap.Contains(ChunkId))
 		{
-			if (!FoundSet.Contains(Item))
-			{
-				return nullptr;
-			}
+			return INDEX_NONE;
 		}
-
-		return RootNode;
 	}
-	return nullptr;
-}
-
 
-int32 UChunkDependencyInfo::FindHighestSharedChunk(const TArray<int32>& ChunkIDs) const
-{
-	if (ChunkIDs.Num() == 0)
+	if (TestChunkIds.Num() == 0)
 	{
-		return 0;
+		return INDEX_NONE;
+	}
+	if (TestChunkIds.Num() == 1)
+	{
+		return TestChunkIds[0];
 	}
-	if (ChunkIDs.Num() == 1)
+
+	TSet<int32> CommonParentSet;
+	CommonParentSet.Append(ChildToParentMap[TestChunkIds[0]]);
+	CommonParentSet.Add(TestChunkIds[0]);
+
+	for (int32 Index = 1; Index < TestChunkIds.Num(); Index++)
 	{
-		return ChunkIDs[0];
+		TSet<int32> OtherSet;
+		OtherSet.Append(ChildToParentMap[TestChunkIds[Index]]);
+		OtherSet.Add(TestChunkIds[Index]);
+		CommonParentSet = CommonParentSet.Intersect(OtherSet);
 	}
 
-	TSet<int32> FoundSet;
-	TSet<int32> FindSet;
-	FindSet.Append(ChunkIDs);
-	if (const FChunkDependencyTreeNode* BestParent = LowestCommonAncestor(&RootTreeNode, FoundSet, FindSet))
+	int32 HighestIndex = INDEX_NONE;
+	for (int32 ChunkId : CommonParentSet)
 	{
-		return BestParent->ChunkID;
+		HighestIndex = FMath::Max(HighestIndex, TopologicallySortedChunks.IndexOfByKey(ChunkId));
 	}
-	else
+	if (HighestIndex == INDEX_NONE)
 	{
 		UE_LOG(LogChunkDependencyInfo, Error, TEXT("Unable to find parent."));
 		return 0;
 	}
-}
+	return TopologicallySortedChunks[HighestIndex];
+}
\ No newline at end of file
diff --git a/Engine/Source/Editor/UnrealEd/Private/Tests/ChunkDependencyTests.cpp b/Engine/Source/Editor/UnrealEd/Private/Tests/ChunkDependencyTests.cpp
new file mode 100644
--- /dev/null
+++ b/Engine/Source/Editor/UnrealEd/Private/Tests/ChunkDependencyTests.cpp
@@ -0,0 +1,130 @@
+// Copyright Epic Games, Inc. All Rights Reserved.
+
+#include "Misc/AutomationTest.h"
+
+#if WITH_AUTOMATION_TESTS
+
+#include "Commandlets/ChunkDependencyInfo.h"
+
+IMPLEMENT_SIMPLE_AUTOMATION_TEST(FChunkDependencyHighestSharedTest, "System.Core.ChunkDependency.HighestSharedChunk",
+	EAutomationTestFlags::EditorContext | EAutomationTestFlags::EngineFilter)
+
+bool FChunkDependencyHighestSharedTest::RunTest(const FString& Parameters)
+{
+	UChunkDependencyInfo* ChunkDependencyInfo = NewObject<UChunkDependencyInfo>();
+	ChunkDependencyInfo->DependencyArray.Empty();
+	ChunkDependencyInfo->DependencyArray.Add(FChunkDependency(5, 0));
+	ChunkDependencyInfo->DependencyArray.Add(FChunkDependency(10, 0));
+	ChunkDependencyInfo->DependencyArray.Add(FChunkDependency(20, 0));
+	ChunkDependencyInfo->DependencyArray.Add(FChunkDependency(30, 20));
+	ChunkDependencyInfo->DependencyArray.Add(FChunkDependency(40, 20));
+	ChunkDependencyInfo->DependencyArray.Add(FChunkDependency(50, 20));
+	ChunkDependencyInfo->DependencyArray.Add(FChunkDependency(60, 40));
+	ChunkDependencyInfo->BuildChunkDependencyGraph(60);
+	//     0
+	//   / |  \
+	//  5  10  20
+	//        / | \
+	//       30 40 50
+	//          | 
+	//          60 
+
+	UTEST_EQUAL(TEXT("Invalid"), ChunkDependencyInfo->FindHighestSharedChunk({INDEX_NONE}), INDEX_NONE);
+	UTEST_EQUAL(TEXT("Non existent"), ChunkDependencyInfo->FindHighestSharedChunk({100}), INDEX_NONE);
+	UTEST_EQUAL(TEXT("Partial Invalid"), ChunkDependencyInfo->FindHighestSharedChunk({0, INDEX_NONE}), INDEX_NONE);
+	UTEST_EQUAL(TEXT("Empty"), ChunkDependencyInfo->FindHighestSharedChunk({}), INDEX_NONE);
+	UTEST_EQUAL(TEXT("Single"), ChunkDependencyInfo->FindHighestSharedChunk({ 10 }), 10);
+	UTEST_EQUAL(TEXT("Duplicate"), ChunkDependencyInfo->FindHighestSharedChunk({ 10, 10 }), 10);
+	UTEST_EQUAL(TEXT("Two leaf nodes are able to find a parent"), ChunkDependencyInfo->FindHighestSharedChunk({ 5, 10 }), 0);
+	UTEST_EQUAL(TEXT("Non leaf nodes are able to find a parent"), ChunkDependencyInfo->FindHighestSharedChunk({ 10, 20 }), 0);
+	UTEST_EQUAL(TEXT("Parent and Leaf results in parent"), ChunkDependencyInfo->FindHighestSharedChunk({ 20, 30 }), 20);
+	UTEST_EQUAL(TEXT("separatend leaves find parent"), ChunkDependencyInfo->FindHighestSharedChunk({ 30, 60 }), 20);
+	UTEST_EQUAL(TEXT("Complex parent"), ChunkDependencyInfo->FindHighestSharedChunk({ 30, 40, 60 }), 20);
+	UTEST_EQUAL(TEXT("(20, 40, 60) = 20"), ChunkDependencyInfo->FindHighestSharedChunk({ 20, 40, 60 }), 20);
+	UTEST_EQUAL(TEXT("(20, 40, 30) = 20"), ChunkDependencyInfo->FindHighestSharedChunk({ 20, 40, 30 }), 20);
+	UTEST_EQUAL(TEXT("(40, 60) = 40"), ChunkDependencyInfo->FindHighestSharedChunk({ 40, 60 }), 40);
+	UTEST_EQUAL(TEXT("(5, 60) = 0"), ChunkDependencyInfo->FindHighestSharedChunk({ 5, 60 }), 0);
+	UTEST_EQUAL(TEXT("(0, 5) = 0"), ChunkDependencyInfo->FindHighestSharedChunk({ 0, 5 }), 0);
+	return true;
+}
+
+/**
+ * multiple parents are not supported by the chunk dependencies.
+ * Only the first parent will be considered the parent.
+ */
+IMPLEMENT_SIMPLE_AUTOMATION_TEST(FChunkDependencyMultipleParentTest, "System.Core.ChunkDependency.MultipleParent",
+	EAutomationTestFlags::EditorContext | EAutomationTestFlags::EngineFilter)
+
+	bool FChunkDependencyMultipleParentTest::RunTest(const FString& Parameters)
+{
+	UChunkDependencyInfo* ChunkDependencyInfo = NewObject<UChunkDependencyInfo>();
+	ChunkDependencyInfo->DependencyArray.Empty();
+	ChunkDependencyInfo->DependencyArray.Add(FChunkDependency(5, 0));
+	ChunkDependencyInfo->DependencyArray.Add(FChunkDependency(10, 0));
+	ChunkDependencyInfo->DependencyArray.Add(FChunkDependency(20, 0));
+	ChunkDependencyInfo->DependencyArray.Add(FChunkDependency(30, 20));
+	ChunkDependencyInfo->DependencyArray.Add(FChunkDependency(40, 20));
+	ChunkDependencyInfo->DependencyArray.Add(FChunkDependency(50, 20));
+	ChunkDependencyInfo->DependencyArray.Add(FChunkDependency(60, 40));
+	ChunkDependencyInfo->DependencyArray.Add(FChunkDependency(60, 30));
+	ChunkDependencyInfo->BuildChunkDependencyGraph(60);
+	//     0
+	//   / |  \
+	//  5  10  20
+	//        / | \
+	//       30 40 50
+	//        \ | 
+	//          60 
+
+
+	UTEST_EQUAL(TEXT("Single"), ChunkDependencyInfo->FindHighestSharedChunk({ 10 }), 10);
+	UTEST_EQUAL(TEXT("Duplicate"), ChunkDependencyInfo->FindHighestSharedChunk({ 10, 10 }), 10);
+	UTEST_EQUAL(TEXT("Two leaf nodes are able to find a parent"), ChunkDependencyInfo->FindHighestSharedChunk({ 5, 10 }), 0);
+	UTEST_EQUAL(TEXT("Non leaf nodes are able to find a parent"), ChunkDependencyInfo->FindHighestSharedChunk({ 10, 20 }), 0);
+	UTEST_EQUAL(TEXT("Parent and Leaf results in parent"), ChunkDependencyInfo->FindHighestSharedChunk({ 20, 30 }), 20);
+	UTEST_EQUAL(TEXT("separatend leaves find parent"), ChunkDependencyInfo->FindHighestSharedChunk({ 30, 60 }), 20);
+	UTEST_EQUAL(TEXT("Complex parent"), ChunkDependencyInfo->FindHighestSharedChunk({ 30, 40, 60 }), 20);
+	UTEST_EQUAL(TEXT("(20, 40, 60) = 20"), ChunkDependencyInfo->FindHighestSharedChunk({ 20, 40, 60 }), 20);
+	UTEST_EQUAL(TEXT("(20, 40, 30) = 20"), ChunkDependencyInfo->FindHighestSharedChunk({ 20, 40, 30 }), 20);
+	UTEST_EQUAL(TEXT("(40, 60) = 40"), ChunkDependencyInfo->FindHighestSharedChunk({ 40, 60 }), 40);
+	UTEST_EQUAL(TEXT("(30, 40, 60) = 20"), ChunkDependencyInfo->FindHighestSharedChunk({ 30, 40, 60 }), 20);
+	UTEST_EQUAL(TEXT("(30, 60) = 20 - multiple parents not supported."), ChunkDependencyInfo->FindHighestSharedChunk({ 30, 60 }), 20);
+	UTEST_EQUAL(TEXT("(10, 60) = 0"), ChunkDependencyInfo->FindHighestSharedChunk({ 10, 60 }), 0);
+	return true;
+}
+
+IMPLEMENT_SIMPLE_AUTOMATION_TEST(FChunkDependencyCycleTest, "System.Core.ChunkDependency.Cycle",
+	EAutomationTestFlags::EditorContext | EAutomationTestFlags::EngineFilter)
+
+	bool FChunkDependencyCycleTest::RunTest(const FString& Parameters)
+{
+	UChunkDependencyInfo* ChunkDependencyInfo = NewObject<UChunkDependencyInfo>();
+	ChunkDependencyInfo->DependencyArray.Empty();
+	ChunkDependencyInfo->DependencyArray.Add(FChunkDependency(5, 0));
+	ChunkDependencyInfo->DependencyArray.Add(FChunkDependency(10, 0));
+	ChunkDependencyInfo->DependencyArray.Add(FChunkDependency(20, 0));
+	ChunkDependencyInfo->DependencyArray.Add(FChunkDependency(30, 20));
+	ChunkDependencyInfo->DependencyArray.Add(FChunkDependency(40, 20));
+	ChunkDependencyInfo->DependencyArray.Add(FChunkDependency(60, 40));
+	ChunkDependencyInfo->DependencyArray.Add(FChunkDependency(60, 20));
+	ChunkDependencyInfo->BuildChunkDependencyGraph(60);
+	//     0
+	//   / |  \
+	//  5  10  20
+	//        / | 
+	//       30 40 
+	//          | 
+	//          60 
+	//          | 
+	//          20 
+
+	UTEST_EQUAL(TEXT("Two leaf nodes are able to find a parent"), ChunkDependencyInfo->FindHighestSharedChunk({ 5, 10 }), 0);
+	UTEST_EQUAL(TEXT("separatend leaves find parent"), ChunkDependencyInfo->FindHighestSharedChunk({ 30, 60 }), 20);
+	UTEST_EQUAL(TEXT("Complex parent"), ChunkDependencyInfo->FindHighestSharedChunk({ 30, 40, 60 }), 20);
+	UTEST_EQUAL(TEXT("(20, 40, 60) = 20"), ChunkDependencyInfo->FindHighestSharedChunk({ 20, 40, 60 }), 20);
+	UTEST_EQUAL(TEXT("(20, 60) = 20"), ChunkDependencyInfo->FindHighestSharedChunk({ 20, 60 }), 20);
+	UTEST_EQUAL(TEXT("(40, 60) = 40"), ChunkDependencyInfo->FindHighestSharedChunk({ 40, 60 }), 40);
+	return true;
+}
+
+#endif // WITH_AUTOMATION_TESTS
\ No newline at end of file
diff --git a/Engine/Source/Editor/UnrealEd/Public/Commandlets/ChunkDependencyInfo.h b/Engine/Source/Editor/UnrealEd/Public/Commandlets/ChunkDependencyInfo.h
--- a/Engine/Source/Editor/UnrealEd/Public/Commandlets/ChunkDependencyInfo.h
+++ b/Engine/Source/Editor/UnrealEd/Public/Commandlets/ChunkDependencyInfo.h
@@ -20,6 +20,14 @@ struct FChunkDependency
 {
 	GENERATED_USTRUCT_BODY()
 
+	FChunkDependency()
+	{};
+
+	FChunkDependency(int32 InChunkId, int32 InParentChunkId)
+		: ChunkID(InChunkId)
+		, ParentChunkID(InParentChunkId)
+	{};
+
 	/** The child chunk */
 	UPROPERTY(EditAnywhere, Category = ChunkInfo)
 	int32 ChunkID = 0;
@@ -80,6 +88,9 @@ private:
 	/** Map of child chunks to all parent chunks, computed in BuildChunkDependencyGraph */
 	TMap<int32, TSet<int32>> ChildToParentMap;
 
+	/** Cached array of topologically sorted chunk. */
+	TArray<int32> TopologicallySortedChunks;
+
 	/** Cached value of HighestChunk at time of building */
 	int32 CachedHighestChunk;
 };
