Make the InputKey property on UK2Node_InputDebugKey Visible only, not editable. This fixes a bug where you could select the unsupported "Any Key" option from the key selector.

diff --git a/Engine/Plugins/EnhancedInput/Source/InputBlueprintNodes/Private/K2Node_InputDebugKey.h b/Engine/Plugins/EnhancedInput/Source/InputBlueprintNodes/Private/K2Node_InputDebugKey.h
--- a/Engine/Plugins/EnhancedInput/Source/InputBlueprintNodes/Private/K2Node_InputDebugKey.h
+++ b/Engine/Plugins/EnhancedInput/Source/InputBlueprintNodes/Private/K2Node_InputDebugKey.h
@@ -19,8 +19,8 @@ class INPUTBLUEPRINTNODES_API UK2Node_InputDebugKey : public UK2Node, public IK2
 {
 	GENERATED_UCLASS_BODY()
 
-	// The key that is bound
-	UPROPERTY(EditAnywhere, Category="Input")
+	// The key that is bound to this debug event. Pressing this key while the game is running will trigger this node's events.
+	UPROPERTY(VisibleDefaultsOnly, Category="Input")
 	FKey InputKey;
 
 	// Should the binding execute even when the game is paused
diff --git a/Engine/Source/Editor/DetailCustomizations/Private/KeyStructCustomization.cpp b/Engine/Source/Editor/DetailCustomizations/Private/KeyStructCustomization.cpp
--- a/Engine/Source/Editor/DetailCustomizations/Private/KeyStructCustomization.cpp
+++ b/Engine/Source/Editor/DetailCustomizations/Private/KeyStructCustomization.cpp
@@ -54,6 +54,9 @@ void FKeyStructCustomization::CustomizeHeader( TSharedRef<class IPropertyHandle>
 {
 	PropertyHandle = StructPropertyHandle;
 
+	// We only want to allow the FKey to be changed if the property is NOT const.
+	const bool bCanChangeKey = !StructPropertyHandle->GetProperty()->HasAnyPropertyFlags(CPF_EditConst);
+	
 	// create struct header
 	HeaderRow.NameContent()
 	[
@@ -69,6 +72,7 @@ void FKeyStructCustomization::CustomizeHeader( TSharedRef<class IPropertyHandle>
 		.Font(StructCustomizationUtils.GetRegularFont())
 		.AllowClear(!StructPropertyHandle->GetProperty()->HasAnyPropertyFlags(CPF_NoClear))
 		.FilterBlueprintBindable(false)
+		.AllowKeyChange(bCanChangeKey)
 	];
 }
 
diff --git a/Engine/Source/Editor/UnrealEd/Private/SKeySelector.cpp b/Engine/Source/Editor/UnrealEd/Private/SKeySelector.cpp
--- a/Engine/Source/Editor/UnrealEd/Private/SKeySelector.cpp
+++ b/Engine/Source/Editor/UnrealEd/Private/SKeySelector.cpp
@@ -103,6 +103,7 @@ void SKeySelector::Construct(const FArguments& InArgs)
 
 	OnKeyChanged = InArgs._OnKeyChanged;
 	CurrentKey = InArgs._CurrentKey;
+	AllowChangeKey = InArgs._AllowKeyChange;
 
 	TMap<FName, FKeyTreeItem> TreeRootsForCatgories;
 
@@ -152,6 +153,7 @@ void SKeySelector::Construct(const FArguments& InArgs)
 			.PressMethod(EButtonPressMethod::DownAndUp)
 			.ToolTipText(this, &SKeySelector::GetKeyTooltip)
 			.OnClicked(this, &SKeySelector::ListenForInput)
+			.IsEnabled(this, &SKeySelector::CanChangeKey)
 			[
 				SNew(SHorizontalBox)
 				+ SHorizontalBox::Slot()
@@ -267,7 +269,7 @@ FSlateColor SKeySelector::GetKeyIconColor() const
 
 FReply SKeySelector::ListenForInput()
 {
-	if (!bListenForNextInput)
+	if (!bListenForNextInput && CanChangeKey())
 	{
 		bListenForNextInput = true;
 		return FReply::Handled().CaptureMouse(SharedThis(this)).SetUserFocus(SharedThis(this));
@@ -277,6 +279,8 @@ FReply SKeySelector::ListenForInput()
 
 FReply SKeySelector::ProcessHeardInput(FKey KeyHeard)
 {
+	ensure(CanChangeKey());
+	
 	if (bListenForNextInput)	// TODO: Unnecessary. Keep it for safety?
 	{
 		const FScopedTransaction Transaction(LOCTEXT("ChangeKey", "Change Key Value"));
@@ -391,7 +395,7 @@ void SKeySelector::OnKeySelectionChanged(FKeyTreeItem Selection, ESelectInfo::Ty
 	}
 
 	// Only handle selection for non-read only items, since STreeViewItem doesn't actually support read-only
-	if (Selection.IsValid())
+	if (Selection.IsValid() && CanChangeKey())
 	{
 		if (Selection->GetKey().IsValid())
 		{
@@ -423,6 +427,11 @@ void SKeySelector::GetKeyChildren(FKeyTreeItem InItem, TArray<FKeyTreeItem>& Out
 	OutChildren = InItem->Children;
 }
 
+bool SKeySelector::CanChangeKey() const
+{
+	return AllowChangeKey.Get();
+}
+
 TSharedRef<SWidget>	SKeySelector::GetMenuContent()
 {
 	if (!MenuContent.IsValid())
@@ -433,7 +442,8 @@ TSharedRef<SWidget>	SKeySelector::GetMenuContent()
 			.SelectionMode(ESelectionMode::Single)
 			.OnGenerateRow(this, &SKeySelector::GenerateKeyTreeRow)
 			.OnSelectionChanged(this, &SKeySelector::OnKeySelectionChanged)
-			.OnGetChildren(this, &SKeySelector::GetKeyChildren);
+			.OnGetChildren(this, &SKeySelector::GetKeyChildren)
+			.IsEnabled(this, &SKeySelector::CanChangeKey);
 
 		SAssignNew(FilterTextBox, SSearchBox)
 			.OnTextChanged(this, &SKeySelector::OnFilterTextChanged)
diff --git a/Engine/Source/Editor/UnrealEd/Public/SKeySelector.h b/Engine/Source/Editor/UnrealEd/Public/SKeySelector.h
--- a/Engine/Source/Editor/UnrealEd/Public/SKeySelector.h
+++ b/Engine/Source/Editor/UnrealEd/Public/SKeySelector.h
@@ -63,6 +63,7 @@ public:
 		, _Font( FAppStyle::GetFontStyle( TEXT("NormalFont") ) )
 		, _FilterBlueprintBindable( true )
 		, _AllowClear( true )
+		, _AllowKeyChange( true )
 		{}
 		SLATE_ATTRIBUTE( TOptional<FKey>, CurrentKey )
 		SLATE_ATTRIBUTE( FOptionalSize, TreeViewWidth )
@@ -71,6 +72,7 @@ public:
 		SLATE_ATTRIBUTE( FSlateFontInfo, Font )
 		SLATE_ARGUMENT( bool, FilterBlueprintBindable )
 		SLATE_ARGUMENT( bool, AllowClear )
+		SLATE_ARGUMENT( bool, AllowKeyChange )
 	SLATE_END_ARGS()
 public:
 	UNREALED_API void Construct(const FArguments& InArgs);
@@ -120,6 +122,11 @@ protected:
 	UNREALED_API void OnKeySelectionChanged(FKeyTreeItem Selection, ESelectInfo::Type SelectInfo);
 	UNREALED_API void GetKeyChildren(FKeyTreeItem InItem, TArray<FKeyTreeItem>& OutChildren);
 
+	/**
+	 * Returns true if this key selector can actually change the that it's assigned to or not
+	 */
+	UNREALED_API bool CanChangeKey() const;
+
 	/** Gets the Menu Content, setting it up if necessary */
 	UNREALED_API virtual TSharedRef<SWidget>	GetMenuContent();
 
@@ -169,6 +176,9 @@ protected:
 	/** Delegate that is called every time the key changes. */
 	FOnKeyChanged				OnKeyChanged;
 
+	/** If true, then we are allowed to change the key in this widget. Otherwise, this should be used for display only. */
+	TAttribute<bool> AllowChangeKey;
+
 	/** Desired width of the tree view widget */
 	TAttribute<FOptionalSize>	TreeViewWidth;
 	/** Desired height of the tree view widget */
