Interchange: If we cancel a re-import (not with new file), the backup system will override the source file with an empty one.
I also fix the FReimportManager::GetNewReimportPath to add all extension we can import, not only the interchange one.

#rnx

diff --git a/Engine/Source/Editor/UnrealEd/Private/Editor.cpp b/Engine/Source/Editor/UnrealEd/Private/Editor.cpp
--- a/Engine/Source/Editor/UnrealEd/Private/Editor.cpp
+++ b/Engine/Source/Editor/UnrealEd/Private/Editor.cpp
@@ -826,41 +826,34 @@ void FReimportManager::GetNewReimportPath(UObject* Obj, TArray<FString>& InOutFi
 
 	// Append the Interchange supported translator formats for this object
 	TMultiMap<uint32, UFactory*> DummyFilterIndexToFactory;
-	if (UInterchangeManager::IsInterchangeImportEnabled())
+	// Get the list of valid factories
+	for (TObjectIterator<UClass> It; It; ++It)
 	{
-		//Get the extension interchange can translate for this object
-		TArray<FString> TranslatorFormats = UInterchangeManager::GetInterchangeManager().GetSupportedFormatsForObject(Obj);
-		ObjectTools::AppendFormatsFileExtensions(TranslatorFormats, FileTypes, AllExtensions, DummyFilterIndexToFactory);
-	}
+		UClass* CurrentClass = (*It);
 
-	// Interchange is either disabled or do not support the given object, check with the legacy factories
-	if (AllExtensions.IsEmpty())
-	{
-		// Get the list of valid factories
-		for (TObjectIterator<UClass> It; It; ++It)
+		if (CurrentClass->IsChildOf(UFactory::StaticClass()) && !(CurrentClass->HasAnyClassFlags(CLASS_Abstract)))
 		{
-			UClass* CurrentClass = (*It);
-
-			if (CurrentClass->IsChildOf(UFactory::StaticClass()) && !(CurrentClass->HasAnyClassFlags(CLASS_Abstract)))
+			UFactory* Factory = Cast<UFactory>(CurrentClass->GetDefaultObject());
+			if (Factory->bEditorImport && Factory->DoesSupportClass(Obj->GetClass()))
 			{
-				UFactory* Factory = Cast<UFactory>(CurrentClass->GetDefaultObject());
-				if (Factory->bEditorImport && Factory->DoesSupportClass(Obj->GetClass()))
-				{
-					Factories.Add(Factory);
-				}
+				Factories.Add(Factory);
 			}
 		}
-
-		if (Factories.Num() <= 0)
-		{
-			// No matching factories for this asset, fail
-			return;
-		}
-
+	}
+	//We now add all the factory extensions
+	if (Factories.Num() > 0)
+	{
 		// Generate the file types and extensions represented by the selected factories
 		ObjectTools::GenerateFactoryFileExtensions(Factories, FileTypes, AllExtensions, DummyFilterIndexToFactory);
 	}
 
+	if (UInterchangeManager::IsInterchangeImportEnabled())
+	{
+		//Get the extension interchange can translate for this object
+		TArray<FString> TranslatorFormats = UInterchangeManager::GetInterchangeManager().GetSupportedFormatsForObject(Obj);
+		ObjectTools::AppendFormatsFileExtensions(TranslatorFormats, FileTypes, AllExtensions, DummyFilterIndexToFactory);
+	}
+
 	FString DefaultFolder;
 	FString DefaultFile;
 	
diff --git a/Engine/Source/Runtime/Interchange/Engine/Private/InterchangeAssetImportData.cpp b/Engine/Source/Runtime/Interchange/Engine/Private/InterchangeAssetImportData.cpp
--- a/Engine/Source/Runtime/Interchange/Engine/Private/InterchangeAssetImportData.cpp
+++ b/Engine/Source/Runtime/Interchange/Engine/Private/InterchangeAssetImportData.cpp
@@ -403,6 +403,9 @@ void UInterchangeAssetImportData::ClearBackupSourceData() const
 void UInterchangeAssetImportData::ReinstateBackupSourceData()
 {
 #if WITH_EDITORONLY_DATA
-	SourceData = SourceDataBackup;
+	if (SourceDataBackup.SourceFiles.Num() > 0)
+	{
+		SourceData = SourceDataBackup;
+	}
 #endif
 }
\ No newline at end of file
