CookPackageSplitter and MPCook: Remove spurious error "PackageSplitter found an existing copy of a package...". This circumstance can now happen during some generated package retraction cases. Handle it instead: restore the package to newly-created state by renaming all UObjects out of it.
#rnx

diff --git a/Engine/Source/Editor/UnrealEd/Classes/CookOnTheSide/CookOnTheFlyServer.h b/Engine/Source/Editor/UnrealEd/Classes/CookOnTheSide/CookOnTheFlyServer.h
--- a/Engine/Source/Editor/UnrealEd/Classes/CookOnTheSide/CookOnTheFlyServer.h
+++ b/Engine/Source/Editor/UnrealEd/Classes/CookOnTheSide/CookOnTheFlyServer.h
@@ -1374,7 +1374,7 @@ private:
 		UE::Cook::FGenerationHelper& GenerationHelper, UE::Cook::FCookGenerationInfo& Info,
 		TArray<ICookPackageSplitter::FGeneratedPackageForPreSave>& GeneratedPackagesForPresave);
 	/** Construct the list of generated packages that is required for some CookPackageSplitter interface calls. */
-	void ConstructGeneratedPackagesForPresave(UE::Cook::FPackageData& PackageData,
+	bool TryConstructGeneratedPackagesForPresave(UE::Cook::FPackageData& PackageData,
 		UE::Cook::FGenerationHelper& GenerationHelper,
 		TArray<ICookPackageSplitter::FGeneratedPackageForPreSave>& GeneratedPackagesForPresave);
 	/**
@@ -1384,9 +1384,6 @@ private:
 	UE::Cook::EPollStatus BeginCachePostMove(UE::Cook::FGenerationHelper& GenerationHelper,
 		UE::Cook::FCookGenerationInfo& Info, UE::Cook::FCookerTimer& Timer);
 
-	/** Try creating (or finding from earlier creation) the generated package for later population */
-	UPackage* TryCreateGeneratedPackage(UE::Cook::FGenerationHelper& GenerationHelper,
-		UE::Cook::FCookGenerationInfo& GeneratedInfo);
 	/** Try calling the CookPackageSplitter's populate to create the package */
 	UE::Cook::EPollStatus TryPopulateGeneratedPackage(UE::Cook::FGenerationHelper& GenerationHelper,
 		UE::Cook::FCookGenerationInfo& GeneratedInfo);
diff --git a/Engine/Source/Editor/UnrealEd/Private/CookOnTheFlyServer.cpp b/Engine/Source/Editor/UnrealEd/Private/CookOnTheFlyServer.cpp
--- a/Engine/Source/Editor/UnrealEd/Private/CookOnTheFlyServer.cpp
+++ b/Engine/Source/Editor/UnrealEd/Private/CookOnTheFlyServer.cpp
@@ -2911,7 +2911,7 @@ void UCookOnTheFlyServer::LoadPackageInQueue(UE::Cook::FPackageData& PackageData
 			return;
 		}
 
-		LoadedPackage = TryCreateGeneratedPackage(*GenerationHelper, *Info);
+		LoadedPackage = GenerationHelper->TryCreateGeneratedPackage(*Info);
 		if (!LoadedPackage)
 		{
 			RejectPackageToLoad(PackageData, TEXT("is a generated package which could not be populated"),
@@ -3460,7 +3460,12 @@ UE::Cook::EPollStatus UCookOnTheFlyServer::BeginCacheObjectsToMove(UE::Cook::FGe
 		TArray<UPackage*> KeepReferencedPackages;
 		if (Info.IsGenerator())
 		{
-			ConstructGeneratedPackagesForPresave(PackageData, GenerationHelper, GeneratedPackagesForPresave);
+			if (!TryConstructGeneratedPackagesForPresave(PackageData, GenerationHelper, GeneratedPackagesForPresave))
+			{
+				UE_LOG(LogCook, Error, TEXT("PackageSplitter unexpected failure: could not ConstructGeneratedPackagesForPreSave. Splitter=%s"),
+					*GenerationHelper.GetSplitDataObjectName().ToString());
+				return EPollStatus::Error;
+			}
 			FScopedActivePackage ScopedActivePackage(*this, GenerationHelper.GetOwner().GetPackageName(),
 				PackageAccessTrackingOps::NAME_CookerBuildObject);
 			bPopulateSucceeded = Splitter->PopulateGeneratorPackage(Package, SplitDataObject,
@@ -3518,7 +3523,13 @@ UE::Cook::EPollStatus UCookOnTheFlyServer::PreSaveGeneratorPackage(UE::Cook::FPa
 	}
 
 	TArray<UPackage*> KeepReferencedPackages;
-	ConstructGeneratedPackagesForPresave(PackageData, GenerationHelper, GeneratedPackagesForPresave);
+	if (!TryConstructGeneratedPackagesForPresave(PackageData, GenerationHelper, GeneratedPackagesForPresave))
+	{
+		UE_LOG(LogCook, Error, TEXT("PackageSplitter unexpected failure: could not ConstructGeneratedPackagesForPreSave. Splitter=%s"),
+			*GenerationHelper.GetSplitDataObjectName().ToString());
+		return EPollStatus::Error;
+	}
+
 	{
 		FScopedActivePackage ScopedActivePackage(*this, GenerationHelper.GetOwner().GetPackageName(),
 			PackageAccessTrackingOps::NAME_CookerBuildObject);
@@ -3534,7 +3545,7 @@ UE::Cook::EPollStatus UCookOnTheFlyServer::PreSaveGeneratorPackage(UE::Cook::FPa
 	return EPollStatus::Success;
 }
 
-void UCookOnTheFlyServer::ConstructGeneratedPackagesForPresave(UE::Cook::FPackageData& PackageData, UE::Cook::FGenerationHelper& GenerationHelper,
+bool UCookOnTheFlyServer::TryConstructGeneratedPackagesForPresave(UE::Cook::FPackageData& PackageData, UE::Cook::FGenerationHelper& GenerationHelper,
 	TArray<ICookPackageSplitter::FGeneratedPackageForPreSave>& GeneratedPackagesForPresave)
 {
 	using namespace UE::Cook;
@@ -3542,10 +3553,8 @@ void UCookOnTheFlyServer::ConstructGeneratedPackagesForPresave(UE::Cook::FPackag
 	if (GeneratedPackagesForPresave.Num() > 0)
 	{
 		// Already constructed, save time by early exiting
-		return;
+		return true;
 	}
-	UPackage* Package = PackageData.GetPackage();
-	check(Package);
 
 	// We need to find or (create empty stub packages for) each of the PackagesToGenerate so that PreSaveGeneratorPackage
 	// can refer to them to create hardlinks in the cooked Generator package
@@ -3558,14 +3567,13 @@ void UCookOnTheFlyServer::ConstructGeneratedPackagesForPresave(UE::Cook::FPackag
 		SplitterData.RelativePath = Info.RelativePath;
 		SplitterData.GeneratedRootPath = Info.GeneratedRootPath;
 		SplitterData.bCreatedAsMap = Info.IsCreateAsMap();
-
-		const FString GeneratedPackageName = Info.PackageData->GetPackageName().ToString();
-		SplitterData.Package = FindObject<UPackage>(nullptr, *GeneratedPackageName);
+		SplitterData.Package = GenerationHelper.TryCreateGeneratedPackage(Info);
 		if (!SplitterData.Package)
 		{
-			SplitterData.Package = GenerationHelper.CreateGeneratedUPackage(Info, Package, *GeneratedPackageName);
+			return false;
 		}
 	}
+	return true;
 }
 
 UE::Cook::EPollStatus UCookOnTheFlyServer::BeginCachePostMove(UE::Cook::FGenerationHelper& GenerationHelper,
@@ -3639,45 +3647,6 @@ UE::Cook::EPollStatus UCookOnTheFlyServer::BeginCachePostMove(UE::Cook::FGenerat
 	return EPollStatus::Success;
 }
 
-UPackage* UCookOnTheFlyServer::TryCreateGeneratedPackage(UE::Cook::FGenerationHelper& GenerationHelper, UE::Cook::FCookGenerationInfo& Info)
-{
-	using namespace UE::Cook;
-	check(Info.PackageData); // Caller is responsible for validating
-	UE::Cook::FPackageData& GeneratedPackageData = *Info.PackageData;
-	const FString GeneratedPackageName = GeneratedPackageData.GetPackageName().ToString();
-	UPackage* OwnerPackage = GenerationHelper.FindOrLoadOwnerPackage(*this);
-	if (!OwnerPackage)
-	{
-		UE_LOG(LogCook, Error, TEXT("TryCreateGeneratedPackage: could not load ParentGeneratorPackage %s for GeneratedPackage %s"),
-			*GenerationHelper.GetOwner().GetPackageName().ToString(), *GeneratedPackageName);
-		return nullptr;
-	}
-
-	UPackage* GeneratedPackage = FindObject<UPackage>(nullptr, *GeneratedPackageName);
-	if (GeneratedPackage)
-	{
-		if (!Info.HasCreatedPackage())
-		{
-			UE_LOG(LogCook, Error, TEXT("PackageSplitter found an existing copy of a package it was trying to populate;")
-				TEXT("this is unexpected since garbage has been collected and the package should have been unreferenced so it should have been collected.")
-				TEXT("Splitter=%s, Generated=%s."),
-				*GenerationHelper.GetSplitDataObjectName().ToString(), *GeneratedPackageName);
-			EReferenceChainSearchMode SearchMode = EReferenceChainSearchMode::Shortest
-				| EReferenceChainSearchMode::PrintAllResults
-				| EReferenceChainSearchMode::FullChain;
-			FReferenceChainSearch RefChainSearch(GeneratedPackage, SearchMode);
-			return nullptr;
-		}
-		// Otherwise this is the package that was created for the generator's presave, and it is still valid because
-		// there has not been a GC since we created it.
-	}
-	else
-	{
-		GeneratedPackage = GenerationHelper.CreateGeneratedUPackage(Info, OwnerPackage, *GeneratedPackageName);
-	}
-	return GeneratedPackage;
-}
-
 UE::Cook::EPollStatus UCookOnTheFlyServer::TryPopulateGeneratedPackage(UE::Cook::FGenerationHelper& GenerationHelper,
 	UE::Cook::FCookGenerationInfo& GeneratedInfo)
 {
diff --git a/Engine/Source/Editor/UnrealEd/Private/Cooker/CookGenerationHelper.cpp b/Engine/Source/Editor/UnrealEd/Private/Cooker/CookGenerationHelper.cpp
--- a/Engine/Source/Editor/UnrealEd/Private/Cooker/CookGenerationHelper.cpp
+++ b/Engine/Source/Editor/UnrealEd/Private/Cooker/CookGenerationHelper.cpp
@@ -16,7 +16,7 @@ namespace UE::Cook
 //////////////////////////////////////////////////////////////////////////
 // FGenerationHelper
 
-FGenerationHelper::FGenerationHelper(UE::Cook::FPackageData& InOwner)
+FGenerationHelper::FGenerationHelper(FPackageData& InOwner)
 : OwnerInfo(InOwner, true /* bInGenerator */)
 {
 }
@@ -121,7 +121,7 @@ UPackage* FGenerationHelper::FindOrLoadPackage(UCookOnTheFlyServer& COTFS, FPack
 
 void FGenerationHelper::SearchForRegisteredSplitDataObject(UCookOnTheFlyServer& COTFS,
 	FName PackageName, UPackage* Package,
-	TOptional<TConstArrayView<UE::Cook::FCachedObjectInOuter>> CachedObjectsInOuter,
+	TOptional<TConstArrayView<FCachedObjectInOuter>> CachedObjectsInOuter,
 	UObject*& OutSplitDataObject, UE::Cook::Private::FRegisteredCookPackageSplitter*& OutRegisteredSplitter,
 	TUniquePtr<ICookPackageSplitter>& OutSplitterInstance)
 {
@@ -368,7 +368,7 @@ bool FGenerationHelper::TryGenerateList()
 		FString PackageName = FPaths::RemoveDuplicateSlashes(FString::Printf(TEXT("/%s/%s/%s"),
 			*PackageRoot, GeneratedPackageSubPath, *SplitterData.RelativePath));
 		const FName PackageFName(*PackageName);
-		UE::Cook::FPackageData* PackageData = PackageDatas.TryAddPackageDataByPackageName(PackageFName,
+		FPackageData* PackageData = PackageDatas.TryAddPackageDataByPackageName(PackageFName,
 			false /* bRequireExists */, bCreateAsMap);
 		if (!PackageData)
 		{
@@ -409,7 +409,7 @@ bool FGenerationHelper::TryGenerateList()
 		GeneratedInfo.PackageDependencies.SetNum(Algo::Unique(GeneratedInfo.PackageDependencies));
 		GeneratedInfo.SetIsCreateAsMap(bCreateAsMap);
 		if (bNeedCachedPlatformDataBeforeSplit ||
-			COTFS.MPCookGeneratorSplit == UE::Cook::EMPCookGeneratorSplit::AllOnSameWorker)
+			COTFS.MPCookGeneratorSplit == EMPCookGeneratorSplit::AllOnSameWorker)
 		{
 			PackageData->SetWorkerAssignmentConstraint(FWorkerId::Local());
 		}
@@ -547,26 +547,54 @@ void FGenerationHelper::OnRequestFencePassed(UCookOnTheFlyServer& COTFS)
 	PreviousGeneratedPackages.Empty();
 }
 
-UPackage* FGenerationHelper::CreateGeneratedUPackage(FCookGenerationInfo& GeneratedInfo,
-	const UPackage* InOwnerPackage, const TCHAR* GeneratedPackageName)
+UPackage* FGenerationHelper::TryCreateGeneratedPackage(FCookGenerationInfo& GeneratedInfo)
 {
 	if (!IsValid())
 	{
 		return nullptr;
 	}
 
+	check(GeneratedInfo.PackageData); // Caller is responsible for validating
+	FPackageData& GeneratedPackageData = *GeneratedInfo.PackageData;
+	const FString GeneratedPackageName = GeneratedPackageData.GetPackageName().ToString();
+	UPackage* LocalOwnerPackage = FindOrLoadOwnerPackage(GetOwner().GetPackageDatas().GetCookOnTheFlyServer());
+	if (!LocalOwnerPackage)
+	{
+		UE_LOG(LogCook, Error,
+			TEXT("TryCreateGeneratedPackage: could not load ParentGeneratorPackage %s for GeneratedPackage %s"),
+			*GetOwner().GetPackageName().ToString(), *GeneratedPackageName);
+		return nullptr;
+	}
+
+	UPackage* GeneratedPackage = FindObject<UPackage>(nullptr, *GeneratedPackageName);
+	if (GeneratedPackage)
+	{
+		// The package might have been created for the generator's presave, or it might have been created and abandoned
+		// by an earlier save attempt of the generated package. Remove all objects from it; generated packages are
+		// created empty and it is the job of the CookPackageSplitter to populate them.
+		TArray<UObject*> ExistingObjects;
+		GetObjectsWithPackage(GeneratedPackage, ExistingObjects, false /* bIncludeNestedObjects */);
+		if (!ExistingObjects.IsEmpty())
+		{
+			UObject* TransientPackage = GetTransientPackage();
+			for (UObject* Existing : ExistingObjects)
+			{
+				Existing->Rename(nullptr, TransientPackage, REN_DontCreateRedirectors);
+			}
+		}
+	}
+	else
+	{
 #if ENABLE_COOK_STATS
-	++DetailedCookStats::NumRequestedLoads;
+		++DetailedCookStats::NumRequestedLoads;
 #endif
-	UPackage* GeneratedPackage = CreatePackage(GeneratedPackageName);
+		GeneratedPackage = CreatePackage(*GeneratedPackageName);
+	}
 	GeneratedPackage->SetSavedHash(GeneratedInfo.PackageHash);
-	GeneratedPackage->SetPersistentGuid(InOwnerPackage->GetPersistentGuid());
+	GeneratedPackage->SetPersistentGuid(LocalOwnerPackage->GetPersistentGuid());
 	GeneratedPackage->SetPackageFlags(PKG_CookGenerated);
+	GeneratedPackage->SetLoadedByEditorPropertiesOnly(LocalOwnerPackage->IsLoadedByEditorPropertiesOnly());
 	GeneratedInfo.SetHasCreatedPackage(true);
-	if (!InOwnerPackage->IsLoadedByEditorPropertiesOnly())
-	{
-		GeneratedPackage->SetLoadedByEditorPropertiesOnly(false);
-	}
 
 	return GeneratedPackage;
 }
diff --git a/Engine/Source/Editor/UnrealEd/Private/Cooker/CookGenerationHelper.h b/Engine/Source/Editor/UnrealEd/Private/Cooker/CookGenerationHelper.h
--- a/Engine/Source/Editor/UnrealEd/Private/Cooker/CookGenerationHelper.h
+++ b/Engine/Source/Editor/UnrealEd/Private/Cooker/CookGenerationHelper.h
@@ -272,9 +272,8 @@ public:
 	void OnRequestFencePassedBroadcast(UCookOnTheFlyServer& COTFS);
 	/** Called from when the RequestFence added from EndQueueGeneratedPackages has passed. */
 	void OnRequestFencePassed(UCookOnTheFlyServer& COTFS);
-	/** Call CreatePackage and set the necessary data on the UPackage, and update status. */
-	UPackage* CreateGeneratedUPackage(FCookGenerationInfo& GenerationInfo, const UPackage* OwnerPackage,
-		const TCHAR* GeneratedPackageName);
+	/** Call CreatePackage and set package header data; or empty it and normalize it if it already exists. */
+	UPackage* TryCreateGeneratedPackage(FCookGenerationInfo& GenerationInfo);
 	/**
 	 * Clear any data that should only be held when an FPackageData is in the save state, for the given Info. The given
 	 * Info might specify the generator package or one of the generated packages.
