IterativeCook:
1) Tweak log messages for iterativevalidate and for config invalidation of iterative and incremental cooks.
2) In incremental cooks, turn off the full-cook invalidation that occurs when any config value changes; incremental cooks now test config dependencies per package.
#rnx

diff --git a/Engine/Build/Graph/Tests/IterativeValidate.xml b/Engine/Build/Graph/Tests/IterativeValidate.xml
--- a/Engine/Build/Graph/Tests/IterativeValidate.xml
+++ b/Engine/Build/Graph/Tests/IterativeValidate.xml
@@ -30,7 +30,7 @@
 	<Property Name="CommonCookArgs" Value="$(CommonCookArgs) -cookprocesscount=1"/>
 
 	<Property Name="CompileNodeName" Value="Compile $(EditorPlatform) $(ProjectName)"/>
-	<Property Name="CookIterativeValidateSinglePhaseName" Value="Cook InteraveValidateSinglePhase $(ProjectName) $(TargetPlatform)"/>
+	<Property Name="CookIterativeValidateSinglePhaseName" Value="Cook IterativeValidateSinglePhase $(ProjectName) $(TargetPlatform)"/>
 	<Property Name="CookIterativeValidatePhase1Name" Value="Cook IterativeValidatePhase1 $(ProjectName) $(TargetPlatform)"/>
 	<Property Name="CookIterativeValidatePhase2Name" Value="Cook IterativeValidatePhase2 $(ProjectName) $(TargetPlatform)"/>
 
diff --git a/Engine/Source/Editor/UnrealEd/Private/CookOnTheFlyServer.cpp b/Engine/Source/Editor/UnrealEd/Private/CookOnTheFlyServer.cpp
--- a/Engine/Source/Editor/UnrealEd/Private/CookOnTheFlyServer.cpp
+++ b/Engine/Source/Editor/UnrealEd/Private/CookOnTheFlyServer.cpp
@@ -7984,7 +7984,7 @@ bool UCookOnTheFlyServer::ArePreviousCookSettingsCompatible(const TMap<FName, FS
 		}
 	}
 
-	if (!bIterativeIgnoreIni && IniSettingsOutOfDate(TargetPlatform))
+	if (!bIterativeIgnoreIni && !bHybridIterativeEnabled && IniSettingsOutOfDate(TargetPlatform))
 	{
 		UE_LOG(LogCook, Display, TEXT("Cook invalidated for platform %s because ini settings have changed. Clearing all cooked content."),
 			*TargetPlatform->PlatformName());
@@ -8074,8 +8074,6 @@ bool UCookOnTheFlyServer::IniSettingsOutOfDate(const ITargetPlatform* TargetPlat
 	}
 
 	TStringBuilder<256> ConfigNameKeyStr;
-	TStringBuilder<128> Filename;
-	TStringBuilder<256> PlatformName;
 	for (const auto& OldIniFile : OldIniSettings)
 	{
 		OldIniFile.Key.ToString(ConfigNameKeyStr);
@@ -8103,18 +8101,18 @@ bool UCookOnTheFlyServer::IniSettingsOutOfDate(const ITargetPlatform* TargetPlat
 		{
 			FName SectionName = OldIniSection.Key;
 			const FConfigSection* IniSection = ConfigFile->FindSection( SectionName.ToString() );
-			auto GetDenyListMessageStart = [&PlatformName, &Filename, SectionName]()
+			auto GetDenyListMessageStart = [&ConfigNameKeyStr, SectionName]()
 				{
 					return FString::Printf(
 						TEXT("To avoid invalidating due to this setting, add a deny list setting")
-						TEXT("\n\tDefaultEditor.ini:[CookSettings]:+CookOnTheFlyConfigSettingDenyList=%s.%s:%s"),
-						*PlatformName, *Filename, *SectionName.ToString());
+						TEXT("\n\tDefaultEditor.ini:[CookSettings]:+CookOnTheFlyConfigSettingDenyList=%s:%s"),
+						*ConfigNameKeyStr, *SectionName.ToString());
 				};
 
 			if ( IniSection == nullptr )
 			{
 				UE_LOG(LogCook, Display,
-					TEXT("Invalidating inisettings: Inisetting used by platform %s is different for %s:[%s], section doesn't exist in current config."),
+					TEXT("Invalidating inisettings: Inisetting used by platform %s is different for %s:[%s]. The section doesn't exist in current config."),
 					*TargetPlatform->PlatformName(), *ConfigNameKeyStr, *SectionName.ToString());
 				UE_LOG(LogCook, Display, TEXT("%s"), *GetDenyListMessageStart());
 				return true;
@@ -8130,7 +8128,7 @@ bool UCookOnTheFlyServer::IniSettingsOutOfDate(const ITargetPlatform* TargetPlat
 				if ( CurrentValues.Num() != OldIniValue.Value.Num() )
 				{
 					UE_LOG(LogCook, Display,
-						TEXT("Invalidating inisettings: Inisetting used by platform %s is different for %s:[%s]:%s, missmatched num array elements %d != %d."),
+						TEXT("Invalidating inisettings: Inisetting used by platform %s is different for %s:[%s]:%s. Mismatched num array elements %d != %d."),
 						*TargetPlatform->PlatformName(), *ConfigNameKeyStr, *SectionName.ToString(),
 						*ValueName.ToString(), CurrentValues.Num(), OldIniValue.Value.Num());
 					UE_LOG(LogCook, Display, TEXT("%s:%s"), *GetDenyListMessageStart(), *ValueName.ToString());
@@ -8142,7 +8140,7 @@ bool UCookOnTheFlyServer::IniSettingsOutOfDate(const ITargetPlatform* TargetPlat
 					if ( FilteredCurrentValue != OldIniValue.Value[Index] )
 					{
 						UE_LOG(LogCook, Display,
-							TEXT("Invalidating inisettings: Inisetting used by platform %s is different for %s:[%s]:%s%s, value '%s' != '%s'."),
+							TEXT("Invalidating inisettings: Inisetting used by platform %s is different for %s:[%s]:%s%s. Value '%s' != '%s'."),
 							*TargetPlatform->PlatformName(), *ConfigNameKeyStr, *SectionName.ToString(),
 							*ValueName.ToString(),
 							(CurrentValues.Num() == 1 ? TEXT("") : *FString::Printf(TEXT(" %d"), Index)),
@@ -10838,7 +10836,9 @@ void UCookOnTheFlyServer::LoadBeginCookIterativeFlagsLocal(FBeginCookContext& Be
 
 		if (bIsDiffOnly)
 		{
-			UE_LOG(LogCook, Display, TEXT("Keeping cooked content for platform %s for DiffOnly"), *TargetPlatform->PlatformName());
+			UE_LOG(LogCook, Display,
+				TEXT("INCREMENTAL COOK: cooking incrementally due to -DiffOnly flag. Keeping cooked content for platform %s."),
+				*TargetPlatform->PlatformName());
 			// When looking for deterministic cooking differences in cooked packages, don't delete the packages on disk
 			PlatformContext.bFullBuild = false;
 			PlatformContext.bAllowIterativeResults = false;
@@ -10851,16 +10851,26 @@ void UCookOnTheFlyServer::LoadBeginCookIterativeFlagsLocal(FBeginCookContext& Be
 			bool bIterativeAllowed = true;
 			if (!bIterative && !PlatformData->bIsSandboxInitialized)
 			{
-				UE_LOG(LogCook, Display, TEXT("Clearing all cooked content for platform %s"), *TargetPlatform->PlatformName());
+				UE_LOG(LogCook, Display,
+					TEXT("FULL COOK: Neither -iterative nor -cookincremental were specified. Clearing all cooked content for platform %s and executing a full cook."),
+					*TargetPlatform->PlatformName());
 				bIterativeAllowed = false;
 			}
 			else if (!ArePreviousCookSettingsCompatible(PlatformContext.CurrentCookSettings, TargetPlatform))
 			{
+				UE_LOG(LogCook, Display,
+					TEXT("FULL COOK: %s was specified, but globally invalidated. Clearing all cooked content for platform %s and executing a full cook."),
+					bHybridIterativeEnabled ? TEXT("-cookincremental") : TEXT("-iterative"),
+					*TargetPlatform->PlatformName());
 				bIterativeAllowed = false;
 			}
 
 			if (bIterativeAllowed)
 			{
+				UE_LOG(LogCook, Display,
+					TEXT("INCREMENTAL COOK: %s was specified and not globally invalidated. Keeping cooked content for platform platform %s."),
+					bHybridIterativeEnabled ? TEXT("-cookincremental") : TEXT("-iterative"),
+					*TargetPlatform->PlatformName());
 				PlatformContext.bFullBuild = false;
 				PlatformContext.bAllowIterativeResults = true;
 				PlatformContext.bClearMemoryResults = false;
diff --git a/Engine/Source/Editor/UnrealEd/Private/Cooker/IterativeValidatePackageWriter.cpp b/Engine/Source/Editor/UnrealEd/Private/Cooker/IterativeValidatePackageWriter.cpp
--- a/Engine/Source/Editor/UnrealEd/Private/Cooker/IterativeValidatePackageWriter.cpp
+++ b/Engine/Source/Editor/UnrealEd/Private/Cooker/IterativeValidatePackageWriter.cpp
@@ -377,10 +377,11 @@ void FIterativeValidatePackageWriter::EndCook(const FCookInfo& Info)
 	case EPhase::AllInOnePhase:
 	{
 		UE_LOG(LogIterativeValidate, Display,
-			TEXT("Modified: %d. DetectedUnmodified: %d. ValidatedUnmodified: %d. IterativeFalseNegative: %d."),
-			ModifiedCount, IterativeValidated.Num() + IterativeFalseNegative.Num(), IterativeValidated.Num(), IterativeFalseNegative.Num());
-		FString Message = FString::Printf(TEXT("IterativeFalseNegative: %d."), IterativeFailed.Num());
-		if (IterativeFalseNegative.Num())
+			TEXT("Modified: %d. DetectedUnmodified: %d. ValidatedUnmodified: %d. IterativeSkipFalsePositive: %d."),
+			ModifiedCount, IterativeValidated.Num() + IterativeSkipFalsePositive.Num(), IterativeValidated.Num(),
+			IterativeSkipFalsePositive.Num());
+		FString Message = FString::Printf(TEXT("IterativeSkipFalsePositive: %d."), IterativeSkipFalsePositive.Num());
+		if (IterativeSkipFalsePositive.Num())
 		{
 			UE_LOG(LogIterativeValidate, Error, TEXT("%s"), *Message);
 		}
@@ -392,17 +393,19 @@ void FIterativeValidatePackageWriter::EndCook(const FCookInfo& Info)
 	}
 	case EPhase::Phase1:
 		UE_LOG(LogIterativeValidate, Display,
-			TEXT("Modified: %d. DetectedUnmodified: %d. ValidatedUnmodified: %d. IterativeFalseNegativeOrIndeterminism: %d."),
-			ModifiedCount, IterativeValidated.Num() + IterativeFailed.Num(), IterativeValidated.Num(), IterativeFailed.Num());
+			TEXT("Modified: %d. DetectedUnmodified: %d. ValidatedUnmodified: %d. IterativeSkipFalsePositiveOrIndeterminism: %d."),
+			ModifiedCount, IterativeValidated.Num() + IterativeFailed.Num(), IterativeValidated.Num(),
+			IterativeFailed.Num());
 		Save();
 		break;
 	case EPhase::Phase2:
 	{
 		UE_LOG(LogIterativeValidate, Display,
 			TEXT("Modified: %d. DetectedUnmodified: %d. ValidatedUnmodified: %d. Indeterminism: %d."),
-			ModifiedCount, IterativeValidated.Num() + IterativeFailed.Num(), IterativeValidated.Num(), IndeterminismFailed.Num());
-		FString Message = FString::Printf(TEXT("IterativeFalseNegative: %d."), IterativeFalseNegative.Num());
-		if (IterativeFalseNegative.Num())
+			ModifiedCount, IterativeValidated.Num() + IterativeFailed.Num(), IterativeValidated.Num(),
+			IndeterminismFailed.Num());
+		FString Message = FString::Printf(TEXT("IterativeSkipFalsePositive: %d."), IterativeSkipFalsePositive.Num());
+		if (IterativeSkipFalsePositive.Num())
 		{
 			UE_LOG(LogIterativeValidate, Error, TEXT("%s"), *Message);
 		}
@@ -479,7 +482,7 @@ bool FIterativeValidatePackageWriter::IsAnotherSaveNeeded(FSavePackageResultStru
 
 			if (bIsDifferent && !bNewPackage)
 			{
-				IterativeFalseNegative.Add(BeginInfo.PackageName);
+				IterativeSkipFalsePositive.Add(BeginInfo.PackageName);
 			}
 			else if (!bNewPackage)
 			{
@@ -584,7 +587,7 @@ void FIterativeValidatePackageWriter::LogIterativeDifferences()
 	// It is called immediately after first-pass package save that is run by our super class FDiffPackageWriter, which
 	// compared it against the version that was written to disk during Phase1. If there are differences
 	// now from Phase1, then this package has a determinism issue. We log that information at display rather
-	// than Warning because this cookmode only logs warnings for IterativeFalseNegatives.
+	// than Warning because this cookmode only logs warnings for IterativeSkipFalsePositives.
 	bool bHasDeterminismIssue = bIsDifferent;
 	if (bHasDeterminismIssue)
 	{
@@ -595,9 +598,9 @@ void FIterativeValidatePackageWriter::LogIterativeDifferences()
 	}
 
 	// Otherwise, no determinism issues, so the differences indicate a bug in Diff Package
-	IterativeFalseNegative.Add(BeginInfo.PackageName);
+	IterativeSkipFalsePositive.Add(BeginInfo.PackageName);
 	FMsg::Logf(__FILE__, __LINE__, LogIterativeValidate.GetCategoryName(), ELogVerbosity::Warning,
-		TEXT("IterativeFalseNegative package %s."), *BeginInfo.PackageName.ToString());
+		TEXT("IterativeSkipFalsePositive package %s."), *BeginInfo.PackageName.ToString());
 	TArray<FMessage>& Messages = IterativeFailed.FindOrAdd(BeginInfo.PackageName);
 	for (const FMessage& Message : Messages)
 	{
diff --git a/Engine/Source/Editor/UnrealEd/Private/Cooker/IterativeValidatePackageWriter.h b/Engine/Source/Editor/UnrealEd/Private/Cooker/IterativeValidatePackageWriter.h
--- a/Engine/Source/Editor/UnrealEd/Private/Cooker/IterativeValidatePackageWriter.h
+++ b/Engine/Source/Editor/UnrealEd/Private/Cooker/IterativeValidatePackageWriter.h
@@ -76,7 +76,7 @@ protected:
 	/**
 	 * Written during Phase1, read during Phase2. Unused in AllInOnePhase. Packages that we thought were iteratively
 	 * unchanged but in which we discovered differences. But the differences might be due to indeterminism. Phase2
-	 * splits this container into IndeterminismFailed entries and IterativeFalseNegative entries.
+	 * splits this container into IndeterminismFailed entries and IterativeSkipFalsePositive entries.
 	 */
 	TMap<FName, TArray<FMessage>> IterativeFailed;
 	/**
@@ -87,9 +87,9 @@ protected:
 	/**
 	 * Read during Phase2. Read/Written during AllInOnePhase. Packages that we thought were iteratively unchanged but that had
 	 * differences when saved, and no indeterminism detected (AllInOnePhase does not search for indeterminism, so all packages
-	 * with differences end up here), so they must be a bug in the iteratvely unchanged decision.
+	 * with differences end up here), so they must be a bug in the decision to iteratively skip.
 	 */
-	TSet<FName> IterativeFalseNegative;
+	TSet<FName> IterativeSkipFalsePositive;
 	/**
 	 * Unused in Phase1,Phase2. In AllInOnePhase, this records packages that we think are iteratively unchanged.
 	 * We do a 2-pass or 3-pass save for these: look for diffs and then save to disk. For packages not in this list
