CookPackageSplitter and MPCook: Fix error
PackageSplitter specified generated package is map in GetGenerateList results, but then in PreSaveGeneratedPackage created it as uasset.
TryCreateGeneratedPackage should not ResetToEmpty when called from the generator to get stub generated packages, because the generated package might already be in progress on it.
#rnx

diff --git a/Engine/Source/Editor/UnrealEd/Private/CookOnTheFlyServer.cpp b/Engine/Source/Editor/UnrealEd/Private/CookOnTheFlyServer.cpp
--- a/Engine/Source/Editor/UnrealEd/Private/CookOnTheFlyServer.cpp
+++ b/Engine/Source/Editor/UnrealEd/Private/CookOnTheFlyServer.cpp
@@ -2911,7 +2911,7 @@ void UCookOnTheFlyServer::LoadPackageInQueue(UE::Cook::FPackageData& PackageData
 			return;
 		}
 
-		LoadedPackage = GenerationHelper->TryCreateGeneratedPackage(*Info);
+		LoadedPackage = GenerationHelper->TryCreateGeneratedPackage(*Info, true /* bResetToEmpty */);
 		if (!LoadedPackage)
 		{
 			RejectPackageToLoad(PackageData, TEXT("is a generated package which could not be populated"),
@@ -3567,7 +3567,7 @@ bool UCookOnTheFlyServer::TryConstructGeneratedPackagesForPresave(UE::Cook::FPac
 		SplitterData.RelativePath = Info.RelativePath;
 		SplitterData.GeneratedRootPath = Info.GeneratedRootPath;
 		SplitterData.bCreatedAsMap = Info.IsCreateAsMap();
-		SplitterData.Package = GenerationHelper.TryCreateGeneratedPackage(Info);
+		SplitterData.Package = GenerationHelper.TryCreateGeneratedPackage(Info, false /* bResetToEmpty */);
 		if (!SplitterData.Package)
 		{
 			return false;
diff --git a/Engine/Source/Editor/UnrealEd/Private/Cooker/CookGenerationHelper.cpp b/Engine/Source/Editor/UnrealEd/Private/Cooker/CookGenerationHelper.cpp
--- a/Engine/Source/Editor/UnrealEd/Private/Cooker/CookGenerationHelper.cpp
+++ b/Engine/Source/Editor/UnrealEd/Private/Cooker/CookGenerationHelper.cpp
@@ -547,7 +547,7 @@ void FGenerationHelper::OnRequestFencePassed(UCookOnTheFlyServer& COTFS)
 	PreviousGeneratedPackages.Empty();
 }
 
-UPackage* FGenerationHelper::TryCreateGeneratedPackage(FCookGenerationInfo& GeneratedInfo)
+UPackage* FGenerationHelper::TryCreateGeneratedPackage(FCookGenerationInfo& GeneratedInfo, bool bResetToEmpty)
 {
 	if (!IsValid())
 	{
@@ -570,16 +570,21 @@ UPackage* FGenerationHelper::TryCreateGeneratedPackage(FCookGenerationInfo& Gene
 	if (GeneratedPackage)
 	{
 		// The package might have been created for the generator's presave, or it might have been created and abandoned
-		// by an earlier save attempt of the generated package. Remove all objects from it; generated packages are
-		// created empty and it is the job of the CookPackageSplitter to populate them.
-		TArray<UObject*> ExistingObjects;
-		GetObjectsWithPackage(GeneratedPackage, ExistingObjects, false /* bIncludeNestedObjects */);
-		if (!ExistingObjects.IsEmpty())
+		// by an earlier save attempt of the generated package.
+		// If bResetToEmpty then we are starting the populate of the generated package and we need to remove all objects
+		// from the package. Generated packages are created empty and it is the job of the CookPackageSplitter to populate
+		// them during save.
+		if (bResetToEmpty)
 		{
-			UObject* TransientPackage = GetTransientPackage();
-			for (UObject* Existing : ExistingObjects)
+			TArray<UObject*> ExistingObjects;
+			GetObjectsWithPackage(GeneratedPackage, ExistingObjects, false /* bIncludeNestedObjects */);
+			if (!ExistingObjects.IsEmpty())
 			{
-				Existing->Rename(nullptr, TransientPackage, REN_DontCreateRedirectors);
+				UObject* TransientPackage = GetTransientPackage();
+				for (UObject* Existing : ExistingObjects)
+				{
+					Existing->Rename(nullptr, TransientPackage, REN_DontCreateRedirectors);
+				}
 			}
 		}
 	}
diff --git a/Engine/Source/Editor/UnrealEd/Private/Cooker/CookGenerationHelper.h b/Engine/Source/Editor/UnrealEd/Private/Cooker/CookGenerationHelper.h
--- a/Engine/Source/Editor/UnrealEd/Private/Cooker/CookGenerationHelper.h
+++ b/Engine/Source/Editor/UnrealEd/Private/Cooker/CookGenerationHelper.h
@@ -273,7 +273,7 @@ public:
 	/** Called from when the RequestFence added from EndQueueGeneratedPackages has passed. */
 	void OnRequestFencePassed(UCookOnTheFlyServer& COTFS);
 	/** Call CreatePackage and set package header data; or empty it and normalize it if it already exists. */
-	UPackage* TryCreateGeneratedPackage(FCookGenerationInfo& GenerationInfo);
+	UPackage* TryCreateGeneratedPackage(FCookGenerationInfo& GenerationInfo, bool bResetToEmpty);
 	/**
 	 * Clear any data that should only be held when an FPackageData is in the save state, for the given Info. The given
 	 * Info might specify the generator package or one of the generated packages.
