Fix crash during build of lighting caused by null widget

diff --git a/Engine/Source/Editor/UnrealEd/Private/EditorBuildUtils.cpp b/Engine/Source/Editor/UnrealEd/Private/EditorBuildUtils.cpp
--- a/Engine/Source/Editor/UnrealEd/Private/EditorBuildUtils.cpp
+++ b/Engine/Source/Editor/UnrealEd/Private/EditorBuildUtils.cpp
@@ -560,7 +560,7 @@ bool FEditorBuildUtils::EditorBuild( UWorld* InWorld, FName Id, const bool bAllo
 		GUnrealEd->Exec( InWorld, TEXT("MAP CHECK DONTDISPLAYDIALOG") );
 	}
 
-	if ( bDoBuild )
+	if ( bDoBuild && BuildProgressWidget.IsValid() )
 	{
 		// Display elapsed build time.
 		UE_LOG(LogEditorBuildUtils, Log,  TEXT("Build time %s"), *BuildProgressWidget.Pin()->BuildElapsedTimeText().ToString() );
