CookPackageSplitter:
Add debug setting IsGeneratorSavedAfterGenerated.
Add utility functions for creating generated package names.
#rnx

diff --git a/Engine/Source/Editor/UnrealEd/Private/CookOnTheFlyServer.cpp b/Engine/Source/Editor/UnrealEd/Private/CookOnTheFlyServer.cpp
--- a/Engine/Source/Editor/UnrealEd/Private/CookOnTheFlyServer.cpp
+++ b/Engine/Source/Editor/UnrealEd/Private/CookOnTheFlyServer.cpp
@@ -225,11 +225,8 @@ UCookOnTheFlyServer::FOnCookByTheBookStarted UCookOnTheFlyServer::CookByTheBookS
 UCookOnTheFlyServer::FOnCookByTheBookFinished UCookOnTheFlyServer::CookByTheBookFinishedEvent;
 
 
-namespace UE
-{
-namespace Cook
+namespace UE::Cook
 {
-const TCHAR* GeneratedPackageSubPath = TEXT("_Generated_");
 
 // Keep the old behavior of cooking all by default until we implement good feedback in the editor about the missing setting
 static bool bCookAllByDefault = true;
@@ -237,7 +234,7 @@ static FAutoConsoleVariableRef CookAllByDefaultCVar(
 	TEXT("Cook.CookAllByDefault"),
 	bCookAllByDefault,
 	TEXT("When FilesInPath is empty. Cook all packages by default."));
-}
+
 }
 
 /* helper structs functions
@@ -3333,6 +3330,20 @@ UE::Cook::EPollStatus UCookOnTheFlyServer::PrepareSaveGenerationPackage(UE::Cook
 		{
 			return EPollStatus::Incomplete;
 		}
+		if (FGenerationHelper::IsGeneratorSavedAfterGenerated() && Info.IsGenerator())
+		{
+			if (GenerationHelper.IsWaitingForQueueResults())
+			{
+				return EPollStatus::Incomplete;
+			}
+			for (FCookGenerationInfo& GeneratedInfo : GenerationHelper.GetPackagesToGenerate())
+			{
+				if (GeneratedInfo.PackageData && GeneratedInfo.PackageData->IsInProgress())
+				{
+					return EPollStatus::Incomplete;
+				}
+			}
+		}
 		Info.SetSaveStateComplete(FCookGenerationInfo::ESaveState::FinishCachePreMove);
 	}
 
@@ -10735,6 +10746,7 @@ void UCookOnTheFlyServer::SetBeginCookConfigSettings(FBeginCookContext& BeginCon
 	bHybridIterativeEnabled = Settings.bHybridIterativeEnabled;
 	bHybridIterativeAllowAllClasses = Settings.bHybridIterativeAllowAllClasses;
 	PackageDatas->SetBeginCookConfigSettings(Settings.CookShowInstigator);
+	UE::Cook::FGenerationHelper::SetBeginCookConfigSettings();
 	SetNeverCookPackageConfigSettings(BeginContext, Settings);
 }
 
diff --git a/Engine/Source/Editor/UnrealEd/Private/Cooker/CookGenerationHelper.cpp b/Engine/Source/Editor/UnrealEd/Private/Cooker/CookGenerationHelper.cpp
--- a/Engine/Source/Editor/UnrealEd/Private/Cooker/CookGenerationHelper.cpp
+++ b/Engine/Source/Editor/UnrealEd/Private/Cooker/CookGenerationHelper.cpp
@@ -7,7 +7,9 @@
 #include "Cooker/CookDirector.h"
 #include "Cooker/IWorkerRequests.h"
 #include "Cooker/PackageTracker.h"
+#include "Misc/CommandLine.h"
 #include "Misc/PackageAccessTrackingOps.h"
+#include "Misc/Parse.h"
 #include "UObject/ReferenceChainSearch.h"
 
 namespace UE::Cook
@@ -362,11 +364,8 @@ bool FGenerationHelper::TryGenerateList()
 		}
 		bool bCreateAsMap = *SplitterData.GetCreateAsMap();
 
-		FString PackageRoot = SplitterData.GeneratedRootPath.IsEmpty() ?
-			OwnerPackageName.ToString() : SplitterData.GeneratedRootPath;
-
-		FString PackageName = FPaths::RemoveDuplicateSlashes(FString::Printf(TEXT("/%s/%s/%s"),
-			*PackageRoot, GeneratedPackageSubPath, *SplitterData.RelativePath));
+		FString PackageName = ICookPackageSplitter::ConstructGeneratedPackageName(OwnerPackageName,
+			SplitterData.RelativePath, SplitterData.GeneratedRootPath);
 		const FName PackageFName(*PackageName);
 		FPackageData* PackageData = PackageDatas.TryAddPackageDataByPackageName(PackageFName,
 			false /* bRequireExists */, bCreateAsMap);
@@ -1224,4 +1223,16 @@ void FCookGenerationInfo::IterativeCookValidateOrClear(FGenerationHelper& Genera
 	}
 }
 
+bool FGenerationHelper::bGeneratorSavedAfterGenerated = false;
+
+void FGenerationHelper::SetBeginCookConfigSettings()
+{
+	bGeneratorSavedAfterGenerated = FParse::Param(FCommandLine::Get(), TEXT("CookGeneratorSavedAfterGenerated"));
+}
+
+bool FGenerationHelper::IsGeneratorSavedAfterGenerated()
+{
+	return bGeneratorSavedAfterGenerated;
+}
+
 }
\ No newline at end of file
diff --git a/Engine/Source/Editor/UnrealEd/Private/Cooker/CookGenerationHelper.h b/Engine/Source/Editor/UnrealEd/Private/Cooker/CookGenerationHelper.h
--- a/Engine/Source/Editor/UnrealEd/Private/Cooker/CookGenerationHelper.h
+++ b/Engine/Source/Editor/UnrealEd/Private/Cooker/CookGenerationHelper.h
@@ -308,6 +308,7 @@ public:
 	void ClearKeepForIterative();
 	void SetKeepForQueueResults();
 	void ClearKeepForQueueResults();
+	bool IsWaitingForQueueResults() const;
 	void SetKeepForGeneratorSave();
 	void ClearKeepForGeneratorSave();
 
@@ -330,6 +331,14 @@ public:
 	/** Helper function for Initialize and for TryCreateValidGenerationHelper. */
 	static UPackage* FindOrLoadPackage(UCookOnTheFlyServer& COTFS, FPackageData& OwnerPackageData);
 
+	/** Initialize settings from commandline and config files. */
+	static void SetBeginCookConfigSettings();
+	/**
+	 * Debug setting that forces a race-condition behavior. Force generator packages to not save until after
+	 * all generated packages save.
+	 */
+	static bool IsGeneratorSavedAfterGenerated();
+
 private:
 	enum class EInitializeStatus : uint8
 	{
@@ -364,6 +373,8 @@ private:
 	bool bUseInternalReferenceToAvoidGarbageCollect = false;
 	bool bNeedCachedPlatformDataBeforeSplit = false;
 	bool bGeneratedList = false;
+
+	static bool bGeneratorSavedAfterGenerated;
 };
 
 
@@ -559,6 +570,11 @@ inline void FGenerationHelper::ClearKeepForQueueResults()
 	ReferenceFromKeepForQueueResults.SafeRelease();
 }
 
+inline bool FGenerationHelper::IsWaitingForQueueResults() const
+{
+	return ReferenceFromKeepForQueueResults.IsValid();
+}
+
 inline void FGenerationHelper::SetKeepForGeneratorSave()
 {
 	ReferenceFromKeepForGeneratorSave = this;
diff --git a/Engine/Source/Editor/UnrealEd/Private/Cooker/CookPackageData.h b/Engine/Source/Editor/UnrealEd/Private/Cooker/CookPackageData.h
--- a/Engine/Source/Editor/UnrealEd/Private/Cooker/CookPackageData.h
+++ b/Engine/Source/Editor/UnrealEd/Private/Cooker/CookPackageData.h
@@ -55,8 +55,6 @@ struct FPackageData;
 struct FPackageDataMonitor;
 struct FPendingCookedPlatformDataCancelManager;
 
-extern const TCHAR* GeneratedPackageSubPath;
-
 /**
  * Events in the lifetime of an object related to BeginCacheForCookedPlatformData. Used by the cooker
  * to track which calls have been made and still need to be made.
diff --git a/Engine/Source/Editor/UnrealEd/Private/Cooker/CookPackageSplitter.cpp b/Engine/Source/Editor/UnrealEd/Private/Cooker/CookPackageSplitter.cpp
--- a/Engine/Source/Editor/UnrealEd/Private/Cooker/CookPackageSplitter.cpp
+++ b/Engine/Source/Editor/UnrealEd/Private/Cooker/CookPackageSplitter.cpp
@@ -2,8 +2,31 @@
 
 #include "CookPackageSplitter.h"
 
+#include "Misc/Paths.h"
+
 #if WITH_EDITOR
 
+const TCHAR* ICookPackageSplitter::GetGeneratedPackageSubPath()
+{
+	return TEXT("_Generated_");
+}
+
+FString ICookPackageSplitter::ConstructGeneratedPackageName(FName OwnerPackageName, FStringView RelPath,
+	FStringView GeneratedRootOverride)
+{
+	FString PackageRoot;
+	if (GeneratedRootOverride.IsEmpty())
+	{
+		PackageRoot = OwnerPackageName.ToString();
+	}
+	else
+	{
+		PackageRoot = GeneratedRootOverride;
+	}
+	return FPaths::RemoveDuplicateSlashes(FString::Printf(TEXT("/%s/%s/%.*s"),
+		*PackageRoot, GetGeneratedPackageSubPath(), RelPath.Len(), RelPath.GetData()));
+}
+
 namespace UE::Cook::Private
 {
 
diff --git a/Engine/Source/Editor/UnrealEd/Public/CookPackageSplitter.h b/Engine/Source/Editor/UnrealEd/Public/CookPackageSplitter.h
--- a/Engine/Source/Editor/UnrealEd/Public/CookPackageSplitter.h
+++ b/Engine/Source/Editor/UnrealEd/Public/CookPackageSplitter.h
@@ -7,6 +7,7 @@
 #include "AssetRegistry/IAssetRegistry.h"
 #include "Containers/Array.h"
 #include "Containers/List.h"
+#include "Containers/StringView.h"
 #include "Containers/UnrealString.h"
 #include "CoreMinimal.h"
 #include "Hash/Blake3.h"
@@ -233,13 +234,22 @@ public:
 
 	/** Called when the Owner package needs to be reloaded after a garbage collect in order to populate a generated package. */
 	virtual void OnOwnerReloaded(UPackage* OwnerPackage, UObject* OwnerObject) {}
+
+	// Utility functions for Splitters
+
+	/** The name of the _Generated_ subdirectory that is the parent directory of a splitter's generated packages. */
+	UNREALED_API static const TCHAR* GetGeneratedPackageSubPath();
+
+	/**
+	 * Return the full packagename that will be used for a GeneratedPackage, based on the GeneratorPackage's name and
+	 * on the RelPath and optional GeneratedRootPath that the splitter provides in the FGeneratedPackage it returns from
+	 * GetGenerateList.
+	 */
+	UNREALED_API static FString ConstructGeneratedPackageName(FName OwnerPackageName, FStringView RelPath,
+		FStringView GeneratedRootOverride = FStringView());
 };
 
-namespace UE
-{
-namespace Cook
-{
-namespace Private
+namespace UE::Cook::Private
 {
 
 /** Interface for internal use only (used by REGISTER_COOKPACKAGE_SPLITTER to register an ICookPackageSplitter for a class) */
@@ -262,8 +272,6 @@ private:
 	TLinkedList<FRegisteredCookPackageSplitter*> GlobalListLink;
 };
 
-}
-}
 }
 
 /**
