CookPackageSplitter and MPCook: Replace CookGeneratorSavedAfterGenerated with the two possible orderings MPCookGeneratorSaveOrder=GeneratorFirst or GeneratedFirst.
Turn GeneratorFirst on by default until we update WorldPartitionCookPackageSplitter to handle any order.
#rnx

diff --git a/Engine/Source/Editor/UnrealEd/Private/CookOnTheFlyServer.cpp b/Engine/Source/Editor/UnrealEd/Private/CookOnTheFlyServer.cpp
--- a/Engine/Source/Editor/UnrealEd/Private/CookOnTheFlyServer.cpp
+++ b/Engine/Source/Editor/UnrealEd/Private/CookOnTheFlyServer.cpp
@@ -3330,7 +3330,7 @@ UE::Cook::EPollStatus UCookOnTheFlyServer::PrepareSaveGenerationPackage(UE::Cook
 		{
 			return EPollStatus::Incomplete;
 		}
-		if (FGenerationHelper::IsGeneratorSavedAfterGenerated() && Info.IsGenerator())
+		if (FGenerationHelper::IsGeneratedSavedFirst() && Info.IsGenerator())
 		{
 			if (GenerationHelper.IsWaitingForQueueResults())
 			{
@@ -3344,6 +3344,13 @@ UE::Cook::EPollStatus UCookOnTheFlyServer::PrepareSaveGenerationPackage(UE::Cook
 				}
 			}
 		}
+		if (FGenerationHelper::IsGeneratorSavedFirst() && !Info.IsGenerator())
+		{
+			if (GenerationHelper.GetOwner().IsInProgress())
+			{
+				return EPollStatus::Incomplete;
+			}
+		}
 		Info.SetSaveStateComplete(FCookGenerationInfo::ESaveState::FinishCachePreMove);
 	}
 
diff --git a/Engine/Source/Editor/UnrealEd/Private/Cooker/CookGenerationHelper.cpp b/Engine/Source/Editor/UnrealEd/Private/Cooker/CookGenerationHelper.cpp
--- a/Engine/Source/Editor/UnrealEd/Private/Cooker/CookGenerationHelper.cpp
+++ b/Engine/Source/Editor/UnrealEd/Private/Cooker/CookGenerationHelper.cpp
@@ -8,6 +8,7 @@
 #include "Cooker/IWorkerRequests.h"
 #include "Cooker/PackageTracker.h"
 #include "Misc/CommandLine.h"
+#include "Misc/ConfigCacheIni.h"
 #include "Misc/PackageAccessTrackingOps.h"
 #include "Misc/Parse.h"
 #include "UObject/ReferenceChainSearch.h"
@@ -1223,16 +1224,55 @@ void FCookGenerationInfo::IterativeCookValidateOrClear(FGenerationHelper& Genera
 	}
 }
 
-bool FGenerationHelper::bGeneratorSavedAfterGenerated = false;
+namespace GenerationHelperPrivate
+{
+
+enum class ERequiredSaveOrder
+{
+	None,
+	GeneratorFirst,
+	GeneratedFirst,
+};
+
+ERequiredSaveOrder RequiredSaveOrder = ERequiredSaveOrder::None;
+
+}
 
 void FGenerationHelper::SetBeginCookConfigSettings()
 {
-	bGeneratorSavedAfterGenerated = FParse::Param(FCommandLine::Get(), TEXT("CookGeneratorSavedAfterGenerated"));
+	const TCHAR* CommandLine = FCommandLine::Get();
+
+	FString SaveOrder;
+	GConfig->GetString(TEXT("CookSettings"), TEXT("MPCookGeneratorSaveOrder"), SaveOrder, GEditorIni);
+	FParse::Value(FCommandLine::Get(), TEXT("-MPCookGeneratorSaveOrder="), SaveOrder);
+	if (SaveOrder == TEXT("None"))
+	{
+		GenerationHelperPrivate::RequiredSaveOrder = GenerationHelperPrivate::ERequiredSaveOrder::None;
+	}
+	else if (SaveOrder == TEXT("GeneratorFirst"))
+	{
+		GenerationHelperPrivate::RequiredSaveOrder = GenerationHelperPrivate::ERequiredSaveOrder::GeneratorFirst;
+	}
+	else if (SaveOrder == TEXT("GeneratedFirst"))
+	{
+		GenerationHelperPrivate::RequiredSaveOrder = GenerationHelperPrivate::ERequiredSaveOrder::GeneratedFirst;
+	}
+	else
+	{
+		// UE-211211: Make RequiredSaveOrder::GeneratorFirst the default until we have
+		// changed WorldPartitionCookPackageSplitter to handle the new allowed ordering.
+		GenerationHelperPrivate::RequiredSaveOrder = GenerationHelperPrivate::ERequiredSaveOrder::GeneratorFirst;
+	}
+}
+
+bool FGenerationHelper::IsGeneratorSavedFirst()
+{
+	return GenerationHelperPrivate::RequiredSaveOrder == GenerationHelperPrivate::ERequiredSaveOrder::GeneratorFirst;
 }
 
-bool FGenerationHelper::IsGeneratorSavedAfterGenerated()
+bool FGenerationHelper::IsGeneratedSavedFirst()
 {
-	return bGeneratorSavedAfterGenerated;
+	return GenerationHelperPrivate::RequiredSaveOrder == GenerationHelperPrivate::ERequiredSaveOrder::GeneratedFirst;
 }
 
 }
\ No newline at end of file
diff --git a/Engine/Source/Editor/UnrealEd/Private/Cooker/CookGenerationHelper.h b/Engine/Source/Editor/UnrealEd/Private/Cooker/CookGenerationHelper.h
--- a/Engine/Source/Editor/UnrealEd/Private/Cooker/CookGenerationHelper.h
+++ b/Engine/Source/Editor/UnrealEd/Private/Cooker/CookGenerationHelper.h
@@ -333,11 +333,10 @@ public:
 
 	/** Initialize settings from commandline and config files. */
 	static void SetBeginCookConfigSettings();
-	/**
-	 * Debug setting that forces a race-condition behavior. Force generator packages to not save until after
-	 * all generated packages save.
-	 */
-	static bool IsGeneratorSavedAfterGenerated();
+	/** Debug setting that forces an otherwise unspecified order. Generator packages save before the generated do. */
+	static bool IsGeneratorSavedFirst();
+	/** Debug setting that forces an otherwise unspecified order. Generated packages save before their generator does. */
+	static bool IsGeneratedSavedFirst();
 
 private:
 	enum class EInitializeStatus : uint8
@@ -373,8 +372,6 @@ private:
 	bool bUseInternalReferenceToAvoidGarbageCollect = false;
 	bool bNeedCachedPlatformDataBeforeSplit = false;
 	bool bGeneratedList = false;
-
-	static bool bGeneratorSavedAfterGenerated;
 };
 
 
