Replaced hardcoded LoadModuleChecked of WorldPartitionEditor by IWorldPartitionEditorModule::Get(). Also removed cached instance pointer as it's not safe to be used during shutdown phase.

diff --git a/Engine/Plugins/Editor/ObjectMixer/ObjectMixer/Source/ObjectMixer/Private/Views/List/Modes/ObjectMixerOutlinerHierarchy.cpp b/Engine/Plugins/Editor/ObjectMixer/ObjectMixer/Source/ObjectMixer/Private/Views/List/Modes/ObjectMixerOutlinerHierarchy.cpp
--- a/Engine/Plugins/Editor/ObjectMixer/ObjectMixer/Source/ObjectMixer/Private/Views/List/Modes/ObjectMixerOutlinerHierarchy.cpp
+++ b/Engine/Plugins/Editor/ObjectMixer/ObjectMixer/Source/ObjectMixer/Private/Views/List/Modes/ObjectMixerOutlinerHierarchy.cpp
@@ -50,8 +50,7 @@ TUniquePtr<FObjectMixerOutlinerHierarchy> FObjectMixerOutlinerHierarchy::Create(
 	GEngine->OnActorFolderAdded().AddRaw(Hierarchy, &FObjectMixerOutlinerHierarchy::OnActorFolderAdded);
 	GEngine->OnActorFoldersUpdatedEvent().AddRaw(Hierarchy, &FObjectMixerOutlinerHierarchy::OnActorFoldersUpdatedEvent);
 
-	IWorldPartitionEditorModule& WorldPartitionEditorModule = FModuleManager::LoadModuleChecked<IWorldPartitionEditorModule>("WorldPartitionEditor");
-	WorldPartitionEditorModule.OnWorldPartitionCreated().AddRaw(Hierarchy, &FObjectMixerOutlinerHierarchy::OnWorldPartitionCreated);
+	IWorldPartitionEditorModule::Get().OnWorldPartitionCreated().AddRaw(Hierarchy, &FObjectMixerOutlinerHierarchy::OnWorldPartitionCreated);
 
 	if (World.IsValid())
 	{
diff --git a/Engine/Plugins/PCG/Source/PCGEditor/Private/WorldPartitionBuilder/PCGWorldPartitionBuilder.cpp b/Engine/Plugins/PCG/Source/PCGEditor/Private/WorldPartitionBuilder/PCGWorldPartitionBuilder.cpp
--- a/Engine/Plugins/PCG/Source/PCGEditor/Private/WorldPartitionBuilder/PCGWorldPartitionBuilder.cpp
+++ b/Engine/Plugins/PCG/Source/PCGEditor/Private/WorldPartitionBuilder/PCGWorldPartitionBuilder.cpp
@@ -574,7 +574,6 @@ void PCGWorldPartitionBuilder::Build(const TArray<FString>& Args)
 {
 	if (UWorld* World = (GEditor ? GEditor->GetEditorWorldContext().World() : nullptr))
 	{
-		IWorldPartitionEditorModule& WorldPartitionEditorModule = FModuleManager::LoadModuleChecked<IWorldPartitionEditorModule>("WorldPartitionEditor");
 		IWorldPartitionEditorModule::FRunBuilderParams Params;
 		Params.BuilderClass = UPCGWorldPartitionBuilder::StaticClass();
 		Params.World = World;
@@ -587,6 +586,6 @@ void PCGWorldPartitionBuilder::Build(const TArray<FString>& Args)
 			Params.ExtraArgs += Arg;
 		}
 
-		WorldPartitionEditorModule.RunBuilder(Params);
+		IWorldPartitionEditorModule::Get().RunBuilder(Params);
 	}
 }
diff --git a/Engine/Source/Editor/DataLayerEditor/Private/DataLayer/DataLayerHierarchy.cpp b/Engine/Source/Editor/DataLayerEditor/Private/DataLayer/DataLayerHierarchy.cpp
--- a/Engine/Source/Editor/DataLayerEditor/Private/DataLayer/DataLayerHierarchy.cpp
+++ b/Engine/Source/Editor/DataLayerEditor/Private/DataLayer/DataLayerHierarchy.cpp
@@ -57,8 +57,7 @@ FDataLayerHierarchy::FDataLayerHierarchy(FDataLayerMode* Mode, const TWeakObject
 		GEngine->OnLevelActorListChanged().AddRaw(this, &FDataLayerHierarchy::OnLevelActorListChanged);
 	}
 
-	IWorldPartitionEditorModule& WorldPartitionEditorModule = FModuleManager::LoadModuleChecked<IWorldPartitionEditorModule>("WorldPartitionEditor");
-	WorldPartitionEditorModule.OnWorldPartitionCreated().AddRaw(this, &FDataLayerHierarchy::OnWorldPartitionCreated);
+	IWorldPartitionEditorModule::Get().OnWorldPartitionCreated().AddRaw(this, &FDataLayerHierarchy::OnWorldPartitionCreated);
 
 	if (World.IsValid())
 	{
@@ -93,8 +92,7 @@ FDataLayerHierarchy::~FDataLayerHierarchy()
 		GEngine->OnLevelActorListChanged().RemoveAll(this);
 	}
 
-	IWorldPartitionEditorModule& WorldPartitionEditorModule = FModuleManager::LoadModuleChecked<IWorldPartitionEditorModule>("WorldPartitionEditor");
-	WorldPartitionEditorModule.OnWorldPartitionCreated().RemoveAll(this);
+	IWorldPartitionEditorModule::Get().OnWorldPartitionCreated().RemoveAll(this);
 
 	if (RepresentingWorld.IsValid())
 	{
diff --git a/Engine/Source/Editor/SceneOutliner/Private/ActorHierarchy.cpp b/Engine/Source/Editor/SceneOutliner/Private/ActorHierarchy.cpp
--- a/Engine/Source/Editor/SceneOutliner/Private/ActorHierarchy.cpp
+++ b/Engine/Source/Editor/SceneOutliner/Private/ActorHierarchy.cpp
@@ -50,8 +50,7 @@ void FActorHierarchy::Create_Internal(FActorHierarchy* Hierarchy, const TWeakObj
 	GEngine->OnActorFolderAdded().AddRaw(Hierarchy, &FActorHierarchy::OnActorFolderAdded);
 	GEngine->OnActorFoldersUpdatedEvent().AddRaw(Hierarchy, &FActorHierarchy::OnActorFoldersUpdatedEvent);
 
-	IWorldPartitionEditorModule& WorldPartitionEditorModule = FModuleManager::LoadModuleChecked<IWorldPartitionEditorModule>("WorldPartitionEditor");
-	WorldPartitionEditorModule.OnWorldPartitionCreated().AddRaw(Hierarchy, &FActorHierarchy::OnWorldPartitionCreated);
+	IWorldPartitionEditorModule::Get().OnWorldPartitionCreated().AddRaw(Hierarchy, &FActorHierarchy::OnWorldPartitionCreated);
 
 	if (World.IsValid())
 	{
@@ -100,8 +99,7 @@ FActorHierarchy::~FActorHierarchy()
 		GEngine->OnActorFoldersUpdatedEvent().RemoveAll(this);
 	}
 
-	IWorldPartitionEditorModule& WorldPartitionEditorModule = FModuleManager::LoadModuleChecked<IWorldPartitionEditorModule>("WorldPartitionEditor");
-	WorldPartitionEditorModule.OnWorldPartitionCreated().RemoveAll(this);
+	IWorldPartitionEditorModule::Get().OnWorldPartitionCreated().RemoveAll(this);
 
 	if (RepresentingWorld.IsValid())
 	{
diff --git a/Engine/Source/Editor/UnrealEd/Private/EditorBuildUtils.cpp b/Engine/Source/Editor/UnrealEd/Private/EditorBuildUtils.cpp
--- a/Engine/Source/Editor/UnrealEd/Private/EditorBuildUtils.cpp
+++ b/Engine/Source/Editor/UnrealEd/Private/EditorBuildUtils.cpp
@@ -1401,9 +1401,8 @@ void FEditorBuildUtils::TriggerHierarchicalLODBuilder(UWorld* InWorld)
 {
 	if (InWorld->IsPartitionedWorld())
 	{
-		IWorldPartitionEditorModule& WorldPartitionEditorModule = FModuleManager::LoadModuleChecked<IWorldPartitionEditorModule>("WorldPartitionEditor");
 		TSubclassOf<UWorldPartitionBuilder> WorldPartitionHLODsBuilder = FindObjectChecked<UClass>(nullptr, TEXT("/Script/UnrealEd.WorldPartitionHLODsBuilder"), true);
-		WorldPartitionEditorModule.RunBuilder(WorldPartitionHLODsBuilder, InWorld);
+		IWorldPartitionEditorModule::Get().RunBuilder(WorldPartitionHLODsBuilder, InWorld);
 	}
 	else
 	{
@@ -1416,9 +1415,8 @@ void FEditorBuildUtils::TriggerMinimapBuilder(UWorld* InWorld)
 {
 	if (InWorld->IsPartitionedWorld())
 	{
-		IWorldPartitionEditorModule& WorldPartitionEditorModule = FModuleManager::LoadModuleChecked<IWorldPartitionEditorModule>("WorldPartitionEditor");
 		TSubclassOf<UWorldPartitionBuilder> WorldPartitionMiniMapBuilder = FindObjectChecked<UClass>(nullptr, TEXT("/Script/UnrealEd.WorldPartitionMiniMapBuilder"), true);
-		WorldPartitionEditorModule.RunBuilder(WorldPartitionMiniMapBuilder, InWorld);
+		IWorldPartitionEditorModule::Get().RunBuilder(WorldPartitionMiniMapBuilder, InWorld);
 	}
 }
 
@@ -1426,9 +1424,8 @@ void FEditorBuildUtils::TriggerLandscapeSplineMeshesBuilder(UWorld* InWorld)
 {
 	if (InWorld->IsPartitionedWorld())
 	{
-		IWorldPartitionEditorModule& WorldPartitionEditorModule = FModuleManager::LoadModuleChecked<IWorldPartitionEditorModule>("WorldPartitionEditor");
 		TSubclassOf<UWorldPartitionBuilder> WorldPartitionLandscapeSplineMeshesBuilder = FindObjectChecked<UClass>(nullptr, TEXT("/Script/UnrealEd.WorldPartitionLandscapeSplineMeshesBuilder"), true);
-		WorldPartitionEditorModule.RunBuilder(WorldPartitionLandscapeSplineMeshesBuilder, InWorld);
+		IWorldPartitionEditorModule::Get().RunBuilder(WorldPartitionLandscapeSplineMeshesBuilder, InWorld);
 	}
 }
 
diff --git a/Engine/Source/Editor/WorldPartitionEditor/Private/WorldPartition/SWorldPartitionEditor.cpp b/Engine/Source/Editor/WorldPartitionEditor/Private/WorldPartition/SWorldPartitionEditor.cpp
--- a/Engine/Source/Editor/WorldPartitionEditor/Private/WorldPartition/SWorldPartitionEditor.cpp
+++ b/Engine/Source/Editor/WorldPartitionEditor/Private/WorldPartition/SWorldPartitionEditor.cpp
@@ -24,9 +24,8 @@ void SWorldPartitionEditor::Construct(const FArguments& InArgs)
 	FWorldBrowserModule& WorldBrowserModule = FModuleManager::LoadModuleChecked<FWorldBrowserModule>("WorldBrowser");
 	WorldBrowserModule.OnBrowseWorld.AddSP(this, &SWorldPartitionEditor::OnBrowseWorld);
 	UWorldPartition::WorldPartitionChangedEvent.AddSP(this, &SWorldPartitionEditor::OnBrowseWorld);
-
-	IWorldPartitionEditorModule& WorldPartitionEditorModule = FModuleManager::LoadModuleChecked<IWorldPartitionEditorModule>("WorldPartitionEditor");
-	WorldPartitionEditorModule.OnWorldPartitionCreated().AddSP(this, &SWorldPartitionEditor::OnBrowseWorld);
+	
+	IWorldPartitionEditorModule::Get().OnWorldPartitionCreated().AddSP(this, &SWorldPartitionEditor::OnBrowseWorld);
 }
 
 SWorldPartitionEditor::~SWorldPartitionEditor()
@@ -35,8 +34,7 @@ SWorldPartitionEditor::~SWorldPartitionEditor()
 	WorldBrowserModule.OnBrowseWorld.RemoveAll(this);
 	UWorldPartition::WorldPartitionChangedEvent.RemoveAll(this);
 	
-	IWorldPartitionEditorModule& WorldPartitionEditorModule = FModuleManager::LoadModuleChecked<IWorldPartitionEditorModule>("WorldPartitionEditor");
-	WorldPartitionEditorModule.OnWorldPartitionCreated().RemoveAll(this);
+	IWorldPartitionEditorModule::Get().OnWorldPartitionCreated().RemoveAll(this);
 
 	if (World.IsValid())
 	{
diff --git a/Engine/Source/Runtime/Engine/Private/IWorldPartitionEditorModule.cpp b/Engine/Source/Runtime/Engine/Private/IWorldPartitionEditorModule.cpp
--- a/Engine/Source/Runtime/Engine/Private/IWorldPartitionEditorModule.cpp
+++ b/Engine/Source/Runtime/Engine/Private/IWorldPartitionEditorModule.cpp
@@ -3,13 +3,13 @@
 #include "WorldPartition/IWorldPartitionEditorModule.h"
 #include "Modules/ModuleManager.h"
 
-IWorldPartitionEditorModule* IWorldPartitionEditorModule::Instance = nullptr;
-
+/**
+* Singleton-like access to this module's interface.  This is just for convenience!
+* Beware of calling this during the shutdown phase, though.  Your module might have been unloaded already.
+*
+* @return Returns singleton instance, loading the module on demand if needed
+*/
 IWorldPartitionEditorModule& IWorldPartitionEditorModule::Get()
 {
-	if (!Instance)
-	{
-		Instance = &FModuleManager::LoadModuleChecked<IWorldPartitionEditorModule>("WorldPartitionEditor");
-	}
-	return *Instance;
+	return FModuleManager::LoadModuleChecked<IWorldPartitionEditorModule>("WorldPartitionEditor");
 }
\ No newline at end of file
diff --git a/Engine/Source/Runtime/Engine/Private/WorldPartition/WorldPartition.cpp b/Engine/Source/Runtime/Engine/Private/WorldPartition/WorldPartition.cpp
--- a/Engine/Source/Runtime/Engine/Private/WorldPartition/WorldPartition.cpp
+++ b/Engine/Source/Runtime/Engine/Private/WorldPartition/WorldPartition.cpp
@@ -607,8 +607,7 @@ void UWorldPartition::Initialize(UWorld* InWorld, const FTransform& InTransform)
 		{			
 			PinnedActors = new FLoaderAdapterPinnedActors(OuterWorld);
 		
-			IWorldPartitionEditorModule& WorldPartitionEditorModule = FModuleManager::LoadModuleChecked<IWorldPartitionEditorModule>("WorldPartitionEditor");
-			ForceLoadedActors = WorldPartitionEditorModule.GetEnableLoadingInEditor() ? nullptr : new FLoaderAdapterActorList(OuterWorld);
+			ForceLoadedActors = IWorldPartitionEditorModule::Get().GetEnableLoadingInEditor() ? nullptr : new FLoaderAdapterActorList(OuterWorld);
 		}
 	}
 
@@ -888,8 +887,7 @@ void UWorldPartition::OnPostBugItGoCalled(const FVector& Loc, const FRotator& Ro
 		const FVector LoadExtent(UWorldPartition::LoadingRangeBugItGo, UWorldPartition::LoadingRangeBugItGo, HALF_WORLD_MAX);
 		const FBox LoadCellsBox(Loc - LoadExtent, Loc + LoadExtent);
 
-		IWorldPartitionEditorModule& WorldPartitionEditorModule = FModuleManager::LoadModuleChecked<IWorldPartitionEditorModule>("WorldPartitionEditor");
-		if (WorldPartitionEditorModule.GetEnableLoadingInEditor())
+		if (IWorldPartitionEditorModule::Get().GetEnableLoadingInEditor())
 		{
 			UWorldPartitionEditorLoaderAdapter* EditorLoaderAdapter = CreateEditorLoaderAdapter<FLoaderAdapterShape>(World, LoadCellsBox, TEXT("BugItGo"));
 			EditorLoaderAdapter->GetLoaderAdapter()->Load();
@@ -1426,9 +1424,7 @@ void UWorldPartition::OnEnableLoadingInEditorChanged()
 		ForceLoadedActors = nullptr;
 	}
 
-	IWorldPartitionEditorModule& WorldPartitionEditorModule = FModuleManager::LoadModuleChecked<IWorldPartitionEditorModule>("WorldPartitionEditor");
-
-	if (!WorldPartitionEditorModule.GetEnableLoadingInEditor())
+	if (!IWorldPartitionEditorModule::Get().GetEnableLoadingInEditor())
 	{
 		UWorld* OuterWorld = GetTypedOuter<UWorld>();
 		check(OuterWorld);
diff --git a/Engine/Source/Runtime/Engine/Private/WorldPartition/WorldPartitionMiniMap.cpp b/Engine/Source/Runtime/Engine/Private/WorldPartition/WorldPartitionMiniMap.cpp
--- a/Engine/Source/Runtime/Engine/Private/WorldPartition/WorldPartitionMiniMap.cpp
+++ b/Engine/Source/Runtime/Engine/Private/WorldPartition/WorldPartitionMiniMap.cpp
@@ -79,8 +79,7 @@ void AWorldPartitionMiniMap::CheckForErrors()
 		int32 EffectiveWorldUnitsPerPixel = 0;
 		GetMiniMapResolution(EffectiveMinimapImageSizeX, EffectiveMinimapImageSizeY, EffectiveWorldUnitsPerPixel);
 
-		IWorldPartitionEditorModule& WorldPartitionEditorModule = FModuleManager::LoadModuleChecked<IWorldPartitionEditorModule>("WorldPartitionEditor");
-		if (EffectiveWorldUnitsPerPixel > WorldPartitionEditorModule.GetMinimapLowQualityWorldUnitsPerPixelThreshold())
+		if (EffectiveWorldUnitsPerPixel > IWorldPartitionEditorModule::Get().GetMinimapLowQualityWorldUnitsPerPixelThreshold())
 		{
 			FMessageLog("MapCheck").Warning()
 				->AddToken(FUObjectToken::Create(this))
diff --git a/Engine/Source/Runtime/Engine/Private/WorldPartition/WorldPartitionStreamingGeneration.cpp b/Engine/Source/Runtime/Engine/Private/WorldPartition/WorldPartitionStreamingGeneration.cpp
--- a/Engine/Source/Runtime/Engine/Private/WorldPartition/WorldPartitionStreamingGeneration.cpp
+++ b/Engine/Source/Runtime/Engine/Private/WorldPartition/WorldPartitionStreamingGeneration.cpp
@@ -1868,7 +1868,7 @@ bool UWorldPartition::GenerateContainerStreaming(const FGenerateStreamingParams&
 	TUniquePtr<FArchive> LogFileAr;
 	TUniquePtr<FHierarchicalLogArchive> HierarchicalLogAr;
 
-	const bool bIsStreamingGenerationLogAllowed = (!bIsPIE || FModuleManager::LoadModuleChecked<IWorldPartitionEditorModule>("WorldPartitionEditor").GetEnableStreamingGenerationLogOnPIE()) && (!InParams.OutputLogType.IsSet() || !InParams.OutputLogType.GetValue().IsEmpty());
+	const bool bIsStreamingGenerationLogAllowed = (!bIsPIE || IWorldPartitionEditorModule::Get().GetEnableStreamingGenerationLogOnPIE()) && (!InParams.OutputLogType.IsSet() || !InParams.OutputLogType.GetValue().IsEmpty());
 	const bool bIsStreamingGenerationLogRelevant = IsMainWorldPartition() && (!GIsBuildMachine || GIsAutomationTesting || IsRunningCookCommandlet());
 
 	if (bIsStreamingGenerationLogAllowed && bIsStreamingGenerationLogRelevant)
diff --git a/Engine/Source/Runtime/Engine/Public/WorldPartition/IWorldPartitionEditorModule.h b/Engine/Source/Runtime/Engine/Public/WorldPartition/IWorldPartitionEditorModule.h
--- a/Engine/Source/Runtime/Engine/Public/WorldPartition/IWorldPartitionEditorModule.h
+++ b/Engine/Source/Runtime/Engine/Public/WorldPartition/IWorldPartitionEditorModule.h
@@ -92,7 +92,4 @@ public:
 
 	/** Return the commandlet post-execution event */
 	virtual FOnPostExecuteCommandlet& OnPostExecuteCommandlet() = 0;
-
-private:
-	static IWorldPartitionEditorModule* Instance;
 };
\ No newline at end of file
