Fbx Exporter: fix LOD 0 material export to be able to have multiple section using the same material.

#rnx

diff --git a/Engine/Source/Editor/UnrealEd/Private/Fbx/FbxSkeletalMeshExport.cpp b/Engine/Source/Editor/UnrealEd/Private/Fbx/FbxSkeletalMeshExport.cpp
--- a/Engine/Source/Editor/UnrealEd/Private/Fbx/FbxSkeletalMeshExport.cpp
+++ b/Engine/Source/Editor/UnrealEd/Private/Fbx/FbxSkeletalMeshExport.cpp
@@ -363,7 +363,7 @@ FbxNode* FFbxExporter::CreateMesh(const USkeletalMesh* SkelMesh, const TCHAR* Me
 		FbxSurfaceMaterial* FbxMaterial = nullptr;
 		if (LODIndex == 0)
 		{
-			if (MatInterface && !FbxMaterials.Find(MatInterface))
+			if (MatInterface)
 			{
 				FbxMaterial = ExportMaterial(MatInterface);
 			}
