CookPackageSplitter and MPCook: Rename NeedCachedPlatformDataBeforeSplit to GeneratedReliesOnGeneratorSave.
GeneratedReliesOnGeneratorSave now causes the splitter to cook generated on the same CookWorker, but also to use the slower but safer IsGeneratedSavedFirst ordering.
Fix DecideNextCookAction to correctly handle the possible dependency of a generator upon its generated when the generator and generated are urgent due to cook on the fly or -cookfirst.
#rnx

diff --git a/Engine/Plugins/Experimental/Mutable/Source/CustomizableObjectEditor/Private/MuCOE/CustomizableObjectCookPackageSplitter.h b/Engine/Plugins/Experimental/Mutable/Source/CustomizableObjectEditor/Private/MuCOE/CustomizableObjectCookPackageSplitter.h
--- a/Engine/Plugins/Experimental/Mutable/Source/CustomizableObjectEditor/Private/MuCOE/CustomizableObjectCookPackageSplitter.h
+++ b/Engine/Plugins/Experimental/Mutable/Source/CustomizableObjectEditor/Private/MuCOE/CustomizableObjectCookPackageSplitter.h
@@ -47,10 +47,11 @@ public:
 	 */
 	virtual bool UseInternalReferenceToAvoidGarbageCollect() override { return true; }
 	/**
-	 * If true, the cooker will not call any splitter work functions (GetGenerateList, Populate*, PreSave*, PostSave*)
-	 * until BeginCacheForCookedPlatformData has been called for all objects in the generator package.
+	 * If true, the cooker will wait for the generator save to complete before calling Populate and PreSave on the
+	 * generated package. And during MPCook the cooker will load and save generated packages only on the same
+	 * CookWorker that saved the generator package that containes the SplitData object. 
 	 */
-	virtual bool NeedCachedPlatformDataBeforeSplit() override { return true; }
+	virtual bool GeneratedReliesOnGeneratorSave() override { return true; }
 
 private:
 
diff --git a/Engine/Source/Editor/UnrealEd/Private/CookOnTheFlyServer.cpp b/Engine/Source/Editor/UnrealEd/Private/CookOnTheFlyServer.cpp
--- a/Engine/Source/Editor/UnrealEd/Private/CookOnTheFlyServer.cpp
+++ b/Engine/Source/Editor/UnrealEd/Private/CookOnTheFlyServer.cpp
@@ -2316,17 +2316,17 @@ UCookOnTheFlyServer::ECookAction UCookOnTheFlyServer::DecideNextCookAction(UE::C
 	UE::Cook::FPackageDataMonitor& Monitor = PackageDatas->GetMonitor();
 	if (Monitor.GetNumUrgent() > 0)
 	{
-		if (Monitor.GetNumUrgent(UE::Cook::EPackageState::Save) > 0)
+		if (!bSaveBusy && Monitor.GetNumUrgent(UE::Cook::EPackageState::Save) > 0)
 		{
 			SetIdleStatus(StackData, EIdleStatus::Active);
 			return ECookAction::Save;
 		}
-		else if (Monitor.GetNumUrgent(UE::Cook::EPackageState::LoadPrepare) > 0)
+		else if (!bLoadBusy && Monitor.GetNumUrgent(UE::Cook::EPackageState::LoadPrepare) > 0)
 		{
 			SetIdleStatus(StackData, EIdleStatus::Active);
 			return ECookAction::Load;
 		}
-		else if (Monitor.GetNumUrgent(UE::Cook::EPackageState::LoadReady) > 0)
+		else if (!bLoadBusy && Monitor.GetNumUrgent(UE::Cook::EPackageState::LoadReady) > 0)
 		{
 			SetIdleStatus(StackData, EIdleStatus::Active);
 			return ECookAction::Load;
@@ -2336,6 +2336,21 @@ UCookOnTheFlyServer::ECookAction UCookOnTheFlyServer::DecideNextCookAction(UE::C
 			SetIdleStatus(StackData, EIdleStatus::Active);
 			return ECookAction::Request;
 		}
+		else if (Monitor.GetNumUrgent(UE::Cook::EPackageState::Save) > 0)
+		{
+			SetIdleStatus(StackData, EIdleStatus::Active);
+			return ECookAction::Save;
+		}
+		else if (Monitor.GetNumUrgent(UE::Cook::EPackageState::LoadPrepare) > 0)
+		{
+			SetIdleStatus(StackData, EIdleStatus::Active);
+			return ECookAction::Load;
+		}
+		else if (Monitor.GetNumUrgent(UE::Cook::EPackageState::LoadReady) > 0)
+		{
+			SetIdleStatus(StackData, EIdleStatus::Active);
+			return ECookAction::Load;
+		}
 
 		if (Monitor.GetNumUrgent(UE::Cook::EPackageState::AssignedToWorker) > 0)
 		{
@@ -3330,7 +3345,13 @@ UE::Cook::EPollStatus UCookOnTheFlyServer::PrepareSaveGenerationPackage(UE::Cook
 		{
 			return EPollStatus::Incomplete;
 		}
-		if (FGenerationHelper::IsGeneratedSavedFirst() && Info.IsGenerator())
+
+		// If generator should not save until after generated, stall it here
+		if (Info.IsGenerator()
+			&& FGenerationHelper::IsGeneratedSavedFirst()
+			// Splitters that declare GeneratedReliesOnGeneratorSave ignore the global setting and
+			// never wait for generated to save
+			&& !GenerationHelper.IsGeneratedReliesOnGeneratorSave())
 		{
 			if (GenerationHelper.IsWaitingForQueueResults())
 			{
@@ -3344,7 +3365,13 @@ UE::Cook::EPollStatus UCookOnTheFlyServer::PrepareSaveGenerationPackage(UE::Cook
 				}
 			}
 		}
-		if (FGenerationHelper::IsGeneratorSavedFirst() && !Info.IsGenerator())
+
+		// If generated should not save until after generator, stall it here
+		if (!Info.IsGenerator()
+			&& (FGenerationHelper::IsGeneratorSavedFirst()
+			// Splitters that declare GeneratedReliesOnGeneratorSave ignore the global setting and
+			// always wait for the generator to save
+				|| GenerationHelper.IsGeneratedReliesOnGeneratorSave()))
 		{
 			if (GenerationHelper.GetOwner().IsInProgress())
 			{
@@ -11484,10 +11511,10 @@ void UCookOnTheFlyServer::GetPackagesToRetract(int32 NumToRetract, TArray<FName>
 		}
 		if (PackageData->IsGenerated())
 		{
-			if (PackageData->IsGeneratedNeedCachedPlatformDataBeforeSplit()
+			if (PackageData->IsGeneratedReliesOnGeneratorSave()
 				|| MPCookGeneratorSplit == UE::Cook::EMPCookGeneratorSplit::AllOnSameWorker)
 			{
-				// With IsNeedCachedPlatformDataBeforeSplit or the AllOnSameWorker setting, GeneratedPackages are
+				// With IsGeneratedReliesOnGeneratorSave or the AllOnSameWorker setting, GeneratedPackages are
 				// constrained to this worker.
 				return false;
 			}
@@ -11498,10 +11525,10 @@ void UCookOnTheFlyServer::GetPackagesToRetract(int32 NumToRetract, TArray<FName>
 				GenerationHelper->GetOwnerInfo().GetSaveState()
 					>= FCookGenerationInfo::ESaveState::QueueGeneratedPackages)
 			{
-				if (GenerationHelper->IsNeedCachedPlatformDataBeforeSplit() ||
+				if (GenerationHelper->IsGeneratedReliesOnGeneratorSave() ||
 					MPCookGeneratorSplit != UE::Cook::EMPCookGeneratorSplit::AnyWorker)
 				{
-					// With IsNeedCachedPlatformDataBeforeSplit or with any MPCookGeneratorSplit setting other than
+					// With IsGeneratedReliesOnGeneratorSave or with any MPCookGeneratorSplit setting other than
 					// AnyWorker, we make assignment decisions based on the worker that saved and queued the generator
 					// package. We do not track queuing separately; we assume it happened on the worker that saved the
 					// package. Therefore, do not allow retraction of a generator package if it has already entered
diff --git a/Engine/Source/Editor/UnrealEd/Private/Cooker/CookGenerationHelper.cpp b/Engine/Source/Editor/UnrealEd/Private/Cooker/CookGenerationHelper.cpp
--- a/Engine/Source/Editor/UnrealEd/Private/Cooker/CookGenerationHelper.cpp
+++ b/Engine/Source/Editor/UnrealEd/Private/Cooker/CookGenerationHelper.cpp
@@ -91,8 +91,8 @@ void FGenerationHelper::Initialize(const UObject* InSplitDataObject,
 	SplitDataObjectName = FName(FStringView(InSplitDataObject->GetFullName()));
 	bUseInternalReferenceToAvoidGarbageCollect =
 		CookPackageSplitterInstance->UseInternalReferenceToAvoidGarbageCollect();
-	bNeedCachedPlatformDataBeforeSplit =
-		CookPackageSplitterInstance->NeedCachedPlatformDataBeforeSplit();
+	bGeneratedReliesOnGeneratorSave =
+		CookPackageSplitterInstance->GeneratedReliesOnGeneratorSave();
 }
 
 void FGenerationHelper::InitializeAsInvalid()
@@ -381,7 +381,7 @@ bool FGenerationHelper::TryGenerateList()
 		check(PackageData->GetParentGenerator().IsNone() ||
 			PackageData->GetParentGenerator() == OwnerPackageName);
 		PackageData->SetGenerated(OwnerPackageName);
-		PackageData->SetGeneratedNeedCachedPlatformDataBeforeSplit(bNeedCachedPlatformDataBeforeSplit);
+		PackageData->SetGeneratedReliesOnGeneratorSave(bGeneratedReliesOnGeneratorSave);
 		if (IFileManager::Get().FileExists(*PackageData->GetFileName().ToString()))
 		{
 			UE_LOG(LogCook, Warning,
@@ -408,7 +408,7 @@ bool FGenerationHelper::TryGenerateList()
 			[](const FAssetDependency& A, const FAssetDependency& B) { return A.LexicalLess(B); });
 		GeneratedInfo.PackageDependencies.SetNum(Algo::Unique(GeneratedInfo.PackageDependencies));
 		GeneratedInfo.SetIsCreateAsMap(bCreateAsMap);
-		if (bNeedCachedPlatformDataBeforeSplit ||
+		if (bGeneratedReliesOnGeneratorSave ||
 			COTFS.MPCookGeneratorSplit == EMPCookGeneratorSplit::AllOnSameWorker)
 		{
 			PackageData->SetWorkerAssignmentConstraint(FWorkerId::Local());
diff --git a/Engine/Source/Editor/UnrealEd/Private/Cooker/CookGenerationHelper.h b/Engine/Source/Editor/UnrealEd/Private/Cooker/CookGenerationHelper.h
--- a/Engine/Source/Editor/UnrealEd/Private/Cooker/CookGenerationHelper.h
+++ b/Engine/Source/Editor/UnrealEd/Private/Cooker/CookGenerationHelper.h
@@ -221,8 +221,8 @@ public:
 	const FName GetSplitDataObjectName() const;
 	/** Return the Splitter's value for virtual bool UseInternalReferenceToAvoidGarbageCollect(). */
 	bool IsUseInternalReferenceToAvoidGarbageCollect() const;
-	/** Return the Splitter's value for virtual bool NeedCachedPlatformDataBeforeSplit(). */
-	bool IsNeedCachedPlatformDataBeforeSplit() const;
+	/** Return the Splitter's value for virtual bool GeneratedReliesOnGeneratorSave(). */
+	bool IsGeneratedReliesOnGeneratorSave() const;
 	/** Return the cached pointer to the SplitDataObject. Returns null if no longer in memory or marked as garbage. */
 	UObject* GetWeakSplitDataObject() const;
 	/**
@@ -370,7 +370,7 @@ private:
 	FWorkerId WorkerIdThatSavedGenerator = FWorkerId::Invalid();
 	EInitializeStatus InitializeStatus = EInitializeStatus::Uninitialized;
 	bool bUseInternalReferenceToAvoidGarbageCollect = false;
-	bool bNeedCachedPlatformDataBeforeSplit = false;
+	bool bGeneratedReliesOnGeneratorSave = false;
 	bool bGeneratedList = false;
 };
 
@@ -519,10 +519,10 @@ inline bool FGenerationHelper::IsUseInternalReferenceToAvoidGarbageCollect() con
 	return bUseInternalReferenceToAvoidGarbageCollect;
 }
 
-inline bool FGenerationHelper::IsNeedCachedPlatformDataBeforeSplit() const
+inline bool FGenerationHelper::IsGeneratedReliesOnGeneratorSave() const
 {
 	ConditionalInitialize();
-	return bNeedCachedPlatformDataBeforeSplit;
+	return bGeneratedReliesOnGeneratorSave;
 }
 
 inline UObject* FGenerationHelper::GetWeakSplitDataObject() const
diff --git a/Engine/Source/Editor/UnrealEd/Private/Cooker/CookPackageData.cpp b/Engine/Source/Editor/UnrealEd/Private/Cooker/CookPackageData.cpp
--- a/Engine/Source/Editor/UnrealEd/Private/Cooker/CookPackageData.cpp
+++ b/Engine/Source/Editor/UnrealEd/Private/Cooker/CookPackageData.cpp
@@ -98,7 +98,7 @@ FPackageData::FPackageData(FPackageDatas& PackageDatas, const FName& InPackageNa
 	, bCookedPlatformDataStarted(0), bCookedPlatformDataCalled(0), bCookedPlatformDataComplete(0)
 	, MonitorCookResult((uint8)ECookResult::NotAttempted)
 	, bInitializedGeneratorSave(0), bCompletedGeneration(0), bGenerated(0), bKeepReferencedDuringGC(0)
-	, bWasCookedThisSession(0), bGeneratedNeedCachedPlatformDataBeforeSplit(0)
+	, bWasCookedThisSession(0), bGeneratedReliesOnGeneratorSave(0)
 {
 	SetState(EPackageState::Idle);
 	SendToState(EPackageState::Idle, ESendFlags::QueueAdd, EStateChangeReason::Discovered);
diff --git a/Engine/Source/Editor/UnrealEd/Private/Cooker/CookPackageData.h b/Engine/Source/Editor/UnrealEd/Private/Cooker/CookPackageData.h
--- a/Engine/Source/Editor/UnrealEd/Private/Cooker/CookPackageData.h
+++ b/Engine/Source/Editor/UnrealEd/Private/Cooker/CookPackageData.h
@@ -640,9 +640,9 @@ public:
 	 */
 	TRefCountPtr<FGenerationHelper> TryCreateValidParentGenerationHelper();
 
-	/** Get/Set true if this is a generated package and its splitter reports NeedCachedPlatformDataBeforeSplit. */
-	bool IsGeneratedNeedCachedPlatformDataBeforeSplit() const;
-	void SetGeneratedNeedCachedPlatformDataBeforeSplit(bool bValue);
+	/** Get/Set true if this is a generated package and its splitter reports GeneratedReliesOnGeneratorSave. */
+	bool IsGeneratedReliesOnGeneratorSave() const;
+	void SetGeneratedReliesOnGeneratorSave(bool bValue);
 
 	/**
 	 * Return the instigator for this package. The Instigator is the first code location or
@@ -825,7 +825,7 @@ private:
 	uint32 bGenerated : 1;
 	uint32 bKeepReferencedDuringGC : 1;
 	uint32 bWasCookedThisSession : 1;
-	uint32 bGeneratedNeedCachedPlatformDataBeforeSplit : 1;
+	uint32 bGeneratedReliesOnGeneratorSave : 1;
 };
 
 /**
@@ -1469,14 +1469,14 @@ inline FName FPackageData::GetParentGenerator() const
 	return ParentGenerator;
 }
 
-inline bool FPackageData::IsGeneratedNeedCachedPlatformDataBeforeSplit() const
+inline bool FPackageData::IsGeneratedReliesOnGeneratorSave() const
 {
-	return bGeneratedNeedCachedPlatformDataBeforeSplit != 0;
+	return bGeneratedReliesOnGeneratorSave != 0;
 }
 
-inline void FPackageData::SetGeneratedNeedCachedPlatformDataBeforeSplit(bool bValue)
+inline void FPackageData::SetGeneratedReliesOnGeneratorSave(bool bValue)
 {
-	bGeneratedNeedCachedPlatformDataBeforeSplit = (uint32)bValue;
+	bGeneratedReliesOnGeneratorSave = (uint32)bValue;
 }
 
 } // namespace UE::Cook
\ No newline at end of file
diff --git a/Engine/Source/Editor/UnrealEd/Private/Cooker/CookWorkerClient.cpp b/Engine/Source/Editor/UnrealEd/Private/Cooker/CookWorkerClient.cpp
--- a/Engine/Source/Editor/UnrealEd/Private/Cooker/CookWorkerClient.cpp
+++ b/Engine/Source/Editor/UnrealEd/Private/Cooker/CookWorkerClient.cpp
@@ -318,7 +318,7 @@ void FCookWorkerClient::ReportDiscoveredPackage(const FPackageData& PackageData,
 	Discovered.Instigator = Instigator;
 	Discovered.Platforms = MoveTemp(ReachablePlatforms);
 	Discovered.Platforms.ConvertToBitfield(OrderedSessionAndSpecialPlatforms);
-	Discovered.bNeedCachedPlatformDataBeforeSplit = PackageData.IsGeneratedNeedCachedPlatformDataBeforeSplit();
+	Discovered.bGeneratedReliesOnGeneratorSave = PackageData.IsGeneratedReliesOnGeneratorSave();
 }
 
 void FCookWorkerClient::ReportGeneratorQueuedGeneratedPackages(FGenerationHelper& GenerationHelper)
diff --git a/Engine/Source/Editor/UnrealEd/Private/Cooker/CookWorkerServer.cpp b/Engine/Source/Editor/UnrealEd/Private/Cooker/CookWorkerServer.cpp
--- a/Engine/Source/Editor/UnrealEd/Private/Cooker/CookWorkerServer.cpp
+++ b/Engine/Source/Editor/UnrealEd/Private/Cooker/CookWorkerServer.cpp
@@ -685,7 +685,7 @@ void FCookWorkerServer::SendPendingPackages()
 		FAssignPackageData& AssignData = AssignDatas.Emplace_GetRef();
 		AssignData.ConstructData = PackageData->CreateConstructData();
 		AssignData.ParentGenerator = PackageData->GetParentGenerator();
-		// We do not record PackageData->IsNeedCachedPlatformDataBeforeSplit(). The CookWorker that created
+		// We do not record PackageData->IsGeneratedReliesOnGeneratorSave(). The CookWorker that created
 		// it knows it, and the Director knows it from FDiscoveredPackageReplication, and the package will not
 		// be sent to other CookWorkers if it is true.
 		AssignData.Instigator = PackageData->GetInstigator();
@@ -1000,7 +1000,7 @@ void FCookWorkerServer::QueueDiscoveredPackage(FDiscoveredPackageReplication&& D
 	if (!DiscoveredPackage.ParentGenerator.IsNone())
 	{
 		PackageData.SetGenerated(DiscoveredPackage.ParentGenerator);
-		PackageData.SetGeneratedNeedCachedPlatformDataBeforeSplit(DiscoveredPackage.bNeedCachedPlatformDataBeforeSplit);
+		PackageData.SetGeneratedReliesOnGeneratorSave(DiscoveredPackage.bGeneratedReliesOnGeneratorSave);
 		FPackageData* GeneratorPackageData = PackageDatas.FindPackageDataByPackageName(
 			DiscoveredPackage.ParentGenerator);
 		if (GeneratorPackageData)
@@ -1047,7 +1047,7 @@ void FCookWorkerServer::QueueDiscoveredPackage(FDiscoveredPackageReplication&& D
 	}
 
 	if (PackageData.IsGenerated()
-		&& (PackageData.IsGeneratedNeedCachedPlatformDataBeforeSplit()
+		&& (PackageData.IsGeneratedReliesOnGeneratorSave()
 				|| COTFS.MPCookGeneratorSplit == EMPCookGeneratorSplit::AllOnSameWorker))
 	{
 		PackageData.SetWorkerAssignmentConstraint(GetWorkerId());
@@ -1374,7 +1374,7 @@ void FDiscoveredPackageReplication::Write(FCbWriter& Writer,
 	Writer << ParentGenerator;
 	Writer << static_cast<uint8>(Instigator.Category);
 	Writer << Instigator.Referencer;
-	Writer << bNeedCachedPlatformDataBeforeSplit;
+	Writer << bGeneratedReliesOnGeneratorSave;
 	WriteToCompactBinary(Writer, Platforms, OrderedSessionAndSpecialPlatforms);
 	Writer.EndArray();
 }
@@ -1404,7 +1404,7 @@ bool FDiscoveredPackageReplication::TryRead(FCbFieldView Field,
 		bOk = false;
 	}
 	bOk = LoadFromCompactBinary(Iter++, Instigator.Referencer) & bOk;
-	bNeedCachedPlatformDataBeforeSplit = Iter->AsBool();
+	bGeneratedReliesOnGeneratorSave = Iter->AsBool();
 	bOk = (!(Iter++)->HasError()) & bOk;
 	bOk = LoadFromCompactBinary(Iter++, Platforms, OrderedSessionAndSpecialPlatforms) & bOk;
 	if (!bOk)
diff --git a/Engine/Source/Editor/UnrealEd/Private/Cooker/CookWorkerServer.h b/Engine/Source/Editor/UnrealEd/Private/Cooker/CookWorkerServer.h
--- a/Engine/Source/Editor/UnrealEd/Private/Cooker/CookWorkerServer.h
+++ b/Engine/Source/Editor/UnrealEd/Private/Cooker/CookWorkerServer.h
@@ -400,7 +400,7 @@ struct FDiscoveredPackageReplication
 	FName ParentGenerator;
 	FInstigator Instigator;
 	FDiscoveredPlatformSet Platforms;
-	bool bNeedCachedPlatformDataBeforeSplit = false;
+	bool bGeneratedReliesOnGeneratorSave = false;
 
 private:
 	void Write(FCbWriter& Writer, TConstArrayView<const ITargetPlatform*> OrderedSessionAndSpecialPlatforms) const;
diff --git a/Engine/Source/Editor/UnrealEd/Public/CookPackageSplitter.h b/Engine/Source/Editor/UnrealEd/Public/CookPackageSplitter.h
--- a/Engine/Source/Editor/UnrealEd/Public/CookPackageSplitter.h
+++ b/Engine/Source/Editor/UnrealEd/Public/CookPackageSplitter.h
@@ -63,12 +63,16 @@ public:
 		return false;
 	}
 	/**
-	 * If true, the cooker will not call any splitter work functions (GetGenerateList, Populate*, PreSave*, PostSave*)
-	 * until BeginCacheForCookedPlatformData has been called for all objects in the generator package. This reduces
-	 * MPCook performance - generated packages have to be saved on the same Worker that saved the generator - but is
-	 * otherwise not a problem.
+	 * If true, the cooker will wait for the generator save to complete before calling Populate and PreSave on the
+	 * generated package. And during MPCook the cooker will load and save generated packages only on the same
+	 * CookWorker that saved the generator package that containes the SplitData object. This election reduces cook
+	 * performance but is otherwise not a problem.
+	 * 
+	 * Examples of dependencies:
+	 *     ShouldSplit call reads data that is written by BeginCacheForCookedPlatformData
+	 *     PopulateGeneratedPackage or PreSaveGeneratedPackage read data that is written by PopulateGeneratorPackage.
 	 */
-	virtual bool NeedCachedPlatformDataBeforeSplit() 
+	virtual bool GeneratedReliesOnGeneratorSave() 
 	{
 		 return false;
 	 }
