[Nanite-Skinning] Disable Nanite rendering (proxy and mesh object) when used with UDebugSkelMeshComponent.

This fixes missing renders in the skeletal mesh editor when Nanite is enabled on a mesh. Rendering falls back to non-Nanite in editor window, but all workflows remain functional. Future improvements would be Nanite-focused tooling (preview, GPU driven debug rendering, etc..)

diff --git a/Engine/Source/Editor/UnrealEd/Classes/Animation/DebugSkelMeshComponent.h b/Engine/Source/Editor/UnrealEd/Classes/Animation/DebugSkelMeshComponent.h
--- a/Engine/Source/Editor/UnrealEd/Classes/Animation/DebugSkelMeshComponent.h
+++ b/Engine/Source/Editor/UnrealEd/Classes/Animation/DebugSkelMeshComponent.h
@@ -376,6 +376,7 @@ class UDebugSkelMeshComponent : public USkeletalMeshComponent
 
 	//~ Begin SkinnedMeshComponent Interface
 	UNREALED_API virtual bool ShouldCPUSkin() override;
+	UNREALED_API virtual bool ShouldNaniteSkin() override;
 	UNREALED_API virtual void PostInitMeshObject(class FSkeletalMeshObject* MeshObject) override;
 	UNREALED_API virtual void RefreshBoneTransforms(FActorComponentTickFunction* TickFunction = NULL) override;
 	virtual int32 GetLODBias() const override { return 0; }
diff --git a/Engine/Source/Editor/UnrealEd/Private/Animation/DebugSkelMeshComponent.cpp b/Engine/Source/Editor/UnrealEd/Private/Animation/DebugSkelMeshComponent.cpp
--- a/Engine/Source/Editor/UnrealEd/Private/Animation/DebugSkelMeshComponent.cpp
+++ b/Engine/Source/Editor/UnrealEd/Private/Animation/DebugSkelMeshComponent.cpp
@@ -708,6 +708,10 @@ bool UDebugSkelMeshComponent::ShouldCPUSkin()
 	return 	GetCPUSkinningEnabled() || bDrawBoneInfluences || bDrawNormals || bDrawTangents || bDrawBinormals || bDrawMorphTargetVerts;
 }
 
+bool UDebugSkelMeshComponent::ShouldNaniteSkin()
+{
+	return false;
+}
 
 void UDebugSkelMeshComponent::PostInitMeshObject(FSkeletalMeshObject* InMeshObject)
 {
diff --git a/Engine/Source/Runtime/Engine/Private/Components/SkinnedMeshComponent.cpp b/Engine/Source/Runtime/Engine/Private/Components/SkinnedMeshComponent.cpp
--- a/Engine/Source/Runtime/Engine/Private/Components/SkinnedMeshComponent.cpp
+++ b/Engine/Source/Runtime/Engine/Private/Components/SkinnedMeshComponent.cpp
@@ -1014,7 +1014,6 @@ void USkinnedMeshComponent::CreateRenderState_Concurrent(FRegisterComponentConte
 					// Also check if skeletal mesh has too many bones/chunk for GPU skinning.
 					if (ShouldNaniteSkin())
 					{
-						// TODO: Nanite-Skinning, do we need to allow falling back to CPU skinning for certain editor tools?
 						MeshObject = ::new FSkeletalMeshObjectNanite(this, SkelMeshRenderData, SceneFeatureLevel);
 					}
 					else if (bRenderStatic)
