World Partition - Builder: As we manage GC frequency ourselves during builds, make sure to turn off automated GC from level streaming

diff --git a/Engine/Source/Editor/UnrealEd/Private/WorldPartition/WorldPartitionBuilder.cpp b/Engine/Source/Editor/UnrealEd/Private/WorldPartition/WorldPartitionBuilder.cpp
--- a/Engine/Source/Editor/UnrealEd/Private/WorldPartition/WorldPartitionBuilder.cpp
+++ b/Engine/Source/Editor/UnrealEd/Private/WorldPartition/WorldPartitionBuilder.cpp
@@ -8,6 +8,7 @@
 #include "EngineModule.h"
 #include "UObject/Linker.h"
 #include "HAL/PlatformFileManager.h"
+#include "Engine/CoreSettings.h"
 #include "Engine/World.h"
 #include "WorldPartition/WorldPartition.h"
 #include "WorldPartition/WorldPartitionHelpers.h"
@@ -170,6 +171,12 @@ FWorldBuilderCellCoord FCellInfo::GetCellCount(const FBox& InBounds, const int32
 
 bool UWorldPartitionBuilder::Run(UWorld* World, FPackageSourceControlHelper& PackageHelper)
 {
+	// As we manage GC frequency ourselves during builds, make sure to turn off automated GC based on other settings
+	TGuardValue<int32> Guard_GLevelStreamingContinuouslyIncrementalGCWhileLevelsPendingPurgeOverride(GLevelStreamingContinuouslyIncrementalGCWhileLevelsPendingPurgeOverride, 0);
+	TGuardValue<int32> Guard_GLevelStreamingContinuouslyIncrementalGCWhileLevelsPendingPurge(GLevelStreamingContinuouslyIncrementalGCWhileLevelsPendingPurge, 0);
+	TGuardValue<int32> Guard_GLevelStreamingLowMemoryPendingPurgeCount(GLevelStreamingLowMemoryPendingPurgeCount, MAX_int32);
+	TGuardValue<int32> Guard_GLevelStreamingForceGCAfterLevelStreamedOut(GLevelStreamingForceGCAfterLevelStreamedOut, 0);
+
 	UWorldPartition* WorldPartition = World->GetWorldPartition();
 	check(WorldPartition || CanProcessNonPartitionedWorlds());
 
