World Partition - HLOD: When requested to build a single HLOD Layer, skip the child actors completely

diff --git a/Engine/Source/Editor/UnrealEd/Private/WorldPartition/WorldPartitionHLODsBuilder.cpp b/Engine/Source/Editor/UnrealEd/Private/WorldPartition/WorldPartitionHLODsBuilder.cpp
--- a/Engine/Source/Editor/UnrealEd/Private/WorldPartition/WorldPartitionHLODsBuilder.cpp
+++ b/Engine/Source/Editor/UnrealEd/Private/WorldPartition/WorldPartitionHLODsBuilder.cpp
@@ -403,31 +403,33 @@ bool UWorldPartitionHLODsBuilder::BuildHLODActors()
 		{
 			TRACE_BOOKMARK(TEXT("BuildHLOD Start - %d"), CurrentActor);
 
-			const FGuid& HLODActorGuid = HLODActorsToBuild[CurrentActor];
+			{
+				const FGuid& HLODActorGuid = HLODActorsToBuild[CurrentActor];
 
-			FWorldPartitionReference ActorRef(WorldPartition, HLODActorGuid);
+				FWorldPartitionReference ActorRef(WorldPartition, HLODActorGuid);
 
-			AWorldPartitionHLOD* HLODActor = CastChecked<AWorldPartitionHLOD>(ActorRef.GetActor());
+				AWorldPartitionHLOD* HLODActor = CastChecked<AWorldPartitionHLOD>(ActorRef.GetActor());
 
-			UE_LOG(LogWorldPartitionHLODsBuilder, Display, TEXT("[%d / %d] Building HLOD actor %s..."), CurrentActor + 1, HLODActorsToBuild.Num(), *HLODActor->GetActorLabel());
+				UE_LOG(LogWorldPartitionHLODsBuilder, Display, TEXT("[%d / %d] Building HLOD actor %s..."), CurrentActor + 1, HLODActorsToBuild.Num(), *HLODActor->GetActorLabel());
 
-			// Simulate an engine tick to make sure engine & render resources that are queued for deletion are processed.
-			FWorldPartitionHelpers::FakeEngineTick(World);
+				// Simulate an engine tick to make sure engine & render resources that are queued for deletion are processed.
+				FWorldPartitionHelpers::FakeEngineTick(World);
 
-			HLODActor->BuildHLOD(bForceBuild);
+				HLODActor->BuildHLOD(bForceBuild);
 
-			bool bSaved = SaveHLODActor(HLODActor);
-			if (!bSaved)
-			{
-				return false;
+				bool bSaved = SaveHLODActor(HLODActor);
+				if (!bSaved)
+				{
+					return false;
+				}
 			}
 
+			TRACE_BOOKMARK(TEXT("BuildHLOD End - %d"), CurrentActor);
+
 			if (FWorldPartitionHelpers::ShouldCollectGarbage())
 			{
 				FWorldPartitionHelpers::DoCollectGarbage();
 			}
-
-			TRACE_BOOKMARK(TEXT("BuildHLOD End - %d"), CurrentActor);
 		}
 
 		UE_LOG(LogWorldPartitionHLODsBuilder, Display, TEXT("#### Built %d HLOD actors ####"), HLODActorsToBuild.Num());
@@ -632,7 +634,15 @@ TArray<TArray<FGuid>> UWorldPartitionHLODsBuilder::GetHLODWorkloads(int32 NumWor
 			continue;
 		}
 
-		HLODParenting.Add(HLODIterator->GetGuid(), HLODActorDesc.GetChildHLODActors());
+		// When requested to build a single HLOD Layer, skip the child actors
+		if (HLODLayerToBuild.IsNone())
+		{
+			HLODParenting.Add(HLODIterator->GetGuid(), HLODActorDesc.GetChildHLODActors());
+		}
+		else
+		{
+			HLODParenting.Add(HLODIterator->GetGuid());
+		}
 	}
 
 	// All child HLODs must be built before their parent HLOD
@@ -717,14 +727,18 @@ bool UWorldPartitionHLODsBuilder::ValidateWorkload(const TArray<FGuid>& Workload
 			return false;
 		}
 
-		const FHLODActorDesc* HLODActorDesc = static_cast<const FHLODActorDesc*>(ActorDescInstance->GetActorDesc());
-
-		for (const FGuid& ChildHLODActorGuid : HLODActorDesc->GetChildHLODActors())
+		// When requested to build a single HLOD Layer, do not validate that child actors are included
+		if (HLODLayerToBuild.IsNone())
 		{
-			if (!ProcessedHLOD.Contains(ChildHLODActorGuid))
+			const FHLODActorDesc* HLODActorDesc = static_cast<const FHLODActorDesc*>(ActorDescInstance->GetActorDesc());
+
+			for (const FGuid& ChildHLODActorGuid : HLODActorDesc->GetChildHLODActors())
 			{
-				UE_LOG(LogWorldPartitionHLODsBuilder, Error, TEXT("Child HLOD actor missing or out of order in HLOD workload, exiting..."));
-				return false;
+				if (!ProcessedHLOD.Contains(ChildHLODActorGuid))
+				{
+					UE_LOG(LogWorldPartitionHLODsBuilder, Error, TEXT("Child HLOD actor missing or out of order in HLOD workload, exiting..."));
+					return false;
+				}
 			}
 		}
 
