Fix warning V1052: Declaring virtual methods in a class marked as 'final' is pointless.
#rnx

diff --git a/Engine/Source/Editor/UnrealEd/Private/PlayLevel.cpp b/Engine/Source/Editor/UnrealEd/Private/PlayLevel.cpp
--- a/Engine/Source/Editor/UnrealEd/Private/PlayLevel.cpp
+++ b/Engine/Source/Editor/UnrealEd/Private/PlayLevel.cpp
@@ -171,24 +171,24 @@ public:
 		}
 	}
 
-	virtual bool CanBeUsedOnMultipleThreads() const final
+	virtual bool CanBeUsedOnMultipleThreads() const override
 	{
 		return true;
 	}
 
 	// FTickableEditorObject Interface
 
-	virtual void Tick(float DeltaTime) final
+	virtual void Tick(float DeltaTime) override
 	{
 		QueuedLines.ConsumeAllFifo(LogLine);
 	}
 
-	virtual ETickableTickType GetTickableTickType() const final
+	virtual ETickableTickType GetTickableTickType() const override
 	{
 		return ETickableTickType::Always;
 	}
 
-	virtual TStatId GetStatId() const final
+	virtual TStatId GetStatId() const override
 	{
 		RETURN_QUICK_DECLARE_CYCLE_STAT(FOutputLogErrorsToMessageLogProxy, STATGROUP_Tickables);
 	}
