Improvements to CookShaders commandlet
Can identify all relevant materials instances in a few seconds
Can cook Niagara
Added validation warnings for any missing but requested shader symbols

#tests local testing console platforms from ShaderSymbols.info from complex cooks

diff --git a/Engine/Source/Editor/UnrealEd/Private/Commandlets/CookShadersCommandlet.cpp b/Engine/Source/Editor/UnrealEd/Private/Commandlets/CookShadersCommandlet.cpp
--- a/Engine/Source/Editor/UnrealEd/Private/Commandlets/CookShadersCommandlet.cpp
+++ b/Engine/Source/Editor/UnrealEd/Private/Commandlets/CookShadersCommandlet.cpp
@@ -173,6 +173,23 @@ namespace CookShadersCommandlet {
 
 		return false;
 	}
+
+	bool GetParentName(const FAssetData* InAssetData, FString& ParentName)
+	{
+		static const FName NAME_Parent = TEXT("Parent");
+		FString ParentPathString = InAssetData->GetTagValueRef<FString>(NAME_Parent);
+
+		int32 FirstCut = INDEX_NONE;
+		ParentPathString.FindChar(L'\'', FirstCut);
+
+		if (FirstCut != INDEX_NONE)
+		{
+			ParentName = ParentPathString.Mid(FirstCut + 1, ParentPathString.Len() - FirstCut - 2);
+			return true;
+		}
+
+		return false;
+	}
 };
 
 using namespace CookShadersCommandlet;
@@ -193,7 +210,7 @@ int32 UCookShadersCommandlet::Main(const FString& Params)
 	if (Switches.Contains("help"))
 	{
 		UE_LOG(LogCookShadersCommandlet, Log, TEXT("CookShadersCommandlet"));
-		UE_LOG(LogCookShadersCommandlet, Log, TEXT("Cook shaders based upon the options, ideal for generating pdbs for shaders you need"));
+		UE_LOG(LogCookShadersCommandlet, Log, TEXT("Cook shaders based upon the options, ideal for generating symbols for shaders you need"));
 		UE_LOG(LogCookShadersCommandlet, Log, TEXT("Options:"));
 		UE_LOG(LogCookShadersCommandlet, Log, TEXT(" Required: -targetPlatform=<platform>     (Which target platform do you want results, e.g. WindowsClient, etc."));
 		UE_LOG(LogCookShadersCommandlet, Log, TEXT(" Required: -ShaderSymbolsExport=<path>    (Set shader symbols output location."));
@@ -201,7 +218,7 @@ int32 UCookShadersCommandlet::Main(const FString& Params)
 		UE_LOG(LogCookShadersCommandlet, Log, TEXT(" Optional: -filter=<string>               (Recommended! Filter to shaders with <string> in their hash or info data, requires -infoFile)."));
 		UE_LOG(LogCookShadersCommandlet, Log, TEXT(" Optional: -material=<string>             (Cook this material if you don't have a .info file, can be Global for global shaders)."));
 		UE_LOG(LogCookShadersCommandlet, Log, TEXT(" Optional: -noglobals                     (Don't do global shaders, even if they match the filter.)"));
-		UE_LOG(LogCookShadersCommandlet, Log, TEXT(" Optional: -nomaterialinstances           (Don't do material instances, this search is currently slow.)"));
+		UE_LOG(LogCookShadersCommandlet, Log, TEXT(" Optional: -nomaterialinstances           (Don't do material instances)"));
 		return 0;
 	}
 
@@ -251,11 +268,29 @@ int32 UCookShadersCommandlet::Main(const FString& Params)
 		}
 		else if (I.Type == TEXT("Material"))
 		{
-			TArray<FString> ShaderTypeNames;
-			ShaderTypeNames.Add(I.Shader);
-			IndividualRequests.Add(FODSCRequestPayload(
-				EShaderPlatform::SP_NumPlatforms, ERHIFeatureLevel::Num,
-				I.Quality, I.Name, I.VertexFactory, I.Pipeline, ShaderTypeNames, I.Permutation, I.Hash));
+			FODSCRequestPayload* Match = IndividualRequests.FindByPredicate(
+				[&I](FODSCRequestPayload& Entry)
+				{
+					return (Entry.QualityLevel == I.Quality) && (Entry.VertexFactoryName == I.VertexFactory) && (Entry.MaterialName == I.Name);
+				}
+			);
+
+			if (Match)
+			{
+				if (Match->ShaderTypeNames.Find(I.Shader) == INDEX_NONE)
+				{
+					Match->ShaderTypeNames.Add(I.Shader);
+				}
+			}
+			else
+			{
+				TArray<FString> ShaderTypeNames;
+				ShaderTypeNames.Add(I.Shader);
+
+				IndividualRequests.Add(FODSCRequestPayload(
+					EShaderPlatform::SP_NumPlatforms, ERHIFeatureLevel::Num,
+					I.Quality, I.Name, I.VertexFactory, I.Pipeline, ShaderTypeNames, I.Permutation, I.Hash));
+			}
 		}
 	}
 
@@ -275,23 +310,25 @@ int32 UCookShadersCommandlet::Main(const FString& Params)
 		}
 	}
 
-	// Locate full paths for the materials we have individual requests for
-	TArray<UMaterialInterface*> MaterialsToFindInstancesOf;
-	for (auto& I : IndividualRequests)
+	// Locate full paths for the materials we have individual requests for & save materials to potentially find instances of
+	TSet<FAssetData> MaterialsToFindInstancesOf;
+	for (auto& Req : IndividualRequests)
 	{
-		for (const FAssetData& It : MaterialList)
-		{
-			if (It.AssetName == I.MaterialName)
+		FAssetData* Match = MaterialList.FindByPredicate(
+			[&Req](FAssetData& Entry)
 			{
-				UMaterialInterface* Material = LoadObject<UMaterialInterface>(nullptr, *It.GetObjectPathString());
-				MaterialsToFindInstancesOf.Add(Material);
-				I.MaterialName = Material->GetPathName();
-				break;
+				return Entry.AssetName == Req.MaterialName;
 			}
+		);
+
+		if (Match)
+		{
+			Req.MaterialName = Match->GetObjectPathString();
+			MaterialsToFindInstancesOf.Add(*Match);
 		}
 	}
 
-	// Locate full paths for the materials which match the command line material param
+	// Also locate and add materials matched from the command line switch
 	TArray<FString> MaterialsRequested;
 	if (!MaterialString.IsEmpty())
 	{
@@ -300,62 +337,86 @@ int32 UCookShadersCommandlet::Main(const FString& Params)
 			FString AssetString = It.AssetName.ToString();
 			if (AssetString.Contains(MaterialString))
 			{
-				UMaterialInterface* Material = LoadObject<UMaterialInterface>(nullptr, *It.GetObjectPathString());
-				MaterialsRequested.Add(Material->GetPathName());
-				MaterialsToFindInstancesOf.Add(Material);
+				MaterialsRequested.Add(*It.GetObjectPathString());
+				MaterialsToFindInstancesOf.Add(It);
 			}
 		}
 	}
 
-	// Iterate material instances, need to find instances which depend upon the list of base materials
+	// Iterate instances and find ones which depend upon the materials we are interested in
 	if (!bNoMaterialInstances && MaterialsToFindInstancesOf.Num())
 	{
-		int32 Count = 0;
-		double StartTime = FPlatformTime::Seconds();
-		TArray<FODSCRequestPayload> InstancedRequests;
-		for (const FAssetData& It : MaterialInstanceList)
+		// for faster name lookups
+		TMap<FString, int32> MaterialInstanceNameToIndex;
+		int32 Index = 0;
+		for (const FAssetData& Instance : MaterialInstanceList)
 		{
-			Count++;
-			if (FPlatformTime::Seconds() > StartTime + 10.0)
-			{
-				StartTime = FPlatformTime::Seconds();
-				UE_LOG(LogCookShadersCommandlet, Display, TEXT("MaterialInstances %d/%d..."), Count, MaterialInstanceList.Num());
-			}
+			MaterialInstanceNameToIndex.Emplace(Instance.GetSoftObjectPath().ToString(), Index++);
+		}
+
+		TSet<FString> MaterialsToFindInstancesOfNames;
+		for (const FAssetData& Instance : MaterialsToFindInstancesOf)
+		{
+			MaterialsToFindInstancesOfNames.Emplace(Instance.GetSoftObjectPath().ToString());
+		}
 
-			UMaterialInterface* MaterialInstance = LoadObject<UMaterialInterface>(nullptr, *It.GetObjectPathString());
-			if (MaterialInstance)
+		TArray<FODSCRequestPayload> InstancedRequests;
+		for (const FAssetData& InstanceIt : MaterialInstanceList)
+		{
+			FString ParentName;
+			const FAssetData* Cur = &InstanceIt;
+			const FAssetData* Parent = Cur;
+			while (Parent && GetParentName(Cur, ParentName))
 			{
-				TArray<FODSCRequestPayload>::TConstIterator RequestIt(IndividualRequests);
-				for (auto* Material : MaterialsToFindInstancesOf)
+				if (MaterialsToFindInstancesOfNames.Find(ParentName))
 				{
-					if (MaterialInstance->IsDependent(Material))
+					FString InstanceName = *InstanceIt.GetSoftObjectPath().ToString();
+					if (IndividualRequests.IsEmpty())
 					{
-						if (IndividualRequests.IsEmpty())
-						{
-							// We are matching a set of materials, simply add to the list
-							MaterialsRequested.Add(MaterialInstance->GetPathName());
-						}
-						else if (RequestIt)
-						{
-							// Make a new individual request, same as base with instance material name
-							FODSCRequestPayload Req(*RequestIt);
-							Req.MaterialName = MaterialInstance->GetPathName();
-							InstancedRequests.Add(Req);
-						}
+						// We are matching a set of materials, and have no specific requests, simply add to the list of materials
+						MaterialsRequested.Add(InstanceName);
 					}
-
-					// if we have individual base requests, go to next one
-					if (RequestIt)
+					else
 					{
-						++RequestIt;
+						// duplicate any relevent material requests using the instance name instead of material name
+						for (auto& ReqIt : IndividualRequests)
+						{
+							if (ReqIt.MaterialName == ParentName)
+							{
+								FODSCRequestPayload Req(ReqIt);
+								Req.MaterialName = InstanceName;
+								InstancedRequests.Add(Req);
+							}
+						}
 					}
+					break;
 				}
+
+				// if our Parent is also an instance, iterate back up the hierarchy, otherwise stop iterating
+				int32* IndexPtr = MaterialInstanceNameToIndex.Find(ParentName);
+				Parent = IndexPtr ? &MaterialInstanceList[*IndexPtr] : nullptr;
+				Cur = Parent;
 			}
 		}
-
 		IndividualRequests.Append(InstancedRequests);
 	}
 
+	// Add all the unique materials found into the materials requested list
+	// This is to make sure if individual requests fail to compile the shaders we want, we catch them
+	// This helps catch niagara shaders, and unusual shader types which don't match their debug info
+	if (!IndividualRequests.IsEmpty())
+	{
+		TSet<FString> UniqueRequestedMaterials;
+		for (auto& I : IndividualRequests)
+		{
+			UniqueRequestedMaterials.Add(I.MaterialName);
+		}
+		for (auto& I : UniqueRequestedMaterials)
+		{
+			MaterialsRequested.Add(*I);
+		}
+	}
+
 	// Did we find anything to do? 
 	if (MaterialsRequested.IsEmpty() && GlobalsToFind.IsEmpty() && IndividualRequests.IsEmpty())
 	{
@@ -418,11 +479,25 @@ int32 UCookShadersCommandlet::Main(const FString& Params)
 				UE_LOG(LogCookShadersCommandlet, Display, TEXT("Cooking Materials..."));
 				Arguments.CommandType = ODSCRecompileCommand::Material;
 				Arguments.MaterialsToLoad = MaterialsRequested;
+				Arguments.ShadersToRecompile.Empty();
 				RecompileShadersForRemote(Arguments, OutputDir);
 			}
 		}
 	}
 
+	// Validate and note any missing symbol files we didn't generate, when we have enough info to do so
+	if (!ExportPath.IsEmpty() && !InfoFilePath.IsEmpty())
+	{
+		for (const auto& I : Info)
+		{
+			FString Path = ExportPath + TEXT("\\") + I.Hash;
+			if (!IFileManager::Get().FileExists(*Path))
+			{
+				UE_LOG(LogCookShadersCommandlet, Warning, TEXT("Did not generate symbol file '%s' for '%s'"), *I.Hash, *I.Name);
+			}
+		}
+	}
+
 	UE_LOG(LogCookShadersCommandlet, Display, TEXT("Done CookShadersCommandlet"));
 	return 0;
 }
