Added DumpLightFunctionMaterialInfo commandlet.

diff --git a/Engine/Source/Editor/UnrealEd/Classes/Commandlets/DumpLightFunctionMaterialInfo.h b/Engine/Source/Editor/UnrealEd/Classes/Commandlets/DumpLightFunctionMaterialInfo.h
new file mode 100644
--- /dev/null
+++ b/Engine/Source/Editor/UnrealEd/Classes/Commandlets/DumpLightFunctionMaterialInfo.h
@@ -0,0 +1,18 @@
+// Copyright Epic Games, Inc. All Rights Reserved.
+
+#pragma once
+
+#include "CoreMinimal.h"
+#include "UObject/ObjectMacros.h"
+#include "Commandlets/Commandlet.h"
+#include "DumpLightFunctionMaterialInfo.generated.h"
+
+UCLASS(config = Editor)
+class UDumpLightFunctionMaterialInfoCommandlet : public UCommandlet
+{
+	GENERATED_UCLASS_BODY()
+
+	//~ Begin UCommandlet Interface
+	virtual int32 Main(const FString& Params) override;
+	//~ End UCommandlet Interface
+};
diff --git a/Engine/Source/Editor/UnrealEd/Private/Commandlets/DumpLightFunctionMaterialInfo.cpp b/Engine/Source/Editor/UnrealEd/Private/Commandlets/DumpLightFunctionMaterialInfo.cpp
new file mode 100644
--- /dev/null
+++ b/Engine/Source/Editor/UnrealEd/Private/Commandlets/DumpLightFunctionMaterialInfo.cpp
@@ -0,0 +1,243 @@
+// Copyright Epic Games, Inc. All Rights Reserved.
+
+#include "Commandlets/DumpLightFunctionMaterialInfo.h"
+#include "Modules/ModuleManager.h"
+#include "AssetRegistry/AssetRegistryModule.h"
+#include "AssetRegistry/AssetData.h"
+#include "Interfaces/ITargetPlatform.h"
+#include "Interfaces/ITargetPlatformManagerModule.h"
+#include "Materials/Material.h"
+#include "MaterialShared.h"
+#include "MaterialDomain.h"
+#include "ShaderCompiler.h"
+
+DEFINE_LOG_CATEGORY_STATIC(LogDumpLightFunctionMaterialInfo, Log, All);
+
+UDumpLightFunctionMaterialInfoCommandlet::UDumpLightFunctionMaterialInfoCommandlet(const FObjectInitializer& ObjectInitializer)
+	: Super(ObjectInitializer)
+{
+}
+
+int32 UDumpLightFunctionMaterialInfoCommandlet::Main(const FString& Params)
+{
+	TArray<FString> Tokens;
+	TArray<FString> Switches;
+	TMap<FString, FString> ParamVals;
+	UCommandlet::ParseCommandLine(*Params, Tokens, Switches, ParamVals);
+
+	// Display help
+	if (Switches.Contains("help"))
+	{
+		UE_LOG(LogDumpLightFunctionMaterialInfo, Log, TEXT("DumpLightFunctionMaterialInfo"));
+		UE_LOG(LogDumpLightFunctionMaterialInfo, Log, TEXT("This commandlet will dump to information about light function materials."));
+		UE_LOG(LogDumpLightFunctionMaterialInfo, Log, TEXT("A typical way to invoke it is: <YourProject> -run=DumpLightFunctionMaterialInfo -targetplatform=Windows -unattended -sm6 -allowcommandletrendering -nomaterialshaderddc."));
+		return 0;
+	}
+
+
+
+	UE_LOG(LogDumpLightFunctionMaterialInfo, Log, TEXT("Searching for materials within the project..."));
+
+	FAssetRegistryModule& AssetRegistryModule = FModuleManager::LoadModuleChecked<FAssetRegistryModule>("AssetRegistry");
+	IAssetRegistry& AssetRegistry = AssetRegistryModule.Get();
+	AssetRegistry.SearchAllAssets(true);
+
+	TArray<FAssetData> MaterialAssets;
+	AssetRegistry.GetAssetsByClass(UMaterial::StaticClass()->GetClassPathName(), MaterialAssets, true);
+
+
+
+	UE_LOG(LogDumpLightFunctionMaterialInfo, Log, TEXT("Found %d materials"), MaterialAssets.Num());
+
+	ITargetPlatformManagerModule* TPM = GetTargetPlatformManager();
+	const TArray<ITargetPlatform*>& Platforms = TPM->GetActiveTargetPlatforms();
+	const EShaderPlatform ShaderPlatform = EShaderPlatform::SP_PCD3D_SM6;
+
+	TArray<UMaterialInterface*> LightFunctionMaterialsCompatible;
+	TArray<UMaterialInterface*> LightFunctionMaterialsNotCompatible;
+	TSet<UMaterialInterface*> MaterialsToCompile;
+	TSet<UMaterialInterface*> MaterialsToAnalyse;
+
+	for (ITargetPlatform* Platform : Platforms)
+	{
+		UE_LOG(LogDumpLightFunctionMaterialInfo, Display, TEXT("Compiling shaders for %s..."), *Platform->PlatformName());
+		for (const FAssetData& AssetData : MaterialAssets)
+		{
+			if (UMaterialInterface* MaterialInterface = Cast<UMaterialInterface>(AssetData.GetAsset()))
+			{
+				UMaterial* Material = MaterialInterface->GetMaterial();
+				if (Material && Material->MaterialDomain == MD_LightFunction)
+				{
+					UE_LOG(LogDumpLightFunctionMaterialInfo, Display, TEXT("BeginCache for %s"), *MaterialInterface->GetFullName());
+					MaterialInterface->BeginCacheForCookedPlatformData(Platform);
+					// need to call this once for all objects before any calls to ProcessAsyncResults as otherwise we'll potentially upload
+					// incremental/incomplete shadermaps to DDC (as this function actually triggers compilation, some compiles for a particular
+					// material may finish before we've even started others - if we call ProcessAsyncResults in that case the associated shader
+					// maps will think they are "finished" due to having no outstanding dependencies).
+					if (!MaterialInterface->IsCachedCookedPlatformDataLoaded(Platform))
+					{
+						MaterialsToCompile.Add(MaterialInterface);
+					}
+				}
+			}
+		}
+		MaterialsToAnalyse = MaterialsToCompile;
+		LightFunctionMaterialsCompatible.Reserve(MaterialsToAnalyse.Num());
+		LightFunctionMaterialsNotCompatible.Reserve(MaterialsToAnalyse.Num());
+
+
+
+		UE_LOG(LogDumpLightFunctionMaterialInfo, Log, TEXT("Found %d light function materials to compile."), MaterialsToCompile.Num());
+
+		static constexpr bool bLimitExecutationTime = false;
+		int32 PreviousOutstandingJobs = 0;
+		constexpr int32 MaxOutstandingJobs = 20000; // Having a max is a way to try to reduce memory usage.. otherwise outstanding jobs can reach 100k+ and use up 300gb committed memory
+		// Submit all the jobs.
+		{
+			TRACE_CPUPROFILER_EVENT_SCOPE(SubmitJobs);
+
+			UE_LOG(LogDumpLightFunctionMaterialInfo, Display, TEXT("Submit Jobs"));
+
+			while (MaterialsToCompile.Num())
+			{
+				for (auto It = MaterialsToCompile.CreateIterator(); It; ++It)
+				{
+					UMaterialInterface* MaterialInterface = *It;
+					if (MaterialInterface->IsCachedCookedPlatformDataLoaded(Platform))
+					{
+						It.RemoveCurrent();
+						UE_LOG(LogDumpLightFunctionMaterialInfo, Display, TEXT("Finished cache for %s."), *MaterialInterface->GetFullName());
+						UE_LOG(LogDumpLightFunctionMaterialInfo, Display, TEXT("Materials remaining: %d"), MaterialsToCompile.Num());
+					}
+
+					GShaderCompilingManager->ProcessAsyncResults(bLimitExecutationTime, false /* bBlockOnGlobalShaderCompilation */);
+
+					while (true)
+					{
+						const int32 CurrentOutstandingJobs = GShaderCompilingManager->GetNumOutstandingJobs();
+						if (CurrentOutstandingJobs != PreviousOutstandingJobs)
+						{
+							UE_LOG(LogDumpLightFunctionMaterialInfo, Display, TEXT("Outstanding Jobs: %d"), CurrentOutstandingJobs);
+							PreviousOutstandingJobs = CurrentOutstandingJobs;
+						}
+
+						// Flush rendering commands to release any RHI resources (shaders and shader maps).
+						// Delete any FPendingCleanupObjects (shader maps).
+						FlushRenderingCommands();
+
+						if (CurrentOutstandingJobs < MaxOutstandingJobs)
+						{
+							break;
+						}
+						FPlatformProcess::Sleep(1);
+					}
+				}
+			}
+		}
+
+
+
+		// Process the shader maps and save to the DDC.
+		{
+			TRACE_CPUPROFILER_EVENT_SCOPE(ProcessShaderCompileResults);
+
+			UE_LOG(LogDumpLightFunctionMaterialInfo, Log, TEXT("ProcessAsyncResults"));
+
+			while (GShaderCompilingManager->IsCompiling())
+			{
+				GShaderCompilingManager->ProcessAsyncResults(bLimitExecutationTime, false /* bBlockOnGlobalShaderCompilation */);
+
+				while (true)
+				{
+					const int32 CurrentOutstandingJobs = GShaderCompilingManager->GetNumOutstandingJobs();
+					if (CurrentOutstandingJobs != PreviousOutstandingJobs)
+					{
+						UE_LOG(LogDumpLightFunctionMaterialInfo, Display, TEXT("Outstanding Jobs: %d"), CurrentOutstandingJobs);
+						PreviousOutstandingJobs = CurrentOutstandingJobs;
+					}
+
+					// Flush rendering commands to release any RHI resources (shaders and shader maps).
+					// Delete any FPendingCleanupObjects (shader maps).
+					FlushRenderingCommands();
+
+					if (CurrentOutstandingJobs < MaxOutstandingJobs)
+					{
+						break;
+					}
+					FPlatformProcess::Sleep(1);
+				}
+			}
+		}
+
+
+
+		// Look up compilation result for our light function materials
+		for(UMaterialInterface* MaterialInterface : MaterialsToAnalyse)
+		{
+			UMaterial* Material = MaterialInterface->GetMaterial();
+			if (Material && Material->MaterialDomain == MD_LightFunction)
+			{
+				TArray<FMaterialResource*> ResourcesToCache;
+				FMaterialResource* CurrentResource = FindOrCreateMaterialResource(ResourcesToCache, Material, nullptr, ERHIFeatureLevel::SM6, EMaterialQualityLevel::High);
+				check(CurrentResource);
+
+				FMaterialRelevance MaterialRelevance = CurrentResource->GetMaterialInterface()->GetRelevance(ERHIFeatureLevel::SM6);
+				bool bIsLightFunctionAtlasCompatible = MaterialRelevance.bIsLightFunctionAtlasCompatible;
+
+
+				if (bIsLightFunctionAtlasCompatible)
+				{
+					LightFunctionMaterialsCompatible.Add(Material);
+				}
+				else
+				{
+					LightFunctionMaterialsNotCompatible.Add(Material);
+				}
+
+				FMaterial::DeferredDeleteArray(ResourcesToCache);
+			}
+		}
+
+
+
+		// Perform cleanup and clear cached data for cooking.
+		{
+			TRACE_CPUPROFILER_EVENT_SCOPE(ClearCachedCookedPlatformData);
+
+			UE_LOG(LogDumpLightFunctionMaterialInfo, Display, TEXT("Clear Cached Cooked Platform Data"));
+
+			for (const FAssetData& AssetData : MaterialAssets)
+			{
+				if (UMaterialInterface* MaterialInterface = Cast<UMaterialInterface>(AssetData.GetAsset()))
+				{
+					MaterialInterface->ClearAllCachedCookedPlatformData();
+				}
+			}
+		}
+
+
+
+		break;	// only run for a single platform as this is enough to know if a light function material will be compatible with that LFAtlas.
+	} // Platforms
+	
+
+	UE_LOG(LogDumpLightFunctionMaterialInfo, Log, TEXT("**********************************"));
+	UE_LOG(LogDumpLightFunctionMaterialInfo, Log, TEXT("* Material compatible with atlas *"));
+	UE_LOG(LogDumpLightFunctionMaterialInfo, Log, TEXT("**********************************"));
+	for (UMaterialInterface* MaterialInterface : LightFunctionMaterialsCompatible)
+	{
+		UE_LOG(LogDumpLightFunctionMaterialInfo, Log, TEXT("  - %s"), *MaterialInterface->GetPathName());
+	}
+	UE_LOG(LogDumpLightFunctionMaterialInfo, Log, TEXT("**************************************"));
+	UE_LOG(LogDumpLightFunctionMaterialInfo, Log, TEXT("* Material not compatible with atlas *"));
+	UE_LOG(LogDumpLightFunctionMaterialInfo, Log, TEXT("**************************************"));
+	for (UMaterialInterface* MaterialInterface : LightFunctionMaterialsNotCompatible)
+	{
+		UE_LOG(LogDumpLightFunctionMaterialInfo, Log, TEXT("  - %s"), *MaterialInterface->GetPathName());
+	}
+
+	UE_LOG(LogDumpLightFunctionMaterialInfo, Log, TEXT("Material compatible with atlas:.....%d."), LightFunctionMaterialsCompatible.Num());
+	UE_LOG(LogDumpLightFunctionMaterialInfo, Log, TEXT("Material not compatible with atlas:.%d."), LightFunctionMaterialsNotCompatible.Num());
+	return 0;
+}
+
