IncrementalCook: Fix recorded dependencies to separate runtime dependencies from build dependencies, so that editor-only build dependencies are not added to the followup incremental cook.
#rnx

diff --git a/Engine/Source/Editor/UnrealEd/Private/CookOnTheFlyServer.cpp b/Engine/Source/Editor/UnrealEd/Private/CookOnTheFlyServer.cpp
--- a/Engine/Source/Editor/UnrealEd/Private/CookOnTheFlyServer.cpp
+++ b/Engine/Source/Editor/UnrealEd/Private/CookOnTheFlyServer.cpp
@@ -4406,16 +4406,23 @@ private: // Used only by UCookOnTheFlyServer, which has private access
 	FSaveCookedPackageContext(UCookOnTheFlyServer& InCOTFS, UE::Cook::FPackageData& InPackageData,
 		TArrayView<const ITargetPlatform*> InPlatformsForPackage, UE::Cook::FTickStackData& StackData);
 
+	// Hooks used by friends
 	void SetupPackage();
-	void SetupPlatform(const ITargetPlatform* InTargetPlatform, bool bFirstPlatform);
+	void SetupPlatform(const ITargetPlatform* InTargetPlatform, int32 InPlatformIndex);
 	void FinishPlatform();
 	void FinishPackage();
 
+	// private helper functions
+	void CalculatePlatformAgnosticRuntimeDependencies();
+	void CalculatePlatformRuntimeDependencies();
+	static void AddDependency(TMap<FPackageData*, EInstigator>& InDependencies, FPackageData* PackageData, bool bHard);
+
 	// General Package Data
 	UCookOnTheFlyServer& COTFS;
 	FPackageData& PackageData;
 	TArrayView<const ITargetPlatform*> PlatformsForPackage;
-	TSet<FPackageData*> SaveReferences;
+	TMap<FPackageData*, EInstigator> PlatformAgnosticDependencies;
+	TArray<TMap<FPackageData*, EInstigator>, TInlineAllocator<1>> PlatformDependencies;
 	FTickStackData& StackData;
 	UPackage* Package;
 	const FString PackageName;
@@ -4424,7 +4431,8 @@ private: // Used only by UCookOnTheFlyServer, which has private access
 	bool bReferencedOnlyByEditorOnlyData = false;
 	bool bHasTimeOut = false;
 	bool bHasRetryErrorCode = false;
-	bool bHasFirstPlatformResults = false;
+	bool bPlatformAgnosticDependenciesCalculated = false;
+	bool bAnySaveSucceeded = false;
 
 	// General Package Data that is delay-loaded the first time we save a platform
 	UWorld* World = nullptr;
@@ -4441,6 +4449,7 @@ private: // Used only by UCookOnTheFlyServer, which has private access
 	ICookedPackageWriter* PackageWriter = nullptr;
 	FString PlatFilename;
 	FSavePackageResultStruct SavePackageResult;
+	int32 PlatformIndex = -1;
 	bool bPlatformSetupSuccessful = false;
 	bool bEndianSwap = false;
 
@@ -6072,10 +6081,10 @@ void UCookOnTheFlyServer::SaveCookedPackage(UE::Cook::FSaveCookedPackageContext&
 	// For legacy reasons we set GIsCookerLoadingPackage == true during save. Some classes use it to conditionally execute cook operations in both save and load
 	TGuardValue<bool> ScopedIsCookerLoadingPackage(GIsCookerLoadingPackage, true);
 
-	bool bFirstPlatform = true;
-	for (const ITargetPlatform* TargetPlatform : Context.PlatformsForPackage)
+	for (int32 PlatformIndex = 0; PlatformIndex < Context.PlatformsForPackage.Num(); ++PlatformIndex)
 	{
-		Context.SetupPlatform(TargetPlatform, bFirstPlatform);
+		const ITargetPlatform* TargetPlatform = Context.PlatformsForPackage[PlatformIndex];
+		Context.SetupPlatform(TargetPlatform, PlatformIndex);
 		if (Context.bPlatformSetupSuccessful)
 		{
 			UE_SCOPED_HIERARCHICAL_COOKTIMER(GEditorSavePackage);
@@ -6142,7 +6151,6 @@ void UCookOnTheFlyServer::SaveCookedPackage(UE::Cook::FSaveCookedPackageContext&
 		}
 
 		Context.FinishPlatform();
-		bFirstPlatform = false;
 	}
 
 	// Need to restore flags before calling FinishPackage because it might need to save again
@@ -6173,6 +6181,7 @@ FSaveCookedPackageContext::FSaveCookedPackageContext(UCookOnTheFlyServer& InCOTF
 	, PackageName(Package ? Package->GetName() : FString())
 	, Filename(PackageData.GetFileName().ToString())
 {
+	PlatformDependencies.SetNum(InPlatformsForPackage.Num());
 }
 
 void FSaveCookedPackageContext::SetupPackage()
@@ -6203,8 +6212,9 @@ void FSaveCookedPackageContext::SetupPackage()
 	Filename = COTFS.ConvertToFullSandboxPath(*Filename, true);
 }
 
-void FSaveCookedPackageContext::SetupPlatform(const ITargetPlatform* InTargetPlatform, bool bFirstPlatform)
+void FSaveCookedPackageContext::SetupPlatform(const ITargetPlatform* InTargetPlatform, int32 InPlatformIndex)
 {
+	PlatformIndex = InPlatformIndex;
 	TargetPlatform = InTargetPlatform;
 	PlatFilename = Filename.Replace(TEXT("[Platform]"), *TargetPlatform->PlatformName());
 	bPlatformSetupSuccessful = false;
@@ -6268,8 +6278,9 @@ void FSaveCookedPackageContext::SetupPlatform(const ITargetPlatform* InTargetPla
 		bHasDelayLoaded = true;
 	}
 
-	UE_CLOG((GCookProgressDisplay & (int32)ECookProgressDisplayMode::Instigators) && bFirstPlatform, LogCook, Display,
-		TEXT("Cooking %s, Instigator: { %s }"), *PackageName, *(PackageData.GetInstigator().ToString()));
+	UE_CLOG((GCookProgressDisplay & (int32)ECookProgressDisplayMode::Instigators) && PlatformIndex == 0,
+		LogCook, Display, TEXT("Cooking %s, Instigator: { %s }"),
+		*PackageName, *(PackageData.GetInstigator().ToString()));
 	UE_CLOG(GCookProgressDisplay & (int32)ECookProgressDisplayMode::PackageNames, LogCook, Display,
 		TEXT("Cooking %s"), *PackageName);
 
@@ -6315,19 +6326,30 @@ bool IsRetryErrorCode(ESavePackageResult Result)
 void FSaveCookedPackageContext::FinishPlatform()
 {
 	TRACE_CPUPROFILER_EVENT_SCOPE(FSaveCookedPackageContext::FinishPlatform);
+	check(PlatformIndex >= 0 && PlatformDependencies.IsValidIndex(PlatformIndex));
 
 	bool bSuccessful = SavePackageResult.IsSuccessful();
 	ECookResult CookResult = bSuccessful ? ECookResult::Succeeded : ECookResult::Failed;
 
 	if (bPlatformSetupSuccessful)
 	{
+		CalculatePlatformAgnosticRuntimeDependencies();
+		CalculatePlatformRuntimeDependencies();
+
 		TOptional<FAssetPackageData> AssetPackageData = COTFS.AssetRegistry->GetAssetPackageDataCopy(Package->GetFName());
 
 		ICookedPackageWriter::FCommitPackageInfo Info;
 		if (COTFS.bHybridIterativeEnabled)
 		{
+			TArray<FName> PlatformDependencyNames;
+			PlatformDependencyNames.Reserve(PlatformDependencies[PlatformIndex].Num());
+			for (const TPair<FPackageData*, EInstigator>& DependencyPair : PlatformDependencies[PlatformIndex])
+			{
+				PlatformDependencyNames.Add(DependencyPair.Key->GetPackageName());
+			}
 			UE_SCOPED_HIERARCHICAL_COOKTIMER(TargetDomainDependencies);
-			UE::TargetDomain::CollectAndStoreCookAttachments(Package, TargetPlatform, &SavePackageResult, Info.Attachments);
+			UE::TargetDomain::CollectAndStoreCookAttachments(Package, TargetPlatform, &SavePackageResult,
+				MoveTemp(PlatformDependencyNames), Info.Attachments);
 		}
 		if (bSuccessful)
 		{
@@ -6385,7 +6407,7 @@ void FSaveCookedPackageContext::FinishPlatform()
 				Dependency.Properties = UE::AssetRegistry::EDependencyProperty::Game;
 			}
 
-			if (!bHasFirstPlatformResults)
+			if (PlatformIndex == 0)
 			{
 				GenerationHelper->FetchExternalActorDependencies();
 				COTFS.RecordExternalActorDependencies(GenerationHelper->GetExternalActorDependencies());
@@ -6440,31 +6462,6 @@ void FSaveCookedPackageContext::FinishPlatform()
 			COTFS);
 	}
 
-	if (bSuccessful && COTFS.bSkipOnlyEditorOnly)
-	{
-		// If the save succeeded, add the imports and softobjectpaths from the save to the cook for the platform
-		FName PackageFName = Package->GetFName();
-		TConstArrayView<const ITargetPlatform*> ReachablePlatforms(&TargetPlatform, 1);
-		for (const TArray<FName>& DependencyNames : { SavePackageResult.ImportPackages, SavePackageResult.SoftPackageReferences })
-		{
-			EInstigator InstigatorType = &DependencyNames == &SavePackageResult.ImportPackages ?
-				EInstigator::SaveTimeHardDependency : EInstigator::SaveTimeSoftDependency;
-			for (FName DependencyName : DependencyNames)
-			{
-				UE::Cook::FPackageData* DependencyData = COTFS.PackageDatas->TryAddPackageDataByPackageName(DependencyName);
-				if (DependencyData)
-				{
-					COTFS.QueueDiscoveredPackage(*DependencyData,
-						UE::Cook::FInstigator(InstigatorType, PackageFName), FDiscoveredPlatformSet(ReachablePlatforms));
-					if (COTFS.IsDebugRecordUnsolicited())
-					{
-						SaveReferences.Add(DependencyData);
-					}
-				}
-			}
-		}
-	}
-
 	// If not retrying, mark the package as cooked, either successfully or with failure
 	bool bIsRetryErrorCode = IsRetryErrorCode(SavePackageResult.Result);
 	if (!bIsRetryErrorCode)
@@ -6488,7 +6485,7 @@ void FSaveCookedPackageContext::FinishPlatform()
 
 	// Accumulate results for SaveCookedPackage_Finish
 	bool bLocalReferencedOnlyByEditorOnlyData = SavePackageResult.Result == ESavePackageResult::ReferencedOnlyByEditorOnlyData;
-	if (bHasFirstPlatformResults && bLocalReferencedOnlyByEditorOnlyData != bReferencedOnlyByEditorOnlyData)
+	if (PlatformIndex > 0 && bLocalReferencedOnlyByEditorOnlyData != bReferencedOnlyByEditorOnlyData)
 	{
 		UE_LOG(LogCook, Error, TEXT("Package %s had different values for IsReferencedOnlyByEditorOnlyData from multiple platforms. ")
 			TEXT("Treating all platforms as IsReferencedOnlyByEditorOnlyData = true; this will cause the package to be ignored on the platforms that need it."),
@@ -6504,82 +6501,75 @@ void FSaveCookedPackageContext::FinishPlatform()
 		PackageData.FindOrAddPlatformData(TargetPlatform).SetSaveTimedOut(true);
 		bHasTimeOut = true;
 	}
+	bAnySaveSucceeded |= bSuccessful;
 
 	bHasRetryErrorCode |= bIsRetryErrorCode;
-	bHasFirstPlatformResults = true;
 
 	ArchiveCookContext.Reset();
+	PlatformIndex = -1;
 }
 
 void FSaveCookedPackageContext::FinishPackage()
 {
-	// Add soft references discovered from the package
-	FName PackageFName = Package->GetFName();
-	TConstArrayView<const ITargetPlatform*> ReachablePlatforms = PlatformsForPackage;
-	if (!COTFS.CookByTheBookOptions->bSkipSoftReferences)
+	// If any save succeeded, add all dependencies from all platforms to the cook for the platform
+	if (bAnySaveSucceeded && !COTFS.CookByTheBookOptions->bSkipSoftReferences)
 	{
-		// Also request any localized variants of this package
-		for (FName LocalizedPackageName : FRequestCluster::GetLocalizationReferences(PackageFName, COTFS))
-		{
-			UE::Cook::FPackageData* LocalizedPackageData = COTFS.PackageDatas->TryAddPackageDataByPackageName(LocalizedPackageName);
-			if (LocalizedPackageData)
-			{
-				COTFS.QueueDiscoveredPackage(*LocalizedPackageData,
-					UE::Cook::FInstigator(UE::Cook::EInstigator::SoftDependency, PackageFName),
-					FDiscoveredPlatformSet(ReachablePlatforms));
-			}
-		}
-		// Also add any references from the package that are required by the AssetManager
-		for (FName AMPackageName : FRequestCluster::GetAssetManagerReferences(PackageFName))
+		FName PackageFName = Package->GetFName();
+		if (PlatformsForPackage.Num() == 1)
 		{
-			UE::Cook::FPackageData* AMPackageData = COTFS.PackageDatas->TryAddPackageDataByPackageName(AMPackageName);
-			if (AMPackageData)
+			TConstArrayView<const ITargetPlatform*> ReachablePlatforms(&TargetPlatform, 1);
+			for (const TPair<FPackageData*, EInstigator>& DependencyPair : PlatformDependencies[0])
 			{
-				COTFS.QueueDiscoveredPackage(*AMPackageData,
-					UE::Cook::FInstigator(UE::Cook::EInstigator::SoftDependency, PackageFName),
-					FDiscoveredPlatformSet(ReachablePlatforms));
+				COTFS.QueueDiscoveredPackage(*DependencyPair.Key,
+					UE::Cook::FInstigator(DependencyPair.Value, PackageFName), FDiscoveredPlatformSet(ReachablePlatforms));
 			}
 		}
-
-		// When using legacy WhatGetsCookedRules, add all the SoftObjectPaths discovered during the package's load, plus any added on
-		// during save, to the cook for all platforms
-		TSet<FName> SoftObjectPackages;
-		GRedirectCollector.ProcessSoftObjectPathPackageList(PackageFName, false /* bGetEditorOnly */, SoftObjectPackages);
-		for (FName SoftObjectPackage : SoftObjectPackages)
+		else
 		{
-			TMap<FSoftObjectPath, FSoftObjectPath> RedirectedPaths;
-
-			// If this is a redirector, extract destination from asset registry
-			if (COTFS.ContainsRedirector(SoftObjectPackage, RedirectedPaths))
+			// Merge the platform dependencies for each package into a single QueueDiscoveredPackage call, if possible
+			// It will not be possible if the different platforms have different instigators, so track a list of
+			// platforms for each package for each instigator type.
+			TMap<FPackageData*, TMap<EInstigator, TArray<const ITargetPlatform*>>> PackagePlatformsForInstigator;
+			for (int32 LocalIndex = 0; LocalIndex < PlatformsForPackage.Num(); ++LocalIndex)
 			{
-				for (TPair<FSoftObjectPath, FSoftObjectPath>& RedirectedPath : RedirectedPaths)
+				const TMap<FPackageData*, EInstigator>& CurrentDependencies = PlatformDependencies[LocalIndex];
+				const ITargetPlatform* CurrentPlatform = PlatformsForPackage[LocalIndex];
+				for (const TPair<FPackageData*, EInstigator>& PackagePlatformPair : CurrentDependencies)
 				{
-					GRedirectCollector.AddAssetPathRedirection(RedirectedPath.Key, RedirectedPath.Value);
+					TMap<EInstigator, TArray<const ITargetPlatform*>>& TargetMap =
+						PackagePlatformsForInstigator.FindOrAdd(PackagePlatformPair.Key);
+					TargetMap.FindOrAdd(PackagePlatformPair.Value).Add(CurrentPlatform);
 				}
 			}
-
-			UE::Cook::FPackageData* SoftObjectPackageData = COTFS.PackageDatas->TryAddPackageDataByPackageName(SoftObjectPackage);
-			if (SoftObjectPackageData)
+			for (const TPair<FPackageData*, TMap<EInstigator, TArray<const ITargetPlatform*>>>& PackagePair
+				: PackagePlatformsForInstigator)
 			{
-				if (!COTFS.bSkipOnlyEditorOnly)
+				for (const TPair<EInstigator, TArray<const ITargetPlatform*>>& InstigatorPair : PackagePair.Value)
 				{
-					COTFS.QueueDiscoveredPackage(*SoftObjectPackageData,
-						UE::Cook::FInstigator(UE::Cook::EInstigator::SoftDependency, PackageFName),
-						FDiscoveredPlatformSet(ReachablePlatforms));
-				}
-
-				if (COTFS.IsDebugRecordUnsolicited())
-				{
-					PackageData.CreateOrGetUnsolicited().Add(SoftObjectPackageData, EInstigator::SoftDependency);
+					COTFS.QueueDiscoveredPackage(*PackagePair.Key,
+						UE::Cook::FInstigator(InstigatorPair.Key, PackageFName),
+						FDiscoveredPlatformSet(InstigatorPair.Value));
 				}
 			}
 		}
 	}
+
 	if (COTFS.IsDebugRecordUnsolicited())
 	{
 		COTFS.ProcessUnsolicitedPackages();
+		TMap<FPackageData*, EInstigator> AllPlatformDependenciesBuffer;
+		TMap<FPackageData*, EInstigator>* AllPlatformDependencies = &PlatformDependencies[0];
+		if (PlatformsForPackage.Num() > 1)
+		{
+			AllPlatformDependencies = &AllPlatformDependenciesBuffer;
+			for (int32 LocalIndex = 0; LocalIndex < PlatformsForPackage.Num(); ++LocalIndex)
+			{
+				AllPlatformDependenciesBuffer.Append(PlatformDependencies[LocalIndex]);
+			}
+		}
+
 		FDiagnostics::AnalyzeHiddenDependencies(COTFS, PackageData, PackageData.DetachUnsolicited(),
-			SaveReferences, ReachablePlatforms, COTFS.bOnlyEditorOnlyDebug, COTFS.bHiddenDependenciesDebug);
+			*AllPlatformDependencies, PlatformsForPackage, COTFS.bOnlyEditorOnlyDebug, COTFS.bHiddenDependenciesDebug);
 	}
 
 	if (!bHasRetryErrorCode)
@@ -6605,6 +6595,108 @@ void FSaveCookedPackageContext::FinishPackage()
 	}
 }
 
+void FSaveCookedPackageContext::CalculatePlatformAgnosticRuntimeDependencies()
+{
+	if (bPlatformAgnosticDependenciesCalculated)
+	{
+		return;
+	}
+	bPlatformAgnosticDependenciesCalculated = true;
+
+	FName PackageFName = PackageData.GetPackageName();
+	for (FName LocalizedPackageName : FRequestCluster::GetLocalizationReferences(PackageFName, COTFS))
+	{
+		UE::Cook::FPackageData* LocalizedPackageData = COTFS.PackageDatas->TryAddPackageDataByPackageName(LocalizedPackageName);
+		if (LocalizedPackageData)
+		{
+			AddDependency(PlatformAgnosticDependencies, LocalizedPackageData, false /* bHard */);
+		}
+	}
+	// Also add any references from the package that are required by the AssetManager
+	for (FName AMPackageName : FRequestCluster::GetAssetManagerReferences(PackageFName))
+	{
+		UE::Cook::FPackageData* AMPackageData = COTFS.PackageDatas->TryAddPackageDataByPackageName(AMPackageName);
+		if (AMPackageData)
+		{
+			AddDependency(PlatformAgnosticDependencies, AMPackageData, false /* bHard */);
+		}
+	}
+
+	// When using legacy WhatGetsCookedRules, add all the SoftObjectPaths discovered during the package's load, plus any added on
+	// during save, to the cook for all platforms
+	TSet<FName> SoftObjectPackages;
+	GRedirectCollector.ProcessSoftObjectPathPackageList(PackageFName, false /* bGetEditorOnly */, SoftObjectPackages);
+	for (FName SoftObjectPackage : SoftObjectPackages)
+	{
+		TMap<FSoftObjectPath, FSoftObjectPath> RedirectedPaths;
+
+		// If this is a redirector, extract destination from asset registry
+		if (COTFS.ContainsRedirector(SoftObjectPackage, RedirectedPaths))
+		{
+			for (TPair<FSoftObjectPath, FSoftObjectPath>& RedirectedPath : RedirectedPaths)
+			{
+				GRedirectCollector.AddAssetPathRedirection(RedirectedPath.Key, RedirectedPath.Value);
+			}
+		}
+
+		UE::Cook::FPackageData* SoftObjectPackageData = COTFS.PackageDatas->TryAddPackageDataByPackageName(SoftObjectPackage);
+		if (SoftObjectPackageData)
+		{
+			if (!COTFS.bSkipOnlyEditorOnly)
+			{
+				AddDependency(PlatformAgnosticDependencies, SoftObjectPackageData, false /* bHard */);
+			}
+
+			if (COTFS.IsDebugRecordUnsolicited())
+			{
+				PackageData.CreateOrGetUnsolicited().Add(SoftObjectPackageData, EInstigator::SoftDependency);
+			}
+		}
+	}
+}
+
+void FSaveCookedPackageContext::CalculatePlatformRuntimeDependencies()
+{
+	PlatformDependencies.IsValidIndex(PlatformIndex);
+	TMap<FPackageData*, EInstigator>& CurrentDependencies = PlatformDependencies[PlatformIndex];
+	CurrentDependencies.Reset();
+	// Add PlatformAgnostic dependencies to the PlatformDependencies
+	CurrentDependencies.Append(PlatformAgnosticDependencies);
+
+	// Add imports and softobjectpaths from the save to the SaveDependencies for the current Platform and for all Platforms
+	if (COTFS.bSkipOnlyEditorOnly)
+	{
+		FName PackageFName = Package->GetFName();
+		TConstArrayView<const ITargetPlatform*> ReachablePlatforms(&TargetPlatform, 1);
+		for (const TArray<FName>& DependencyNames : { SavePackageResult.ImportPackages, SavePackageResult.SoftPackageReferences })
+		{
+			bool bHard = &DependencyNames == &SavePackageResult.ImportPackages;
+			for (FName DependencyName : DependencyNames)
+			{
+				UE::Cook::FPackageData* DependencyData = COTFS.PackageDatas->TryAddPackageDataByPackageName(DependencyName);
+				if (DependencyData)
+				{
+					AddDependency(CurrentDependencies, DependencyData, bHard);
+				}
+			}
+		}
+	}
+}
+
+void FSaveCookedPackageContext::AddDependency(TMap<FPackageData*, EInstigator>& InDependencies,
+	FPackageData* PackageData, bool bHard)
+{
+	EInstigator& Existing = InDependencies.FindOrAdd(PackageData, EInstigator::Unspecified);
+	if (bHard)
+	{
+		Existing = EInstigator::SaveTimeHardDependency;
+	}
+	else if (Existing == EInstigator::Unspecified)
+	{
+		Existing = EInstigator::SoftDependency;
+	}
+}
+
 } // namespace UE::Cook
 
 void UCookOnTheFlyServer::RecordExternalActorDependencies(TConstArrayView<FName> ExternalActorDependencies)
diff --git a/Engine/Source/Editor/UnrealEd/Private/Cooker/CookDiagnostics.cpp b/Engine/Source/Editor/UnrealEd/Private/Cooker/CookDiagnostics.cpp
--- a/Engine/Source/Editor/UnrealEd/Private/Cooker/CookDiagnostics.cpp
+++ b/Engine/Source/Editor/UnrealEd/Private/Cooker/CookDiagnostics.cpp
@@ -101,7 +101,7 @@ public:
 };
 
 void FDiagnostics::AnalyzeHiddenDependencies(UCookOnTheFlyServer& COTFS, FPackageData& PackageData,
-	TMap<FPackageData*, EInstigator>&& UnsolicitedForPackage, TSet<FPackageData*>& SaveReferences,
+	TMap<FPackageData*, EInstigator>&& UnsolicitedForPackage, TMap<FPackageData*, EInstigator>& SaveReferences,
 	TConstArrayView<const ITargetPlatform*> ReachablePlatforms, bool bOnlyEditorOnlyDebug,
 	bool bHiddenDependenciesDebug)
 {
diff --git a/Engine/Source/Editor/UnrealEd/Private/Cooker/CookDiagnostics.h b/Engine/Source/Editor/UnrealEd/Private/Cooker/CookDiagnostics.h
--- a/Engine/Source/Editor/UnrealEd/Private/Cooker/CookDiagnostics.h
+++ b/Engine/Source/Editor/UnrealEd/Private/Cooker/CookDiagnostics.h
@@ -23,7 +23,7 @@ public:
 	 * with diagnostics for unexpected differences.
 	 */
 	static void AnalyzeHiddenDependencies(UCookOnTheFlyServer& COTFS, FPackageData& PackageData,
-		TMap<FPackageData*, EInstigator>&& UnsolicitedForPackage, TSet<FPackageData*>& SaveReferences,
+		TMap<FPackageData*, EInstigator>&& UnsolicitedForPackage, TMap<FPackageData*, EInstigator>& SaveReferences,
 		TConstArrayView<const ITargetPlatform*> ReachablePlatforms, bool bOnlyEditorOnlyDebug,
 		bool bHiddenDependenciesDebug);
 
diff --git a/Engine/Source/Editor/UnrealEd/Private/Cooker/CookRequestCluster.cpp b/Engine/Source/Editor/UnrealEd/Private/Cooker/CookRequestCluster.cpp
--- a/Engine/Source/Editor/UnrealEd/Private/Cooker/CookRequestCluster.cpp
+++ b/Engine/Source/Editor/UnrealEd/Private/Cooker/CookRequestCluster.cpp
@@ -1129,15 +1129,10 @@ void FRequestCluster::FGraphSearch::ExploreVertexEdges(FVertexData& Vertex)
 					bIterativelyUnmodified = true;
 					PackagePlatformData.SetIterativelyUnmodified(true);
 				}
-				if (bExploreDependencies)
+				if (bExploreDependencies && Cluster.bAllowSoftDependencies)
 				{
-					AddPlatformDependencyRange(CookDependencies.GetPackageDependencies(), PlatformIndex,
-						EInstigator::HardDependency);
-					if (Cluster.bAllowSoftDependencies)
-					{
-						AddPlatformDependencyRange(CookDependencies.GetRuntimePackageDependencies(), PlatformIndex,
-							EInstigator::HardDependency);
-					}
+					AddPlatformDependencyRange(CookDependencies.GetRuntimePackageDependencies(), PlatformIndex,
+						EInstigator::SoftDependency);
 				}
 
 				if (Cluster.bPreQueueBuildDefinitions)
@@ -1208,9 +1203,11 @@ void FRequestCluster::FGraphSearch::ExploreVertexEdges(FVertexData& Vertex)
 			Cluster.AssetRegistry.GetDependencies(PackageName, CookerLoadingDependencies, EDependencyCategory::Package,
 				EDependencyQuery::Hard);
 
-			// ITERATIVECOOK_TODO: Build dependencies need to be stored and used to mark package loads as expected
-			// But we can't use them to explore packages that will be loaded during cook because they might not be;
-			// some build dependencies might be a conservative list but unused by the asset, or unused on
+			// ITERATIVECOOK_TODO: To improve cooker load performance, we should declare EDependencyQuery::Build
+			// packages as packages that will be loaded during the cook, by adding them as edges for the
+			// CookerLoadingPlatformIndex platform.
+			// But we can't do that yet; in some important cases the build dependencies are declared by a class but not
+			// always used - some build dependencies might be a conservative list but unused by the asset, or unused on
 			// targetplatform.
 			// Adding BuildDependencies also sets up many circular dependencies, because maps declare their external
 			// actors as build dependencies and the external actors declare the map as a build or hard dependency.
diff --git a/Engine/Source/Editor/UnrealEd/Private/EditorDomain/EditorDomainUtils.cpp b/Engine/Source/Editor/UnrealEd/Private/EditorDomain/EditorDomainUtils.cpp
--- a/Engine/Source/Editor/UnrealEd/Private/EditorDomain/EditorDomainUtils.cpp
+++ b/Engine/Source/Editor/UnrealEd/Private/EditorDomain/EditorDomainUtils.cpp
@@ -1917,7 +1917,7 @@ bool TrySavePackage(UPackage* Package)
 	if (bStorageResultValid)
 	{
 		TArray<IPackageWriter::FCommitAttachmentInfo, TInlineAllocator<2>> Attachments;
-		UE::TargetDomain::CollectAndStoreCookAttachments(Package, nullptr, nullptr, Attachments);
+		UE::TargetDomain::CollectAndStoreCookAttachments(Package, nullptr, nullptr, TArray<FName>(), Attachments);
 		if (!Attachments.IsEmpty())
 		{
 			UE::TargetDomain::CommitEditorDomainCookAttachments(Package->GetFName(), Attachments);
diff --git a/Engine/Source/Editor/UnrealEd/Private/TargetDomain/TargetDomainUtils.cpp b/Engine/Source/Editor/UnrealEd/Private/TargetDomain/TargetDomainUtils.cpp
--- a/Engine/Source/Editor/UnrealEd/Private/TargetDomain/TargetDomainUtils.cpp
+++ b/Engine/Source/Editor/UnrealEd/Private/TargetDomain/TargetDomainUtils.cpp
@@ -5,6 +5,7 @@
 #include "Algo/BinarySearch.h"
 #include "Algo/IsSorted.h"
 #include "Algo/Sort.h"
+#include "Algo/Unique.h"
 #include "AssetRegistry/IAssetRegistry.h"
 #include "Containers/Set.h"
 #include "Containers/UnrealString.h"
@@ -34,7 +35,7 @@
 namespace UE::TargetDomain
 {
 
-constexpr uint32 CookDependenciesVersion = 0x00000001;
+constexpr uint32 CookDependenciesVersion = 0x00000002;
 static const FUtf8StringView CookDependenciesAttachmentKey = UTF8TEXTVIEW("CookDependencies");
 static const FUtf8StringView BuildDefinitionsAttachmentKey = UTF8TEXTVIEW("BuildDefinitionsAttachmentKey");
 /**
@@ -148,7 +149,7 @@ bool FCookDependencies::TryCalculateCurrentKey(FString* OutErrorMessage)
 
 	KeyBuilder.Update(&PackageDigest.Hash, sizeof(PackageDigest.Hash));
 
-	for (FName PackageDependency : PackageDependencies)
+	for (FName PackageDependency : BuildPackageDependencies)
 	{
 		PackageDigest = EditorDomain->GetPackageDigest(PackageDependency);
 		if (!PackageDigest.IsSuccessful())
@@ -221,7 +222,7 @@ bool FCookDependencies::TryCalculateCurrentKey(FString* OutErrorMessage)
 
 void FCookDependencies::Reset()
 {
-	PackageDependencies.Reset();
+	BuildPackageDependencies.Reset();
 	ConfigDependencies.Reset();
 	RuntimePackageDependencies.Reset();
 	PackageName = FName();
@@ -233,14 +234,27 @@ void FCookDependencies::Reset()
 void FCookDependencies::Empty()
 {
 	Reset();
-	PackageDependencies.Empty();
+	BuildPackageDependencies.Empty();
 	ConfigDependencies.Empty();
 	RuntimePackageDependencies.Empty();
 }
 
 FCookDependencies FCookDependencies::Collect(UPackage* Package, const ITargetPlatform* TargetPlatform,
-	FSavePackageResultStruct* SaveResult, FString* OutErrorMessage)
+	FSavePackageResultStruct* SaveResult, TArray<FName>&& RuntimeDependencies, FString* OutErrorMessage)
 {
+	TStringBuilder<256> StringBuffer;
+	FName TransientPackageName = GetTransientPackage()->GetFName();
+	auto IsTransientPackageName = [&StringBuffer, TransientPackageName](FName InPackageName)
+		{
+			if (InPackageName == TransientPackageName)
+			{
+				return true;
+			}
+			InPackageName.ToString(StringBuffer);
+			return FPackageName::IsMemoryPackage(StringBuffer) ||
+				FPackageName::IsScriptPackage(StringBuffer);
+		};
+
 	if (!Package)
 	{
 		if (OutErrorMessage) *OutErrorMessage = TEXT("Invalid null package.");
@@ -261,12 +275,15 @@ FCookDependencies FCookDependencies::Collect(UPackage* Package, const ITargetPla
 
 	FCookDependencies Result;
 	Result.PackageName = Package->GetFName();
-	TSet<FName> PackageDependenciesSet;
-	TSet<FName> RuntimePackageDependenciesSet;
+	TSet<FName> BuildDependenciesSet;
 
 	TArray<FName> AssetDependencies;
 	AssetRegistry->GetDependencies(Result.PackageName, AssetDependencies,
 		UE::AssetRegistry::EDependencyCategory::Package, UE::AssetRegistry::EDependencyQuery::Game);
+	RuntimeDependencies.Append(MoveTemp(AssetDependencies));
+	RuntimeDependencies.RemoveAllSwap(IsTransientPackageName, EAllowShrinking::No);
+	RuntimeDependencies.Sort(FNameLexicalLess());
+	RuntimeDependencies.SetNum(Algo::Unique(RuntimeDependencies), EAllowShrinking::Yes);
 
 	FPackageBuildDependencyTracker& Tracker = FPackageBuildDependencyTracker::Get();
 
@@ -275,52 +292,27 @@ FCookDependencies FCookDependencies::Collect(UPackage* Package, const ITargetPla
 	{
 		TArray<FBuildDependencyAccessData> AccessDatas = Tracker.GetAccessDatas(Result.PackageName);
 
-		PackageDependenciesSet.Reserve(AccessDatas.Num());
+		BuildDependenciesSet.Reserve(AccessDatas.Num());
 		for (FBuildDependencyAccessData& AccessData : AccessDatas)
 		{
 			if (AccessData.TargetPlatform == TargetPlatform || AccessData.TargetPlatform == nullptr)
 			{
-				PackageDependenciesSet.Add(AccessData.ReferencedPackage);
-			}
-		}
-
-		RuntimePackageDependenciesSet.Reserve(AssetDependencies.Num());
-		for (FName DependencyName : AssetDependencies)
-		{
-			if (!PackageDependenciesSet.Contains(DependencyName))
-			{
-				RuntimePackageDependenciesSet.Add(DependencyName);
+				BuildDependenciesSet.Add(AccessData.ReferencedPackage);
 			}
 		}
 	}
 	else
 #endif
 	{
-		// Defensively treat all asset dependencies as build dependencies and have zero runtime only dependencies
-		PackageDependenciesSet.Append(AssetDependencies);
+		// Defensively treat all asset dependencies as both build and runtime dependencies
+		BuildDependenciesSet.Append(AssetDependencies);
 	}
 
-	TStringBuilder<256> StringBuffer;
-	FName TransientPackageName = GetTransientPackage()->GetFName();
-	auto IsTransientPackageName = [&StringBuffer, TransientPackageName](FName InPackageName)
-	{
-		if (InPackageName == TransientPackageName)
-		{
-			return true;
-		}
-		InPackageName.ToString(StringBuffer);
-		return FPackageName::IsMemoryPackage(StringBuffer) ||
-			FPackageName::IsScriptPackage(StringBuffer);
-	};
-	Result.PackageDependencies = PackageDependenciesSet.Array();
-	Result.PackageDependencies.RemoveAllSwap(IsTransientPackageName, EAllowShrinking::No);
-	Result.PackageDependencies.Sort(FNameLexicalLess());
-	Result.PackageDependencies.Shrink();
+	Result.BuildPackageDependencies = BuildDependenciesSet.Array();
+	Result.BuildPackageDependencies.RemoveAllSwap(IsTransientPackageName, EAllowShrinking::Yes);
+	Result.BuildPackageDependencies.Sort(FNameLexicalLess());
 
-	Result.RuntimePackageDependencies = RuntimePackageDependenciesSet.Array();
-	Result.RuntimePackageDependencies.RemoveAllSwap(IsTransientPackageName, EAllowShrinking::No);
-	Result.RuntimePackageDependencies.Sort(FNameLexicalLess());
-	Result.RuntimePackageDependencies.Shrink();
+	Result.RuntimePackageDependencies = MoveTemp(RuntimeDependencies);
 
 #if UE_WITH_CONFIG_TRACKING
 	{
@@ -381,9 +373,9 @@ bool LoadFromCompactBinary(FCbObjectView ObjectView, UE::TargetDomain::FCookDepe
 				return false;
 			}
 		}
-		if (FieldView.GetName().Equals(UTF8TEXTVIEW("PackageDependencies")))
+		if (FieldView.GetName().Equals(UTF8TEXTVIEW("BuildPackageDependencies")))
 		{
-			if (!LoadFromCompactBinary(FieldView++, Dependencies.PackageDependencies))
+			if (!LoadFromCompactBinary(FieldView++, Dependencies.BuildPackageDependencies))
 			{
 				return false;
 			}
@@ -429,9 +421,9 @@ FCbWriter& operator<<(FCbWriter& Writer, const UE::TargetDomain::FCookDependenci
 	Writer.BeginObject();
 	Writer << "Version" << CookDependenciesVersion;
 	Writer << "StoredKey" << CookDependencies.StoredKey;
-	if (!CookDependencies.PackageDependencies.IsEmpty())
+	if (!CookDependencies.BuildPackageDependencies.IsEmpty())
 	{
-		Writer << "PackageDependencies" << CookDependencies.PackageDependencies;
+		Writer << "BuildPackageDependencies" << CookDependencies.BuildPackageDependencies;
 	}
 	if (!CookDependencies.ConfigDependencies.IsEmpty())
 	{
@@ -546,10 +538,12 @@ void FCookAttachments::Empty()
 }
 
 bool TryCollectAndStoreCookDependencies(UPackage* Package, const ITargetPlatform* TargetPlatform,
-	FSavePackageResultStruct* SaveResult, IPackageWriter::FCommitAttachmentInfo& OutResult)
+	FSavePackageResultStruct* SaveResult, TArray<FName>&& RuntimeDependencies,
+	IPackageWriter::FCommitAttachmentInfo& OutResult)
 {
 	FString ErrorMessage;
-	FCookDependencies CookDependencies = FCookDependencies::Collect(Package, TargetPlatform, SaveResult, &ErrorMessage);
+	FCookDependencies CookDependencies = FCookDependencies::Collect(Package, TargetPlatform, SaveResult, 
+		MoveTemp(RuntimeDependencies), &ErrorMessage);
 	if (!CookDependencies.IsValid())
 	{
 		// CookPackageSplitterTODO: This error occurs for generated packages. Need to register them with EditorDomain.
diff --git a/Engine/Source/Editor/UnrealEd/Private/TargetDomain/TargetDomainUtils.h b/Engine/Source/Editor/UnrealEd/Private/TargetDomain/TargetDomainUtils.h
--- a/Engine/Source/Editor/UnrealEd/Private/TargetDomain/TargetDomainUtils.h
+++ b/Engine/Source/Editor/UnrealEd/Private/TargetDomain/TargetDomainUtils.h
@@ -55,14 +55,14 @@ public:
 	FCookDependencies& operator=(const FCookDependencies&);
 	FCookDependencies& operator=(FCookDependencies&&);
 	// Build Dependencies
-	const TArray<FName>& GetPackageDependencies() const { return PackageDependencies; }
-	const TArray<FString>& GetConfigDependencies() const { return ConfigDependencies; }
+	const TArray<FName>& GetBuildPackageDependencies() const;
+	const TArray<FString>& GetConfigDependencies() const;
 
 	// Runtime Dependencies
-	const TArray<FName>& GetRuntimePackageDependencies() const { return RuntimePackageDependencies; }
+	const TArray<FName>& GetRuntimePackageDependencies() const;
 
 	// Cook Dependencies
-	const TArray<UE::Cook::FCookDependency>& GetCookDependencies() const { return CookDependencies; }
+	const TArray<UE::Cook::FCookDependency>& GetCookDependencies() const;
 
 	// Data about the Key and Package that are not dependencies
 	/**
@@ -70,9 +70,9 @@ public:
 	 * key for the package. False if the stucture is default, has been reset, or was marked invalid.
 	 */
 	bool IsValid() const;
-	FName GetPackageName() const { return PackageName; }
-	const FIoHash& GetStoredKey() const { return StoredKey; }
-	const FIoHash& GetCurrentKey() const { return CurrentKey; }
+	FName GetPackageName() const;
+	const FIoHash& GetStoredKey() const;
+	const FIoHash& GetCurrentKey() const;
 	bool HasKeyMatch();
 
 	// Modifying the structure data
@@ -86,13 +86,13 @@ public:
 	 * that have recorded its data during the package's load/other/save operations in the current cook session.
 	 */
 	static FCookDependencies Collect(UPackage* Package, const ITargetPlatform* TargetPlatform,
-		FSavePackageResultStruct* SaveResult, FString* OutErrorMessage = nullptr);
+		FSavePackageResultStruct* SaveResult, TArray<FName>&& RuntimeDependencies, FString* OutErrorMessage = nullptr);
 
 	// Fetch function to load the dependencies from a PackageStore is not yet implemented independently for
 	// this structure. Use FCookAttachments instead. 
 
 private:
-	TArray<FName> PackageDependencies;
+	TArray<FName> BuildPackageDependencies;
 	TArray<FString> ConfigDependencies;
 	TArray<FName> RuntimePackageDependencies;
 	TArray<UE::Cook::FCookDependency> CookDependencies;
@@ -142,16 +142,17 @@ struct FCookAttachments
 };
 
 bool TryCollectAndStoreCookDependencies(UPackage* Package, const ITargetPlatform* TargetPlatform,
-	FSavePackageResultStruct* SaveResult, IPackageWriter::FCommitAttachmentInfo& OutResult);
+	FSavePackageResultStruct* SaveResult, TArray<FName>&& RuntimeDependencies,
+	IPackageWriter::FCommitAttachmentInfo& OutResult);
 bool TryCollectAndStoreBuildDefinitionList(UPackage* Package, const ITargetPlatform* TargetPlatform,
 	IPackageWriter::FCommitAttachmentInfo& OutResult);
 
 template <typename ArrayType>
 void CollectAndStoreCookAttachments(UPackage* Package, const ITargetPlatform* TargetPlatform,
-	FSavePackageResultStruct* SaveResult, ArrayType& Output)
+	FSavePackageResultStruct* SaveResult, TArray<FName>&& RuntimeDependencies, ArrayType& Output)
 {
 	IPackageWriter::FCommitAttachmentInfo Result;
-	if (TryCollectAndStoreCookDependencies(Package, TargetPlatform, SaveResult, Result))
+	if (TryCollectAndStoreCookDependencies(Package, TargetPlatform, SaveResult, MoveTemp(RuntimeDependencies), Result))
 	{
 		Output.Add(MoveTemp(Result));
 	}
@@ -167,4 +168,44 @@ bool IsIterativeEnabled(FName PackageName, bool bAllowAllClasses);
 /** Store extra information derived during save and used by the cooker for the given EditorDomain package. */
 void CommitEditorDomainCookAttachments(FName PackageName, TArrayView<IPackageWriter::FCommitAttachmentInfo> Attachments);
 
+
+////////////////////////////////
+// Inline Implementations
+////////////////////////////////
+
+inline const TArray<FName>& FCookDependencies::GetBuildPackageDependencies() const
+{
+	return BuildPackageDependencies;
+}
+
+inline const TArray<FString>& FCookDependencies::GetConfigDependencies() const
+{
+	return ConfigDependencies;
+}
+
+inline const TArray<FName>& FCookDependencies::GetRuntimePackageDependencies() const
+{
+	return RuntimePackageDependencies;
 }
+
+inline const TArray<UE::Cook::FCookDependency>& FCookDependencies::GetCookDependencies() const
+{
+	return CookDependencies;
+}
+
+inline FName FCookDependencies::GetPackageName() const
+{
+	return PackageName;
+}
+
+inline const FIoHash& FCookDependencies::GetStoredKey() const
+{
+	return StoredKey;
+}
+
+inline const FIoHash& FCookDependencies::GetCurrentKey() const
+{
+	return CurrentKey;
+}
+
+} // namespace UE::TargetDomain
