Fix animation root motion preview when turntable
- Use the root motion mode to reset debug skel mesh transform to ensure correct initial placement
- Do not set identiy transform in UDebugSkelMeshComponent::ConsumeRootMotion so that turn table works for non-rootmotion anims too

diff --git a/Engine/Source/Editor/Persona/Private/SAnimationEditorViewport.cpp b/Engine/Source/Editor/Persona/Private/SAnimationEditorViewport.cpp
--- a/Engine/Source/Editor/Persona/Private/SAnimationEditorViewport.cpp
+++ b/Engine/Source/Editor/Persona/Private/SAnimationEditorViewport.cpp
@@ -1194,12 +1194,7 @@ void SAnimationEditorViewportTabBody::OnSetTurnTableMode(int32 ModeIndex)
 	UDebugSkelMeshComponent* PreviewComponent = GetPreviewScene()->GetPreviewMeshComponent();
 	if (PreviewComponent)
 	{
-		PreviewComponent->TurnTableMode = SelectedTurnTableMode;
-
-		if (SelectedTurnTableMode == EPersonaTurnTableMode::Stopped)
-		{
-			PreviewComponent->SetRelativeRotation(FRotator::ZeroRotator);
-		}
+		PreviewComponent->SetTurnTableMode(SelectedTurnTableMode);
 	}
 }
 
diff --git a/Engine/Source/Editor/UnrealEd/Classes/Animation/DebugSkelMeshComponent.h b/Engine/Source/Editor/UnrealEd/Classes/Animation/DebugSkelMeshComponent.h
--- a/Engine/Source/Editor/UnrealEd/Classes/Animation/DebugSkelMeshComponent.h
+++ b/Engine/Source/Editor/UnrealEd/Classes/Animation/DebugSkelMeshComponent.h
@@ -625,6 +625,10 @@ protected:
 	UNREALED_API virtual void SendRenderDynamicData_Concurrent() override;
 
 public:
+
+	/** Set TurnTableMode and reset the turntable rotation if needed. */
+	UNREALED_API void SetTurnTableMode(EPersonaTurnTableMode::Type NewMode);
+	
 	/** Current turn table mode */
 	EPersonaTurnTableMode::Type TurnTableMode;
 	/** Current turn table speed scaling */
diff --git a/Engine/Source/Editor/UnrealEd/Private/Animation/DebugSkelMeshComponent.cpp b/Engine/Source/Editor/UnrealEd/Private/Animation/DebugSkelMeshComponent.cpp
--- a/Engine/Source/Editor/UnrealEd/Private/Animation/DebugSkelMeshComponent.cpp
+++ b/Engine/Source/Editor/UnrealEd/Private/Animation/DebugSkelMeshComponent.cpp
@@ -339,15 +339,6 @@ void UDebugSkelMeshComponent::ConsumeRootMotion(const FVector& FloorMin, const F
 				AddLocalTransform(RootMotionDelta);
 			}
 		}
-		else
-		{
-			// No root motion, reset to identity for consistency.
-			const FTransform RootMotionTransform = GetRelativeTransform();
-			if (!RootMotionTransform.Equals(FTransform::Identity))
-			{
-				SetRelativeTransform(FTransform::Identity);
-			}
-		}
 	}
 }
 
@@ -452,37 +443,41 @@ void UDebugSkelMeshComponent::SetProcessRootMotionModeInternal(EProcessRootMotio
 {
 	ProcessRootMotionMode = Mode;
 
-	if (!DoesCurrentAssetHaveRootMotion())
-	{
-		return;
-	}
-
 	FTransform RootMotionTransform = FTransform::Identity;
-	
-	if (ProcessRootMotionMode == EProcessRootMotionMode::LoopAndReset || ProcessRootMotionMode == EProcessRootMotionMode::Loop)
+
+	if (DoesCurrentAssetHaveRootMotion())
 	{
-		// Reset transform
-		const float CurrentTime = PreviewInstance->GetCurrentTime();
-		float SectionStartPosition = 0.0f;
-		if (const UAnimMontage* Montage = Cast<UAnimMontage>(PreviewInstance->CurrentAsset))
+		if (ProcessRootMotionMode == EProcessRootMotionMode::LoopAndReset || ProcessRootMotionMode == EProcessRootMotionMode::Loop)
 		{
-			const int32 PreviewStartSectionIdx = Montage->CompositeSections.IsValidIndex(PreviewInstance->MontagePreviewStartSectionIdx) ? PreviewInstance->MontagePreviewStartSectionIdx : Montage->GetSectionIndexFromPosition(CurrentTime);
-			const int32 FirstSectionIdx = PreviewInstance->MontagePreview_FindFirstSectionAsInMontage(PreviewStartSectionIdx);
-			const int32 LastSectionIdx = PreviewInstance->MontagePreview_FindLastSection(FirstSectionIdx);
-			float StartTime = 0.0f, EndTime = 0.0f;
-			Montage->GetSectionStartAndEndTime(LastSectionIdx, StartTime, EndTime);
-			SectionStartPosition = StartTime;
-		}
-	
-		const FTransform InitialTransform = UE::Anim::Private::CalculateInitialTransformFromAssetAndTime(*this, PreviewInstance->CurrentAsset, SectionStartPosition);
-		const FTransform RootMotionDelta = UE::Anim::ExtractRootMotionFromAnimationAsset(PreviewInstance->CurrentAsset, PreviewInstance->GetMirrorDataTable(), SectionStartPosition, CurrentTime);
-		RootMotionTransform = RootMotionDelta * InitialTransform;
+			// Reset transform
+			const float CurrentTime = PreviewInstance->GetCurrentTime();
+			float SectionStartPosition = 0.0f;
+			if (const UAnimMontage* Montage = Cast<UAnimMontage>(PreviewInstance->CurrentAsset))
+			{
+				const int32 PreviewStartSectionIdx = Montage->CompositeSections.IsValidIndex(PreviewInstance->MontagePreviewStartSectionIdx) ? PreviewInstance->MontagePreviewStartSectionIdx : Montage->GetSectionIndexFromPosition(CurrentTime);
+				const int32 FirstSectionIdx = PreviewInstance->MontagePreview_FindFirstSectionAsInMontage(PreviewStartSectionIdx);
+				const int32 LastSectionIdx = PreviewInstance->MontagePreview_FindLastSection(FirstSectionIdx);
+				float StartTime = 0.0f, EndTime = 0.0f;
+				Montage->GetSectionStartAndEndTime(LastSectionIdx, StartTime, EndTime);
+				SectionStartPosition = StartTime;
+			}
 		
-		// Reference transform is always relative to the beginning of the animation sequence. 
-		RootMotionReferenceTransform = UE::Anim::ExtractRootTransformFromAnimationAsset(PreviewInstance->CurrentAsset, 0.0f);
+			const FTransform InitialTransform = UE::Anim::Private::CalculateInitialTransformFromAssetAndTime(*this, PreviewInstance->CurrentAsset, SectionStartPosition);
+			const FTransform RootMotionDelta = UE::Anim::ExtractRootMotionFromAnimationAsset(PreviewInstance->CurrentAsset, PreviewInstance->GetMirrorDataTable(), SectionStartPosition, CurrentTime);
+			RootMotionTransform = RootMotionDelta * InitialTransform;
+			
+			// Reference transform is always relative to the beginning of the animation sequence. 
+			RootMotionReferenceTransform = UE::Anim::ExtractRootTransformFromAnimationAsset(PreviewInstance->CurrentAsset, 0.0f);
+		}
+		else if (ProcessRootMotionMode == EProcessRootMotionMode::Ignore)
+		{
+			RootMotionTransform = FTransform::Identity;
+			RootMotionReferenceTransform = FTransform::Identity;
+		}
 	}
-	else if (ProcessRootMotionMode == EProcessRootMotionMode::Ignore)
+	else
 	{
+		// No root motion in animation, reset transform. 
 		RootMotionTransform = FTransform::Identity;
 		RootMotionReferenceTransform = FTransform::Identity;
 	}
@@ -1209,6 +1204,21 @@ void UDebugSkelMeshComponent::CheckClothTeleport()
 	// doesn't need cloth teleport while previewing
 }
 
+void UDebugSkelMeshComponent::SetTurnTableMode(EPersonaTurnTableMode::Type NewMode)
+{
+	if (TurnTableMode == NewMode)
+	{
+		return;
+	}
+	TurnTableMode = NewMode;
+	
+	if (TurnTableMode == EPersonaTurnTableMode::Stopped)
+	{
+		// Use SetProcessRootMotionModeInternal with current value to reset the root motion preview state.
+		SetProcessRootMotionModeInternal(ProcessRootMotionMode);
+	}
+}
+
 void UDebugSkelMeshComponent::TickComponent(float DeltaTime, enum ELevelTick TickType, FActorComponentTickFunction *ThisTickFunction)
 {
 	//Do not tick a skeletalmesh component if the skeletalmesh is compiling
