Fbx scene importer: Make sure all asset are registered to the asset registry even the one the importer return.

#rnx

diff --git a/Engine/Source/Editor/UnrealEd/Private/Fbx/FbxSceneImportFactory.cpp b/Engine/Source/Editor/UnrealEd/Private/Fbx/FbxSceneImportFactory.cpp
--- a/Engine/Source/Editor/UnrealEd/Private/Fbx/FbxSceneImportFactory.cpp
+++ b/Engine/Source/Editor/UnrealEd/Private/Fbx/FbxSceneImportFactory.cpp
@@ -2144,15 +2144,12 @@ void UFbxSceneImportFactory::ImportAllStaticMesh(void* VoidRootNodeToImport, voi
 
 	NewStaticMesh = Cast<UStaticMesh>(Object);
 
-	// Make sure to notify the asset registry of all assets created other than the one returned, which will notify the asset registry automatically.
+	// Make sure to notify the asset registry of all assets created.
 	for (auto AssetItKvp = AllNewAssets.CreateIterator(); AssetItKvp; ++AssetItKvp)
 	{
 		UObject* Asset = AssetItKvp.Value();
-		if (Asset != NewStaticMesh)
-		{
-			FAssetRegistryModule::AssetCreated(Asset);
-			Asset->MarkPackageDirty();
-		}
+		FAssetRegistryModule::AssetCreated(Asset);
+		Asset->MarkPackageDirty();
 	}
 	ImportedMeshCount = AllNewAssets.Num();
 	if (ImportedMeshCount == 1 && NewStaticMesh)
