IterativeCook: Fix incorrect package counts in status displays when cooking incrementally.
#rnx

diff --git a/Engine/Source/Editor/UnrealEd/Private/CookOnTheFlyServer.cpp b/Engine/Source/Editor/UnrealEd/Private/CookOnTheFlyServer.cpp
--- a/Engine/Source/Editor/UnrealEd/Private/CookOnTheFlyServer.cpp
+++ b/Engine/Source/Editor/UnrealEd/Private/CookOnTheFlyServer.cpp
@@ -1778,22 +1778,33 @@ void UCookOnTheFlyServer::UpdateDisplay(UE::Cook::FTickStackData& StackData, boo
 		return;
 	}
 
-	const int32 CookedPackagesCount = PackageDatas->GetNumCooked() - PackageDatas->GetNumCooked(ECookResult::NeverCookPlaceholder) - PackageDataFromBaseGameNum;
-	const int32 CookPendingCount = WorkerRequests->GetNumExternalRequests() + PackageDatas->GetMonitor().GetNumInProgress();
+	int32 CookedCountWitness = PackageDatas->GetNumCooked();
+	int32 PendingCountWitness = WorkerRequests->GetNumExternalRequests()
+		+ PackageDatas->GetMonitor().GetNumInProgress();
 	if (bForceDisplay ||
-		(DeltaProgressDisplayTime >= GCookProgressUpdateTime && CookPendingCount != 0 &&
-			(LastCookedPackagesCount != CookedPackagesCount || LastCookPendingCount != CookPendingCount || DeltaProgressDisplayTime > GCookProgressRepeatTime)))
-	{
-		UE_CLOG(!(StackData.TickFlags & ECookTickFlags::HideProgressDisplay) && (GCookProgressDisplay & (int32)ECookProgressDisplayMode::RemainingPackages),
-			LogCook,
-			Display,
+		(DeltaProgressDisplayTime >= GCookProgressUpdateTime && PendingCountWitness != 0 &&
+			(LastCookedPackagesCount != CookedCountWitness || LastCookPendingCount != PendingCountWitness
+				|| DeltaProgressDisplayTime > GCookProgressRepeatTime)))
+	{
+		const int32 CookedPackagesCount = PackageDatas->GetNumCooked()
+			- PackageDatas->GetNumCooked(ECookResult::NeverCookPlaceholder) - PackageDataFromBaseGameNum;
+		int32 CookPendingCount = WorkerRequests->GetNumExternalRequests()
+			+ PackageDatas->GetMonitor().GetNumInProgress();
+		// When a RequestCluster is doing a graph search, it marks uncookable packages as to-be-demoted, and
+		// iteratively skippable packages as cooked, but those packages remain in the request state until the cluster
+		// search is complete so we still count them as inprogress. Subtract them from the inprogress count.
+		for (const FRequestCluster& Cluster : PackageDatas->GetRequestQueue().GetRequestClusters())
+		{
+			CookPendingCount -= Cluster.GetPackagesToMarkNotInProgress();
+		}
+		UE_CLOG(!(StackData.TickFlags & ECookTickFlags::HideProgressDisplay)
+			&& (GCookProgressDisplay & (int32)ECookProgressDisplayMode::RemainingPackages),
+			LogCook, Display,
 			TEXT("Cooked packages %d Packages Remain %d Total %d"),
-			CookedPackagesCount,
-			CookPendingCount,
-			CookedPackagesCount + CookPendingCount);
+			CookedPackagesCount, CookPendingCount, CookedPackagesCount + CookPendingCount);
 
-		LastCookedPackagesCount = CookedPackagesCount;
-		LastCookPendingCount = CookPendingCount;
+		LastCookedPackagesCount = CookedCountWitness;
+		LastCookPendingCount = PendingCountWitness;
 		LastProgressDisplayTime = CurrentTime;
 	}
 	const double DeltaDiagnosticsDisplayTime = CurrentTime - LastDiagnosticsDisplayTime;
@@ -2472,16 +2483,23 @@ UCookOnTheFlyServer::ECookAction UCookOnTheFlyServer::DecideNextCookAction(UE::C
 
 int32 UCookOnTheFlyServer::NumMultiprocessLocalWorkerAssignments() const
 {
+	using namespace UE::Cook;
+
 	if (!CookDirector.IsValid())
 	{
 		return 0;
 	}
 	UE::Cook::FPackageDataMonitor& Monitor = PackageDatas->GetMonitor();
-	return WorkerRequests->GetNumExternalRequests() +
+	int32 Result = WorkerRequests->GetNumExternalRequests() +
 		PackageDatas->GetRequestQueue().Num() +
 		PackageDatas->GetLoadPrepareQueue().Num() +
 		PackageDatas->GetLoadReadyQueue().Num() +
 		PackageDatas->GetSaveQueue().Num();
+	for (const FRequestCluster& Cluster : PackageDatas->GetRequestQueue().GetRequestClusters())
+	{
+		Result -= Cluster.GetPackagesToMarkNotInProgress();
+	}
+	return Result;
 }
 
 void UCookOnTheFlyServer::PumpExternalRequests(const UE::Cook::FCookerTimer& CookerTimer)
diff --git a/Engine/Source/Editor/UnrealEd/Private/Cooker/CookRequestCluster.cpp b/Engine/Source/Editor/UnrealEd/Private/Cooker/CookRequestCluster.cpp
--- a/Engine/Source/Editor/UnrealEd/Private/Cooker/CookRequestCluster.cpp
+++ b/Engine/Source/Editor/UnrealEd/Private/Cooker/CookRequestCluster.cpp
@@ -98,9 +98,10 @@ FRequestCluster::FRequestCluster(UCookOnTheFlyServer& InCOTFS, FPackageDataSet&&
 	ReserveInitialRequests(InRequests.Num());
 	for (FPackageData* PackageData : InRequests)
 	{
-		ESuppressCookReason& Existing = OwnedPackageDatas.FindOrAdd(PackageData, ESuppressCookReason::Invalid);
-		check(Existing == ESuppressCookReason::Invalid);
-		Existing = ESuppressCookReason::NotSuppressed;
+		check(PackageData);
+		bool bExisted;
+		SetPackageDataSuppressReason(*PackageData, ESuppressCookReason::NotSuppressed, &bExisted);
+		check(!bExisted);
 	}
 	InRequests.Empty();
 }
@@ -150,13 +151,13 @@ FRequestCluster::FRequestCluster(UCookOnTheFlyServer& InCOTFS, TRingBuffer<FDisc
 		{
 			// If there are no new reachable platforms, add it to the cluster for cooking if it needs
 			// it, otherwise let it remain where it is
-			FDiscoveryQueueElement PoppedDiscovery = DiscoveryQueue.PopFrontValue();
-			Discovery = &PoppedDiscovery;
+			DiscoveryQueue.PopFrontValue();
+			Discovery = nullptr;
 			if (!PackageData.IsInProgress() && PackageData.GetPlatformsNeedingCookingNum() == 0)
 			{
 				PackageData.SendToState(EPackageState::Request, ESendFlags::QueueRemove,
 					EStateChangeReason::RequestCluster);
-				OwnedPackageDatas.Add(&PackageData, ESuppressCookReason::NotSuppressed);
+				SetPackageDataSuppressReason(PackageData, ESuppressCookReason::NotSuppressed);
 			}
 			continue;
 		}
@@ -210,7 +211,7 @@ FRequestCluster::FRequestCluster(UCookOnTheFlyServer& InCOTFS, TRingBuffer<FDisc
 		// and add it to this cluster.
 		PackageData.SendToState(EPackageState::Request, ESendFlags::QueueRemove, EStateChangeReason::RequestCluster);
 		PackageData.AddUrgency(Discovery->bUrgent, false /* bAllowUpdateState */);
-		OwnedPackageDatas.Add(&PackageData, ESuppressCookReason::NotSuppressed);
+		SetPackageDataSuppressReason(PackageData, ESuppressCookReason::NotSuppressed);
 	}
 }
 
@@ -313,8 +314,9 @@ void FRequestCluster::ReserveInitialRequests(int32 RequestNum)
 
 void FRequestCluster::PullIntoCluster(FPackageData& PackageData)
 {
-	ESuppressCookReason& Existing = OwnedPackageDatas.FindOrAdd(&PackageData, ESuppressCookReason::Invalid);
-	if (Existing == ESuppressCookReason::Invalid)
+	bool bExisted;
+	SetPackageDataSuppressReason(PackageData, ESuppressCookReason::NotSuppressed, &bExisted);
+	if (!bExisted)
 	{
 		// Steal it from wherever it is and send it to Request State. It has already been added to this cluster
 		if (PackageData.GetState() == EPackageState::Request)
@@ -326,7 +328,6 @@ void FRequestCluster::PullIntoCluster(FPackageData& PackageData)
 			PackageData.SendToState(EPackageState::Request, ESendFlags::QueueRemove,
 				EStateChangeReason::RequestCluster);
 		}
-		Existing = ESuppressCookReason::NotSuppressed;
 	}
 }
 
@@ -350,9 +351,9 @@ void FRequestCluster::StartAsync(const FCookerTimer& CookerTimer, bool& bOutComp
 			// If the EditorDomain is active, then batch-download all packages to cook from remote cache into local
 			TArray<FName> BatchDownload;
 			BatchDownload.Reserve(OwnedPackageDatas.Num());
-			for (TPair<FPackageData*, ESuppressCookReason>& Pair : OwnedPackageDatas)
+			for (TPair<FPackageData*, FProcessingFlags>& Pair : OwnedPackageDatas)
 			{
-				if (Pair.Value == ESuppressCookReason::NotSuppressed)
+				if (Pair.Value.GetSuppressReason() == ESuppressCookReason::NotSuppressed)
 				{
 					BatchDownload.Add(Pair.Key->GetPackageName());
 				}
@@ -371,10 +372,16 @@ int32 FRequestCluster::NumPackageDatas() const
 
 void FRequestCluster::RemovePackageData(FPackageData* PackageData)
 {
-	if (OwnedPackageDatas.Remove(PackageData) == 0)
+	FProcessingFlags RemovedFlags;
+	if (!OwnedPackageDatas.RemoveAndCopyValue(PackageData, RemovedFlags))
 	{
 		return;
 	}
+	check(RemovedFlags.IsValid());
+	if (RemovedFlags.ShouldMarkNotInProgress())
+	{
+		--PackagesToMarkNotInProgressCount;
+	}
 
 	if (GraphSearch)
 	{
@@ -382,6 +389,47 @@ void FRequestCluster::RemovePackageData(FPackageData* PackageData)
 	}
 }
 
+void FRequestCluster::SetPackageDataSuppressReason(FPackageData& PackageData, ESuppressCookReason Reason, bool* bOutExisted)
+{
+	check(Reason != ESuppressCookReason::Invalid);
+
+	FProcessingFlags& Existing = OwnedPackageDatas.FindOrAdd(&PackageData);
+	if (bOutExisted)
+	{
+		*bOutExisted = Existing.IsValid();
+	}
+	if (Existing.ShouldMarkNotInProgress())
+	{
+		--PackagesToMarkNotInProgressCount;
+	}
+	Existing.SetValid();
+	Existing.SetSuppressReason(Reason);
+	if (Existing.ShouldMarkNotInProgress())
+	{
+		++PackagesToMarkNotInProgressCount;
+	}
+}
+
+void FRequestCluster::SetPackageDataWasMarkedCooked(FPackageData& PackageData, bool bValue,
+	bool* bOutExisted)
+{
+	FProcessingFlags& Existing = OwnedPackageDatas.FindOrAdd(&PackageData);
+	if (bOutExisted)
+	{
+		*bOutExisted = Existing.IsValid();
+	}
+	if (Existing.ShouldMarkNotInProgress())
+	{
+		--PackagesToMarkNotInProgressCount;
+	}
+	Existing.SetValid();
+	Existing.SetWasMarkedCooked(bValue);
+	if (Existing.ShouldMarkNotInProgress())
+	{
+		++PackagesToMarkNotInProgressCount;
+	}
+}
+
 void FRequestCluster::OnNewReachablePlatforms(FPackageData* PackageData)
 {
 	if (GraphSearch)
@@ -443,15 +491,15 @@ void FRequestCluster::ClearAndDetachOwnedPackageDatas(TArray<FPackageData*>& Out
 		check(!GraphSearch);
 		OutRequestsToLoad.Reset();
 		OutRequestsToDemote.Reset();
-		for (TPair<FPackageData*, ESuppressCookReason>& Pair : OwnedPackageDatas)
+		for (TPair<FPackageData*, FProcessingFlags>& Pair : OwnedPackageDatas)
 		{
-			if (Pair.Value == ESuppressCookReason::NotSuppressed)
+			if (Pair.Value.GetSuppressReason() == ESuppressCookReason::NotSuppressed)
 			{
 				OutRequestsToLoad.Add(Pair.Key);
 			}
 			else
 			{
-				OutRequestsToDemote.Add(Pair);
+				OutRequestsToDemote.Add({ Pair.Key, Pair.Value.GetSuppressReason() });
 			}
 		}
 		OutRequestGraph = MoveTemp(RequestGraph);
@@ -459,7 +507,7 @@ void FRequestCluster::ClearAndDetachOwnedPackageDatas(TArray<FPackageData*>& Out
 	else
 	{
 		OutRequestsToLoad.Reset();
-		for (TPair<FPackageData*, ESuppressCookReason>& Pair : OwnedPackageDatas)
+		for (TPair<FPackageData*, FProcessingFlags>& Pair : OwnedPackageDatas)
 		{
 			OutRequestsToLoad.Add(Pair.Key);
 		}
@@ -468,6 +516,7 @@ void FRequestCluster::ClearAndDetachOwnedPackageDatas(TArray<FPackageData*>& Out
 	}
 	FilePlatformRequests.Empty();
 	OwnedPackageDatas.Empty();
+	PackagesToMarkNotInProgressCount = 0;
 	GraphSearch.Reset();
 	RequestGraph.Reset();
 }
@@ -517,9 +566,9 @@ void FRequestCluster::PumpExploration(const FCookerTimer& CookerTimer, bool& bOu
 
 	TArray<FPackageData*> SortedPackages;
 	SortedPackages.Reserve(OwnedPackageDatas.Num());
-	for (TPair<FPackageData*, ESuppressCookReason>& Pair : OwnedPackageDatas)
+	for (TPair<FPackageData*, FProcessingFlags>& Pair : OwnedPackageDatas)
 	{
-		if (Pair.Value == ESuppressCookReason::NotSuppressed)
+		if (Pair.Value.GetSuppressReason() == ESuppressCookReason::NotSuppressed)
 		{
 			SortedPackages.Add(Pair.Key);
 		}
@@ -611,7 +660,7 @@ void FRequestCluster::FGraphSearch::VisitWithoutDependencies()
 {
 	// PumpExploration is responsible for marking all requests as explored and cookable/uncoookable.
 	// If we're skipping the dependencies search, handle that responsibility for the initial requests and return.
-	for (TPair<FPackageData*, ESuppressCookReason>& Pair : Cluster.OwnedPackageDatas)
+	for (TPair<FPackageData*, FProcessingFlags>& Pair : Cluster.OwnedPackageDatas)
 	{
 		FVertexData Vertex;
 		Vertex.PackageData = Pair.Key;
@@ -622,13 +671,12 @@ void FRequestCluster::FGraphSearch::VisitWithoutDependencies()
 void FRequestCluster::FGraphSearch::StartSearch()
 {
 	Frontier.Reserve(Cluster.OwnedPackageDatas.Num());
-	for (TPair<FPackageData*, ESuppressCookReason>& Pair : Cluster.OwnedPackageDatas)
+	for (TPair<FPackageData*, FProcessingFlags>& Pair : Cluster.OwnedPackageDatas)
 	{
 		FVertexData& Vertex = FindOrAddVertex(Pair.Key->GetPackageName(), *Pair.Key);
 		check(Vertex.PackageData);
 		// We're iterating over OwnedPackageDatas and the Vertex is already in the Cluster so we don't need to call
 		// AddToFrontier; just add it directly.
-		check(Pair.Value != ESuppressCookReason::Invalid);
 		Frontier.Add(&Vertex);
 	}
 }
@@ -926,7 +974,7 @@ void FRequestCluster::FGraphSearch::VisitVertex(FVertexData& Vertex)
 		{
 			check(SuppressCookReason == ESuppressCookReason::NotSuppressed);
 		}
-		Cluster.OwnedPackageDatas.FindOrAdd(Vertex.PackageData) = SuppressCookReason;
+		Cluster.SetPackageDataSuppressReason(*Vertex.PackageData, SuppressCookReason);
 		Vertex.bAnyCookable = bAnyCookable;
 	}
 
@@ -1150,6 +1198,7 @@ void FRequestCluster::FGraphSearch::ExploreVertexEdges(FVertexData& Vertex)
 				// Call SetPlatformCooked instead of just PackagePlatformData.SetCookResults because we might also need
 				// to set OnFirstCookedPlatformAdded
 				PackageData.SetPlatformCooked(TargetPlatform, ECookResult::Succeeded);
+				Cluster.SetPackageDataWasMarkedCooked(PackageData, true);
 				if (PlatformIndex == FirstSessionPlatformIndex)
 				{
 					COOK_STAT(++DetailedCookStats::NumPackagesIterativelySkipped);
diff --git a/Engine/Source/Editor/UnrealEd/Private/Cooker/CookRequestCluster.h b/Engine/Source/Editor/UnrealEd/Private/Cooker/CookRequestCluster.h
--- a/Engine/Source/Editor/UnrealEd/Private/Cooker/CookRequestCluster.h
+++ b/Engine/Source/Editor/UnrealEd/Private/Cooker/CookRequestCluster.h
@@ -78,6 +78,10 @@ public:
 	void ClearAndDetachOwnedPackageDatas(TArray<FPackageData*>& OutRequestsToLoad,
 		TArray<TPair<FPackageData*, ESuppressCookReason>>& OutRequestsToDemote,
 		TMap<FPackageData*, TArray<FPackageData*>>& OutRequestGraph);
+	/**
+	 * Report packages that are in request state and assigned to this Cluster, but that should not be counted as in
+	 * progress for progress displays because this cluster has marked them as already cooked or as to be demoted. */
+	int32 GetPackagesToMarkNotInProgress() const;
 
 	static TConstArrayView<FName> GetLocalizationReferences(FName PackageName, UCookOnTheFlyServer& InCOTFS);
 	static TArray<FName> GetAssetManagerReferences(FName PackageName);
@@ -323,6 +327,28 @@ private:
 		FEventRef AsyncResultsReadyEvent;
 	};
 
+	/** Tracks flags about this cluster's processing state for its OwnedPackageDatas. */
+	struct FProcessingFlags
+	{
+		/** Whether this struct has been set valid, used to identify whether values exist in a TMap. */
+		bool IsValid() const;
+		/** The package's SuppressCookReason, either NotSuppressed or a reason it was suppressed. */
+		ESuppressCookReason GetSuppressReason() const;
+		/** Whether the package was marked as cooked for any platform by this cluster. */
+		bool WasMarkedCooked() const;
+		/** Whether the values indicate the package should be added to the cluster's PackagesToMarkNotInProgress. */
+		bool ShouldMarkNotInProgress() const;
+
+		void SetValid();
+		void SetSuppressReason(ESuppressCookReason Value);
+		void SetWasMarkedCooked(bool bValue);
+
+	private:
+		ESuppressCookReason SuppressCookReason = ESuppressCookReason::NotSuppressed;
+		bool bValid = false;
+		bool bWasMarkedCooked = false;
+	};
+
 private:
 	explicit FRequestCluster(UCookOnTheFlyServer& COTFS);
 	void ReserveInitialRequests(int32 RequestNum);
@@ -331,6 +357,9 @@ private:
 	void PumpExploration(const FCookerTimer& CookerTimer, bool& bOutComplete);
 	void StartAsync(const FCookerTimer& CookerTimer, bool& bOutComplete);
 	bool IsIncrementalCook() const;
+	void SetPackageDataSuppressReason(FPackageData& PackageData, ESuppressCookReason Reason,
+		bool* bOutExisted = nullptr);
+	void SetPackageDataWasMarkedCooked(FPackageData& PackageData, bool bValue, bool* bOutExisted = nullptr);
 	void IsRequestCookable(const ITargetPlatform* TargetPlatform, FName PackageName, FPackageData& PackageData,
 		ESuppressCookReason& OutReason, bool& bOutCookable, bool& bOutExplorable);
 	static void IsRequestCookable(const ITargetPlatform* TargetPlatform, FName PackageName, FPackageData& PackageData,
@@ -340,7 +369,13 @@ private:
 		const TMap<FPackageData*, TArray<FPackageData*>>& Edges);
 
 	TArray<FFilePlatformRequest> FilePlatformRequests;
-	TFastPointerMap<FPackageData*, ESuppressCookReason> OwnedPackageDatas;
+	/**
+	 * Set of all packageDatas owned by this cluster (they are in the request state and this is the requeststate
+	 * container that records them). The count of PackageDatas matching certain properties is stored in
+	 * PackagesToMarkNotInProgress and must be updated whenever values change in OwnedPackageDatas. Call
+	 * SetPackageData... functions or RemovePackageData instead of modifying it directly.
+	 */
+	TFastPointerMap<FPackageData*, FProcessingFlags> OwnedPackageDatas;
 	TMap<FPackageData*, TArray<FPackageData*>> RequestGraph;
 	FString DLCPath;
 	TUniquePtr<FGraphSearch> GraphSearch; // Needs to be dynamic-allocated because of large alignment
@@ -349,6 +384,7 @@ private:
 	IAssetRegistry& AssetRegistry;
 	FPackageTracker& PackageTracker;
 	FBuildDefinitions& BuildDefinitions;
+	int32 PackagesToMarkNotInProgressCount = 0;
 	bool bAllowHardDependencies = true;
 	bool bAllowSoftDependencies = true;
 	bool bErrorOnEngineContentUse = false;
@@ -359,4 +395,50 @@ private:
 	bool bPreQueueBuildDefinitions = true;
 };
 
+
+///////////////////////////////////////////////////////
+// Inline implementations
+///////////////////////////////////////////////////////
+
+inline int32 FRequestCluster::GetPackagesToMarkNotInProgress() const
+{
+	return PackagesToMarkNotInProgressCount;
+}
+
+inline bool FRequestCluster::FProcessingFlags::IsValid() const
+{
+	return bValid;
+}
+
+inline ESuppressCookReason FRequestCluster::FProcessingFlags::GetSuppressReason() const
+{
+	return SuppressCookReason;
+}
+
+inline bool FRequestCluster::FProcessingFlags::WasMarkedCooked() const
+{
+	return bWasMarkedCooked;
+}
+
+inline bool FRequestCluster::FProcessingFlags::ShouldMarkNotInProgress() const
+{
+	return bValid & 
+		((SuppressCookReason != ESuppressCookReason::NotSuppressed) | bWasMarkedCooked);
+}
+
+inline void FRequestCluster::FProcessingFlags::SetValid()
+{
+	bValid = true;
+}
+
+inline void FRequestCluster::FProcessingFlags::SetSuppressReason(ESuppressCookReason Value)
+{
+	SuppressCookReason = Value;
+}
+
+inline void FRequestCluster::FProcessingFlags::SetWasMarkedCooked(bool bValue)
+{
+	bWasMarkedCooked = bValue;
+}
+
 }
