Adding an additional fallback so we can spawn the Slow Task widget if we are a Program that compiles agains the editor, but we don't use the mainframe module.

diff --git a/Engine/Source/Editor/UnrealEd/Private/FeedbackContextEditor.cpp b/Engine/Source/Editor/UnrealEd/Private/FeedbackContextEditor.cpp
--- a/Engine/Source/Editor/UnrealEd/Private/FeedbackContextEditor.cpp
+++ b/Engine/Source/Editor/UnrealEd/Private/FeedbackContextEditor.cpp
@@ -481,6 +481,13 @@ void FFeedbackContextEditor::StartSlowTask( const FText& Task, bool bShowCancelB
 			ParentWindow = MainFrame.GetParentWindow();
 		}
 
+		// If we still don't have a valid window, attempt to parent the slow task window to the top level window of the application
+		// This might be a program compiling against (and acting as) editor, but that does not use the Mainframe module
+		if (!ParentWindow.IsValid() && FPlatformProperties::IsProgram())
+		{
+			ParentWindow = FSlateApplication::Get().GetActiveTopLevelWindow();
+		}
+
 		if (ParentWindow.IsValid())
 		{
 			GSlowTaskOccurred = GIsSlowTask;
