Cooker: Fix DiffOnly -diffdenylist classes to include all subclasses of the ini list, since denylist classes from ini specify only the base class.
#rnx

diff --git a/Engine/Source/Editor/UnrealEd/Private/Cooker/DiffPackageWriter.cpp b/Engine/Source/Editor/UnrealEd/Private/Cooker/DiffPackageWriter.cpp
--- a/Engine/Source/Editor/UnrealEd/Private/Cooker/DiffPackageWriter.cpp
+++ b/Engine/Source/Editor/UnrealEd/Private/Cooker/DiffPackageWriter.cpp
@@ -88,7 +88,20 @@ void FDiffPackageWriter::ParseCmds()
 
 	if (FParse::Param(CommandLine, TEXT("DiffDenyList")))
 	{
-		CompareDenyListClasses = UE::EditorDomain::ConstructTargetIterativeClassBlockList();
+		for (FTopLevelAssetPath& DenyBaseClassPath : UE::EditorDomain::ConstructTargetIterativeClassBlockList())
+		{
+			UClass* DenyBaseClass = FindObject<UClass>(DenyBaseClassPath);
+			if (DenyBaseClass)
+			{
+				CompareDenyListClasses.Add(DenyBaseClassPath);
+				TArray<UClass*> DerivedClasses;
+				GetDerivedClasses(DenyBaseClass, DerivedClasses);
+				for (UClass* DerivedClass : DerivedClasses)
+				{
+					CompareDenyListClasses.Add(DerivedClass->GetClassPathName());
+				}
+			}
+		}
 	}
 }
 
