Fixes a crash if FocusViewportOnBox was called on an Editor Viewport while the Viewport tab was closed.

Updates Gameplay Tag Assets to use the Asset Definition api

diff --git a/Engine/Plugins/Editor/GameplayTagsEditor/Source/GameplayTagsEditor/GameplayTagsEditor.Build.cs b/Engine/Plugins/Editor/GameplayTagsEditor/Source/GameplayTagsEditor/GameplayTagsEditor.Build.cs
--- a/Engine/Plugins/Editor/GameplayTagsEditor/Source/GameplayTagsEditor/GameplayTagsEditor.Build.cs
+++ b/Engine/Plugins/Editor/GameplayTagsEditor/Source/GameplayTagsEditor/GameplayTagsEditor.Build.cs
@@ -23,6 +23,7 @@ namespace UnrealBuildTool.Rules
 					"CoreUObject",
 					"Engine",
 					"ApplicationCore",
+					"AssetDefinition",
 					"AssetTools",
 					"AssetRegistry",
 					"GameplayTags",
diff --git a/Engine/Plugins/Editor/GameplayTagsEditor/Source/GameplayTagsEditor/Private/AssetDefinition_GameplayTagAssetBase.cpp b/Engine/Plugins/Editor/GameplayTagsEditor/Source/GameplayTagsEditor/Private/AssetDefinition_GameplayTagAssetBase.cpp
new file mode 100644
--- /dev/null
+++ b/Engine/Plugins/Editor/GameplayTagsEditor/Source/GameplayTagsEditor/Private/AssetDefinition_GameplayTagAssetBase.cpp
@@ -0,0 +1,148 @@
+// Copyright Epic Games, Inc. All Rights Reserved.
+
+#include "AssetDefinition_GameplayTagAssetBase.h"
+
+#include "Framework/Application/SlateApplication.h"
+#include "Interfaces/IMainFrameModule.h"
+#include "SGameplayTagPicker.h"
+#include "ToolMenuSection.h"
+
+#define LOCTEXT_NAMESPACE "UAssetDefinition_GameplayTagAssetBase"
+
+void UAssetDefinition_GameplayTagAssetBase::AddGameplayTagsEditMenuExtension(FToolMenuSection& InSection, TArray<UObject*> InObjects, const FName& OwnedGameplayTagPropertyName)
+{
+	TArray<UObject*> Objects;
+	TArray<FGameplayTagContainer> Containers;
+
+	for (int32 ObjIdx = 0; ObjIdx < InObjects.Num(); ++ObjIdx)
+	{
+		if (UObject* CurObj = InObjects[ObjIdx])
+		{
+			const FStructProperty* Property = FindFProperty<FStructProperty>(CurObj->GetClass(), OwnedGameplayTagPropertyName);
+			if (Property != nullptr)
+			{
+				const FGameplayTagContainer& Container = *Property->ContainerPtrToValuePtr<FGameplayTagContainer>(CurObj);
+				Objects.Add(CurObj);
+				Containers.Add(Container);
+			}
+		}
+	}
+
+	if (Containers.Num() > 0)
+	{
+		const TAttribute<FText> Label = LOCTEXT("GameplayTags_Edit", "Edit Gameplay Tags...");
+		const TAttribute<FText> ToolTip = LOCTEXT("GameplayTags_EditToolTip", "Opens the Gameplay Tag Editor.");
+		const FSlateIcon Icon = FSlateIcon();
+		const FToolMenuExecuteAction UIAction = FToolMenuExecuteAction::CreateLambda([Objects, Containers, OwnedGameplayTagPropertyName](const FToolMenuContext& InContext)
+			{
+				OpenGameplayTagEditor(Objects, Containers, OwnedGameplayTagPropertyName);
+			});
+		InSection.AddMenuEntry("GameplayTags_Edit", Label, ToolTip, Icon, UIAction);
+	}
+}
+
+void UAssetDefinition_GameplayTagAssetBase::OpenGameplayTagEditor(TArray<UObject*> Objects, TArray<FGameplayTagContainer> Containers, const FName& OwnedGameplayTagPropertyName)
+{
+	if (!Objects.Num() || !Containers.Num())
+	{
+		return;
+	}
+
+	check(Objects.Num() == Containers.Num());
+
+	for (UObject* Object : Objects)
+	{
+		check(IsValid(Object));
+		Object->SetFlags(RF_Transactional);
+	}
+
+	FText Title;
+	FText AssetName;
+
+	const int32 NumAssets = Containers.Num();
+	if (NumAssets > 1)
+	{
+		AssetName = FText::Format(LOCTEXT("AssetTypeActions_GameplayTagAssetBaseMultipleAssets", "{0} Assets"), FText::AsNumber(NumAssets));
+		Title = FText::Format(LOCTEXT("AssetTypeActions_GameplayTagAssetBaseEditorTitle", "Tag Editor: Owned Gameplay Tags: {0}"), AssetName);
+	}
+	else
+	{
+		AssetName = FText::FromString(GetNameSafe(Objects[0]));
+		Title = FText::Format(LOCTEXT("AssetTypeActions_GameplayTagAssetBaseEditorTitle", "Tag Editor: Owned Gameplay Tags: {0}"), AssetName);
+	}
+
+	TSharedPtr<SWindow> Window = SNew(SWindow)
+		.Title(Title)
+		.ClientSize(FVector2D(500, 600))
+		[
+			SNew(SGameplayTagPicker)
+				.TagContainers(Containers)
+				.MaxHeight(0) // unbounded
+				.OnRefreshTagContainers_Lambda([Objects, OwnedGameplayTagPropertyName = OwnedGameplayTagPropertyName](SGameplayTagPicker& TagPicker)
+					{
+						// Refresh tags from objects, this is called e.g. on post undo/redo. 
+						TArray<FGameplayTagContainer> Containers;
+						for (UObject* Object : Objects)
+						{
+							// Adding extra entry even if the object has gone invalid to keep the container count the same as object count.
+							FGameplayTagContainer& NewContainer = Containers.AddDefaulted_GetRef();
+							if (IsValid(Object))
+							{
+								const FStructProperty* Property = FindFProperty<FStructProperty>(Object->GetClass(), OwnedGameplayTagPropertyName);
+								if (Property != nullptr)
+								{
+									NewContainer = *Property->ContainerPtrToValuePtr<FGameplayTagContainer>(Object);
+								}
+							}
+						}
+						TagPicker.SetTagContainers(Containers);
+					})
+				.OnTagChanged_Lambda([Objects, OwnedGameplayTagPropertyName = OwnedGameplayTagPropertyName](const TArray<FGameplayTagContainer>& TagContainers)
+					{
+						// Sanity check that our arrays are in sync.
+						if (Objects.Num() != TagContainers.Num())
+						{
+							return;
+						}
+
+						for (int32 Index = 0; Index < Objects.Num(); Index++)
+						{
+							UObject* Object = Objects[Index];
+							if (!IsValid(Object))
+							{
+								continue;
+							}
+
+							FStructProperty* Property = FindFProperty<FStructProperty>(Object->GetClass(), OwnedGameplayTagPropertyName);
+							if (!Property)
+							{
+								continue;
+							}
+
+							Object->Modify();
+							FGameplayTagContainer& Container = *Property->ContainerPtrToValuePtr<FGameplayTagContainer>(Object);
+
+							FEditPropertyChain PropertyChain;
+							PropertyChain.AddHead(Property);
+							Object->PreEditChange(PropertyChain);
+
+							Container = TagContainers[Index];
+
+							FPropertyChangedEvent PropertyEvent(Property);
+							Object->PostEditChangeProperty(PropertyEvent);
+						}
+					})
+		];
+
+	IMainFrameModule& MainFrameModule = FModuleManager::LoadModuleChecked<IMainFrameModule>(TEXT("MainFrame"));
+	if (MainFrameModule.GetParentWindow().IsValid())
+	{
+		FSlateApplication::Get().AddWindowAsNativeChild(Window.ToSharedRef(), MainFrameModule.GetParentWindow().ToSharedRef());
+	}
+	else
+	{
+		FSlateApplication::Get().AddWindow(Window.ToSharedRef());
+	}
+}
+
+#undef LOCTEXT_NAMESPACE
diff --git a/Engine/Plugins/Editor/GameplayTagsEditor/Source/GameplayTagsEditor/Private/AssetTypeActions_GameplayTagAssetBase.cpp b/Engine/Plugins/Editor/GameplayTagsEditor/Source/GameplayTagsEditor/Private/AssetTypeActions_GameplayTagAssetBase.cpp
deleted file mode 100644
diff --git a/Engine/Plugins/Editor/GameplayTagsEditor/Source/GameplayTagsEditor/Public/AssetDefinition_GameplayTagAssetBase.h b/Engine/Plugins/Editor/GameplayTagsEditor/Source/GameplayTagsEditor/Public/AssetDefinition_GameplayTagAssetBase.h
new file mode 100644
--- /dev/null
+++ b/Engine/Plugins/Editor/GameplayTagsEditor/Source/GameplayTagsEditor/Public/AssetDefinition_GameplayTagAssetBase.h
@@ -0,0 +1,45 @@
+// Copyright Epic Games, Inc. All Rights Reserved.
+
+#pragma once
+
+#include "AssetDefinitionDefault.h"
+
+#include "AssetDefinition_GameplayTagAssetBase.generated.h"
+
+struct FToolMenuSection;
+struct FGameplayTagContainer;
+
+/** Base asset type actions for any classes with gameplay tagging */
+UCLASS(Abstract)
+class GAMEPLAYTAGSEDITOR_API UAssetDefinition_GameplayTagAssetBase : public UAssetDefinitionDefault
+{
+	GENERATED_BODY()
+
+public:
+	UAssetDefinition_GameplayTagAssetBase() {};
+
+	/** Traditionally these are implemented in a MenuExtension namespace. However, UAssetDefinition_GameplayTagAssetBase is an abstract class,
+	* and the derived classes need to invoke this in their static MenuExtension functions**/
+	static void AddGameplayTagsEditMenuExtension(FToolMenuSection& InSection, TArray<UObject*> InObjects, const FName& OwnedGameplayTagPropertyName);
+
+	// UAssetDefinition Begin
+	/** Overridden to specify misc category */
+	virtual TConstArrayView<FAssetCategoryPath> GetAssetCategories() const override
+	{
+		static const auto Categories = { EAssetCategoryPaths::Misc };
+		return Categories;
+	}
+	// UAssetDefinition End
+
+private:
+	/**
+	 * Open the gameplay tag editor
+	 *
+	 * @param TagAssets	Assets to open the editor with
+	 */
+	static void OpenGameplayTagEditor(TArray<UObject*> Objects, TArray<FGameplayTagContainer> Containers, const FName& OwnedGameplayTagPropertyName);
+};
+
+#if UE_ENABLE_INCLUDE_ORDER_DEPRECATED_IN_5_2
+#include "CoreMinimal.h"
+#endif
diff --git a/Engine/Plugins/Editor/GameplayTagsEditor/Source/GameplayTagsEditor/Public/AssetTypeActions_GameplayTagAssetBase.h b/Engine/Plugins/Editor/GameplayTagsEditor/Source/GameplayTagsEditor/Public/AssetTypeActions_GameplayTagAssetBase.h
deleted file mode 100644
diff --git a/Engine/Source/Editor/UnrealEd/Private/EditorViewportClient.cpp b/Engine/Source/Editor/UnrealEd/Private/EditorViewportClient.cpp
--- a/Engine/Source/Editor/UnrealEd/Private/EditorViewportClient.cpp
+++ b/Engine/Source/Editor/UnrealEd/Private/EditorViewportClient.cpp
@@ -401,7 +401,7 @@ FEditorViewportClient::FEditorViewportClient(FEditorModeTools* InModeTools, FPre
 	, CameraSpeedScalar(1.0f)
 	, ImmersiveDelegate()
 	, VisibilityDelegate()
-	, Viewport(NULL)
+	, Viewport(nullptr)
 	, ViewportType(LVT_Perspective)
 	, ViewState()
 	, StereoViewStates()
@@ -556,7 +556,7 @@ FEditorViewportClient::~FEditorViewportClient()
 
 	if(Viewport)
 	{
-		UE_LOG(LogEditorViewport, Fatal, TEXT("Viewport != NULL in FEditorViewportClient destructor."));
+		UE_LOG(LogEditorViewport, Fatal, TEXT("Viewport != nullptr in FEditorViewportClient destructor."));
 	}
 
 	if(GEditor)
@@ -796,45 +796,50 @@ void FEditorViewportClient::ToggleOrbitCamera( bool bEnableOrbitCamera )
 
 void FEditorViewportClient::FocusViewportOnBox( const FBox& BoundingBox, bool bInstant /* = false */ )
 {
+	if (!Viewport)
+	{
+		return;
+	}
+
 	const FVector Position = BoundingBox.GetCenter();
 	float Radius = FMath::Max<FVector::FReal>(BoundingBox.GetExtent().Size(), 10.f);
 
 	float AspectToUse = AspectRatio;
 	FIntPoint ViewportSize = Viewport->GetSizeXY();
-	if(!bUseControllingActorViewInfo && ViewportSize.X > 0 && ViewportSize.Y > 0)
+	if (!bUseControllingActorViewInfo && ViewportSize.X > 0 && ViewportSize.Y > 0)
 	{
 		AspectToUse = Viewport->GetDesiredAspectRatio();
 	}
 
 	CameraController->ResetVelocity();
 
-	const bool bEnable=false;
+	const bool bEnable = false;
 	ToggleOrbitCamera(bEnable);
 
 	{
 		FViewportCameraTransform& ViewTransform = GetViewTransform();
 
-		if(!IsOrtho())
+		if (!IsOrtho())
 		{
-		   /**
-			* We need to make sure we are fitting the sphere into the viewport completely, so if the height of the viewport is less
-			* than the width of the viewport, we scale the radius by the aspect ratio in order to compensate for the fact that we have
-			* less visible vertically than horizontally.
-			*/
-			if( AspectToUse > 1.0f )
+			/**
+				* We need to make sure we are fitting the sphere into the viewport completely, so if the height of the viewport is less
+				* than the width of the viewport, we scale the radius by the aspect ratio in order to compensate for the fact that we have
+				* less visible vertically than horizontally.
+				*/
+			if (AspectToUse > 1.0f)
 			{
 				Radius *= AspectToUse;
 			}
 
 			/**
-			 * Now that we have a adjusted radius, we are taking half of the viewport's FOV,
-			 * converting it to radians, and then figuring out the camera's distance from the center
-			 * of the bounding sphere using some simple trig.  Once we have the distance, we back up
-			 * along the camera's forward vector from the center of the sphere, and set our new view location.
-			 */
+				* Now that we have a adjusted radius, we are taking half of the viewport's FOV,
+				* converting it to radians, and then figuring out the camera's distance from the center
+				* of the bounding sphere using some simple trig.  Once we have the distance, we back up
+				* along the camera's forward vector from the center of the sphere, and set our new view location.
+				*/
 
-			const float HalfFOVRadians = FMath::DegreesToRadians( ViewFOV / 2.0f);
-			const float DistanceFromSphere = Radius / FMath::Tan( HalfFOVRadians );
+			const float HalfFOVRadians = FMath::DegreesToRadians(ViewFOV / 2.0f);
+			const float DistanceFromSphere = Radius / FMath::Tan(HalfFOVRadians);
 			FVector CameraOffsetVector = ViewTransform.GetRotation().Vector() * -DistanceFromSphere;
 
 			ViewTransform.SetLookAt(Position);
@@ -847,7 +852,7 @@ void FEditorViewportClient::FocusViewportOnBox( const FBox& BoundingBox, bool bI
 			//SetViewLocation( Position );
 			ViewTransform.TransitionToLocation(Position, EditorViewportWidget, bInstant);
 
-			if( !(Viewport->KeyState(EKeys::LeftControl) || Viewport->KeyState(EKeys::RightControl)) )
+			if (!(Viewport->KeyState(EKeys::LeftControl) || Viewport->KeyState(EKeys::RightControl)))
 			{
 				/**
 				* We also need to zoom out till the entire volume is in view.  The following block of code first finds the minimum dimension
@@ -858,8 +863,8 @@ void FEditorViewportClient::FocusViewportOnBox( const FBox& BoundingBox, bool bI
 				uint32 MinAxisSize = (AspectToUse > 1.0f) ? Viewport->GetSizeXY().Y : Viewport->GetSizeXY().X;
 				float Zoom = Radius / (MinAxisSize / 2.0f);
 
-				NewOrthoZoom = Zoom * (Viewport->GetSizeXY().X*15.0f);
-				NewOrthoZoom = FMath::Clamp<float>( NewOrthoZoom, GetMinimumOrthoZoom(), MAX_ORTHOZOOM );
+				NewOrthoZoom = Zoom * (Viewport->GetSizeXY().X * 15.0f);
+				NewOrthoZoom = FMath::Clamp<float>(NewOrthoZoom, GetMinimumOrthoZoom(), MAX_ORTHOZOOM);
 				ViewTransform.SetOrthoZoom(NewOrthoZoom);
 			}
 		}
@@ -1792,7 +1797,7 @@ bool FEditorViewportClient::IsActiveViewportType(ELevelViewportType InViewportTy
 void FEditorViewportClient::UpdateCameraMovement( float DeltaTime )
 {
 	// We only want to move perspective cameras around like this
-	if( Viewport != NULL && IsPerspective() && !ShouldOrbitCamera() )
+	if( Viewport != nullptr && IsPerspective() && !ShouldOrbitCamera() )
 	{
 		const bool bEnable = false;
 		ToggleOrbitCamera(bEnable);
@@ -5083,7 +5088,7 @@ bool FEditorViewportClient::ShouldOrbitCamera() const
 /** Returns true if perspective flight camera input mode is currently active in this viewport */
 bool FEditorViewportClient::IsFlightCameraInputModeActive() const
 {
-	if( (Viewport != NULL) && IsPerspective() )
+	if( (Viewport != nullptr ) && IsPerspective() )
 	{
 		if( CameraController != NULL )
 		{
@@ -5152,7 +5157,7 @@ bool FEditorViewportClient::IsVisible() const
 void FEditorViewportClient::GetViewportDimensions( FIntPoint& OutOrigin, FIntPoint& Outize )
 {
 	OutOrigin = FIntPoint(0,0);
-	if ( Viewport != NULL )
+	if ( Viewport != nullptr )
 	{
 		Outize.X = Viewport->GetSizeXY().X;
 		Outize.Y = Viewport->GetSizeXY().Y;
