Changed UE::Anim::ExtractRootTransformFromAnimationAsset() to handle reference pose.

diff --git a/Engine/Source/Editor/UnrealEd/Private/Animation/DebugSkelMeshComponent.cpp b/Engine/Source/Editor/UnrealEd/Private/Animation/DebugSkelMeshComponent.cpp
--- a/Engine/Source/Editor/UnrealEd/Private/Animation/DebugSkelMeshComponent.cpp
+++ b/Engine/Source/Editor/UnrealEd/Private/Animation/DebugSkelMeshComponent.cpp
@@ -30,22 +30,6 @@
 //////////////////////////////////////////////////////////////////////////
 // UDebugSkelMeshComponent
 
-namespace UE::Anim::Private
-{
-	FTransform CalculateInitialTransformFromAssetAndTime(const UDebugSkelMeshComponent& InDebugSkelMeshComponent, const UAnimationAsset* InAnimAsset, const float InTime)
-	{
-		const FTransform InitialRootBoneTransform = InDebugSkelMeshComponent.GetReferenceSkeleton().GetRefBonePose()[0];
-		FTransform InitialTransform = UE::Anim::ExtractRootTransformFromAnimationAsset(InAnimAsset, InTime);
-		if (InAnimAsset->IsValidAdditive())
-		{
-			// Additive animations have zero scale as "no scaling" value.
-			// This function is used to set the initial transform, we explicitly set the scale to start at 1.
-			InitialTransform.SetScale3D(FVector::OneVector);
-		}
-		return InitialRootBoneTransform * InitialTransform;
-	}
-}
-
 UDebugSkelMeshComponent::UDebugSkelMeshComponent(const FObjectInitializer& ObjectInitializer)
 	: Super(ObjectInitializer)
 {
@@ -342,7 +326,7 @@ void UDebugSkelMeshComponent::ConsumeRootMotion(const FVector& FloorMin, const F
 				{
 					if (bLooped)
 					{
-						const FTransform InitialTransform = UE::Anim::Private::CalculateInitialTransformFromAssetAndTime(*this, PreviewInstance->CurrentAsset, SectionStartPosition);
+						const FTransform InitialTransform = UE::Anim::ExtractRootTransformFromAnimationAsset(PreviewInstance->CurrentAsset, SectionStartPosition);
 						const FTransform RootMotionDelta = UE::Anim::ExtractRootMotionFromAnimationAsset(PreviewInstance->CurrentAsset, PreviewInstance->GetMirrorDataTable(), SectionStartPosition, CurrentTime);
 						const FTransform RootMotionTransform = RootMotionDelta * InitialTransform;
 
@@ -487,7 +471,7 @@ void UDebugSkelMeshComponent::SetProcessRootMotionModeInternal(EProcessRootMotio
 				SectionStartPosition = StartTime;
 			}
 		
-			const FTransform InitialTransform = UE::Anim::Private::CalculateInitialTransformFromAssetAndTime(*this, PreviewInstance->CurrentAsset, SectionStartPosition);
+			const FTransform InitialTransform = UE::Anim::ExtractRootTransformFromAnimationAsset(PreviewInstance->CurrentAsset, SectionStartPosition);
 			const FTransform RootMotionDelta = UE::Anim::ExtractRootMotionFromAnimationAsset(PreviewInstance->CurrentAsset, PreviewInstance->GetMirrorDataTable(), SectionStartPosition, CurrentTime);
 			RootMotionTransform = RootMotionDelta * InitialTransform;
 			
@@ -772,7 +756,7 @@ void UDebugSkelMeshComponent::OnMirrorDataTableChanged()
 			SectionStartPosition = StartTime;
 		}
 
-		const FTransform InitialTransform = UE::Anim::Private::CalculateInitialTransformFromAssetAndTime(*this, PreviewInstance->CurrentAsset, SectionStartPosition);
+		const FTransform InitialTransform = UE::Anim::ExtractRootTransformFromAnimationAsset(PreviewInstance->CurrentAsset, SectionStartPosition);
 		const FTransform RootMotionDelta = UE::Anim::ExtractRootMotionFromAnimationAsset(PreviewInstance->CurrentAsset, PreviewInstance->GetMirrorDataTable(), SectionStartPosition, CurrentTime);
 		const FTransform RootMotionTransform = RootMotionDelta * InitialTransform;
 
diff --git a/Engine/Source/Runtime/Engine/Classes/Animation/AnimSequenceHelpers.h b/Engine/Source/Runtime/Engine/Classes/Animation/AnimSequenceHelpers.h
--- a/Engine/Source/Runtime/Engine/Classes/Animation/AnimSequenceHelpers.h
+++ b/Engine/Source/Runtime/Engine/Classes/Animation/AnimSequenceHelpers.h
@@ -205,6 +205,7 @@ namespace Anim {
 
 	/**
 	 * Extract absolute root transform at specific position in time from the animation asset.
+	 * Note: this function applies skeleton reference pose the transform (similar to ExtractRootMotionFromAnimationAsset, but unlike UAnimSequence::ExtractRootTrackTransform()). 
 	 *
 	 * @param	Animation		Animation asset to extract the root motion from.
 	 * @param	Position		Start Position of the animation.
diff --git a/Engine/Source/Runtime/Engine/Private/Animation/AnimSequenceHelpers.cpp b/Engine/Source/Runtime/Engine/Private/Animation/AnimSequenceHelpers.cpp
--- a/Engine/Source/Runtime/Engine/Private/Animation/AnimSequenceHelpers.cpp
+++ b/Engine/Source/Runtime/Engine/Private/Animation/AnimSequenceHelpers.cpp
@@ -236,6 +236,19 @@ FTransform ExtractRootMotionFromAnimationAsset(const UAnimationAsset* Animation,
 
 FTransform ExtractRootTransformFromAnimationAsset(const UAnimationAsset* Animation, float Time)
 {
+	const FVector DefaultScale(1.f);
+	bool bUseNormalizedRootMotionScale = false;
+
+	FTransform RootTransformRefPose = FTransform::Identity;
+	if (const USkeleton* MySkeleton = Animation->GetSkeleton())
+	{
+		const FReferenceSkeleton& RefSkeleton = MySkeleton->GetReferenceSkeleton();
+		if (RefSkeleton.GetNum() > 0)
+		{
+			RootTransformRefPose = RefSkeleton.GetRefBonePose()[0];
+		}
+	}
+
 	FTransform Result = FTransform::Identity;
 	if (const UAnimMontage* AnimMontage = Cast<UAnimMontage>(Animation))
 	{
@@ -245,14 +258,34 @@ FTransform ExtractRootTransformFromAnimationAsset(const UAnimationAsset* Animati
 			{
 				const float AnimSequenceTime = Segment->ConvertTrackPosToAnimPos(Time);
 				Result = AnimSequence->ExtractRootTrackTransform(AnimSequenceTime, nullptr);
+				bUseNormalizedRootMotionScale = AnimSequence->bUseNormalizedRootMotionScale;
 			}	
 		}
 	}
 	else if (const UAnimSequence* AnimSequence = Cast<UAnimSequence>(Animation))
 	{
 		Result = AnimSequence->ExtractRootTrackTransform(Time, nullptr);
+		bUseNormalizedRootMotionScale = AnimSequence->bUseNormalizedRootMotionScale;
 	}
 
+	// Use old calculation if needed.
+	if (bUseNormalizedRootMotionScale)
+	{
+		//Clear scale as it will muck up GetRelativeTransform
+		Result.SetScale3D(FVector(1.f));
+	}
+	else
+	{
+		if (Animation->IsValidAdditive())
+		{
+			Result.SetScale3D(Result.GetScale3D() + DefaultScale);
+		}
+	}
+
+	// Transform to Component Space
+	const FTransform RootToComponent = RootTransformRefPose.Inverse();
+	Result = RootToComponent * Result;
+
 	return Result;
 }
 
