World Partition - HLOD: Fix crash in the HLOD builder while processing non-wp maps containing level instances

diff --git a/Engine/Source/Editor/UnrealEd/Private/WorldPartition/WorldPartitionHLODsBuilder.cpp b/Engine/Source/Editor/UnrealEd/Private/WorldPartition/WorldPartitionHLODsBuilder.cpp
--- a/Engine/Source/Editor/UnrealEd/Private/WorldPartition/WorldPartitionHLODsBuilder.cpp
+++ b/Engine/Source/Editor/UnrealEd/Private/WorldPartition/WorldPartitionHLODsBuilder.cpp
@@ -449,21 +449,28 @@ bool UWorldPartitionHLODsBuilder::BuildHLODActors()
 			return bSaved;
 		});
 
-		uint32 NumHLODActors = 0;
+		TArray<IWorldPartitionHLODProvider*> HLODProviders;
+
+		// Gather all HLOD providers
 		for (TActorIterator<AActor> ActorIt(World); ActorIt; ++ActorIt)
 		{
 			if (IWorldPartitionHLODProvider* HLODProvider = Cast<IWorldPartitionHLODProvider>(*ActorIt))
 			{
-				bool bBuildResult = HLODProvider->BuildHLODActor(BuildHLODActorParams);
-				if (!bBuildResult)
-				{
-					return false;
-				}
+				HLODProviders.Add(HLODProvider);
+			}
+		}
 
-				NumHLODActors++;
+		// Process them one by one
+		for (IWorldPartitionHLODProvider* HLODProvider : HLODProviders)
+		{
+			bool bBuildResult = HLODProvider->BuildHLODActor(BuildHLODActorParams);
+			if (!bBuildResult)
+			{
+				return false;
 			}
 		}
-		UE_LOG(LogWorldPartitionHLODsBuilder, Display, TEXT("#### Built %d HLOD actor ####"), NumHLODActors);
+
+		UE_LOG(LogWorldPartitionHLODsBuilder, Display, TEXT("#### Built %d HLOD actor ####"), HLODProviders.Num());
 	}
 
 
