Cooker: CookPackageSplitter support for EGeneratedRequiresGenerator::Populate. This new setting is required by WorldPartition, which expects the alwaysloaded actors in the generator map to be loaded and kept in memory for referencing from every generated streaming cell.
#rnx

diff --git a/Engine/Plugins/Experimental/Mutable/Source/CustomizableObjectEditor/Private/MuCOE/CustomizableObjectCookPackageSplitter.h b/Engine/Plugins/Experimental/Mutable/Source/CustomizableObjectEditor/Private/MuCOE/CustomizableObjectCookPackageSplitter.h
--- a/Engine/Plugins/Experimental/Mutable/Source/CustomizableObjectEditor/Private/MuCOE/CustomizableObjectCookPackageSplitter.h
+++ b/Engine/Plugins/Experimental/Mutable/Source/CustomizableObjectEditor/Private/MuCOE/CustomizableObjectCookPackageSplitter.h
@@ -47,11 +47,13 @@ public:
 	 */
 	virtual bool UseInternalReferenceToAvoidGarbageCollect() override { return true; }
 	/**
-	 * If true, the cooker will wait for the generator save to complete before calling Populate and PreSave on the
-	 * generated package. And during MPCook the cooker will load and save generated packages only on the same
-	 * CookWorker that saved the generator package that containes the SplitData object. 
+	 * Return capability setting which indicates which splitter functions acting on the parent generator package must
+	 * be called on the splitter before splitter functions acting on the generated packages can be called. 
 	 */
-	virtual bool GeneratedReliesOnGeneratorSave() override { return true; }
+	virtual EGeneratedRequiresGenerator DoesGeneratedRequireGenerator() override
+	{
+		return EGeneratedRequiresGenerator::Save;
+	}
 
 private:
 
diff --git a/Engine/Source/Editor/UnrealEd/Private/CookOnTheFlyServer.cpp b/Engine/Source/Editor/UnrealEd/Private/CookOnTheFlyServer.cpp
--- a/Engine/Source/Editor/UnrealEd/Private/CookOnTheFlyServer.cpp
+++ b/Engine/Source/Editor/UnrealEd/Private/CookOnTheFlyServer.cpp
@@ -3367,9 +3367,10 @@ UE::Cook::EPollStatus UCookOnTheFlyServer::PrepareSaveGenerationPackage(UE::Cook
 		// If generator should not save until after generated, stall it here
 		if (Info.IsGenerator()
 			&& FGenerationHelper::IsGeneratedSavedFirst()
-			// Splitters that declare GeneratedReliesOnGeneratorSave ignore the global setting and
+			// Splitters that declare DoesGeneratedRequireGenerator=Save ignore the global setting and
 			// never wait for generated to save
-			&& !GenerationHelper.IsGeneratedReliesOnGeneratorSave())
+			&& GenerationHelper.DoesGeneratedRequireGenerator() <
+				ICookPackageSplitter::EGeneratedRequiresGenerator::Save)
 		{
 			if (GenerationHelper.IsWaitingForQueueResults())
 			{
@@ -3387,9 +3388,10 @@ UE::Cook::EPollStatus UCookOnTheFlyServer::PrepareSaveGenerationPackage(UE::Cook
 		// If generated should not save until after generator, stall it here
 		if (!Info.IsGenerator()
 			&& (FGenerationHelper::IsGeneratorSavedFirst()
-			// Splitters that declare GeneratedReliesOnGeneratorSave ignore the global setting and
+			// Splitters that declare DoesGeneratedRequireGenerator=Save ignore the global setting and
 			// always wait for the generator to save
-				|| GenerationHelper.IsGeneratedReliesOnGeneratorSave()))
+				|| GenerationHelper.DoesGeneratedRequireGenerator() >= 
+					ICookPackageSplitter::EGeneratedRequiresGenerator::Save))
 		{
 			if (GenerationHelper.GetOwner().IsInProgress())
 			{
@@ -3519,42 +3521,49 @@ UE::Cook::EPollStatus UCookOnTheFlyServer::BeginCacheObjectsToMove(UE::Cook::FGe
 	if (Info.GetSaveState() <= FCookGenerationInfo::ESaveState::CallObjectsToMove)
 	{
 		bool bPopulateSucceeded = false;
+		if (Info.IsGenerator() || GenerationHelper.DoesGeneratedRequireGenerator()
+			>= ICookPackageSplitter::EGeneratedRequiresGenerator::Populate)
+		{
+			if (!GenerationHelper.TryCallPopulateGeneratorPackage(GeneratedPackagesForPresave))
+			{
+				return EPollStatus::Error;
+			}
+		}
 		TArray<UObject*> ObjectsToMove;
-		TArray<UPackage*> KeepReferencedPackages;
 		if (Info.IsGenerator())
 		{
-			if (!TryConstructGeneratedPackagesForPresave(PackageData, GenerationHelper, GeneratedPackagesForPresave))
+			ObjectsToMove.Reserve(GenerationHelper.GetOwnerObjectsToMove().Num());
+			for (const FWeakObjectPtr& ObjectToMove : GenerationHelper.GetOwnerObjectsToMove())
 			{
-				UE_LOG(LogCook, Error, TEXT("PackageSplitter unexpected failure: could not ConstructGeneratedPackagesForPreSave. Splitter=%s"),
-					*GenerationHelper.GetSplitDataObjectName().ToString());
-				return EPollStatus::Error;
+				UObject* Object = ObjectToMove.Get();
+				if (Object)
+				{
+					ObjectsToMove.Add(Object);
+				}
 			}
-			FScopedActivePackage ScopedActivePackage(*this, GenerationHelper.GetOwner().GetPackageName(),
-				PackageAccessTrackingOps::NAME_CookerBuildObject);
-			bPopulateSucceeded = Splitter->PopulateGeneratorPackage(Package, SplitDataObject,
-				GeneratedPackagesForPresave, ObjectsToMove, KeepReferencedPackages);
 		}
 		else
 		{
+			TArray<UPackage*> KeepReferencedPackages;
 			ICookPackageSplitter::FGeneratedPackageForPopulate SplitterInfo{ Info.RelativePath,
 				Info.GeneratedRootPath, Package, Info.IsCreateAsMap() };
 			FScopedActivePackage ScopedActivePackage(*this, GenerationHelper.GetOwner().GetPackageName(),
 				PackageAccessTrackingOps::NAME_CookerBuildObject);
 			bPopulateSucceeded = Splitter->PopulateGeneratedPackage(Package, SplitDataObject, SplitterInfo,
 				ObjectsToMove, KeepReferencedPackages);
-		}
+			if (!bPopulateSucceeded)
+			{
+				UE_LOG(LogCook, Error,
+					TEXT("CookPackageSplitter returned false from PopulateGeneratedPackage. Splitter=%s")
+					TEXT("\nGeneratedPackage: %s"),
+					*GenerationHelper.GetSplitDataObjectName().ToString(),
+					*PackageData.GetPackageName().ToString());
+				return EPollStatus::Error;
+			}
 
-		if (!bPopulateSucceeded)
-		{
-			UE_LOG(LogCook, Error, TEXT("CookPackageSplitter returned false from %s. Splitter=%s%s"),
-				Info.IsGenerator() ? TEXT("PopulateGeneratorPackage") : TEXT("PopulateGeneratedPackage"),
-				*GenerationHelper.GetSplitDataObjectName().ToString(), 
-				Info.IsGenerator() ? TEXT("") : *FString::Printf(TEXT("\nGeneratedPackage: %s"),
-					*PackageData.GetPackageName().ToString()));
-			return EPollStatus::Error;
+			Info.AddKeepReferencedPackages(GenerationHelper, KeepReferencedPackages);
 		}
 
-		Info.AddKeepReferencedPackages(KeepReferencedPackages);
 		Info.TakeOverCachedObjectsAndAddMoved(GenerationHelper, PackageData.GetCachedObjectsInOuter(), ObjectsToMove);
 		Info.SetSaveStateComplete(FCookGenerationInfo::ESaveState::CallObjectsToMove);
 	}
@@ -3603,7 +3612,7 @@ UE::Cook::EPollStatus UCookOnTheFlyServer::PreSaveGeneratorPackage(UE::Cook::FPa
 			return EPollStatus::Error;
 		}
 	}
-	Info.AddKeepReferencedPackages(KeepReferencedPackages);
+	Info.AddKeepReferencedPackages(GenerationHelper, KeepReferencedPackages);
 
 	return EPollStatus::Success;
 }
@@ -3756,7 +3765,7 @@ UE::Cook::EPollStatus UCookOnTheFlyServer::TryPopulateGeneratedPackage(UE::Cook:
 			return EPollStatus::Error;
 		}
 	}
-	GeneratedInfo.AddKeepReferencedPackages(KeepReferencedPackages);
+	GeneratedInfo.AddKeepReferencedPackages(GenerationHelper, KeepReferencedPackages);
 	bool bPackageIsMap = GeneratedPackage->ContainsMap();
 	if (bPackageIsMap != GeneratedInfo.IsCreateAsMap())
 	{
@@ -4643,6 +4652,7 @@ void UCookOnTheFlyServer::PumpSaves(UE::Cook::FTickStackData& StackData, uint32
 					{
 						// We're waiting on something other than pendingcookedplatformdatas; this loop does not yet handle
 						// updating anything else, so break out
+						break;
 					}
 					// sleep for a bit
 					FPlatformProcess::Sleep(0.0f);
@@ -11621,10 +11631,10 @@ void UCookOnTheFlyServer::GetPackagesToRetract(int32 NumToRetract, TArray<FName>
 		}
 		if (PackageData->IsGenerated())
 		{
-			if (PackageData->IsGeneratedReliesOnGeneratorSave()
+			if (PackageData->DoesGeneratedRequireGenerator() >= ICookPackageSplitter::EGeneratedRequiresGenerator::Save
 				|| MPCookGeneratorSplit == UE::Cook::EMPCookGeneratorSplit::AllOnSameWorker)
 			{
-				// With IsGeneratedReliesOnGeneratorSave or the AllOnSameWorker setting, GeneratedPackages are
+				// With EGeneratedRequiresGenerator::Save or the AllOnSameWorker setting, GeneratedPackages are
 				// constrained to this worker.
 				return false;
 			}
@@ -11635,10 +11645,12 @@ void UCookOnTheFlyServer::GetPackagesToRetract(int32 NumToRetract, TArray<FName>
 				GenerationHelper->GetOwnerInfo().GetSaveState()
 					>= FCookGenerationInfo::ESaveState::QueueGeneratedPackages)
 			{
-				if (GenerationHelper->IsGeneratedReliesOnGeneratorSave() ||
+				if (GenerationHelper->DoesGeneratedRequireGenerator()
+					>= ICookPackageSplitter::EGeneratedRequiresGenerator::Save
+					||
 					MPCookGeneratorSplit != UE::Cook::EMPCookGeneratorSplit::AnyWorker)
 				{
-					// With IsGeneratedReliesOnGeneratorSave or with any MPCookGeneratorSplit setting other than
+					// With EGeneratedRequiresGenerator::Save or with any MPCookGeneratorSplit setting other than
 					// AnyWorker, we make assignment decisions based on the worker that saved and queued the generator
 					// package. We do not track queuing separately; we assume it happened on the worker that saved the
 					// package. Therefore, do not allow retraction of a generator package if it has already entered
diff --git a/Engine/Source/Editor/UnrealEd/Private/Cooker/CookGenerationHelper.cpp b/Engine/Source/Editor/UnrealEd/Private/Cooker/CookGenerationHelper.cpp
--- a/Engine/Source/Editor/UnrealEd/Private/Cooker/CookGenerationHelper.cpp
+++ b/Engine/Source/Editor/UnrealEd/Private/Cooker/CookGenerationHelper.cpp
@@ -91,8 +91,8 @@ void FGenerationHelper::Initialize(const UObject* InSplitDataObject,
 	SplitDataObjectName = FName(FStringView(InSplitDataObject->GetFullName()));
 	bUseInternalReferenceToAvoidGarbageCollect =
 		CookPackageSplitterInstance->UseInternalReferenceToAvoidGarbageCollect();
-	bGeneratedReliesOnGeneratorSave =
-		CookPackageSplitterInstance->GeneratedReliesOnGeneratorSave();
+	DoesGeneratedRequireGeneratorValue =
+		CookPackageSplitterInstance->DoesGeneratedRequireGenerator();
 }
 
 void FGenerationHelper::InitializeAsInvalid()
@@ -381,7 +381,7 @@ bool FGenerationHelper::TryGenerateList()
 		check(PackageData->GetParentGenerator().IsNone() ||
 			PackageData->GetParentGenerator() == OwnerPackageName);
 		PackageData->SetGenerated(OwnerPackageName);
-		PackageData->SetGeneratedReliesOnGeneratorSave(bGeneratedReliesOnGeneratorSave);
+		PackageData->SetDoesGeneratedRequireGenerator(DoesGeneratedRequireGeneratorValue);
 		if (IFileManager::Get().FileExists(*PackageData->GetFileName().ToString()))
 		{
 			UE_LOG(LogCook, Warning,
@@ -408,7 +408,7 @@ bool FGenerationHelper::TryGenerateList()
 			[](const FAssetDependency& A, const FAssetDependency& B) { return A.LexicalLess(B); });
 		GeneratedInfo.PackageDependencies.SetNum(Algo::Unique(GeneratedInfo.PackageDependencies));
 		GeneratedInfo.SetIsCreateAsMap(bCreateAsMap);
-		if (bGeneratedReliesOnGeneratorSave ||
+		if (DoesGeneratedRequireGenerator() >= ICookPackageSplitter::EGeneratedRequiresGenerator::Save ||
 			COTFS.MPCookGeneratorSplit == EMPCookGeneratorSplit::AllOnSameWorker)
 		{
 			PackageData->SetWorkerAssignmentConstraint(FWorkerId::Local());
@@ -422,6 +422,68 @@ bool FGenerationHelper::TryGenerateList()
 	return true;
 }
 
+bool FGenerationHelper::TryCallPopulateGeneratorPackage(
+	TArray<ICookPackageSplitter::FGeneratedPackageForPreSave>& InOutGeneratedPackagesForPresave)
+{
+	if (bCalledPopulateGeneratorPackage)
+	{
+		return true;
+	}
+	FPackageData& OwnerPackageData = GetOwner();
+	FName OwnerPackageName = OwnerPackageData.GetPackageName();
+	UCookOnTheFlyServer& COTFS = OwnerPackageData.GetPackageDatas().GetCookOnTheFlyServer();
+	if (!bGeneratedList)
+	{
+		// Unexpected, caller should not call in this case
+		UE_LOG(LogCook, Error, TEXT("TryCallPopulateGeneratorPackage called for package %s without a previous successful call to TryGenerateList."),
+			*OwnerPackageName.ToString());
+		FDebug::DumpStackTraceToLog(ELogVerbosity::Warning);
+		return false;
+	}
+	check(IsValid()); // Could not have set bGeneratedList=true without being valid. 
+	UObject* LocalSplitDataObject = this->FindOrLoadSplitDataObject();
+	if (!LocalSplitDataObject)
+	{
+		UE_LOG(LogCook, Error,
+			TEXT("Failed to call PopulateGeneratorPackage, CookPackageSplitter missing. Splitter=%s"),
+			*GetSplitDataObjectName().ToString());
+		return false;
+	}
+	UPackage* LocalOwnerPackage = LocalSplitDataObject->GetPackage();
+	if (!COTFS.TryConstructGeneratedPackagesForPresave(OwnerPackageData, *this, InOutGeneratedPackagesForPresave))
+	{
+		UE_LOG(LogCook, Error,
+			TEXT("PackageSplitter unexpected failure: could not ConstructGeneratedPackagesForPreSave. Splitter=%s"),
+			*GetSplitDataObjectName().ToString());
+		return false;
+	}
+	UCookOnTheFlyServer::FScopedActivePackage ScopedActivePackage(COTFS, OwnerPackageName,
+		PackageAccessTrackingOps::NAME_CookerBuildObject);
+
+	TArray<UPackage*> KeepReferencedPackages;
+	TArray<UObject*> ObjectsToMove;
+	bool bPopulateSucceeded = CookPackageSplitterInstance->PopulateGeneratorPackage(LocalOwnerPackage,
+		LocalSplitDataObject, InOutGeneratedPackagesForPresave, ObjectsToMove, KeepReferencedPackages);
+	if (!bPopulateSucceeded)
+	{
+		UE_LOG(LogCook, Error, TEXT("CookPackageSplitter returned false from PopulateGeneratorPackage. Splitter=%s"),
+			*GetSplitDataObjectName().ToString());
+		return false;
+	}
+	OwnerInfo.AddKeepReferencedPackages(*this, KeepReferencedPackages);
+	OwnerObjectsToMove.Reserve(ObjectsToMove.Num());
+	for (UObject* Object : ObjectsToMove)
+	{
+		if (Object)
+		{
+			OwnerObjectsToMove.Emplace(Object);
+		}
+	}
+
+	bCalledPopulateGeneratorPackage = true;
+	return true;
+}
+
 void FGenerationHelper::StartOwnerSave()
 {
 	if (!IsValid())
@@ -622,7 +684,7 @@ void FGenerationHelper::ResetSaveState(FCookGenerationInfo& Info, UPackage* Pack
 				TEXT("PackageSplitter: %s on %s was GarbageCollected before we finished saving it. This prevents us from calling PostSave and may corrupt other packages that it altered during Populate. Splitter=%s."),
 				(!Package ? TEXT("UPackage") :
 					(!LocalOwnerPackage ? TEXT("ParentGenerator UPackage") : TEXT("SplitDataObject"))),
-				Info.PackageData ? *Info.PackageData->GetPackageName().ToString() : *Info.RelativePath,
+				*Info.GetPackageName(),
 				*GetSplitDataObjectName().ToString());
 		}
 		else
@@ -711,14 +773,27 @@ void FGenerationHelper::ResetSaveState(FCookGenerationInfo& Info, UPackage* Pack
 				TEXT("This will allow the objects to be garbage collected and cause failures in the splitter which expects them to remain loaded.\n")
 				TEXT("Package=%s, Splitter=%s, ReleaseSaveReason=%s"),
 				Info.IsGenerator() ? TEXT("generator") : TEXT("generated"),
-				Info.PackageData ? *Info.PackageData->GetPackageName().ToString() : *Info.RelativePath,
+				*Info.GetPackageName(),
 				*GetSplitDataObjectName().ToString(), LexToString(ReleaseSaveReason));
 		}
 		Info.CachedObjectsInOuterInfo.Empty();
 		Info.SetHasTakenOverCachedCookedPlatformData(false);
 	}
 	Info.SetHasIssuedUndeclaredMovedObjectsWarning(false);
-	Info.KeepReferencedPackages.Reset();
+
+	// Clear KeepReferencedPackages; we no longer have a contract that we keep them referenced, except for the
+	// generator. If the splitter requires EGeneratedRequiresGenerator::Populate, then we are required to keep them
+	// referenced until all packages have saved as well, so we keep them referenced for the lifetime of the
+	// GenerationHelper.
+	if (!Info.IsGenerator() || DoesGeneratedRequireGenerator() <
+		ICookPackageSplitter::EGeneratedRequiresGenerator::Populate)
+	{
+		Info.KeepReferencedPackages.Reset();
+	}
+	if (Info.IsGenerator())
+	{
+		OwnerObjectsToMove.Empty();
+	}
 }
 
 void FGenerationHelper::FetchExternalActorDependencies()
@@ -771,6 +846,19 @@ void FGenerationHelper::SetPreviousGeneratedPackages(TMap<FName, FIoHash>&& Pack
 	PreviousGeneratedPackages = MoveTemp(Packages);
 }
 
+template <typename T>
+static void AppendWeakPtrsToObjectPtrArray(TArray<T*>& Out, TArray<TWeakObjectPtr<T>>& In)
+{
+	Out.Reserve(Out.Num() + In.Num());
+	for (const TWeakObjectPtr<T>& WeakPtr : In)
+	{
+		T* Object = WeakPtr.Get();
+		if (Object)
+		{
+			Out.Add(Object);
+		}
+	}
+}
 void FGenerationHelper::PreGarbageCollect(FCookGenerationInfo& Info, TArray<TObjectPtr<UObject>>& GCKeepObjects,
 	TArray<UPackage*>& GCKeepPackages, TArray<FPackageData*>& GCKeepPackageDatas, bool& bOutShouldDemote)
 {
@@ -781,44 +869,82 @@ void FGenerationHelper::PreGarbageCollect(FCookGenerationInfo& Info, TArray<TObj
 
 	bOutShouldDemote = false;
 	check(Info.PackageData); // Caller validates this is non-null
-	if (Info.GetSaveState() > FCookGenerationInfo::ESaveState::CallPopulate)
+	if (!IsUseInternalReferenceToAvoidGarbageCollect() && !Info.PackageData->GetIsCookLast())
 	{
-		if (IsUseInternalReferenceToAvoidGarbageCollect() || Info.PackageData->GetIsCookLast())
+		// If we don't have a contract to keep the packagedata referenced during GC, don't report
+		// anything to garbage collection, and demote the package if it has progressed too far
+		if (Info.GetSaveState() > FCookGenerationInfo::ESaveState::CallPopulate)
 		{
-			UPackage* Package = Info.PackageData->GetPackage();
-			if (Package)
-			{
-				GCKeepPackages.Add(Package);
-				GCKeepPackageDatas.Add(Info.PackageData);
-			}
+			bOutShouldDemote = true;
 		}
-		else
+		return;
+	}
+
+	// When we have a contract to keep the packagedata referenced, keep its various object pointers referenced.
+
+	// We have a contract that KeepReferencedPackages in any Info are kept referenced. COTFS calls PreGarbageCollect
+	// for any info in the save state; with one exception only packages in the savestate can have
+	// non-empty KeepReferencedPackages. The one exception is the generator, which can keep 
+	// KeepReferencedPackages beyond its save state for the benefit of generated packages that rely on it.
+	bool bKeepingAnyObjects = false;
+	bool bNeedsGeneratorPackage = false;
+	if (&Info == &OwnerInfo)
+	{
+		// Handled by bCurrentGCHasKeptGeneratorKeepPackages
+	}
+	else if (!Info.KeepReferencedPackages.IsEmpty())
+	{
+		bKeepingAnyObjects = true;
+		AppendWeakPtrsToObjectPtrArray(GCKeepPackages, Info.KeepReferencedPackages);
+	}
+	if (!bCurrentGCHasKeptGeneratorKeepPackages)
+	{
+		bCurrentGCHasKeptGeneratorKeepPackages = true;
+		if (!OwnerInfo.KeepReferencedPackages.IsEmpty())
 		{
-			bOutShouldDemote = true;
+			bNeedsGeneratorPackage = true;
+			AppendWeakPtrsToObjectPtrArray(GCKeepPackages, OwnerInfo.KeepReferencedPackages);
 		}
 	}
+
+	// Keep the objects returned from GetObjectsToMove* functions referenced
 	if (Info.HasTakenOverCachedCookedPlatformData())
 	{
-		if (IsUseInternalReferenceToAvoidGarbageCollect())
+		bKeepingAnyObjects = true;
+		for (FCachedObjectInOuter& CachedObjectInOuter : Info.PackageData->GetCachedObjectsInOuter())
+		{
+			UObject* Object = CachedObjectInOuter.Object.Get();
+			if (Object)
+			{
+				GCKeepObjects.Add(Object);
+			}
+		}
+	}
+
+	// Keep the generator and generated package referenced if we've passed the call to populate, or if we are keeping
+	// any other objects referenced
+	if (bKeepingAnyObjects || Info.GetSaveState() > FCookGenerationInfo::ESaveState::CallPopulate)
+	{
+		bNeedsGeneratorPackage = true;
+		if (&Info != &OwnerInfo)
 		{
-			// For the UseInternalReferenceToAvoidGarbageCollect case, part of the CookPackageSplitter contract is that
-			// the Cooker will keep referenced the package and all objects returned from GetObjectsToMove* functions
-			// until the PostSave function is called
 			UPackage* Package = Info.PackageData->GetPackage();
 			if (Package)
 			{
 				GCKeepPackages.Add(Package);
 				GCKeepPackageDatas.Add(Info.PackageData);
 			}
-			GCKeepPackages.Append(Info.KeepReferencedPackages);
-			for (FCachedObjectInOuter& CachedObjectInOuter : Info.PackageData->GetCachedObjectsInOuter())
-			{
-				UObject* Object = CachedObjectInOuter.Object.Get();
-				if (Object)
-				{
-					GCKeepObjects.Add(Object);
-				}
-			}
+		}
+	}
+
+	if (bNeedsGeneratorPackage && !bCurrentGCHasKeptGeneratorPackage)
+	{
+		bCurrentGCHasKeptGeneratorPackage = true;
+		UPackage* Package = OwnerInfo.PackageData->GetPackage();
+		if (Package)
+		{
+			GCKeepPackages.Add(Package);
+			GCKeepPackageDatas.Add(Info.PackageData);
 		}
 	}
 }
@@ -829,6 +955,8 @@ void FGenerationHelper::PostGarbageCollect()
 	{
 		return;
 	}
+	bCurrentGCHasKeptGeneratorPackage = false;
+	bCurrentGCHasKeptGeneratorKeepPackages = false;
 
 	FPackageData& Owner = GetOwner();
 	if (Owner.GetState() == EPackageState::Save)
@@ -842,17 +970,16 @@ void FGenerationHelper::PostGarbageCollect()
 				TEXT("\n\tSplitter=%s."), *GetSplitDataObjectName().ToString());
 		}
 	}
-	else
+	else if (!IsUseInternalReferenceToAvoidGarbageCollect())
 	{
 		// After the Generator Package is saved, we drop our references to it and it can be garbage collected
 		// If we have any packages left to populate, our splitter contract requires that it be garbage collected
 		// because we promise that the package is not partially GC'd during calls to TryPopulateGeneratedPackage
 		// The splitter can opt-out of this contract and keep it referenced itself if it desires.
-		UPackage* LocalOwnerPackage = FindObject<UPackage>(nullptr, *Owner.GetPackageName().ToString());
-		if (LocalOwnerPackage)
+		if (!Owner.IsInProgress() && !Owner.IsKeepReferencedDuringGC())
 		{
-			if (!Owner.IsInProgress() && !Owner.IsKeepReferencedDuringGC() &&
-				!IsUseInternalReferenceToAvoidGarbageCollect())
+			UPackage* LocalOwnerPackage = FindObject<UPackage>(nullptr, *Owner.GetPackageName().ToString());
+			if (LocalOwnerPackage)
 			{
 				UE_LOG(LogCook, Error,
 					TEXT("PackageSplitter found the Generator package still in memory after it should have been deleted by GC.")
@@ -869,14 +996,14 @@ void FGenerationHelper::PostGarbageCollect()
 	bool bHasIssuedWarning = false;
 	for (FCookGenerationInfo& Info : PackagesToGenerate)
 	{
-		if (FindObject<UPackage>(nullptr, *Info.PackageData->GetPackageName().ToString()))
+		if (Info.PackageData && FindObject<UPackage>(nullptr, *Info.PackageData->GetPackageName().ToString()))
 		{
 			if (!Info.PackageData->IsKeepReferencedDuringGC() && !Info.HasSaved() && !bHasIssuedWarning)
 			{
 				UE_LOG(LogCook, Warning,
 					TEXT("PackageSplitter found a package it generated that was not removed from memory during garbage collection. This will cause errors later during population.")
 					TEXT("\n\tSplitter=%s, Generated=%s."), *GetSplitDataObjectName().ToString(),
-					*Info.PackageData->GetPackageName().ToString());
+					*Info.GetPackageName());
 				
 				{
 					// Compute UCookOnTheFlyServer's references so they are gathered by OBJ REFS below 
@@ -893,6 +1020,20 @@ void FGenerationHelper::PostGarbageCollect()
 		{
 			Info.SetHasCreatedPackage(false);
 		}
+		for (TArray<TWeakObjectPtr<UPackage>>::TIterator Iter(Info.KeepReferencedPackages); Iter; ++Iter)
+		{
+			const TWeakObjectPtr<UPackage>& KeepPtr = *Iter;
+			if (!KeepPtr.Get())
+			{
+				UE_LOG(LogCook, Warning,
+					TEXT("PackageSplitter returned a package in OutKeepReferencedPackages that the cooker tried to keep referenced, but it was removed by garbage collection anyway.")
+					TEXT(" This might cause errors during save of the generated packages.")
+					TEXT("\n\tSplitter=%s, Generated=%s."),
+					*GetSplitDataObjectName().ToString(),
+					*Info.GetPackageName());
+				Iter.RemoveCurrentSwap();
+			}
+		}
 	}
 }
 
@@ -977,7 +1118,8 @@ void FGenerationHelper::UpdateSaveAfterGarbageCollect(const FPackageData& Packag
 				*Info->PackageData->GetPackageName().ToString()));
 	}
 
-	// Remove raw pointers from RootMovedObjects if they no longer exist in the weakpointers in CachedObjectsInOuter
+	// Remove raw pointers from CachedObjectsInOuterInfo if they no longer exist in the weakpointers 
+	// in CachedObjectsInOuter
 	TSet<UObject*> CachedObjectsInOuterSet;
 	for (FCachedObjectInOuter& CachedObjectInOuter : Info->PackageData->GetCachedObjectsInOuter())
 	{
@@ -1130,6 +1272,27 @@ EPollStatus FCookGenerationInfo::RefreshPackageObjects(FGenerationHelper& Genera
 	return EPollStatus::Success;
 }
 
+void FCookGenerationInfo::AddKeepReferencedPackages(FGenerationHelper& GenerationHelper,
+	TArray<UPackage*>& InKeepReferencedPackages)
+{
+	KeepReferencedPackages.Reserve(KeepReferencedPackages.Num() + InKeepReferencedPackages.Num());
+	for (UPackage* Package : InKeepReferencedPackages)
+	{
+		TWeakObjectPtr<UPackage>& WeakPtr = KeepReferencedPackages.Emplace_GetRef(Package);
+		if (!WeakPtr.Get())
+		{
+			UE_LOG(LogCook, Warning,
+				TEXT("PackageSplitter returned a package in OutKeepReferencedPackages that is already marked as garbage.")
+				TEXT(" This might cause errors during save of the generated packages.")
+				TEXT("\n\tSplitter=%s, Generated=%s."),
+				*GenerationHelper.GetSplitDataObjectName().ToString(),
+				*GetPackageName());
+
+			KeepReferencedPackages.Pop(EAllowShrinking::No);
+		}
+	}
+}
+
 void FCookGenerationInfo::CreatePackageHash()
 {
 	FBlake3 Blake3;
diff --git a/Engine/Source/Editor/UnrealEd/Private/Cooker/CookGenerationHelper.h b/Engine/Source/Editor/UnrealEd/Private/Cooker/CookGenerationHelper.h
--- a/Engine/Source/Editor/UnrealEd/Private/Cooker/CookGenerationHelper.h
+++ b/Engine/Source/Editor/UnrealEd/Private/Cooker/CookGenerationHelper.h
@@ -116,7 +116,7 @@ public:
 	EPollStatus RefreshPackageObjects(FGenerationHelper& GenerationHelper, UPackage* Package,
 		bool& bOutFoundNewObjects, ESaveState DemotionState);
 
-	void AddKeepReferencedPackages(TArray<UPackage*>& InKeepReferencedPackages);
+	void AddKeepReferencedPackages(FGenerationHelper& GenerationHelper, TArray<UPackage*>& InKeepReferencedPackages);
 
 	/** Create the hash for this generated package, based on dependencies and GenerationHash. */
 	void CreatePackageHash();
@@ -132,6 +132,9 @@ public:
 
 	TConstArrayView<FAssetDependency> GetDependencies() const;
 
+	/** Return the packagename, for use in debug messages. Handles PackageData==nullptr by returning RelativePath. */
+	FString GetPackageName() const;
+
 public:
 	FIoHash PackageHash;
 	FString RelativePath;
@@ -139,7 +142,7 @@ public:
 	FBlake3Hash GenerationHash;
 	TArray<FAssetDependency> PackageDependencies;
 	FPackageData* PackageData = nullptr;
-	TArray<UPackage*> KeepReferencedPackages;
+	TArray<TWeakObjectPtr<UPackage>> KeepReferencedPackages;
 	TMap<UObject*, FCachedObjectInOuterGeneratorInfo> CachedObjectsInOuterInfo;
 private:
 	ESaveState GeneratorSaveState = ESaveState::StartSave;
@@ -221,8 +224,8 @@ public:
 	const FName GetSplitDataObjectName() const;
 	/** Return the Splitter's value for virtual bool UseInternalReferenceToAvoidGarbageCollect(). */
 	bool IsUseInternalReferenceToAvoidGarbageCollect() const;
-	/** Return the Splitter's value for virtual bool GeneratedReliesOnGeneratorSave(). */
-	bool IsGeneratedReliesOnGeneratorSave() const;
+	/** Return the Splitter's value for virtual bool DoesGeneratedRequireGenerator(). */
+	ICookPackageSplitter::EGeneratedRequiresGenerator DoesGeneratedRequireGenerator() const;
 	/** Return the cached pointer to the SplitDataObject. Returns null if no longer in memory or marked as garbage. */
 	UObject* GetWeakSplitDataObject() const;
 	/**
@@ -230,6 +233,8 @@ public:
 	 * return null if !IsValid or if not found even after loading the package.
 	 */
 	UObject* FindOrLoadSplitDataObject();
+	/** Return the ObjectsToMove that were returned for the generator package. Does not call initialize. */
+	TConstArrayView<FWeakObjectPtr> GetOwnerObjectsToMove() const;
 
 	/** Find the OwnerPackage in memory, returns null if invalid or not already loaded. Does not call Initialize. */
 	UPackage* GetOwnerPackage();
@@ -248,6 +253,12 @@ public:
 
 	/** Call the Splitter's GetGenerateList and create the PackageDatas. Logs errors and returns false on failure. */
 	bool TryGenerateList();
+	/**
+	 * Call the Splitter's PopulateGeneratorPackage if not yet called. Assumes GenerateList has been called. Logs
+	 * errors and returns false on failure.
+	 */
+	bool TryCallPopulateGeneratorPackage(
+		TArray<ICookPackageSplitter::FGeneratedPackageForPreSave>& InOutGeneratedPackagesForPresave);
 	/**
 	 * Mark that the SavePackage of the Owner is starting. Keeps a reference to keep the generator alive until save
 	 * is finished.
@@ -363,15 +374,20 @@ private:
 	TWeakObjectPtr<UPackage> OwnerPackage;
 	TMap<FName, FIoHash> PreviousGeneratedPackages;
 	TArray<FName> ExternalActorDependencies;
+	TArray<FWeakObjectPtr> OwnerObjectsToMove;
 	TRefCountPtr<FGenerationHelper> ReferenceFromKeepForIterative;
 	TRefCountPtr<FGenerationHelper> ReferenceFromKeepForQueueResults;
 	TRefCountPtr<FGenerationHelper> ReferenceFromKeepForGeneratorSave;
 	int32 MPCookNextAssignmentIndex = 0;
 	FWorkerId WorkerIdThatSavedGenerator = FWorkerId::Invalid();
 	EInitializeStatus InitializeStatus = EInitializeStatus::Uninitialized;
+	ICookPackageSplitter::EGeneratedRequiresGenerator DoesGeneratedRequireGeneratorValue =
+		ICookPackageSplitter::EGeneratedRequiresGenerator::None;
 	bool bUseInternalReferenceToAvoidGarbageCollect = false;
-	bool bGeneratedReliesOnGeneratorSave = false;
 	bool bGeneratedList = false;
+	bool bCalledPopulateGeneratorPackage = false;
+	bool bCurrentGCHasKeptGeneratorPackage = false;
+	bool bCurrentGCHasKeptGeneratorKeepPackages = false;
 };
 
 
@@ -449,14 +465,14 @@ inline void FCookGenerationInfo::SetIsGenerator(bool bValue)
 	bGenerator = bValue;
 }
 
-inline void FCookGenerationInfo::AddKeepReferencedPackages(TArray<UPackage*>& InKeepReferencedPackages)
+inline TConstArrayView<FAssetDependency> FCookGenerationInfo::GetDependencies() const
 {
-	KeepReferencedPackages.Append(InKeepReferencedPackages);
+	return PackageDependencies;
 }
 
-inline TConstArrayView<FAssetDependency> FCookGenerationInfo::GetDependencies() const
+inline FString FCookGenerationInfo::GetPackageName() const
 {
-	return PackageDependencies;
+	return PackageData ? *PackageData->GetPackageName().ToString() : *RelativePath;
 }
 
 inline bool FGenerationHelper::IsInitialized() const
@@ -519,10 +535,10 @@ inline bool FGenerationHelper::IsUseInternalReferenceToAvoidGarbageCollect() con
 	return bUseInternalReferenceToAvoidGarbageCollect;
 }
 
-inline bool FGenerationHelper::IsGeneratedReliesOnGeneratorSave() const
+inline ICookPackageSplitter::EGeneratedRequiresGenerator FGenerationHelper::DoesGeneratedRequireGenerator() const
 {
 	ConditionalInitialize();
-	return bGeneratedReliesOnGeneratorSave;
+	return DoesGeneratedRequireGeneratorValue;
 }
 
 inline UObject* FGenerationHelper::GetWeakSplitDataObject() const
@@ -531,6 +547,11 @@ inline UObject* FGenerationHelper::GetWeakSplitDataObject() const
 	return SplitDataObject.Get();
 }
 
+inline TConstArrayView<FWeakObjectPtr> FGenerationHelper::GetOwnerObjectsToMove() const
+{
+	return OwnerObjectsToMove;
+}
+
 inline TConstArrayView<FName> FGenerationHelper::GetExternalActorDependencies()
 {
 	return ExternalActorDependencies;
diff --git a/Engine/Source/Editor/UnrealEd/Private/Cooker/CookPackageData.cpp b/Engine/Source/Editor/UnrealEd/Private/Cooker/CookPackageData.cpp
--- a/Engine/Source/Editor/UnrealEd/Private/Cooker/CookPackageData.cpp
+++ b/Engine/Source/Editor/UnrealEd/Private/Cooker/CookPackageData.cpp
@@ -98,7 +98,8 @@ FPackageData::FPackageData(FPackageDatas& PackageDatas, const FName& InPackageNa
 	, bCookedPlatformDataStarted(0), bCookedPlatformDataCalled(0), bCookedPlatformDataComplete(0)
 	, MonitorCookResult((uint8)ECookResult::NotAttempted)
 	, bInitializedGeneratorSave(0), bCompletedGeneration(0), bGenerated(0), bKeepReferencedDuringGC(0)
-	, bWasCookedThisSession(0), bGeneratedReliesOnGeneratorSave(0)
+	, bWasCookedThisSession(0)
+	, DoesGeneratedRequireGeneratorValue(static_cast<uint32>(ICookPackageSplitter::EGeneratedRequiresGenerator::None))
 {
 	SetState(EPackageState::Idle);
 	SendToState(EPackageState::Idle, ESendFlags::QueueAdd, EStateChangeReason::Discovered);
diff --git a/Engine/Source/Editor/UnrealEd/Private/Cooker/CookPackageData.h b/Engine/Source/Editor/UnrealEd/Private/Cooker/CookPackageData.h
--- a/Engine/Source/Editor/UnrealEd/Private/Cooker/CookPackageData.h
+++ b/Engine/Source/Editor/UnrealEd/Private/Cooker/CookPackageData.h
@@ -640,9 +640,12 @@ public:
 	 */
 	TRefCountPtr<FGenerationHelper> TryCreateValidParentGenerationHelper();
 
-	/** Get/Set true if this is a generated package and its splitter reports GeneratedReliesOnGeneratorSave. */
-	bool IsGeneratedReliesOnGeneratorSave() const;
-	void SetGeneratedReliesOnGeneratorSave(bool bValue);
+	/**
+	 * Get/Set the package's parent's CookPackageSplitter's value for DoesGeneratedRequireGenerator.
+	 * Should only be called for generated packages.
+	 */
+	ICookPackageSplitter::EGeneratedRequiresGenerator DoesGeneratedRequireGenerator() const;
+	void SetDoesGeneratedRequireGenerator(ICookPackageSplitter::EGeneratedRequiresGenerator Value);
 
 	/**
 	 * Return the instigator for this package. The Instigator is the first code location or
@@ -825,7 +828,8 @@ private:
 	uint32 bGenerated : 1;
 	uint32 bKeepReferencedDuringGC : 1;
 	uint32 bWasCookedThisSession : 1;
-	uint32 bGeneratedReliesOnGeneratorSave : 1;
+	static_assert(static_cast<uint32>(ICookPackageSplitter::EGeneratedRequiresGenerator::Count) <= 4, "We are storing Enum value in 2 bits");
+	uint32 DoesGeneratedRequireGeneratorValue : 2;
 };
 
 /**
@@ -1469,14 +1473,14 @@ inline FName FPackageData::GetParentGenerator() const
 	return ParentGenerator;
 }
 
-inline bool FPackageData::IsGeneratedReliesOnGeneratorSave() const
+inline ICookPackageSplitter::EGeneratedRequiresGenerator FPackageData::DoesGeneratedRequireGenerator() const
 {
-	return bGeneratedReliesOnGeneratorSave != 0;
+	return static_cast<ICookPackageSplitter::EGeneratedRequiresGenerator>(DoesGeneratedRequireGeneratorValue);
 }
 
-inline void FPackageData::SetGeneratedReliesOnGeneratorSave(bool bValue)
+inline void FPackageData::SetDoesGeneratedRequireGenerator(ICookPackageSplitter::EGeneratedRequiresGenerator Value)
 {
-	bGeneratedReliesOnGeneratorSave = (uint32)bValue;
+	DoesGeneratedRequireGeneratorValue = static_cast<uint32>(Value);
 }
 
 } // namespace UE::Cook
\ No newline at end of file
diff --git a/Engine/Source/Editor/UnrealEd/Private/Cooker/CookWorkerClient.cpp b/Engine/Source/Editor/UnrealEd/Private/Cooker/CookWorkerClient.cpp
--- a/Engine/Source/Editor/UnrealEd/Private/Cooker/CookWorkerClient.cpp
+++ b/Engine/Source/Editor/UnrealEd/Private/Cooker/CookWorkerClient.cpp
@@ -318,7 +318,7 @@ void FCookWorkerClient::ReportDiscoveredPackage(const FPackageData& PackageData,
 	Discovered.Instigator = Instigator;
 	Discovered.Platforms = MoveTemp(ReachablePlatforms);
 	Discovered.Platforms.ConvertToBitfield(OrderedSessionAndSpecialPlatforms);
-	Discovered.bGeneratedReliesOnGeneratorSave = PackageData.IsGeneratedReliesOnGeneratorSave();
+	Discovered.DoesGeneratedRequireGenerator = PackageData.DoesGeneratedRequireGenerator();
 }
 
 void FCookWorkerClient::ReportGeneratorQueuedGeneratedPackages(FGenerationHelper& GenerationHelper)
@@ -854,6 +854,7 @@ void FCookWorkerClient::AssignPackages(FAssignPackagesMessage& Message)
 			if (!AssignData.ParentGenerator.IsNone())
 			{
 				PackageData.SetGenerated(AssignData.ParentGenerator);
+				PackageData.SetDoesGeneratedRequireGenerator(AssignData.DoesGeneratedRequireGenerator);
 			}
 			if (!AssignData.GeneratorPreviousGeneratedPackages.IsEmpty())
 			{
diff --git a/Engine/Source/Editor/UnrealEd/Private/Cooker/CookWorkerServer.cpp b/Engine/Source/Editor/UnrealEd/Private/Cooker/CookWorkerServer.cpp
--- a/Engine/Source/Editor/UnrealEd/Private/Cooker/CookWorkerServer.cpp
+++ b/Engine/Source/Editor/UnrealEd/Private/Cooker/CookWorkerServer.cpp
@@ -685,9 +685,7 @@ void FCookWorkerServer::SendPendingPackages()
 		FAssignPackageData& AssignData = AssignDatas.Emplace_GetRef();
 		AssignData.ConstructData = PackageData->CreateConstructData();
 		AssignData.ParentGenerator = PackageData->GetParentGenerator();
-		// We do not record PackageData->IsGeneratedReliesOnGeneratorSave(). The CookWorker that created
-		// it knows it, and the Director knows it from FDiscoveredPackageReplication, and the package will not
-		// be sent to other CookWorkers if it is true.
+		AssignData.DoesGeneratedRequireGenerator = PackageData->DoesGeneratedRequireGenerator();
 		AssignData.Instigator = PackageData->GetInstigator();
 		SessionPlatformNeedsCook.Init(false, OrderedSessionPlatforms.Num());
 		int32 PlatformIndex = 0;
@@ -1000,7 +998,7 @@ void FCookWorkerServer::QueueDiscoveredPackage(FDiscoveredPackageReplication&& D
 	if (!DiscoveredPackage.ParentGenerator.IsNone())
 	{
 		PackageData.SetGenerated(DiscoveredPackage.ParentGenerator);
-		PackageData.SetGeneratedReliesOnGeneratorSave(DiscoveredPackage.bGeneratedReliesOnGeneratorSave);
+		PackageData.SetDoesGeneratedRequireGenerator(DiscoveredPackage.DoesGeneratedRequireGenerator);
 		FPackageData* GeneratorPackageData = PackageDatas.FindPackageDataByPackageName(
 			DiscoveredPackage.ParentGenerator);
 		if (GeneratorPackageData)
@@ -1047,7 +1045,7 @@ void FCookWorkerServer::QueueDiscoveredPackage(FDiscoveredPackageReplication&& D
 	}
 
 	if (PackageData.IsGenerated()
-		&& (PackageData.IsGeneratedReliesOnGeneratorSave()
+		&& (PackageData.DoesGeneratedRequireGenerator() >= ICookPackageSplitter::EGeneratedRequiresGenerator::Save
 				|| COTFS.MPCookGeneratorSplit == EMPCookGeneratorSplit::AllOnSameWorker))
 	{
 		PackageData.SetWorkerAssignmentConstraint(GetWorkerId());
@@ -1184,6 +1182,8 @@ void FAssignPackageData::Write(FCbWriter& Writer,
 	Writer << Instigator;
 	WriteToCompactBinary(Writer, NeedCookPlatforms, OrderedSessionPlatforms);
 	Writer << GeneratorPreviousGeneratedPackages;
+	static_assert(sizeof(ICookPackageSplitter::EGeneratedRequiresGenerator) <= sizeof(uint8), "We are storing it in a uint8");
+	Writer << static_cast<uint8>(DoesGeneratedRequireGenerator);
 	Writer.EndArray();
 }
 
@@ -1196,6 +1196,17 @@ bool FAssignPackageData::TryRead(FCbFieldView Field, TConstArrayView<const ITarg
 	bOk = LoadFromCompactBinary(*It++, Instigator) & bOk;
 	bOk = LoadFromCompactBinary(*It++, NeedCookPlatforms, OrderedSessionPlatforms) & bOk;
 	bOk = LoadFromCompactBinary(*It++, GeneratorPreviousGeneratedPackages) & bOk;
+	uint8 DoesGeneratedRequireGeneratorInt = It->AsUInt8();
+	if (!(It++)->HasError() && DoesGeneratedRequireGeneratorInt
+		< static_cast<uint8>(ICookPackageSplitter::EGeneratedRequiresGenerator::Count))
+	{
+		DoesGeneratedRequireGenerator =
+			static_cast<ICookPackageSplitter::EGeneratedRequiresGenerator>(DoesGeneratedRequireGeneratorInt);
+	}
+	else
+	{
+		bOk = false;
+	}
 	return bOk;
 }
 
@@ -1374,7 +1385,8 @@ void FDiscoveredPackageReplication::Write(FCbWriter& Writer,
 	Writer << ParentGenerator;
 	Writer << static_cast<uint8>(Instigator.Category);
 	Writer << Instigator.Referencer;
-	Writer << bGeneratedReliesOnGeneratorSave;
+	Writer << static_cast<uint8>(DoesGeneratedRequireGenerator);
+	static_assert(sizeof(ICookPackageSplitter::EGeneratedRequiresGenerator) <= sizeof(uint8), "We are storing it in a uint8");
 	WriteToCompactBinary(Writer, Platforms, OrderedSessionAndSpecialPlatforms);
 	Writer.EndArray();
 }
@@ -1404,8 +1416,17 @@ bool FDiscoveredPackageReplication::TryRead(FCbFieldView Field,
 		bOk = false;
 	}
 	bOk = LoadFromCompactBinary(Iter++, Instigator.Referencer) & bOk;
-	bGeneratedReliesOnGeneratorSave = Iter->AsBool();
-	bOk = (!(Iter++)->HasError()) & bOk;
+	uint8 DoesGeneratedRequireGeneratorInt = Iter->AsUInt8();
+	if (!(Iter++)->HasError() && DoesGeneratedRequireGeneratorInt
+		< static_cast<uint8>(ICookPackageSplitter::EGeneratedRequiresGenerator::Count))
+	{
+		DoesGeneratedRequireGenerator = static_cast<ICookPackageSplitter::EGeneratedRequiresGenerator>(
+			DoesGeneratedRequireGeneratorInt);
+	}
+	else
+	{
+		bOk = false;
+	}
 	bOk = LoadFromCompactBinary(Iter++, Platforms, OrderedSessionAndSpecialPlatforms) & bOk;
 	if (!bOk)
 	{
diff --git a/Engine/Source/Editor/UnrealEd/Private/Cooker/CookWorkerServer.h b/Engine/Source/Editor/UnrealEd/Private/Cooker/CookWorkerServer.h
--- a/Engine/Source/Editor/UnrealEd/Private/Cooker/CookWorkerServer.h
+++ b/Engine/Source/Editor/UnrealEd/Private/Cooker/CookWorkerServer.h
@@ -13,6 +13,7 @@
 #include "Cooker/CookTypes.h"
 #include "Cooker/MPCollector.h"
 #include "CookOnTheSide/CookOnTheFlyServer.h"
+#include "CookPackageSplitter.h"
 #include "HAL/CriticalSection.h"
 #include "IO/IoHash.h"
 #include "Logging/LogVerbosity.h"
@@ -252,6 +253,8 @@ struct FAssignPackageData
 	FInstigator Instigator;
 	FDiscoveredPlatformSet NeedCookPlatforms;
 	TMap<FName, FIoHash> GeneratorPreviousGeneratedPackages;
+	ICookPackageSplitter::EGeneratedRequiresGenerator DoesGeneratedRequireGenerator
+		= ICookPackageSplitter::EGeneratedRequiresGenerator::None;
 
 private:
 	void Write(FCbWriter& Writer, TConstArrayView<const ITargetPlatform*> OrderedSessionPlatforms) const;
@@ -400,7 +403,8 @@ struct FDiscoveredPackageReplication
 	FName ParentGenerator;
 	FInstigator Instigator;
 	FDiscoveredPlatformSet Platforms;
-	bool bGeneratedReliesOnGeneratorSave = false;
+	ICookPackageSplitter::EGeneratedRequiresGenerator DoesGeneratedRequireGenerator =
+		ICookPackageSplitter::EGeneratedRequiresGenerator::None;
 
 private:
 	void Write(FCbWriter& Writer, TConstArrayView<const ITargetPlatform*> OrderedSessionAndSpecialPlatforms) const;
diff --git a/Engine/Source/Editor/UnrealEd/Public/CookPackageSplitter.h b/Engine/Source/Editor/UnrealEd/Public/CookPackageSplitter.h
--- a/Engine/Source/Editor/UnrealEd/Public/CookPackageSplitter.h
+++ b/Engine/Source/Editor/UnrealEd/Public/CookPackageSplitter.h
@@ -62,20 +62,52 @@ public:
 	{
 		return false;
 	}
+
+	/**
+	 * Return value for the DoesGeneratedRequireGenerator function. All levels behave correctly, but provide
+	 * different tradeoffs of guarantees to the splitter versus performance.
+	 */
+	enum class EGeneratedRequiresGenerator : uint8
+	{
+		/**
+		 * GetGenerateList will be called before PopulateGeneratedPackage. PopulateGenerator and PreSaveGenerator 
+		 * might or might not be called before. OutKeepReferencedPackages from PopulateGenerator will not be kept
+		 * referenced after PostSaveGenerator. Best for performance.
+		 */
+		None,
+		/**
+		 * GetGenerateList and PopulateGenerator will be called before PopulateGeneratedPackage.
+		 * OutKeepReferencedPackages from PopulateGenerator will be kept referenced until all generated and generator
+		 * packages call PostSave or until the splitter is destroyed. Performance cost: Possible extra calls to
+		 * PopulateGeneratedPackage, possible unnecessary memory increase due to OutKeepReferencedPackages.
+		 */
+		Populate,
+		/**
+		 * GetGenerateList PopulateGenerator, PreSaveGenerator, and PostSaveGenerator will be called before
+		 * PopulateGeneratedPackage. Performance cost: Progress on generated packages will be delayed until generator
+		 * finishes saving. Possible unnecessary memory increase due to OutKeepReferencedPackages. Retraction is not
+		 * possible in MPCook for the generated packages; they must all be saved on the same CookWorker that saves the
+		 * generator.
+		 */
+		Save,
+		Count,
+	};
 	/**
-	 * If true, the cooker will wait for the generator save to complete before calling Populate and PreSave on the
-	 * generated package. And during MPCook the cooker will load and save generated packages only on the same
-	 * CookWorker that saved the generator package that containes the SplitData object. This election reduces cook
-	 * performance but is otherwise not a problem.
+	 * Return capability setting which indicates which splitter functions acting on the parent generator package must
+	 * be called on the splitter before splitter functions acting on the generated packages can be called. Also impacts
+	 * the lifetime of memory guarantees for the generator functions. @see EGeneratedRequiresGenerator. Default is
+	 * EGeneratedRequiresGenerator::None, which provides the best performance but the fewest guarantees.
 	 * 
-	 * Examples of dependencies:
-	 *     ShouldSplit call reads data that is written by BeginCacheForCookedPlatformData
-	 *     PopulateGeneratedPackage or PreSaveGeneratedPackage read data that is written by PopulateGeneratorPackage.
+	 * Examples of dependencies and what capability level should be used:
+	 *		ShouldSplit call reads data that is written by BeginCacheForCookedPlatformData:
+	 *			EGeneratedRequiresGenerator::Save
+	 *     PopulateGeneratedPackage or PreSaveGeneratedPackage read data that is written by PopulateGeneratorPackage:
+	 *			EGeneratedRequiresGenerator::Populate
 	 */
-	virtual bool GeneratedReliesOnGeneratorSave() 
+	virtual EGeneratedRequiresGenerator DoesGeneratedRequireGenerator()
 	{
-		 return false;
-	 }
+		 return EGeneratedRequiresGenerator::None;
+	}
 
 	/** Data sent to the cooker to describe each desired generated package */
 	struct FGeneratedPackage
@@ -129,9 +161,10 @@ public:
 	 * @param OwnerObject				The SplitDataClass instance that this CookPackageSplitter instance was created for
 	 * @param GeneratedPackages			Placeholder UPackage and relative path information for all packages that will be generated
 	 * @param OutObjectsToMove			List of all the objects that will be moved into the Generator package during its save
-	 * @param OutKeepReferencedPackages A list of packages which should be kept referenced until all generated packages for
-	 *                                  the generator have finished saving.
-	 * 
+	 * @param OutKeepReferencedPackages Packages to keep referenced until the generator package finishes save.
+	 *                                  If DoesGeneratedRequireGenerator() >= Populate, these will also be kept referenced until
+	 *                                  all generated packages finish saving or the splitter is destroyed.
+	 *
 	 * @return							True if successfully populated, false on error (this will cause a cook error).
 	 */
 	virtual bool PopulateGeneratorPackage(UPackage* OwnerPackage, UObject* OwnerObject,
@@ -148,8 +181,9 @@ public:
 	 * @param OwnerPackage				The generator package being split
 	 * @param OwnerObject				The SplitDataClass instance that this CookPackageSplitter instance was created for
 	 * @param GeneratedPackages			Placeholder UPackage and relative path information for all packages that will be generated
-	 * @param OutKeepReferencedPackages A list of packages which should be kept referenced until all generated packages for
-	 *                                  the generator have finished saving.
+	 * @param OutKeepReferencedPackages Packages to keep referenced until the generator package finishes save.
+	 *                                  If DoesGeneratedRequireGenerator() >= Populate, these will also be kept referenced until
+	 *                                  all generated packages finish saving or the splitter is destroyed.
 	 *
 	 * @return							True if successfully presaved, false on error (this will cause a cook error).
 	 */
@@ -196,8 +230,7 @@ public:
 	 * @param OwnerObject				The SplitDataClass instance that this CookPackageSplitter instance was created for
 	 * @param GeneratedPackage			Pointer and information about the package to populate
 	 * @param OutObjectsToMove			List of all the objects that will be moved into the generated package during its save
-	 * @param OutKeepReferencedPackages A list of packages which should be kept referenced until all generated packages for
-	 *                                  for the generator have finished saving.
+	 * @param OutKeepReferencedPackages Packages to keep referenced until the generated package finishes save.
 	 * 
 	 * @return							True if successfully populated, false on error (this will cause a cook error).
 	 */
@@ -215,8 +248,7 @@ public:
 	 * @param OwnerPackage				The parent package being split
 	 * @param OwnerObject				The SplitDataClass instance that this CookPackageSplitter instance was created for
 	 * @param GeneratedPackage			Pointer and information about the package to populate
-	 * @param OutKeepReferencedPackages A list of packages which should be kept referenced until all generated packages for
-	 *                                  for the generator have finished saving.
+	 * @param OutKeepReferencedPackages Packages to keep referenced until the generated package finishes save.
 	 * 
 	 * @return							True if successfully presaved, false on error (this will cause a cook error).
 	 */
diff --git a/Engine/Source/Runtime/Engine/Private/WorldPartition/WorldPartitionCookPackageSplitter.h b/Engine/Source/Runtime/Engine/Private/WorldPartition/WorldPartitionCookPackageSplitter.h
--- a/Engine/Source/Runtime/Engine/Private/WorldPartition/WorldPartitionCookPackageSplitter.h
+++ b/Engine/Source/Runtime/Engine/Private/WorldPartition/WorldPartitionCookPackageSplitter.h
@@ -20,7 +20,15 @@ public:
 	virtual ~FWorldPartitionCookPackageSplitter();
 
 	virtual void Teardown(ETeardown Status) override;
-	virtual bool UseInternalReferenceToAvoidGarbageCollect() override { return true; }
+	virtual bool UseInternalReferenceToAvoidGarbageCollect() override
+	{
+		return true;
+	}
+	virtual EGeneratedRequiresGenerator DoesGeneratedRequireGenerator() override
+	{
+		return EGeneratedRequiresGenerator::Populate;
+	}
+
 	virtual TArray<ICookPackageSplitter::FGeneratedPackage> GetGenerateList(const UPackage* OwnerPackage, const UObject* OwnerObject) override;
 	virtual bool PopulateGeneratedPackage(UPackage* OwnerPackage, UObject* OwnerObject,
 		const FGeneratedPackageForPopulate& GeneratedPackage, TArray<UObject*>& OutObjectsToMove, TArray<UPackage*>& OutModifiedPackages) override;
