Use TypedElement selection in MoveActorsToLevel / CopySelectedActorsToClipboard to use SetAllowSubRootSelection(true)

This fixes copy/paste when selecting sub actors of different level instances, the copy paste code relies on being able to do partial selection of the original selection when
actors do not belong to the same level, it also will restore the original selection after the copy operation and this was failing for sub actors without that SetAllowSubRootSelection(true) flag

#rnx

diff --git a/Engine/Source/Editor/UnrealEd/Private/EditorServer.cpp b/Engine/Source/Editor/UnrealEd/Private/EditorServer.cpp
--- a/Engine/Source/Editor/UnrealEd/Private/EditorServer.cpp
+++ b/Engine/Source/Editor/UnrealEd/Private/EditorServer.cpp
@@ -3019,20 +3019,34 @@ public:
 		ULevel* OldCurrentLevel = World->GetCurrentLevel();
 		World->SetCurrentLevel( SrcLevel );
 
-		// Set the selection set to be precisely the actors belonging to this job,
-		// but make sure not to deselect selected BSP surfaces. 
-		GEditor->SelectNone( false, true );
-		for ( int32 ActorIndex = 0 ; ActorIndex < Actors.Num() ; ++ActorIndex )
+		UTypedElementSelectionSet* SelectionSet = GEditor->GetSelectedActors()->GetElementSelectionSet();
+		check(SelectionSet);
 		{
-			AActor* Actor = Actors[ ActorIndex ];
-			GEditor->SelectActor( Actor, true, false );
+			FTypedElementList::FScopedClearNewPendingChange ClearNewPendingChange = SelectionSet->GetScopedClearNewPendingChange();
+
+			const FTypedElementSelectionOptions SelectionOptions = FTypedElementSelectionOptions()
+				.SetAllowHidden(false)
+				.SetWarnIfLocked(false)
+				.SetAllowLegacyNotifications(false)
+				.SetAllowSubRootSelection(true);
 
-			// Groups cannot contain actors in different levels.  If the current actor is in a group but not being moved to the same level as the group
-			// then remove the actor from the group
-			AGroupActor* GroupActor = AGroupActor::GetParentForActor( Actor );
-			if( GroupActor && GroupActor->GetLevel() != DestLevel )
+			// Set the selection set to be precisely the actors belonging to this job,
+			// but make sure not to deselect selected BSP surfaces. 
+			GEditor->SelectNone(false, true);
+			for (int32 ActorIndex = 0; ActorIndex < Actors.Num(); ++ActorIndex)
 			{
-				GroupActor->Remove( *Actor );
+				AActor* Actor = Actors[ActorIndex];
+				FTypedElementHandle ElementHandle = UEngineElementsLibrary::AcquireEditorActorElementHandle(Actor);
+
+				SelectionSet->SelectElement(ElementHandle, SelectionOptions);
+
+				// Groups cannot contain actors in different levels.  If the current actor is in a group but not being moved to the same level as the group
+				// then remove the actor from the group
+				AGroupActor* GroupActor = AGroupActor::GetParentForActor(Actor);
+				if (GroupActor && GroupActor->GetLevel() != DestLevel)
+				{
+					GroupActor->Remove(*Actor);
+				}
 			}
 		}
 
@@ -3477,11 +3491,25 @@ void UEditorEngine::CopySelectedActorsToClipboard( UWorld* InWorld, bool bShould
 
 		// Restore old selection
 		GEditor->SelectNone( false, true );
-		for (const TWeakObjectPtr<AActor>& Actor : CurrentlySelectedActors)
+
+		UTypedElementSelectionSet* SelectionSet = GEditor->GetSelectedActors()->GetElementSelectionSet();
+		check(SelectionSet);
 		{
-			if (AActor* ActorPtr = Actor.Get())
+			FTypedElementList::FScopedClearNewPendingChange ClearNewPendingChange = SelectionSet->GetScopedClearNewPendingChange();
+			const FTypedElementSelectionOptions SelectionOptions = FTypedElementSelectionOptions()
+				.SetAllowHidden(false)
+				.SetWarnIfLocked(false)
+				.SetAllowLegacyNotifications(false)
+				.SetAllowSubRootSelection(true);
+
+			for (const TWeakObjectPtr<AActor>& Actor : CurrentlySelectedActors)
 			{
-				GEditor->SelectActor(ActorPtr, true, false);
+				if (AActor* ActorPtr = Actor.Get())
+				{
+					FTypedElementHandle ElementHandle = UEngineElementsLibrary::AcquireEditorActorElementHandle(ActorPtr);
+					
+					SelectionSet->SelectElement(ElementHandle, SelectionOptions);
+				}
 			}
 		}
 	}
