Exposed AssetViewUtils::GetPackageLengthForCooking to scripting

This is available via UEditorAssetSubsystem as GetAssetFilenameLengthForCooking or GetLoadedAssetFilenameLengthForCooking

diff --git a/Engine/Source/Editor/UnrealEd/Private/Subsystems/EditorAssetSubsystem.cpp b/Engine/Source/Editor/UnrealEd/Private/Subsystems/EditorAssetSubsystem.cpp
--- a/Engine/Source/Editor/UnrealEd/Private/Subsystems/EditorAssetSubsystem.cpp
+++ b/Engine/Source/Editor/UnrealEd/Private/Subsystems/EditorAssetSubsystem.cpp
@@ -5,11 +5,13 @@
 #include "AssetRegistry/IAssetRegistry.h"
 #include "AssetRegistry/AssetRegistryModule.h"
 #include "AssetToolsModule.h"
+#include "AssetViewUtils.h"
 #include "Editor.h"
 #include "Editor/Transactor.h"
 #include "EditorScriptingHelpers.h"
 #include "FileHelpers.h"
 #include "HAL/FileManager.h"
+#include "Misc/EngineBuildSettings.h"
 #include "Misc/PackageName.h"
 #include "Misc/ScopedSlowTask.h"
 #include "Modules/ModuleManager.h"
@@ -1273,6 +1275,21 @@ bool UEditorAssetSubsystem::SaveAsset(const FString& AssetsToSave, bool bOnlyIfI
 	return UEditorLoadingAndSavingUtils::SavePackages(Packages, bOnlyIfIsDirty);
 }
 
+int32 UEditorAssetSubsystem::GetAssetFilenameLengthForCooking(const FString& AssetPath)
+{
+	const FString PackagePath = FPackageName::ObjectPathToPackageName(AssetPath);
+	return AssetViewUtils::GetPackageLengthForCooking(PackagePath, FEngineBuildSettings::IsInternalBuild());
+}
+
+int32 UEditorAssetSubsystem::GetLoadedAssetFilenameLengthForCooking(const UObject* Asset)
+{
+	if (Asset)
+	{
+		return AssetViewUtils::GetPackageLengthForCooking(Asset->GetPackage()->GetName(), FEngineBuildSettings::IsInternalBuild());
+	}
+	return 0;
+}
+
 bool UEditorAssetSubsystem::SaveDirectory(const FString& DirectoryPath, bool bOnlyIfIsDirty, bool bRecursive)
 {
 	using UE::EditorAssetUtils::EPackageEnumerationFilter;
diff --git a/Engine/Source/Editor/UnrealEd/Public/Subsystems/EditorAssetSubsystem.h b/Engine/Source/Editor/UnrealEd/Public/Subsystems/EditorAssetSubsystem.h
--- a/Engine/Source/Editor/UnrealEd/Public/Subsystems/EditorAssetSubsystem.h
+++ b/Engine/Source/Editor/UnrealEd/Public/Subsystems/EditorAssetSubsystem.h
@@ -272,6 +272,14 @@ public:
 	UFUNCTION(BlueprintCallable, Category = "Editor Scripting | Asset")
 	UNREALED_API bool SaveAsset(const FString& AssetToSave, bool bOnlyIfIsDirty = true);
 
+	/** Returns the length of the computed cooked package name and path */
+	UFUNCTION(BlueprintCallable, Category = "Editor Scripting | Asset")
+	UNREALED_API int32 GetAssetFilenameLengthForCooking(const FString& AssetPath);
+
+	/** Returns the length of the computed cooked package name and path */
+	UFUNCTION(BlueprintCallable, Category = "Editor Scripting | Asset")
+	UNREALED_API int32 GetLoadedAssetFilenameLengthForCooking(const UObject* Asset);
+
 	/**
 	 * Save the packages the assets live in inside the directory. All objects that are in the directory will be saved.
 	 * Will try to checkout the file first. Assets will be loaded before being saved.
