Move Ctrl+O to global commands

This makes Ctrl/Cmd+O available when the level editor doesn't have focus, such as after clicking on a tab (Mac and Windows) or on the title bar (Mac only).

diff --git a/Engine/Source/Editor/EditorStyle/Private/StarshipStyle.cpp b/Engine/Source/Editor/EditorStyle/Private/StarshipStyle.cpp
--- a/Engine/Source/Editor/EditorStyle/Private/StarshipStyle.cpp
+++ b/Engine/Source/Editor/EditorStyle/Private/StarshipStyle.cpp
@@ -4344,7 +4344,7 @@ void FStarshipEditorStyle::FStyle::SetupLevelEditorStyle()
 
 #if WITH_EDITOR
 		Set( "LevelEditor.NewLevel",      new IMAGE_BRUSH_SVG( "Starship/Common/LevelNew", Icon16x16 ) );
-		Set( "LevelEditor.OpenLevel",     new IMAGE_BRUSH_SVG( "Starship/Common/LevelOpen", Icon16x16 ) );
+		Set( "SystemWideCommands.OpenLevel",     new IMAGE_BRUSH_SVG( "Starship/Common/LevelOpen", Icon16x16 ) );
 		Set( "LevelEditor.Save",          new IMAGE_BRUSH_SVG( "Starship/Common/SaveCurrent", Icon16x16 ) );
 		Set( "LevelEditor.SaveAs",        new IMAGE_BRUSH_SVG( "Starship/Common/SaveCurrentAs", Icon16x16 ) );
 		Set( "LevelEditor.SaveAllLevels", new IMAGE_BRUSH_SVG( "Starship/Common/LevelSaveAll", Icon16x16 ) );
diff --git a/Engine/Source/Editor/LevelEditor/Private/LevelEditor.cpp b/Engine/Source/Editor/LevelEditor/Private/LevelEditor.cpp
--- a/Engine/Source/Editor/LevelEditor/Private/LevelEditor.cpp
+++ b/Engine/Source/Editor/LevelEditor/Private/LevelEditor.cpp
@@ -657,7 +657,6 @@ void FLevelEditorModule::BindGlobalLevelEditorCommands()
 	ActionList.MapAction( Commands.BrowseDocumentation, FExecuteAction::CreateStatic( &FLevelEditorActionCallbacks::BrowseDocumentation ) );
 	ActionList.MapAction( Commands.BrowseViewportControls, FExecuteAction::CreateStatic( &FLevelEditorActionCallbacks::BrowseViewportControls ) );
 	ActionList.MapAction( Commands.NewLevel, FExecuteAction::CreateStatic( &FLevelEditorActionCallbacks::NewLevel ), FCanExecuteAction::CreateStatic( &FLevelEditorActionCallbacks::NewLevel_CanExecute ) );
-	ActionList.MapAction(Commands.OpenLevel, FExecuteAction::CreateStatic(&FLevelEditorActionCallbacks::OpenLevel), FCanExecuteAction::CreateStatic(&FLevelEditorActionCallbacks::OpenLevel_CanExecute));
 	ActionList.MapAction( Commands.Save, FExecuteAction::CreateStatic( &FLevelEditorActionCallbacks::Save ), FCanExecuteAction::CreateStatic( &FLevelEditorActionCallbacks::CanSaveWorld ) );
 	ActionList.MapAction( Commands.SaveAs, FExecuteAction::CreateStatic( &FLevelEditorActionCallbacks::SaveCurrentAs ), FCanExecuteAction::CreateStatic( &FLevelEditorActionCallbacks::CanSaveCurrentAs), FGetActionCheckState(), FIsActionButtonVisible::CreateStatic( &FLevelEditorActionCallbacks::CanSaveCurrentAs) );
 	ActionList.MapAction( Commands.SaveAllLevels, FExecuteAction::CreateStatic( &FLevelEditorActionCallbacks::SaveAllLevels ), FCanExecuteAction::CreateStatic( &FLevelEditorActionCallbacks::CanSaveUnpartitionedWorld), FGetActionCheckState(), FIsActionButtonVisible::CreateStatic(&FLevelEditorActionCallbacks::CanSaveUnpartitionedWorld) );
diff --git a/Engine/Source/Editor/LevelEditor/Private/LevelEditorActions.cpp b/Engine/Source/Editor/LevelEditor/Private/LevelEditorActions.cpp
--- a/Engine/Source/Editor/LevelEditor/Private/LevelEditorActions.cpp
+++ b/Engine/Source/Editor/LevelEditor/Private/LevelEditorActions.cpp
@@ -3643,7 +3643,6 @@ void FLevelEditorCommands::RegisterCommands()
 	UI_COMMAND( BrowseDocumentation, "Level Editor Documentation", "Details on how to use the Level Editor", EUserInterfaceActionType::Button, FInputChord( EKeys::F1 ) );
 	UI_COMMAND( BrowseViewportControls, "Viewport Controls", "Ways to move around in the 3D viewport", EUserInterfaceActionType::Button, FInputChord() );
 	UI_COMMAND( NewLevel, "New Level...", "Create a new level, or choose a level template to start from.", EUserInterfaceActionType::Button, FInputChord( EModifierKey::Control, EKeys::N ) );
-	UI_COMMAND( OpenLevel, "Open Level...", "Loads an existing level", EUserInterfaceActionType::Button, FInputChord( EModifierKey::Control, EKeys::O ) );
 	UI_COMMAND( Save, "Save Current Level", "Saves the current level to disk", EUserInterfaceActionType::Button, FInputChord(EModifierKey::Control, EKeys::S) );
 	UI_COMMAND( SaveAs, "Save Current Level As...", "Save the current level as...", EUserInterfaceActionType::Button, FInputChord( EModifierKey::Control|EModifierKey::Alt, EKeys::S ) );
 	UI_COMMAND( SaveAllLevels, "Save All Levels", "Saves all unsaved levels to disk", EUserInterfaceActionType::Button, FInputChord() );
diff --git a/Engine/Source/Editor/LevelEditor/Private/LevelEditorMenu.cpp b/Engine/Source/Editor/LevelEditor/Private/LevelEditorMenu.cpp
--- a/Engine/Source/Editor/LevelEditor/Private/LevelEditorMenu.cpp
+++ b/Engine/Source/Editor/LevelEditor/Private/LevelEditorMenu.cpp
@@ -26,6 +26,7 @@
 #include "AssetSelection.h"
 #include "EditorBuildUtils.h"
 #include "EditorViewportCommands.h"
+#include "Toolkits/GlobalEditorCommonCommands.h"
 
 #define LOCTEXT_NAMESPACE "LevelEditorMenu"
 
@@ -45,7 +46,7 @@ void FLevelEditorMenu::RegisterLevelEditorMenus()
 			OpenSection.AddMenuEntry( FLevelEditorCommands::Get().NewLevel ).InsertPosition = InsertPos;
 
 			// Open Level
-			OpenSection.AddMenuEntry( FLevelEditorCommands::Get().OpenLevel ).InsertPosition = InsertPos;
+			OpenSection.AddMenuEntry( FGlobalEditorCommonCommands::Get().OpenLevel ).InsertPosition = InsertPos;
 
 			FToolMenuSection& AssetSection = Menu->FindOrAddSection("FileAsset");
 			
diff --git a/Engine/Source/Editor/LevelEditor/Public/LevelEditorActions.h b/Engine/Source/Editor/LevelEditor/Public/LevelEditorActions.h
--- a/Engine/Source/Editor/LevelEditor/Public/LevelEditorActions.h
+++ b/Engine/Source/Editor/LevelEditor/Public/LevelEditorActions.h
@@ -44,6 +44,7 @@ public:
 
 	/** Level file commands */
 	TSharedPtr< FUICommandInfo > NewLevel;
+	UE_DEPRECATED(5.4, "This command has been moved to FGlobalEditorCommonCommands and is no longer registered.")
 	TSharedPtr< FUICommandInfo > OpenLevel;
 	TSharedPtr< FUICommandInfo > Save;
 	TSharedPtr< FUICommandInfo > SaveAs;
diff --git a/Engine/Source/Editor/UnrealEd/Private/Toolkits/GlobalEditorCommonCommands.cpp b/Engine/Source/Editor/UnrealEd/Private/Toolkits/GlobalEditorCommonCommands.cpp
--- a/Engine/Source/Editor/UnrealEd/Private/Toolkits/GlobalEditorCommonCommands.cpp
+++ b/Engine/Source/Editor/UnrealEd/Private/Toolkits/GlobalEditorCommonCommands.cpp
@@ -15,6 +15,7 @@
 #include "Toolkits/SGlobalOpenAssetDialog.h"
 #include "Toolkits/SGlobalTabSwitchingDialog.h"
 #include "StatusBarSubsystem.h"
+#include "LevelEditorActions.h"
 
 #define LOCTEXT_NAMESPACE "GlobalEditorCommonCommands"
 
@@ -55,6 +56,8 @@ void FGlobalEditorCommonCommands::RegisterCommands()
 
 	UI_COMMAND(OpenDocumentation, "Open Documentation...", "Opens documentation for this tool", EUserInterfaceActionType::Button, FInputChord(EKeys::F1));
 
+	UI_COMMAND(OpenLevel, "Open Level...", "Loads an existing level", EUserInterfaceActionType::Button, FInputChord(EModifierKey::Control, EKeys::O));
+
 #if PLATFORM_MAC
 	// On mac command and ctrl are automatically swapped. Command + Space is spotlight search so we use ctrl+space on mac to avoid the conflict
 	UI_COMMAND(OpenContentBrowserDrawer, "Open Content Browser Drawer", "Opens the content browser drawer from the status bar and focuses the search field", EUserInterfaceActionType::Button, FInputChord(EModifierKey::Command, EKeys::SpaceBar));
@@ -102,6 +105,11 @@ void FGlobalEditorCommonCommands::MapActions(TSharedRef<FUICommandList>& Toolkit
 	ToolkitCommands->MapAction(
 		Get().OpenOutputLogDrawer,
 		FExecuteAction::CreateStatic(&FGlobalEditorCommonCommands::OnOpenOutputLogDrawer));
+
+	ToolkitCommands->MapAction(
+		Get().OpenLevel,
+		FExecuteAction::CreateStatic(&FLevelEditorActionCallbacks::OpenLevel),
+		FCanExecuteAction::CreateStatic(&FLevelEditorActionCallbacks::OpenLevel_CanExecute));
 }
 
 void FGlobalEditorCommonCommands::OnPressedCtrlTab(TSharedPtr<FUICommandInfo> TriggeringCommand)
diff --git a/Engine/Source/Editor/UnrealEd/Public/Toolkits/GlobalEditorCommonCommands.h b/Engine/Source/Editor/UnrealEd/Public/Toolkits/GlobalEditorCommonCommands.h
--- a/Engine/Source/Editor/UnrealEd/Public/Toolkits/GlobalEditorCommonCommands.h
+++ b/Engine/Source/Editor/UnrealEd/Public/Toolkits/GlobalEditorCommonCommands.h
@@ -46,5 +46,8 @@ public:
 	TSharedPtr<FUICommandInfo> SelectNextConsoleExecutor;
 	TSharedPtr<FUICommandInfo> OpenOutputLogDrawer;
 	TSharedPtr<FUICommandInfo> OpenContentBrowserDrawer;
+
+	/** Level file commands */
+	TSharedPtr<FUICommandInfo> OpenLevel;
 };
 
