Cooker DiffOnly: Fix scriptimports not being convertable to string when dumping the header diffs in the case when the 2 versions of the package saved into memory are different.
#rnx

diff --git a/Engine/Source/Editor/UnrealEd/Private/Cooker/DiffWriterArchive.cpp b/Engine/Source/Editor/UnrealEd/Private/Cooker/DiffWriterArchive.cpp
--- a/Engine/Source/Editor/UnrealEd/Private/Cooker/DiffWriterArchive.cpp
+++ b/Engine/Source/Editor/UnrealEd/Private/Cooker/DiffWriterArchive.cpp
@@ -592,6 +592,8 @@ void FAccumulator::OnFirstSaveComplete(FStringView LooseFilePath, int64 InHeader
 
 void FAccumulator::OnSecondSaveComplete(int64 InHeaderSize)
 {
+	Globals.Initialize(PackageHeaderFormat);
+
 	check(bFirstSaveComplete); // Should have been set in OnFirstSaveComplete
 	if (HeaderSize != InHeaderSize)
 	{
diff --git a/Engine/Source/Editor/UnrealEd/Private/Cooker/DiffWriterZenHeader.cpp b/Engine/Source/Editor/UnrealEd/Private/Cooker/DiffWriterZenHeader.cpp
--- a/Engine/Source/Editor/UnrealEd/Private/Cooker/DiffWriterZenHeader.cpp
+++ b/Engine/Source/Editor/UnrealEd/Private/Cooker/DiffWriterZenHeader.cpp
@@ -80,7 +80,7 @@ FStringView FDiffWriterZenHeader::GetObjectIndexPathName(FPackageObjectIndex Pac
 		}
 		else
 		{
-			*PathName = TEXT("<UnknownScriptImport>");
+			*PathName = FString::Printf(TEXT("<UnknownScriptImport:%" UINT64_FMT ">"), PackageObjectIndex.Value());
 		}
 	}
 	else if (PackageObjectIndex.IsPackageImport())
