Substrate - improved pin documentation as requested by QA.

diff --git a/Engine/Source/Editor/UnrealEd/Private/MaterialGraph.cpp b/Engine/Source/Editor/UnrealEd/Private/MaterialGraph.cpp
--- a/Engine/Source/Editor/UnrealEd/Private/MaterialGraph.cpp
+++ b/Engine/Source/Editor/UnrealEd/Private/MaterialGraph.cpp
@@ -264,7 +264,7 @@ void UMaterialGraph::RebuildGraphInternal(const TMap<UMaterialExpression*, TArra
 
 		MaterialInputs.Add(FMaterialInputInfo(FMaterialAttributeDefinitionMap::GetDisplayNameForMaterial(MP_PixelDepthOffset, Material), MP_PixelDepthOffset, LOCTEXT("PixelDepthOffsetToolTip", "Pixel Depth Offset")));
 		MaterialInputs.Add(FMaterialInputInfo(FMaterialAttributeDefinitionMap::GetDisplayNameForMaterial(MP_ShadingModel, Material), MP_ShadingModel, LOCTEXT("ShadingModelToolTip", "Selects which shading model should be used per pixel")));
-		MaterialInputs.Add(FMaterialInputInfo(FMaterialAttributeDefinitionMap::GetDisplayNameForMaterial(MP_SurfaceThickness, Material), MP_SurfaceThickness, LOCTEXT("SurfaceThicknessToolTip", "Defines the surface's thickness when IsThinSurface is enabled")));
+		MaterialInputs.Add(FMaterialInputInfo(FMaterialAttributeDefinitionMap::GetDisplayNameForMaterial(MP_SurfaceThickness, Material), MP_SurfaceThickness, LOCTEXT("SurfaceThicknessToolTip", "Defines the surface's thickness when IsThinSurface is enabled on this node. Only available when Substrate is enabled for the project.")));
 		MaterialInputs.Add(FMaterialInputInfo(FMaterialAttributeDefinitionMap::GetDisplayNameForMaterial(MP_Displacement, Material), MP_Displacement, LOCTEXT("DisplacementToolTip", "Specifies scalar vertex displacement."))); 
 
 		//^^^ New material properties go above here. ^^^^
