Fix localization of "Recent {Type} Assets" menu label

diff --git a/Engine/Source/Editor/UnrealEd/Private/Subsystems/AssetEditorSubsystem.cpp b/Engine/Source/Editor/UnrealEd/Private/Subsystems/AssetEditorSubsystem.cpp
--- a/Engine/Source/Editor/UnrealEd/Private/Subsystems/AssetEditorSubsystem.cpp
+++ b/Engine/Source/Editor/UnrealEd/Private/Subsystems/AssetEditorSubsystem.cpp
@@ -1341,11 +1341,9 @@ void UAssetEditorSubsystem::CreateRecentAssetsMenuForEditor(const IAssetEditorIn
 		else
 		{
 			// Example submenu name: "Recent Material Assets" for the Material Editor
-			const FName RecentAssetsMenuName("Recent " + EditingAssetTypeName.ToString() + " Assets");
-			
 			InSection.AddSubMenu(
 				"RecentAssetEditorAssetsSubmenu",
-				FText::Format(LOCTEXT("RecentAssetEditorAssetsSubmenu_Label", "{0}"), FText::FromName(RecentAssetsMenuName)),
+				FText::Format(LOCTEXT("RecentAssetEditorAssetsSubmenu_Label", "Recent {0} Assets"), FText::FromName(EditingAssetTypeName)),
 				FText::Format(LOCTEXT("RecentAssetEditorAssetsSubmenu_Tooltip", "Access your recently opened {0} assets"), FText::FromName(EditingAssetTypeName)),
 				FNewToolMenuDelegate::CreateUObject(this, &UAssetEditorSubsystem::CreateRecentAssetsMenu, AssetEditorName),
 				false,
