Changed UEditorEngine::SaveMapsForPlaySession to reload the saved level in the case where it's the initial save of a newly created world
This ensures the ActorDescContainerInstance is up to date with all the renames that happened.

diff --git a/Engine/Source/Editor/UnrealEd/Private/PlayLevel.cpp b/Engine/Source/Editor/UnrealEd/Private/PlayLevel.cpp
--- a/Engine/Source/Editor/UnrealEd/Private/PlayLevel.cpp
+++ b/Engine/Source/Editor/UnrealEd/Private/PlayLevel.cpp
@@ -985,17 +985,10 @@ void UEditorEngine::CancelRequestPlaySession()
 
 bool UEditorEngine::SaveMapsForPlaySession()
 {
-	// Prompt the user to save the level if it has not been saved before. 
-	// An unmodified but unsaved blank template level does not appear in the dirty packages check below.
-	if (FEditorFileUtils::GetFilename(GWorld).Len() == 0)
-	{
-		if (!FEditorFileUtils::SaveCurrentLevel())
-		{
-			return false;
-		}
-	}
+	// Detect if we're dealing with an unsaved newly created world
+	const bool bNewlyCreatedWorld = (FEditorFileUtils::GetFilename(GWorld).Len() == 0);
 
-	// Also save dirty packages, this is required because we're going to be launching a session outside of our normal process
+	// Save dirty packages, this is required because we're going to be launching a session outside of our normal process
 	const bool bPromptUserToSave      = true;
 	const bool bSaveMapPackages       = true;
 	const bool bSaveContentPackages   = true;
@@ -1007,6 +1000,24 @@ bool UEditorEngine::SaveMapsForPlaySession()
 		return false;
 	}
 
+	// Prompt the user to save the level if it has not been saved before. 
+	// An unmodified but unsaved blank template level does not appear in the dirty packages check below.
+	if (FEditorFileUtils::GetFilename(GWorld).Len() == 0)
+	{		
+		bool bSaved = FEditorFileUtils::SaveCurrentLevel();
+		
+		if (!bSaved)
+		{
+			return false;
+		}
+	}
+
+	// In the case we're dealing with a newly created world it must be reloaded after being saved
+	if (bNewlyCreatedWorld)
+	{
+		FEditorFileUtils::LoadMap(FEditorFileUtils::GetFilename(GWorld));
+	}
+
 	return true;
 }
 
