Fix for use of uninitialized CurrentOutVal and PreviousOutVal in FFbxCurvesAPI::GetConvertedTransformCurveData().

diff --git a/Engine/Source/Editor/UnrealEd/Private/Fbx/FbxCurveDataImport.cpp b/Engine/Source/Editor/UnrealEd/Private/Fbx/FbxCurveDataImport.cpp
--- a/Engine/Source/Editor/UnrealEd/Private/Fbx/FbxCurveDataImport.cpp
+++ b/Engine/Source/Editor/UnrealEd/Private/Fbx/FbxCurveDataImport.cpp
@@ -657,7 +657,6 @@ namespace UnFbx {
 					auto EulerRotZIt = EulerRotationZ.GetKeyHandleIterator();
 
 					FVector PreviousOutVal;
-					FVector CurrentOutVal;
 					bool bFirst = true;
 					while (EulerRotXIt && EulerRotYIt && EulerRotZIt)
 					{
@@ -665,14 +664,10 @@ namespace UnFbx {
 						float Y = EulerRotationY.GetKeyValue(*EulerRotYIt);
 						float Z = EulerRotationZ.GetKeyValue(*EulerRotZIt);
 
-						if (!bFirst)
+						FVector CurrentOutVal(X, Y, Z);
+						if (bFirst)
 						{
 							PreviousOutVal = CurrentOutVal;
-							CurrentOutVal = FVector(X, Y, Z);
-						}
-						else
-						{
-							CurrentOutVal = FVector(X, Y, Z);
 							bFirst = false;
 						}
 						
@@ -698,6 +693,8 @@ namespace UnFbx {
 						++EulerRotXIt;
 						++EulerRotYIt;
 						++EulerRotZIt;
+
+						PreviousOutVal = CurrentOutVal;
 					}
 				}
 			}
