Fixed some uninitialized FVectors in Editor.

diff --git a/Engine/Source/Editor/DetailCustomizations/Private/ComponentTransformDetails.cpp b/Engine/Source/Editor/DetailCustomizations/Private/ComponentTransformDetails.cpp
--- a/Engine/Source/Editor/DetailCustomizations/Private/ComponentTransformDetails.cpp
+++ b/Engine/Source/Editor/DetailCustomizations/Private/ComponentTransformDetails.cpp
@@ -1070,9 +1070,9 @@ void FComponentTransformDetails::OnZScaleMirrored()
 
 void FComponentTransformDetails::CacheDetails()
 {
-	FVector CurLoc;
+	FVector CurLoc = FVector::ZeroVector;
 	FRotator CurRot;
-	FVector CurScale;
+	FVector CurScale = FVector::ZeroVector;
 	bIsEnabledCache = true;
 	for( int32 ObjectIndex = 0; ObjectIndex < SelectedObjects.Num(); ++ObjectIndex )
 	{
diff --git a/Engine/Source/Editor/UnrealEd/Private/EditorEngine.cpp b/Engine/Source/Editor/UnrealEd/Private/EditorEngine.cpp
--- a/Engine/Source/Editor/UnrealEd/Private/EditorEngine.cpp
+++ b/Engine/Source/Editor/UnrealEd/Private/EditorEngine.cpp
@@ -6719,17 +6719,17 @@ TArray<AActor*> UEditorEngine::AddExportTextActors(const FString& ExportText, bo
 				// Send notification about a new actor being created
 				ULevel::LevelDirtiedEvent.Broadcast();
 				NoteSelectionChange();
+
+	            if( !bSilent )
+	            {
+		            UE_LOG(LogEditor, Log,
+			            TEXT("Added '%d' actor(s) to level at %0.2f,%0.2f,%0.2f"),
+			            NewActors.Num(), Location.X, Location.Y, Location.Z );
+	            }
 			}
 		}
 	}
 
-	if( NewActors.Num() > 0 && !bSilent )
-	{
-		UE_LOG(LogEditor, Log,
-			TEXT("Added '%d' actor(s) to level at %0.2f,%0.2f,%0.2f"),
-			NewActors.Num(), Location.X, Location.Y, Location.Z );
-	}
-
 	return NewActors;
 }
 
diff --git a/Engine/Source/Editor/UnrealEd/Private/ImportUtils/StaticMeshImportUtils.cpp b/Engine/Source/Editor/UnrealEd/Private/ImportUtils/StaticMeshImportUtils.cpp
--- a/Engine/Source/Editor/UnrealEd/Private/ImportUtils/StaticMeshImportUtils.cpp
+++ b/Engine/Source/Editor/UnrealEd/Private/ImportUtils/StaticMeshImportUtils.cpp
@@ -540,7 +540,8 @@ bool StaticMeshImportUtils::AddCapsuleGeomFromVerts(const TArray<FVector3f>& Ver
 		return false;
 	}
 
-	FVector3f AxisStart, AxisEnd;
+	FVector3f AxisStart = FVector3f::ZeroVector;
+	FVector3f AxisEnd = FVector3f::ZeroVector;
 	float MaxDistSqr = 0.f;
 
 	for (int32 IndexA = 0; IndexA < Verts.Num() - 1; IndexA++)
diff --git a/Engine/Source/Editor/UnrealEd/Private/StaticLightingSystem/StaticLightingDebug.cpp b/Engine/Source/Editor/UnrealEd/Private/StaticLightingSystem/StaticLightingDebug.cpp
--- a/Engine/Source/Editor/UnrealEd/Private/StaticLightingSystem/StaticLightingDebug.cpp
+++ b/Engine/Source/Editor/UnrealEd/Private/StaticLightingSystem/StaticLightingDebug.cpp
@@ -220,7 +220,7 @@ void SetDebugLightmapSample(TArray<UActorComponent*>* Components, UModel* Model,
 				if (bUseTextureMap)
 				{
 					double ClosestPlaneDistance = std::numeric_limits<double>::max();
-					FVector ClosestPlaneBaryCentricWeights;
+					FVector ClosestPlaneBaryCentricWeights = FVector::ZeroVector;
 					int32 ClosestPlaneTriangleIndex = -1;
 
 					// Search through the static mesh's triangles for the one that was hit (since we can't get triangle index from a line check)
diff --git a/Engine/Source/Editor/UnrealEd/Private/UnrealWidgetRender.cpp b/Engine/Source/Editor/UnrealEd/Private/UnrealWidgetRender.cpp
--- a/Engine/Source/Editor/UnrealEd/Private/UnrealWidgetRender.cpp
+++ b/Engine/Source/Editor/UnrealEd/Private/UnrealWidgetRender.cpp
@@ -1326,11 +1326,10 @@ void FWidget::DrawThickArc (const FThickArcParams& InParams, const FVector& Axis
 
 			// Push out the arc line borders so they dont z-fight with the mesh arcs
 			// DrawLine needs vertices in world space, but this is fine because it takes FVectors and works with LWC well
-			FVector StartLinePos = LastWorldVertex;
 			FVector EndLinePos = VertexPosition + InParams.Position;
 			if (VertexIndex != 0)
 			{
-				InParams.PDI->DrawLine(StartLinePos,EndLinePos,RingColor,SDPG_Foreground);
+				InParams.PDI->DrawLine(LastWorldVertex,EndLinePos,RingColor,SDPG_Foreground);
 			}
 			LastWorldVertex = EndLinePos;
 		}
diff --git a/Engine/Source/Editor/ViewportInteraction/Private/Gizmo/VIUniformScaleGizmoHandle.cpp b/Engine/Source/Editor/ViewportInteraction/Private/Gizmo/VIUniformScaleGizmoHandle.cpp
--- a/Engine/Source/Editor/ViewportInteraction/Private/Gizmo/VIUniformScaleGizmoHandle.cpp
+++ b/Engine/Source/Editor/ViewportInteraction/Private/Gizmo/VIUniformScaleGizmoHandle.cpp
@@ -45,7 +45,7 @@ void UUniformScaleGizmoHandleGroup::UpdateGizmoHandleGroup( const FTransform& Lo
 	UStaticMeshComponent* UniformScaleHandle = Handle.HandleMesh;
 	if (UniformScaleHandle != nullptr)	// Can be null if no handle for this specific placement
 	{
-		FVector HandleRelativeLocation;
+		FVector HandleRelativeLocation = FVector::ZeroVector;
 		if ( !bUsePivotAsLocation )
 		{
 			UniformScaleHandle->SetRelativeLocation( HandleRelativeLocation );
