Fixed bTemporaryEditorActor not being set in the fallback path

#rnx

diff --git a/Engine/Source/Editor/UnrealEd/Private/AssetSelection.cpp b/Engine/Source/Editor/UnrealEd/Private/AssetSelection.cpp
--- a/Engine/Source/Editor/UnrealEd/Private/AssetSelection.cpp
+++ b/Engine/Source/Editor/UnrealEd/Private/AssetSelection.cpp
@@ -676,6 +676,7 @@ TArray<FTypedElementHandle> PlaceAssetUsingFactory(UObject* Asset, TScriptInterf
 			FActorSpawnParameters SpawnParams;
 			SpawnParams.ObjectFlags = ObjectFlags;
 			SpawnParams.Name = Name;
+			SpawnParams.bTemporaryEditorActor = FLevelEditorViewportClient::IsDroppingPreviewActor();
 
 			if (AActor* PlacedActor = ActorFactory->CreateActor(Asset, DesiredLevel, ActorTransform, SpawnParams))
 			{
