FBX Export Sequencer:
- Add support of Vector tracks
- Improve support of Color tracks

diff --git a/Engine/Source/Editor/UnrealEd/Private/Fbx/FbxMainExport.cpp b/Engine/Source/Editor/UnrealEd/Private/Fbx/FbxMainExport.cpp
--- a/Engine/Source/Editor/UnrealEd/Private/Fbx/FbxMainExport.cpp
+++ b/Engine/Source/Editor/UnrealEd/Private/Fbx/FbxMainExport.cpp
@@ -85,11 +85,13 @@
 #include "Tracks/MovieSceneDoubleTrack.h"
 #include "Tracks/MovieSceneFloatTrack.h"
 #include "Tracks/MovieSceneSkeletalAnimationTrack.h"
+#include "Tracks/MovieSceneVectorTrack.h"
 #include "Sections/MovieSceneSkeletalAnimationSection.h"
 #include "Sections/MovieScene3DTransformSection.h"
 #include "Sections/MovieSceneColorSection.h"
 #include "Sections/MovieSceneDoubleSection.h"
 #include "Sections/MovieSceneFloatSection.h"
+#include "Sections/MovieSceneVectorSection.h"
 #include "Evaluation/MovieScenePlayback.h"
 #include "Evaluation/MovieSceneEvaluationTemplateInstance.h"
 #include "MovieSceneSequenceID.h"
@@ -1995,6 +1997,10 @@ bool FFbxExporter::ExportLevelSequenceTracks(UMovieSceneSequence* MovieSceneSequ
 		{
 			ExportLevelSequenceColorTrack(FbxActor, *Cast<UMovieSceneColorTrack>(Track), BoundObject, MovieScene->GetPlaybackRange(), RootToLocalTransform);
 		}
+		else if (Track->IsA(UMovieSceneDoubleVectorTrack::StaticClass()))
+		{
+			ExportLevelSequenceVectorTrack(FbxActor, *Cast<UMovieSceneDoubleVectorTrack>(Track), BoundObject, MovieScene->GetPlaybackRange(), RootToLocalTransform);
+		}
 		else
 		{
 			bool bBakeChannels = false;
@@ -3026,20 +3032,83 @@ void FFbxExporter::ExportLevelSequenceColorTrack(FbxNode* FbxNode, UMovieSceneCo
 
 	FbxLight* FbxLight = FbxNode->GetLight();	
 
+	const FString PropertyName = ColorTrack.GetTrackName().ToString();
+
 	FbxProperty Property;
-	if(FbxLight)
+
+	if(PropertyName == TEXT("LightColor") && FbxLight)
 	{
 		Property = FbxLight->Color;
+	}
+
+	if(Property == 0)
+	{
+		CreateAnimatableUserProperty(FbxNode, FbxDouble3{ 0 }, TCHAR_TO_UTF8(*PropertyName), TCHAR_TO_UTF8(*PropertyName), FbxDouble3DT);
+		Property = FbxNode->FindProperty(TCHAR_TO_UTF8(*PropertyName));
+	}
+
+	if(Property == 0)
+	{
+		return;
+	}
+
+	FbxAnimLayer* BaseLayer = AnimStack->GetMember<FbxAnimLayer>(0);
+	FbxAnimCurve* CurveRed = Property.GetCurve(BaseLayer, FBXSDK_CURVENODE_COLOR_RED, true);
+	FbxAnimCurve* CurveGreen = Property.GetCurve(BaseLayer, FBXSDK_CURVENODE_COLOR_GREEN, true);
+	FbxAnimCurve* CurveBlue = Property.GetCurve(BaseLayer, FBXSDK_CURVENODE_COLOR_BLUE, true);
+
+	FFrameRate TickResolution = ColorTrack.GetTypedOuter<UMovieScene>()->GetTickResolution();
+	ExportChannelToFbxCurve(*CurveRed, ColorSection->GetRedChannel(), TickResolution, ERichCurveValueMode::Default, false, RootToLocalTransform);
+	ExportChannelToFbxCurve(*CurveGreen, ColorSection->GetGreenChannel(), TickResolution, ERichCurveValueMode::Default, false, RootToLocalTransform);
+	ExportChannelToFbxCurve(*CurveBlue, ColorSection->GetBlueChannel(), TickResolution, ERichCurveValueMode::Default, false, RootToLocalTransform);
+}
+
+void FFbxExporter::ExportLevelSequenceVectorTrack(FbxNode* FbxNode, UMovieSceneDoubleVectorTrack& VectorTrack, UObject* BoundObject, const TRange<FFrameNumber>& InPlaybackRange, const FMovieSceneSequenceTransform& RootToLocalTransform)
+{
+	UMovieSceneDoubleVectorSection* VectorSection = VectorTrack.GetAllSections().Num() > 0
+		? Cast<UMovieSceneDoubleVectorSection>(VectorTrack.GetAllSections()[0])
+		: nullptr;
+
+	if(!VectorSection)
+	{
+		return;
+	}
+
+	if(!FbxNode)
+	{
+		FbxNode = CreateNode(VectorTrack.GetDisplayName().ToString());
+	}
+
+	const FString PropertyName = VectorTrack.GetTrackName().ToString();
 
-		FbxAnimLayer* BaseLayer = AnimStack->GetMember<FbxAnimLayer>(0);
-		FbxAnimCurve* CurveRed =  Property.GetCurve(BaseLayer, FBXSDK_CURVENODE_COLOR_RED, true);
-		FbxAnimCurve* CurveGreen =  Property.GetCurve(BaseLayer, FBXSDK_CURVENODE_COLOR_GREEN, true);
-		FbxAnimCurve* CurveBlue =  Property.GetCurve(BaseLayer, FBXSDK_CURVENODE_COLOR_BLUE, true);
 
-		FFrameRate TickResolution = ColorTrack.GetTypedOuter<UMovieScene>()->GetTickResolution();
-		ExportChannelToFbxCurve(*CurveRed, ColorSection->GetRedChannel(), TickResolution, ERichCurveValueMode::Default, false, RootToLocalTransform);
-		ExportChannelToFbxCurve(*CurveGreen, ColorSection->GetGreenChannel(), TickResolution, ERichCurveValueMode::Default, false, RootToLocalTransform);
-		ExportChannelToFbxCurve(*CurveBlue, ColorSection->GetBlueChannel(), TickResolution, ERichCurveValueMode::Default, false, RootToLocalTransform);
+	//VectorSections supports up to 2-4 channels
+	const int32 NumChannelsUsed = VectorSection->GetChannelsUsed();
+
+	CreateAnimatableUserProperty(FbxNode,
+								 NumChannelsUsed == 2 ? FbxDouble2{ 0 } :
+								 NumChannelsUsed == 3 ? FbxDouble3{ 0 } :
+														FbxDouble4{ 0 },
+								 TCHAR_TO_UTF8(*PropertyName),
+								 TCHAR_TO_UTF8(*PropertyName),
+								 NumChannelsUsed == 2 ? FbxDouble2DT :
+								 NumChannelsUsed == 3 ? FbxDouble3DT :
+														FbxDouble4DT);
+	FbxProperty Property = FbxNode->FindProperty(TCHAR_TO_UTF8(*PropertyName));
+
+	if(Property == 0)
+	{
+		return;
+	}
+
+	const char* CurvesNames[] = { "X", "Y", "Z", "W" };
+	FbxAnimLayer* BaseLayer = AnimStack->GetMember<FbxAnimLayer>(0);
+	FFrameRate TickResolution = VectorTrack.GetTypedOuter<UMovieScene>()->GetTickResolution();
+
+	for(int32 Index = 0; Index < NumChannelsUsed; ++Index)
+	{
+		FbxAnimCurve* Curve = Property.GetCurve(BaseLayer, CurvesNames[Index], true);
+		ExportChannelToFbxCurve(*Curve, VectorSection->GetChannel(Index), TickResolution, ERichCurveValueMode::Default, false, RootToLocalTransform);
 	}
 }
 
diff --git a/Engine/Source/Editor/UnrealEd/Private/FbxExporter.h b/Engine/Source/Editor/UnrealEd/Private/FbxExporter.h
--- a/Engine/Source/Editor/UnrealEd/Private/FbxExporter.h
+++ b/Engine/Source/Editor/UnrealEd/Private/FbxExporter.h
@@ -34,6 +34,7 @@ class UMovieScene3DTransformTrack;
 class UMovieSceneColorTrack;
 class UMovieScenePropertyTrack;
 class UMovieSceneTrack;
+class UMovieSceneDoubleVectorTrack;
 class USkeletalMesh;
 class USkeletalMeshComponent;
 class USplineMeshComponent;
@@ -455,9 +456,14 @@ private:
 	UNREALED_API void ExportLevelSequenceBaked3DTransformTrack(IAnimTrackAdapter& AnimTrackAdapter, FbxNode* FbxActor, IMovieScenePlayer* MovieScenePlayer, FMovieSceneSequenceIDRef InSequenceID, TArray<TWeakObjectPtr<UMovieScene3DTransformTrack> > TransformTracks, UObject* BoundObject, const TRange<FFrameNumber>& InPlaybackRange, const FMovieSceneSequenceTransform& RootToLocalTransform);
 
 	/**
-	 * Exports a level sequence color track into the FBX animation stack. (only works for light)
+	 * Exports a level sequence color track into the FBX animation stack.
 	 */
 	UNREALED_API void ExportLevelSequenceColorTrack(FbxNode* FbxActor, UMovieSceneColorTrack& ColorTrack, UObject* BoundObject, const TRange<FFrameNumber>& InPlaybackRange, const FMovieSceneSequenceTransform& RootToLocalTransform);
+	
+	/**
+	 * Exports a level sequence vector track into the FBX animation stack.
+	 */
+	UNREALED_API void ExportLevelSequenceVectorTrack(FbxNode* FbxActor, UMovieSceneDoubleVectorTrack& VectorTrack, UObject* BoundObject, const TRange<FFrameNumber>& InPlaybackRange, const FMovieSceneSequenceTransform& RootToLocalTransform);
 
 	/**
 	 * Exports a level sequence property track into the FBX animation stack. 
