Cooker SkipOnlyEditorOnly: Fix processing of the discovery queue to skip expected-load editoronly references. We were skipping these packages before adding them to the discovery queue already, but we were not handling the case of the editoronly reference coming itself from a discoveredload package.
#rnx

diff --git a/Engine/Source/Editor/UnrealEd/Private/Cooker/CookRequestCluster.cpp b/Engine/Source/Editor/UnrealEd/Private/Cooker/CookRequestCluster.cpp
--- a/Engine/Source/Editor/UnrealEd/Private/Cooker/CookRequestCluster.cpp
+++ b/Engine/Source/Editor/UnrealEd/Private/Cooker/CookRequestCluster.cpp
@@ -162,10 +162,35 @@ FRequestCluster::FRequestCluster(UCookOnTheFlyServer& InCOTFS, TRingBuffer<FDisc
 			continue;
 		}
 
-		// Startup packages and Generated packages do not need to add hidden dependencies or log warnings
+		// Startup packages and Generated packages are expected discovery types and do not need to add hidden
+		// dependencies.
 		if (Discovery->Instigator.Category != EInstigator::StartupPackage &&
 			Discovery->Instigator.Category != EInstigator::GeneratedPackage)
 		{
+			// For unsolicited packages, we need to check load-reachability to decide whether the load was expected.
+			bool bExpectedDiscoveryType = false;
+			if (Discovery->Instigator.Category == EInstigator::Unsolicited)
+			{
+				bExpectedDiscoveryType = PackageData.FindOrAddPlatformData(CookerLoadingPlatformKey).IsReachable();
+				if (bExpectedDiscoveryType && COTFS.bSkipOnlyEditorOnly)
+				{
+					// In SkipOnlyEditorOnly mode, expected-load unsolicited packages are skipped; merely loading
+					// a package is not sufficient to add it to the cook. So take no action on this package.
+					DiscoveryQueue.PopFrontValue();
+					Discovery = nullptr;
+					continue;
+				}
+			}
+			else
+			{
+				// For other instigator types, the discovery is either expected or unexpected depending on type.
+				// Adding packages to the cook should happen only for a few types of instigators, from external
+				// package requests, or during cluster exploration. If not expected, add a diagnostic message.
+				bExpectedDiscoveryType = Discovery->Instigator.Category == EInstigator::SaveTimeHardDependency ||
+					Discovery->Instigator.Category == EInstigator::SaveTimeSoftDependency ||
+					Discovery->Instigator.Category == EInstigator::ForceExplorableSaveTimeSoftDependency;
+			}
+
 			// If there are other discovered packages we have already added to this cluster, then defer this one
 			// until we have explored those; add this one to the next cluster. Exploring those earlier discoveries
 			// might add this one through cluster exploration and not require a hidden dependency.
@@ -190,11 +215,7 @@ FRequestCluster::FRequestCluster(UCookOnTheFlyServer& InCOTFS, TRingBuffer<FDisc
 				}
 			}
 
-			// Adding packages to the cook should happen only for a few types of instigators, from external package
-			// requests, or during cluster exploration. If not expected, add a diagnostic message.
-			if (Discovery->Instigator.Category != EInstigator::SaveTimeHardDependency &&
-				Discovery->Instigator.Category != EInstigator::SaveTimeSoftDependency &&
-				Discovery->Instigator.Category != EInstigator::ForceExplorableSaveTimeSoftDependency)
+			if (!bExpectedDiscoveryType)
 			{
 				COTFS.OnDiscoveredPackageDebug(PackageData.GetPackageName(), Discovery->Instigator);
 			}
