Cooker: CookPackageSplitter saveorder: Remove the workaround of setting MPCooKGeneratedSaveOrder=GeneratorFirst, now that WorldPartition has been fixed to handle arbitrary ordering  in CL 32885960.
#rnx

diff --git a/Engine/Source/Editor/UnrealEd/Private/Cooker/CookGenerationHelper.cpp b/Engine/Source/Editor/UnrealEd/Private/Cooker/CookGenerationHelper.cpp
--- a/Engine/Source/Editor/UnrealEd/Private/Cooker/CookGenerationHelper.cpp
+++ b/Engine/Source/Editor/UnrealEd/Private/Cooker/CookGenerationHelper.cpp
@@ -1408,11 +1408,7 @@ void FGenerationHelper::SetBeginCookConfigSettings()
 	FString SaveOrder;
 	GConfig->GetString(TEXT("CookSettings"), TEXT("MPCookGeneratorSaveOrder"), SaveOrder, GEditorIni);
 	FParse::Value(FCommandLine::Get(), TEXT("-MPCookGeneratorSaveOrder="), SaveOrder);
-	if (SaveOrder == TEXT("None"))
-	{
-		GenerationHelperPrivate::RequiredSaveOrder = GenerationHelperPrivate::ERequiredSaveOrder::None;
-	}
-	else if (SaveOrder == TEXT("GeneratorFirst"))
+	if (SaveOrder == TEXT("GeneratorFirst"))
 	{
 		GenerationHelperPrivate::RequiredSaveOrder = GenerationHelperPrivate::ERequiredSaveOrder::GeneratorFirst;
 	}
@@ -1422,9 +1418,13 @@ void FGenerationHelper::SetBeginCookConfigSettings()
 	}
 	else
 	{
-		// UE-211211: Make RequiredSaveOrder::GeneratorFirst the default until we have
-		// changed WorldPartitionCookPackageSplitter to handle the new allowed ordering.
-		GenerationHelperPrivate::RequiredSaveOrder = GenerationHelperPrivate::ERequiredSaveOrder::GeneratorFirst;
+		if (!SaveOrder.IsEmpty() && SaveOrder != TEXT("None"))
+		{
+			UE_LOG(LogCook, Error,
+				TEXT("Invalid setting Editor:[CookSettings]:MPCookGeneratorSaveOrder=%s. Expected values are 'GeneratorFirst', 'GeneratedFirst', or 'None'. Falling back to default 'None'."),
+				*SaveOrder);
+		}
+		GenerationHelperPrivate::RequiredSaveOrder = GenerationHelperPrivate::ERequiredSaveOrder::None;
 	}
 }
 
