Fix tab history right-click menu behaviour

diff --git a/Engine/Source/Editor/UnrealEd/Private/WorkflowOrientedApp/WorkflowTabManager.cpp b/Engine/Source/Editor/UnrealEd/Private/WorkflowOrientedApp/WorkflowTabManager.cpp
--- a/Engine/Source/Editor/UnrealEd/Private/WorkflowOrientedApp/WorkflowTabManager.cpp
+++ b/Engine/Source/Editor/UnrealEd/Private/WorkflowOrientedApp/WorkflowTabManager.cpp
@@ -302,31 +302,48 @@ TSharedRef< SWidget > FTabInfo::CreateHistoryNavigationWidget()
 {
 	if(!HistoryNavigationWidget.IsValid())
 	{
-		TWeakPtr< SMenuAnchor > BackMenuAnchorPtr;
-		TWeakPtr< SMenuAnchor > FwdMenuAnchorPtr;
+		const TSharedRef<SMenuAnchor> BackMenuAnchorPtr =
+			SNew(SMenuAnchor)
+			.Placement( MenuPlacement_BelowAnchor )
+			.OnGetMenuContent( this, &FTabInfo::CreateHistoryMenu, true )
+			[
+				SNew(SButton)
+				.OnClicked( this, &FTabInfo::OnGoBackInHistory )
+				.ButtonStyle( FAppStyle::Get(), "GraphBreadcrumbButton" )
+				.IsEnabled(this, &FTabInfo::CanStepBackwardInHistory)
+				.ToolTipText(NSLOCTEXT("WorkflowNavigationBrowser", "Backward_Tooltip", "Step backward in the tab history. Right click to see full history."))
+				[
+					SNew(SImage)
+					.Image( FAppStyle::GetBrush("GraphBreadcrumb.BrowseBack") )
+				]
+			];
+
+		const TSharedRef<SMenuAnchor> FwdMenuAnchorPtr =
+			SNew(SMenuAnchor)
+			.Placement( MenuPlacement_BelowAnchor )
+			.OnGetMenuContent( this, &FTabInfo::CreateHistoryMenu, false )
+			[
+				SNew(SButton)
+				.OnClicked( this, &FTabInfo::OnGoForwardInHistory )
+				.ButtonStyle( FAppStyle::Get(), "GraphBreadcrumbButton" )
+				.IsEnabled(this, &FTabInfo::CanStepForwardInHistory)
+				.ToolTipText(NSLOCTEXT("WorkflowNavigationBrowser", "Forward_Tooltip", "Step forward in the tab history. Right click to see full history."))
+				[
+					SNew(SImage)
+					.Image( FAppStyle::GetBrush("GraphBreadcrumb.BrowseForward") )
+				]
+			];
+
 		HistoryNavigationWidget = 
 			SNew(SHorizontalBox)
 			+SHorizontalBox::Slot()
 			.AutoWidth()
 			[
 				SNew(SBorder)
-				.OnMouseButtonDown(this, &FTabInfo::OnMouseDownHistory, BackMenuAnchorPtr)
+				.OnMouseButtonDown(this, &FTabInfo::OnMouseDownHistory, BackMenuAnchorPtr.ToWeakPtr())
 				.BorderImage( FAppStyle::GetBrush("NoBorder") )
 				[
-					SAssignNew(BackMenuAnchorPtr, SMenuAnchor)
-					.Placement( MenuPlacement_BelowAnchor )
-					.OnGetMenuContent( this, &FTabInfo::CreateHistoryMenu, true )
-					[
-						SNew(SButton)
-						.OnClicked( this, &FTabInfo::OnGoBackInHistory )
-						.ButtonStyle( FAppStyle::Get(), "GraphBreadcrumbButton" )
-						.IsEnabled(this, &FTabInfo::CanStepBackwardInHistory)
-						.ToolTipText(NSLOCTEXT("WorkflowNavigationBrowser", "Backward_Tooltip", "Step backward in the tab history. Right click to see full history."))
-						[
-							SNew(SImage)
-							.Image( FAppStyle::GetBrush("GraphBreadcrumb.BrowseBack") )
-						]
-					]
+					BackMenuAnchorPtr
 				]
 			]
 
@@ -334,23 +351,10 @@ TSharedRef< SWidget > FTabInfo::CreateHistoryNavigationWidget()
 			.AutoWidth()
 			[
 				SNew(SBorder)
-				.OnMouseButtonDown(this, &FTabInfo::OnMouseDownHistory, FwdMenuAnchorPtr)
+				.OnMouseButtonDown(this, &FTabInfo::OnMouseDownHistory, FwdMenuAnchorPtr.ToWeakPtr())
 				.BorderImage( FAppStyle::GetBrush("NoBorder") )
 				[
-					SAssignNew(FwdMenuAnchorPtr, SMenuAnchor)
-					.Placement( MenuPlacement_BelowAnchor )
-					.OnGetMenuContent( this, &FTabInfo::CreateHistoryMenu, false )
-					[
-						SNew(SButton)
-						.OnClicked( this, &FTabInfo::OnGoForwardInHistory )
-						.ButtonStyle( FAppStyle::Get(), "GraphBreadcrumbButton" )
-						.IsEnabled(this, &FTabInfo::CanStepForwardInHistory)
-						.ToolTipText(NSLOCTEXT("WorkflowNavigationBrowser", "Forward_Tooltip", "Step forward in the tab history. Right click to see full history."))
-						[
-							SNew(SImage)
-							.Image( FAppStyle::GetBrush("GraphBreadcrumb.BrowseForward") )
-						]
-					]
+					FwdMenuAnchorPtr
 				]
 			];
 	}
