Sequencer: Remove deprecations (4.25, 4.26, 4.27)

#rnx

diff --git a/Engine/Plugins/Cameras/GameplayCameras/Source/GameplayCameras/Private/Shakes/LegacyCameraShake.cpp b/Engine/Plugins/Cameras/GameplayCameras/Source/GameplayCameras/Private/Shakes/LegacyCameraShake.cpp
--- a/Engine/Plugins/Cameras/GameplayCameras/Source/GameplayCameras/Private/Shakes/LegacyCameraShake.cpp
+++ b/Engine/Plugins/Cameras/GameplayCameras/Source/GameplayCameras/Private/Shakes/LegacyCameraShake.cpp
@@ -409,11 +409,6 @@ bool ULegacyCameraShake::IsLooping() const
 	return ActualOscillationDuration < 0.0f;
 }
 
-void ULegacyCameraShake::SetCurrentTimeAndApplyShake(float NewTime, FMinimalViewInfo& POV)
-{
-	ScrubAndApplyCameraShake(NewTime, 1.f, POV);
-}
-
 ULegacyCameraShake* ULegacyCameraShake::StartLegacyCameraShake(APlayerCameraManager* PlayerCameraManager, TSubclassOf<ULegacyCameraShake> ShakeClass, float Scale, ECameraShakePlaySpace PlaySpace, FRotator UserPlaySpaceRot)
 {
 	if (PlayerCameraManager)
diff --git a/Engine/Plugins/Cameras/GameplayCameras/Source/GameplayCameras/Public/Shakes/LegacyCameraShake.h b/Engine/Plugins/Cameras/GameplayCameras/Source/GameplayCameras/Public/Shakes/LegacyCameraShake.h
--- a/Engine/Plugins/Cameras/GameplayCameras/Source/GameplayCameras/Public/Shakes/LegacyCameraShake.h
+++ b/Engine/Plugins/Cameras/GameplayCameras/Source/GameplayCameras/Public/Shakes/LegacyCameraShake.h
@@ -236,10 +236,6 @@ public:
 	/** Returns true if this camera shake will loop forever */
 	bool IsLooping() const;
 
-	/** Sets current playback time and applies the shake (both oscillation and camera animation sequence) to the given POV. */
-	UE_DEPRECATED(4.27, "SetCurrentTimeAndApplyShake is deprecated, please use ScrubAndApplyCameraShake")
-	void SetCurrentTimeAndApplyShake(float NewTime, FMinimalViewInfo& POV);
-
 private:
 
 	void DoStartShake(const FCameraShakePatternStartParams& Params);
diff --git a/Engine/Plugins/VirtualProduction/Takes/Source/TakeRecorder/Public/Recorder/TakeRecorderPanel.h b/Engine/Plugins/VirtualProduction/Takes/Source/TakeRecorder/Public/Recorder/TakeRecorderPanel.h
--- a/Engine/Plugins/VirtualProduction/Takes/Source/TakeRecorder/Public/Recorder/TakeRecorderPanel.h
+++ b/Engine/Plugins/VirtualProduction/Takes/Source/TakeRecorder/Public/Recorder/TakeRecorderPanel.h
@@ -74,10 +74,6 @@ public:
 	UFUNCTION(BlueprintCallable, Category="Take Recorder|Panel", DisplayName="Set Mode (Read-Only Level Sequence)")
 	void SetupForViewing(ULevelSequence* LevelSequenceAsset);
 
-	UE_DEPRECATED(4.25, "Please use ClearPendingTake instead")
-	UFUNCTION(BlueprintCallable, Category = "Take Recorder|Panel", DisplayName = "New Take", meta = (DeprecatedFunction, DeprecationMessage = "Please use ClearPendingTake instead"))
-	void NewTake() { ClearPendingTake(); }
-
 	/*
 	 * Clear the pending take level sequence
 	 */
diff --git a/Engine/Plugins/VirtualProduction/Takes/Source/TakesCore/Private/TakeRecorderSources.cpp b/Engine/Plugins/VirtualProduction/Takes/Source/TakesCore/Private/TakeRecorderSources.cpp
--- a/Engine/Plugins/VirtualProduction/Takes/Source/TakesCore/Private/TakeRecorderSources.cpp
+++ b/Engine/Plugins/VirtualProduction/Takes/Source/TakesCore/Private/TakeRecorderSources.cpp
@@ -266,11 +266,6 @@ void UTakeRecorderSources::StartRecordingSource(TArray<UTakeRecorderSource *> In
 }
 
 
-void UTakeRecorderSources::StartRecordingSource(TArray<UTakeRecorderSource *> InSources, const FTimecode& CurrentTimecode)
-{
-	StartRecordingSource(InSources, FQualifiedFrameTime(CurrentTimecode, TargetLevelSequenceDisplayRate)); 
-}
-
 void UTakeRecorderSources::StartRecordingTheseSources(const TArray<UTakeRecorderSource *>& InSources, const FQualifiedFrameTime& CurrentFrameTime)
 {
 	FTimecode CurrentTimecode = FTimecode::FromFrameNumber(FFrameRate::TransformTime(CurrentFrameTime.Time, TargetLevelSequenceTickResolution, TargetLevelSequenceDisplayRate).FloorToFrame(), TargetLevelSequenceDisplayRate);
diff --git a/Engine/Plugins/VirtualProduction/Takes/Source/TakesCore/Public/TakeRecorderSources.h b/Engine/Plugins/VirtualProduction/Takes/Source/TakesCore/Public/TakeRecorderSources.h
--- a/Engine/Plugins/VirtualProduction/Takes/Source/TakesCore/Public/TakeRecorderSources.h
+++ b/Engine/Plugins/VirtualProduction/Takes/Source/TakesCore/Public/TakeRecorderSources.h
@@ -102,9 +102,6 @@ public:
 	UFUNCTION(BlueprintCallable, Category = "Take Recorder")
 	void StartRecordingSource(TArray<UTakeRecorderSource*> InSources, const FQualifiedFrameTime& CurrentFrameTime);
 	
-	UE_DEPRECATED(4.27, "StartRecordingSource with FTimecode has been deprecated, please use StartRecordingSource with FQualifiedFrameTime")
-	void StartRecordingSource(TArray<UTakeRecorderSource*> InSources, const FTimecode& CurrentTimecode);
-
 public:
 
 	/**
diff --git a/Engine/Source/Editor/UnrealEd/Public/EditorViewportClient.h b/Engine/Source/Editor/UnrealEd/Public/EditorViewportClient.h
--- a/Engine/Source/Editor/UnrealEd/Public/EditorViewportClient.h
+++ b/Engine/Source/Editor/UnrealEd/Public/EditorViewportClient.h
@@ -355,12 +355,6 @@ public:
 	 */
 	UNREALED_API void AddRealtimeOverride(bool bShouldBeRealtime, FText SystemDisplayName);
 
-	UE_DEPRECATED(4.26, "SetRealtimeOverride is replaced with AddRealtimeOverride, as multiple overrides can now be added.")
-	void SetRealtimeOverride(bool bShouldBeRealtime, FText SystemDisplayName)
-	{
-		AddRealtimeOverride(bShouldBeRealtime, SystemDisplayName);
-	}
-	
 	/**
 	 * Returns whether there's a realtime override registered with the given system name.
 	 */
@@ -374,12 +368,6 @@ public:
 	 */
 	UNREALED_API bool RemoveRealtimeOverride(FText SystemDisplayName, bool bCheckMissingOverride = true);
 
-	UE_DEPRECATED(4.26, "RemoveRealtimeOverride now takes a system display name to help remove the correct override.")
-	void RemoveRealtimeOverride()
-	{
-		PopRealtimeOverride();
-	}
-
 	/**
 	 * Removes the last added realtime override.
 	 * WARNING: it's recommended to use RemoveRealtimeOverride to prevent removing an override that belongs to another systems than yours.
