[Viewport Toolbar] Fix commented out generation of show flag menu entries in SCommonEditorViewportToolbarBase class

Updated them to the UToolMenu representation to emulate what the level editor toolbar is doing.

diff --git a/Engine/Source/Editor/LevelEditor/Private/SLevelViewport.cpp b/Engine/Source/Editor/LevelEditor/Private/SLevelViewport.cpp
--- a/Engine/Source/Editor/LevelEditor/Private/SLevelViewport.cpp
+++ b/Engine/Source/Editor/LevelEditor/Private/SLevelViewport.cpp
@@ -74,7 +74,6 @@
 #include "Slate/SGameLayerManager.h"
 #include "FoliageType.h"
 #include "IVREditorModule.h"
-#include "ShowFlagMenuCommands.h"
 #include "AssetRegistry/AssetRegistryModule.h"
 #include "AssetRegistry/IAssetRegistry.h"
 #include "BufferVisualizationMenuCommands.h"
@@ -1425,7 +1424,6 @@ void SLevelViewport::BindCommands()
 
 	BindOptionCommands( UICommandListRef );
 	BindViewCommands( UICommandListRef );
-	BindShowCommands( UICommandListRef );
 	BindDropCommands( UICommandListRef );
 
 	if ( ParentLevelEditor.IsValid() )
@@ -1456,7 +1454,7 @@ void SLevelViewport::BindOptionCommands( FUICommandList& OutCommandList )
 	OutCommandList.MapAction(
 		ViewportActions.ToggleGameView,
 		FExecuteAction::CreateSP( this, &SLevelViewport::ToggleGameView ),
-		FCanExecuteAction(),
+		FCanExecuteAction::CreateSP( this, &SLevelViewport::CanToggleGameView ),
 		FIsActionChecked::CreateSP( this, &SLevelViewport::IsInGameView ) );
 
 	OutCommandList.MapAction(
@@ -1639,7 +1637,7 @@ void SLevelViewport::BindShowCommands( FUICommandList& OutCommandList )
 		LevelViewportCommands.UseDefaultShowFlags,
 		FExecuteAction::CreateSP( this, &SLevelViewport::OnUseDefaultShowFlags, false ) );
 
-	FShowFlagMenuCommands::Get().BindCommands(OutCommandList, Client);
+	SEditorViewport::BindShowCommands( OutCommandList );
 
 	// Show Volumes
 	{
@@ -2044,18 +2042,21 @@ void SLevelViewport::OnTakeHighResScreenshot()
 
 void SLevelViewport::ToggleGameView()
 {
-	bool bGameViewEnable = !LevelViewportClient->IsInGameView();
+	if( LevelViewportClient->IsPerspective() )
+	{
+		bool bGameViewEnable = !LevelViewportClient->IsInGameView();
 
-	// "Mode Widget" should not automatically be reactivated by selecting an actor after "Game View" is enabled
-	LevelViewportClient->bAlwaysShowModeWidgetAfterSelectionChanges = bGameViewEnable ? false : true;
+		// "Mode Widget" should not automatically be reactivated by selecting an actor after "Game View" is enabled
+		LevelViewportClient->bAlwaysShowModeWidgetAfterSelectionChanges = bGameViewEnable ? false : true;
 
-	LevelViewportClient->SetGameView(bGameViewEnable);
+		LevelViewportClient->SetGameView(bGameViewEnable);
 
-	if (!bGameViewEnable)
-	{
-		// LevelViewportClient->bShowWidget is set to "false" when entering game mode
-		// Need to turn it back to "true" when exiting game mode
-		LevelViewportClient->ShowWidget(true);
+		if (!bGameViewEnable)
+		{
+			// LevelViewportClient->bShowWidget is set to "false" when entering game mode
+			// Need to turn it back to "true" when exiting game mode
+			LevelViewportClient->ShowWidget(true);
+		}
 	}
 }
 
diff --git a/Engine/Source/Editor/LevelEditor/Public/SLevelViewport.h b/Engine/Source/Editor/LevelEditor/Public/SLevelViewport.h
--- a/Engine/Source/Editor/LevelEditor/Public/SLevelViewport.h
+++ b/Engine/Source/Editor/LevelEditor/Public/SLevelViewport.h
@@ -195,7 +195,6 @@ public:
 	/**
 	 * @return true if we can toggle game view
 	 */
-	UE_DEPRECATED(5.5, "CanToggleGameView is deprecated. Please update usage to assume that this function would return true in the future. All level viewports should support game view moving forward.")
 	bool CanToggleGameView() const;
 
 	/**
@@ -672,7 +671,7 @@ private:
 	 *
 	 * @param CommandList	The list to bind commands to
 	 */
-	void BindShowCommands( FUICommandList& CommandList );
+	void BindShowCommands( FUICommandList& CommandList ) override;
 
 	/**
 	 * Binds commands for our drag-drop context menu
diff --git a/Engine/Source/Editor/UnrealEd/Private/SCommonEditorViewportToolbarBase.cpp b/Engine/Source/Editor/UnrealEd/Private/SCommonEditorViewportToolbarBase.cpp
--- a/Engine/Source/Editor/UnrealEd/Private/SCommonEditorViewportToolbarBase.cpp
+++ b/Engine/Source/Editor/UnrealEd/Private/SCommonEditorViewportToolbarBase.cpp
@@ -24,7 +24,9 @@
 #include "AssetEditorViewportLayout.h"
 #include "SAssetEditorViewport.h"
 #include "ToolMenu.h"
+#include "ToolMenus.h"
 #include "ToolMenuSection.h"
+#include "ShowFlagMenuCommands.h"
 
 
 #define LOCTEXT_NAMESPACE "SCommonEditorViewportToolbarBase"
@@ -636,98 +638,37 @@ TSharedRef<SWidget> SCommonEditorViewportToolbarBase::GenerateCameraMenu() const
 TSharedRef<SWidget> SCommonEditorViewportToolbarBase::GenerateShowMenu() const
 {
 	GetInfoProvider().OnFloatingButtonClicked();
-	TSharedRef<SEditorViewport> ViewportRef = GetInfoProvider().GetViewportWidget();
-
-#if 0
-	const FLevelViewportCommands& Actions = FLevelViewportCommands::Get();
-
-	const TArray<FShowFlagData>& ShowFlagData = GetShowFlagMenuItems();
-
-	TArray<CommonEditorViewportUtils::FShowMenuCommand> ShowMenu[SFG_Max];
 
-	// Get each show flag command and put them in their corresponding groups
-	for( int32 ShowFlag = 0; ShowFlag < ShowFlagData.Num(); ++ShowFlag )
-	{
-		const FShowFlagData& SFData = ShowFlagData[ShowFlag];
-		
-		ShowMenu[SFData.Group].Add( Actions.ShowFlagCommands[ ShowFlag ] );
-	}
-#endif
-
-	const bool bInShouldCloseWindowAfterMenuSelection = true;
-	FMenuBuilder ShowMenuBuilder(bInShouldCloseWindowAfterMenuSelection, ViewportRef->GetCommandList());
+	static const FName MenuName("ViewportToolbarBase.Show");
+	if (!UToolMenus::Get()->IsMenuRegistered(MenuName))
 	{
-
-#if 0
-		ShowMenuBuilder.AddMenuEntry( Actions.UseDefaultShowFlags );
-		
-		if( ShowMenu[SFG_Normal].Num() > 0 )
+		UToolMenu* ShowMenu = UToolMenus::Get()->RegisterMenu(MenuName);
+		ShowMenu->AddDynamicSection("Flags", FNewToolMenuDelegate::CreateLambda([](UToolMenu* InMenu)
 		{
-			// Generate entries for the standard show flags
-			ShowMenuBuilder.BeginSection("LevelViewportShowFlagsCommon", LOCTEXT("CommonShowFlagHeader", "Common Show Flags") );
+			if (UCommonViewportToolbarBaseMenuContext* ContextObject = InMenu->FindContext<UCommonViewportToolbarBaseMenuContext>())
 			{
-				for( int32 EntryIndex = 0; EntryIndex < ShowMenu[SFG_Normal].Num(); ++EntryIndex )
+				if (TSharedPtr<const SCommonEditorViewportToolbarBase> ToolbarWidgetPin = ContextObject->ToolbarWidget.Pin())
 				{
-					ShowMenuBuilder.AddMenuEntry( ShowMenu[SFG_Normal][ EntryIndex ].ShowMenuItem, NAME_None, ShowMenu[SFG_Normal][ EntryIndex ].LabelOverride );
+					ToolbarWidgetPin->FillShowFlagsMenu(InMenu);
 				}
 			}
-			ShowMenuBuilder.EndSection();
-		}
-
-		// Generate entries for the different show flags groups
-		ShowMenuBuilder.BeginSection("LevelViewportShowFlags");
-		{
-			ShowMenuBuilder.AddSubMenu( LOCTEXT("PostProcessShowFlagsMenu", "Post Processing"), LOCTEXT("PostProcessShowFlagsMenu_ToolTip", "Post process show flags"),
-				FNewMenuDelegate::CreateStatic(&CommonEditorViewportUtils::FillShowMenu, ShowMenu[SFG_PostProcess], 0));
-
-			ShowMenuBuilder.AddSubMenu( LOCTEXT("LightTypesShowFlagsMenu", "Light Types"), LOCTEXT("LightTypesShowFlagsMenu_ToolTip", "Light Types show flags"),
-				FNewMenuDelegate::CreateStatic(&CommonEditorViewportUtils::FillShowMenu, ShowMenu[SFG_LightTypes], 0));
-
-			ShowMenuBuilder.AddSubMenu( LOCTEXT("LightingComponentsShowFlagsMenu", "Lighting Components"), LOCTEXT("LightingComponentsShowFlagsMenu_ToolTip", "Lighting Components show flags"),
-				FNewMenuDelegate::CreateStatic(&CommonEditorViewportUtils::FillShowMenu, ShowMenu[SFG_LightingComponents], 0));
-
-			ShowMenuBuilder.AddSubMenu( LOCTEXT("LightingFeaturesShowFlagsMenu", "Lighting Features"), LOCTEXT("LightingFeaturesShowFlagsMenu_ToolTip", "Lighting Features show flags"),
-				FNewMenuDelegate::CreateStatic(&CommonEditorViewportUtils::FillShowMenu, ShowMenu[SFG_LightingFeatures], 0));
-
-			ShowMenuBuilder.AddSubMenu(LOCTEXT("LumenShowFlagsMenu", "Lumen"), LOCTEXT("LumenShowFlagsMenu_ToolTip", "Lumen show flags"),
-				FNewMenuDelegate::CreateStatic(&CommonEditorViewportUtils::FillShowMenu, ShowMenu[SFG_Lumen], 0));
-
-			ShowMenuBuilder.AddSubMenu(LOCTEXT("NaniteShowFlagsMenu", "Nanite"), LOCTEXT("NaniteShowFlagsMenu_ToolTip", "Nanite show flags"),
-				FNewMenuDelegate::CreateStatic(&CommonEditorViewportUtils::FillShowMenu, ShowMenu[SFG_Nanite], 0));
-
-			ShowMenuBuilder.AddSubMenu( LOCTEXT("DeveloperShowFlagsMenu", "Developer"), LOCTEXT("DeveloperShowFlagsMenu_ToolTip", "Developer show flags"),
-				FNewMenuDelegate::CreateStatic(&CommonEditorViewportUtils::FillShowMenu, ShowMenu[SFG_Developer], 0));
-
-			ShowMenuBuilder.AddSubMenu( LOCTEXT("VisualizeShowFlagsMenu", "Visualize"), LOCTEXT("VisualizeShowFlagsMenu_ToolTip", "Visualize show flags"),
-				FNewMenuDelegate::CreateStatic(&CommonEditorViewportUtils::FillShowMenu, ShowMenu[SFG_Visualize], 0));
-
-			ShowMenuBuilder.AddSubMenu( LOCTEXT("AdvancedShowFlagsMenu", "Advanced"), LOCTEXT("AdvancedShowFlagsMenu_ToolTip", "Advanced show flags"),
-				FNewMenuDelegate::CreateStatic(&CommonEditorViewportUtils::FillShowMenu, ShowMenu[SFG_Advanced], 0));
-		}
-		ShowMenuBuilder.EndSection();
-
-
-		FText ShowAllLabel = LOCTEXT("ShowAllLabel", "Show All");
-		FText HideAllLabel = LOCTEXT("HideAllLabel", "Hide All");
-
-		// Show Volumes sub-menu
-		{
-			TArray< FLevelViewportCommands::FShowMenuCommand > ShowVolumesMenu;
-
-			// 'Show All' and 'Hide All' buttons
-			ShowVolumesMenu.Add( FLevelViewportCommands::FShowMenuCommand( Actions.ShowAllVolumes, ShowAllLabel ) );
-			ShowVolumesMenu.Add( FLevelViewportCommands::FShowMenuCommand( Actions.HideAllVolumes, HideAllLabel ) );
-
-			// Get each show flag command and put them in their corresponding groups
-			ShowVolumesMenu += Actions.ShowVolumeCommands;
+		}));
+	}
 
-			ShowMenuBuilder.AddSubMenu( LOCTEXT("ShowVolumesMenu", "Volumes"), LOCTEXT("ShowVolumesMenu_ToolTip", "Show volumes flags"),
-				FNewMenuDelegate::CreateStatic( &FillShowMenu, ShowVolumesMenu, 2 ) );
-		}
-#endif
+	FToolMenuContext NewMenuContext;
+	UCommonViewportToolbarBaseMenuContext* ContextObject = NewObject<UCommonViewportToolbarBaseMenuContext>();
+	ContextObject->ToolbarWidget = SharedThis(this);
+	NewMenuContext.AddObject(ContextObject);
+	if (TSharedPtr<SEditorViewport> ViewportWidget = GetInfoProvider().GetViewportWidget())
+	{
+		NewMenuContext.AppendCommandList(GetInfoProvider().GetViewportWidget()->GetCommandList());
 	}
+	return UToolMenus::Get()->GenerateWidget(MenuName, NewMenuContext);
+}
 
-	return ShowMenuBuilder.MakeWidget();
+void SCommonEditorViewportToolbarBase::FillShowFlagsMenu(UToolMenu* InMenu) const
+{
+	FShowFlagMenuCommands::Get().BuildShowFlagsMenu(InMenu);
 }
 
 TSharedRef<SWidget> SCommonEditorViewportToolbarBase::GenerateFOVMenu() const
diff --git a/Engine/Source/Editor/UnrealEd/Private/SEditorViewport.cpp b/Engine/Source/Editor/UnrealEd/Private/SEditorViewport.cpp
--- a/Engine/Source/Editor/UnrealEd/Private/SEditorViewport.cpp
+++ b/Engine/Source/Editor/UnrealEd/Private/SEditorViewport.cpp
@@ -28,6 +28,7 @@
 #include "Widgets/Colors/SComplexGradient.h"
 #include "Modules/ModuleManager.h"
 #include "ISettingsModule.h"
+#include "ShowFlagMenuCommands.h"
 #if WITH_DUMPGPU
 	#include "RenderGraph.h"
 #endif
@@ -415,7 +416,6 @@ void SEditorViewport::BindCommands()
 	MAP_VIEWMODE_ACTION( Commands.WireframeMode, VMI_BrushWireframe );
 	MAP_VIEWMODE_ACTION( Commands.UnlitMode, VMI_Unlit );
 	MAP_VIEWMODE_ACTION( Commands.LitMode, VMI_Lit );
-	MAP_VIEWMODE_ACTION( Commands.LitWireframeMode, VMI_Lit_Wireframe);
 #if RHI_RAYTRACING
 	if (IsRayTracingAllowed())
 	{
@@ -457,8 +457,6 @@ void SEditorViewport::BindCommands()
 	MAP_VIEWMODE_ACTION( Commands.CollisionPawn, VMI_CollisionPawn);
 	MAP_VIEWMODE_ACTION( Commands.CollisionVisibility, VMI_CollisionVisibility);
 
-	MAP_VIEWMODE_ACTION( Commands.VisualizeLWCComplexity, VMI_LWCComplexity);
-
 	if (GEnableGPUSkinCache)
 	{
 		MAP_VIEWMODE_ACTION(Commands.VisualizeGPUSkinCacheMode, VMI_VisualizeGPUSkinCache);
@@ -477,6 +475,13 @@ void SEditorViewport::BindCommands()
 		MAP_VIEWMODEPARAM_ACTION( Commands.TexStreamAccMaterialTextureScaleSingle[TextureIndex], TextureIndex );
 		MAP_VIEWMODEPARAM_ACTION( Commands.RequiredTextureResolutionSingle[TextureIndex], TextureIndex );
 	}
+
+	BindShowCommands( CommandListRef );
+}
+
+void SEditorViewport::BindShowCommands( FUICommandList& OutCommandList )
+{
+	FShowFlagMenuCommands::Get().BindCommands(OutCommandList, Client);
 }
 
 EVisibility SEditorViewport::OnGetViewportContentVisibility() const
diff --git a/Engine/Source/Editor/UnrealEd/Public/SCommonEditorViewportToolbarBase.h b/Engine/Source/Editor/UnrealEd/Public/SCommonEditorViewportToolbarBase.h
--- a/Engine/Source/Editor/UnrealEd/Public/SCommonEditorViewportToolbarBase.h
+++ b/Engine/Source/Editor/UnrealEd/Public/SCommonEditorViewportToolbarBase.h
@@ -10,6 +10,10 @@
 #include "Framework/MultiBox/MultiBoxBuilder.h"
 #include "SViewportToolBar.h"
 #include "IPreviewProfileController.h"
+#include "UObject/Object.h"
+#include "UObject/ObjectMacros.h"
+
+#include "SCommonEditorViewportToolbarBase.generated.h"
 
 // This is the interface that the host of a SCommonEditorViewportToolbarBase must implement
 class ICommonEditorViewportToolbarInfoProvider
@@ -33,18 +37,21 @@ namespace CommonEditorViewportUtils
 		TSharedPtr<FUICommandInfo> ShowMenuItem;
 		FText LabelOverride;
 
+		UE_DEPRECATED(5.5, "Use the version of the show flags builder in FShowFlagMenuCommands")
 		FShowMenuCommand(TSharedPtr<FUICommandInfo> InShowMenuItem, const FText& InLabelOverride)
 			: ShowMenuItem(InShowMenuItem)
 			, LabelOverride(InLabelOverride)
 		{
 		}
 
+		UE_DEPRECATED(5.5, "Use the version of the show flags builder in FShowFlagMenuCommands")
 		FShowMenuCommand(TSharedPtr<FUICommandInfo> InShowMenuItem)
 			: ShowMenuItem(InShowMenuItem)
 		{
 		}
 	};
 
+	UE_DEPRECATED(5.5, "Use the version of the show flags builder in FShowFlagMenuCommands::BuildShowFlagsMenu which takes a UToolMenu instead")
 	static inline void FillShowMenu(class FMenuBuilder& MenuBuilder, TArray<FShowMenuCommand> MenuCommands, int32 EntryOffset)
 	{
 		// Generate entries for the standard show flags
@@ -60,6 +67,15 @@ namespace CommonEditorViewportUtils
 	}
 }
 
+UCLASS()
+class UNREALED_API UCommonViewportToolbarBaseMenuContext : public UObject
+{
+	GENERATED_BODY()
+
+public:
+	TWeakPtr<const class SCommonEditorViewportToolbarBase> ToolbarWidget;
+};
+
 /**
  * A viewport toolbar widget for an asset or level editor that is placed in a viewport
  */
@@ -179,14 +195,14 @@ protected:
 	/** Extension allowing derived classes to add to left-aligned portion of the toolbar slots.*/
 	virtual void ExtendLeftAlignedToolbarSlots(TSharedPtr<SHorizontalBox> MainBoxPtr, TSharedPtr<SViewportToolBar> ParentToolBarPtr) const {}
 
-protected:
+	UNREALED_API virtual void FillShowFlagsMenu(class UToolMenu* InMenu) const;
+
 	// Returns the info provider for this viewport
 	UNREALED_API ICommonEditorViewportToolbarInfoProvider& GetInfoProvider() const;
 
 	// Get the viewport client
 	UNREALED_API class FEditorViewportClient& GetViewportClient() const;
 
-protected:
 	// Creates the view menu widget (override point for children)
 	UNREALED_API virtual TSharedRef<class SEditorViewportViewMenu> MakeViewMenu();
 
@@ -210,7 +226,7 @@ protected:
 	/** Invoked when the asset viewer profile combo box selection changes. */
 	UNREALED_API void OnAssetViewerProfileComboBoxSelectionChanged(TSharedPtr<FString> NewSelection, ESelectInfo::Type /*SelectInfo*/);
 
-	/** Creates and returns the asset viewr profile combo box.*/
+	/** Creates and returns the asset viewer profile combo box.*/
 	UNREALED_API TSharedRef<SWidget> MakeAssetViewerProfileComboBox();
 
 private:
diff --git a/Engine/Source/Editor/UnrealEd/Public/SEditorViewport.h b/Engine/Source/Editor/UnrealEd/Public/SEditorViewport.h
--- a/Engine/Source/Editor/UnrealEd/Public/SEditorViewport.h
+++ b/Engine/Source/Editor/UnrealEd/Public/SEditorViewport.h
@@ -152,7 +152,8 @@ protected:
 	virtual void PopulateViewportOverlays(TSharedRef<SOverlay> Overlay) { }
 
 	UNREALED_API virtual void BindCommands();
-	virtual const FSlateBrush* OnGetViewportBorderBrush() const { return NULL; }
+	UNREALED_API virtual void BindShowCommands( FUICommandList& OutCommandList );
+	virtual const FSlateBrush* OnGetViewportBorderBrush() const { return nullptr; }
 	virtual FSlateColor OnGetViewportBorderColorAndOpacity() const { return FLinearColor::Black; }
 	
 	/**
diff --git a/Engine/Source/Editor/UnrealEd/UnrealEd.Build.cs b/Engine/Source/Editor/UnrealEd/UnrealEd.Build.cs
--- a/Engine/Source/Editor/UnrealEd/UnrealEd.Build.cs
+++ b/Engine/Source/Editor/UnrealEd/UnrealEd.Build.cs
@@ -92,7 +92,8 @@ public class UnrealEd : ModuleRules
 				"AssetTools", 
 				"WidgetRegistration",
 				"GameplayTasks",
-				"TargetPlatform"
+				"TargetPlatform",
+				"CommonMenuExtensions",
 			}
 		);
 
