AssetSizeQueryCommandlet:
-- Adding class filter
-- Adding optional/streaming/installed sizes

diff --git a/Engine/Source/Editor/UnrealEd/Private/Commandlets/AssetSizeQueryCommandlet.cpp b/Engine/Source/Editor/UnrealEd/Private/Commandlets/AssetSizeQueryCommandlet.cpp
--- a/Engine/Source/Editor/UnrealEd/Private/Commandlets/AssetSizeQueryCommandlet.cpp
+++ b/Engine/Source/Editor/UnrealEd/Private/Commandlets/AssetSizeQueryCommandlet.cpp
@@ -51,7 +51,8 @@ int32 UAssetSizeQueryCommandlet::Main(const FString& FullCommandLine)
 		UE_LOG(LogAssetSize, Display, TEXT("    -AssetRegistry=path                     Provides the path to the Development asset registry. This"));
 		UE_LOG(LogAssetSize, Display, TEXT("                                            asset registry must have staging size metadata via ProjectSettings->"));
 		UE_LOG(LogAssetSize, Display, TEXT("                                            Packaging->WriteBackMetadataToAssetRegistry."));
-		UE_LOG(LogAssetSize, Display, TEXT("    -Filter=wildcard                        (optional) Filteres the list of assets using a wildcard match."));
+		UE_LOG(LogAssetSize, Display, TEXT("    -Filter=wildcard                        (optional) Filters the list of assets using a wildcard match."));
+		UE_LOG(LogAssetSize, Display, TEXT("    -FilterClass=class                      (optional) Filters the list of assets to the given class. e.g. Texture2D not /Script/Engine.Texture2D"));
 		UE_LOG(LogAssetSize, Display, TEXT("    -Show=#                                 (optional) Shows only the top # classes, sorted on size (0 is all, default 10)."));
 		UE_LOG(LogAssetSize, Display, TEXT("    -CSV=path                               (optional) Output the filtered per class infomation to the given CSV file."));
 		UE_LOG(LogAssetSize, Display, TEXT("    -CSVType=(Assets,Classes)               (optional) Specifies whether to write the class summary or all matching assets to the csv file."));
@@ -68,9 +69,15 @@ int32 UAssetSizeQueryCommandlet::Main(const FString& FullCommandLine)
 		return 1;
 	}
 
+	UE_LOG(LogAssetSize, Display, TEXT("Using: %s"), *FileName);
+
 	FString AssetFilter;
 	FParse::Value(*FullCommandLine, TEXT("Filter="), AssetFilter);
 
+	FString ClassFilter;
+	FParse::Value(*FullCommandLine, TEXT("FilterClass="), ClassFilter);
+	FName ClassFilterName = ClassFilter.Len() ? FName(ClassFilter) : NAME_None;
+
 	int ShowCount = 10;
 	FString ShowCountString;
 	if (FParse::Value(*FullCommandLine, TEXT("Show="), ShowCountString))
@@ -131,7 +138,10 @@ int32 UAssetSizeQueryCommandlet::Main(const FString& FullCommandLine)
 	struct FMatchedAssetInfo
 	{
 		FSoftObjectPath ObjectPath;
-		int64 CompressedSize;
+		int64 CompressedSize = 0;
+		int64 OptionalSize = 0;
+		int64 InstalledSize = 0;
+		int64 StreamingSize = 0;
 	};
 	TMap<FTopLevelAssetPath /* AssetClass */, int64> FilteredClassCompressedSizes;
 	TMap<FTopLevelAssetPath /* AssetClass */, TArray<FMatchedAssetInfo>> FilteredClassMatchedAssets;
@@ -142,7 +152,8 @@ int32 UAssetSizeQueryCommandlet::Main(const FString& FullCommandLine)
 		 &FilteredClassMatchedAssets, 
 		 &FilteredCompressedSize, 
 		 &FilteredClassCompressedSizes, 
-		 &AssetFilter](const FAssetData& AssetData)
+		 &AssetFilter,
+		ClassFilterName](const FAssetData& AssetData)
 	{
 		FString CompressedSize;
 		if (AssetData.GetTagValue(UE::AssetRegistry::Stage_ChunkCompressedSizeFName, CompressedSize) == false ||
@@ -160,7 +171,12 @@ int32 UAssetSizeQueryCommandlet::Main(const FString& FullCommandLine)
 		{
 			FString ObjectPath = AssetData.GetObjectPathString();
 			bMatched = FWildcardString::IsMatchSubstring(*AssetFilter, *ObjectPath, *ObjectPath + ObjectPath.Len(), ESearchCase::IgnoreCase);
-		} // end if filter exists
+		}
+		if (bMatched &&
+			ClassFilterName != NAME_None)
+		{
+			bMatched = AssetData.AssetClassPath.GetAssetName() == ClassFilterName;
+		}
 
 		if (bMatched)
 		{
@@ -169,6 +185,10 @@ int32 UAssetSizeQueryCommandlet::Main(const FString& FullCommandLine)
 			FMatchedAssetInfo& Info = FilteredClassMatchedAssets.FindOrAdd(AssetData.AssetClassPath).AddDefaulted_GetRef();
 			Info.ObjectPath = AssetData.GetSoftObjectPath();
 			Info.CompressedSize = AssetCompressedSize;
+
+			AssetData.GetTagValue(UE::AssetRegistry::Stage_ChunkInstalledSizeFName, Info.InstalledSize);
+			AssetData.GetTagValue(UE::AssetRegistry::Stage_ChunkOptionalSizeFName, Info.OptionalSize);
+			AssetData.GetTagValue(UE::AssetRegistry::Stage_ChunkStreamingSizeFName, Info.StreamingSize);
 			
 			FilteredCompressedSize += AssetCompressedSize;
 			int64& FilteredClassCompressedSize = FilteredClassCompressedSizes.FindOrAdd(AssetData.AssetClassPath);
@@ -229,13 +249,13 @@ int32 UAssetSizeQueryCommandlet::Main(const FString& FullCommandLine)
 		}
 		else if (OutputCSVType == EOutputCSVType::Assets)
 		{
-			Lines.Add(TEXT("AssetName,AssetType,CompressedSize"));
+			Lines.Add(TEXT("AssetName,AssetType,CompressedSize,InstalledSize,OptionalSize,StreamingSize"));
 			for (TPair<FTopLevelAssetPath, TArray<FMatchedAssetInfo>>& ClassAssetsPair : FilteredClassMatchedAssets)
 			{
 				// we add to both maps at the same time to we know the lookup succeeds.
 				for (const FMatchedAssetInfo& AssetInfo : ClassAssetsPair.Value)
 				{
-					Lines.Add(FString::Printf(TEXT("%s,%s,%lld"), *AssetInfo.ObjectPath.ToString(), *ClassAssetsPair.Key.ToString(), AssetInfo.CompressedSize));
+					Lines.Add(FString::Printf(TEXT("%s,%s,%lld,%lld,%lld,%lld"), *AssetInfo.ObjectPath.ToString(), *ClassAssetsPair.Key.ToString(), AssetInfo.CompressedSize, AssetInfo.InstalledSize, AssetInfo.OptionalSize, AssetInfo.StreamingSize));
 				}
 			}
 		}
