Fix removal of viewmode ui command entries in 33047928

diff --git a/Engine/Source/Editor/UnrealEd/Private/SEditorViewport.cpp b/Engine/Source/Editor/UnrealEd/Private/SEditorViewport.cpp
--- a/Engine/Source/Editor/UnrealEd/Private/SEditorViewport.cpp
+++ b/Engine/Source/Editor/UnrealEd/Private/SEditorViewport.cpp
@@ -416,6 +416,7 @@ void SEditorViewport::BindCommands()
 	MAP_VIEWMODE_ACTION( Commands.WireframeMode, VMI_BrushWireframe );
 	MAP_VIEWMODE_ACTION( Commands.UnlitMode, VMI_Unlit );
 	MAP_VIEWMODE_ACTION( Commands.LitMode, VMI_Lit );
+	MAP_VIEWMODE_ACTION( Commands.LitWireframeMode, VMI_Lit_Wireframe);
 #if RHI_RAYTRACING
 	if (IsRayTracingAllowed())
 	{
@@ -457,6 +458,8 @@ void SEditorViewport::BindCommands()
 	MAP_VIEWMODE_ACTION( Commands.CollisionPawn, VMI_CollisionPawn);
 	MAP_VIEWMODE_ACTION( Commands.CollisionVisibility, VMI_CollisionVisibility);
 
+	MAP_VIEWMODE_ACTION( Commands.VisualizeLWCComplexity, VMI_LWCComplexity);
+
 	if (GEnableGPUSkinCache)
 	{
 		MAP_VIEWMODE_ACTION(Commands.VisualizeGPUSkinCacheMode, VMI_VisualizeGPUSkinCache);
