Fixed crash when Building Virtual Textures with the separate mobile SVT enabled

diff --git a/Engine/Source/Editor/UnrealEd/Private/EditorBuildUtils.cpp b/Engine/Source/Editor/UnrealEd/Private/EditorBuildUtils.cpp
--- a/Engine/Source/Editor/UnrealEd/Private/EditorBuildUtils.cpp
+++ b/Engine/Source/Editor/UnrealEd/Private/EditorBuildUtils.cpp
@@ -59,6 +59,9 @@
 #include "WorldPartition/WorldPartitionRuntimeVirtualTextureBuilder.h"
 #include "AssetCompilingManager.h"
 #include "ComponentRecreateRenderStateContext.h"
+#include "ShaderCompiler.h"
+#include "Materials/Material.h"
+#include "Materials/MaterialInstance.h"
 
 DEFINE_LOG_CATEGORY_STATIC(LogEditorBuildUtils, Log, All);
 
@@ -1753,46 +1756,68 @@ bool FEditorBuildUtils::EditorBuildVirtualTexture(UWorld* InWorld)
 		CollectGarbage(GARBAGE_COLLECTION_KEEPFLAGS);
 	};
 
+	FWorldPartitionHelpers::FForEachActorWithLoadingResult ForEachActorWithLoadingResult;
+	if (UWorldPartition* WorldPartition = InWorld->GetWorldPartition())
 	{
-		FWorldPartitionHelpers::FForEachActorWithLoadingResult ForEachActorWithLoadingResult;
-		if (UWorldPartition* WorldPartition = InWorld->GetWorldPartition())
-		{
-			FScopedSlowTask BuildTask(1.0f, LOCTEXT("VirtualTextureLoadActors", "Loading Actors"));
-			BuildTask.MakeDialog();
-			UWorldPartitionRuntimeVirtualTextureBuilder::LoadRuntimeVirtualTextureActors(WorldPartition, ForEachActorWithLoadingResult);
-		}
+		FScopedSlowTask BuildTask(1.0f, LOCTEXT("VirtualTextureLoadActors", "Loading Actors"));
+		BuildTask.MakeDialog();
+		UWorldPartitionRuntimeVirtualTextureBuilder::LoadRuntimeVirtualTextureActors(WorldPartition, ForEachActorWithLoadingResult);
+	}
 
-		// We will need to build VTs for both shading paths
-		const ERHIFeatureLevel::Type CurFeatureLevel = InWorld->GetFeatureLevel();
-		const ERHIFeatureLevel::Type AltFeatureLevel = (CurFeatureLevel == ERHIFeatureLevel::ES3_1 ? GMaxRHIFeatureLevel : ERHIFeatureLevel::ES3_1);
-		const EShadingPath CurShadingPath = FSceneInterface::GetShadingPath(CurFeatureLevel);
-		const EShadingPath AltShadingPath = FSceneInterface::GetShadingPath(AltFeatureLevel);
+	// We will need to build VTs for both shading paths
+	const ERHIFeatureLevel::Type CurFeatureLevel = InWorld->GetFeatureLevel();
+	const ERHIFeatureLevel::Type AltFeatureLevel = (CurFeatureLevel == ERHIFeatureLevel::ES3_1 ? GMaxRHIFeatureLevel : ERHIFeatureLevel::ES3_1);
+	const EShadingPath CurShadingPath = FSceneInterface::GetShadingPath(CurFeatureLevel);
+	const EShadingPath AltShadingPath = FSceneInterface::GetShadingPath(AltFeatureLevel);
 
-		TArray<URuntimeVirtualTextureComponent*> Components[2];
-		for (TObjectIterator<URuntimeVirtualTextureComponent> It; It; ++It)
+	TArray<URuntimeVirtualTextureComponent*> Components[2];
+	for (TObjectIterator<URuntimeVirtualTextureComponent> It; It; ++It)
+	{
+		if (Module->HasStreamedMips(CurShadingPath, *It))
 		{
-			if (Module->HasStreamedMips(CurShadingPath, *It))
-			{
-				Components[0].Add(*It);
-			}
-
-			if (Module->HasStreamedMips(AltShadingPath, *It))
-			{
-				Components[1].Add(*It);
-			}
+			Components[0].Add(*It);
 		}
-		
-		// Build for a current feature level first
-		if (!BuildVirtualTextureComponents(Module, CurShadingPath, Components[0]))
+
+		if (Module->HasStreamedMips(AltShadingPath, *It))
 		{
-			return false;
+			Components[1].Add(*It);
 		}
-		
-		// Build for others if any
-		bool bResult = true;
-		if (Components[1].Num() != 0)
+	}
+
+	// Build for a current feature level first
+	int32 NumFeatureLevelsToBuild = 0;
+	FScopedSlowTask BuildTask(static_cast<float>((Components[0].IsEmpty() ? 0 : 1) + (Components[1].IsEmpty() ? 0 : 1)));
+	BuildTask.MakeDialog(true);
+
+	auto EnterProgressForFeatureLevel = [&BuildTask](ERHIFeatureLevel::Type InFeatureLevel)
+	{
+		BuildTask.EnterProgressFrame(1.0f, FText::Format(LOCTEXT("VirtualTextureBuildFeatureLevel", "Building Virtual Textures for {0}"), FText::FromString(LexToString(InFeatureLevel))));
+	};
+
+	EnterProgressForFeatureLevel(CurFeatureLevel);
+	if (!BuildVirtualTextureComponents(Module, CurShadingPath, Components[0]))
+	{
+		return false;
+	}
+
+	// Build for the other feature level if any
+	bool bResult = true;
+	if (Components[1].Num() != 0)
+	{
+		EnterProgressForFeatureLevel(AltFeatureLevel);
+
+		FScopedSlowTask SubBuildTask(3.0f, BuildTask.GetCurrentMessage());
 		{
+			FText SwitchingFeatureLevelText = FText::Format(LOCTEXT("VirtualTextureSwitchToAltFeatureLevel", "Switching feature level to {0}"), FText::FromString(LexToString(AltFeatureLevel)));
+			SubBuildTask.EnterProgressFrame(1.0f, SwitchingFeatureLevelText);
+
+			UMaterialInterface::SetGlobalRequiredFeatureLevel(AltFeatureLevel, true);
+			UMaterial::AllMaterialsCacheResourceShadersForRendering(/*bUpdateProgressDialog = */true, /*bCacheAllRemainingShaders = */true);
+			UMaterialInstance::AllMaterialsCacheResourceShadersForRendering(/*bUpdateProgressDialog = */true, /*bCacheAllRemainingShaders = */true);
+			CompileGlobalShaderMap(AltFeatureLevel);
+
 			InWorld->ChangeFeatureLevel(AltFeatureLevel);
+
 			// Make sure all assets are finished compiling. Recreate render state after shader compilation complete
 			{
 				UMaterialInterface::SubmitRemainingJobsForWorld(InWorld);
@@ -1800,15 +1825,19 @@ bool FEditorBuildUtils::EditorBuildVirtualTexture(UWorld* InWorld)
 				FAssetCompilingManager::Get().ProcessAsyncTasks();
 				FGlobalComponentRecreateRenderStateContext Context;
 			}
-			bResult = BuildVirtualTextureComponents(Module, AltShadingPath, Components[1]);
 		}
-		
+
+		SubBuildTask.EnterProgressFrame(1.0f);
+		bResult = BuildVirtualTextureComponents(Module, AltShadingPath, Components[1]);
+
 		// Restore world feature level
-		InWorld->ChangeFeatureLevel(CurFeatureLevel);
-		return bResult;
+		{
+			SubBuildTask.EnterProgressFrame(1.0f, FText::Format(LOCTEXT("VirtualTextureSwitchToAltFeatureLevel", "Switching back feature level to {0}"), FText::FromString(LexToString(CurFeatureLevel))));
+			UMaterialInterface::SetGlobalRequiredFeatureLevel(CurFeatureLevel, /*bShouldCompile = */ false);
+			InWorld->ChangeFeatureLevel(CurFeatureLevel);
+		}
 	}
-	
-	return true;
+	return bResult;
 }
 
 void FEditorBuildUtils::EditorBuildAllLandscape(UWorld* InWorld)
diff --git a/Engine/Source/Editor/VirtualTexturingEditor/Private/RuntimeVirtualTextureBuildStreamingMips.cpp b/Engine/Source/Editor/VirtualTexturingEditor/Private/RuntimeVirtualTextureBuildStreamingMips.cpp
--- a/Engine/Source/Editor/VirtualTexturingEditor/Private/RuntimeVirtualTextureBuildStreamingMips.cpp
+++ b/Engine/Source/Editor/VirtualTexturingEditor/Private/RuntimeVirtualTextureBuildStreamingMips.cpp
@@ -161,6 +161,8 @@ namespace RuntimeVirtualTexture
 
 	bool BuildStreamedMips(EShadingPath ShadingPath, URuntimeVirtualTextureComponent* InComponent, FLinearColor const& FixedColor)
 	{
+		TRACE_CPUPROFILER_EVENT_SCOPE(BuildStreamedMips);
+
 		if (!HasStreamedMips(ShadingPath, InComponent))
 		{
 			return true;
