Upgrade to Preview Scene Settings to support "Preview" and "Edit" modes out-of-the-box

- Added concept of "default" engine profile (can be reset but not deleted)
- Added support for setting Show Flags from profile in animation viewports
- Added profile dropdown to animation viewport
- Skin paint tool now swaps to "Edit" mode and restores prior mode on Shutdown()

diff --git a/Engine/Plugins/Experimental/MeshModelingToolsetExp/Source/MeshModelingToolsEditorOnlyExp/MeshModelingToolsEditorOnlyExp.Build.cs b/Engine/Plugins/Experimental/MeshModelingToolsetExp/Source/MeshModelingToolsEditorOnlyExp/MeshModelingToolsEditorOnlyExp.Build.cs
--- a/Engine/Plugins/Experimental/MeshModelingToolsetExp/Source/MeshModelingToolsEditorOnlyExp/MeshModelingToolsEditorOnlyExp.Build.cs
+++ b/Engine/Plugins/Experimental/MeshModelingToolsetExp/Source/MeshModelingToolsEditorOnlyExp/MeshModelingToolsEditorOnlyExp.Build.cs
@@ -55,7 +55,8 @@ public class MeshModelingToolsEditorOnlyExp : ModuleRules
 				"RenderCore",
 				"InputCore",
 				"UnrealEd",
-				"Persona"
+				"Persona",
+				"AdvancedPreviewScene"
 				// ... add private dependencies that you statically link with here ...	
 			}
 		);
diff --git a/Engine/Plugins/Experimental/MeshModelingToolsetExp/Source/MeshModelingToolsEditorOnlyExp/Private/SkeletalMesh/SkinWeightsPaintTool.cpp b/Engine/Plugins/Experimental/MeshModelingToolsetExp/Source/MeshModelingToolsEditorOnlyExp/Private/SkeletalMesh/SkinWeightsPaintTool.cpp
--- a/Engine/Plugins/Experimental/MeshModelingToolsetExp/Source/MeshModelingToolsEditorOnlyExp/Private/SkeletalMesh/SkinWeightsPaintTool.cpp
+++ b/Engine/Plugins/Experimental/MeshModelingToolsetExp/Source/MeshModelingToolsEditorOnlyExp/Private/SkeletalMesh/SkinWeightsPaintTool.cpp
@@ -1,6 +1,8 @@
 // Copyright Epic Games, Inc. All Rights Reserved.
 
 #include "SkeletalMesh/SkinWeightsPaintTool.h"
+
+#include "AssetViewerSettings.h"
 #include "Engine/SkeletalMesh.h"
 #include "InteractiveToolManager.h"
 #include "SkeletalMeshAttributes.h"
@@ -22,7 +24,6 @@
 #include "DynamicMesh/MeshAdapterUtil.h"
 #include "Selection/PolygonSelectionMechanic.h"
 #include "Spatial/PointSetHashTable.h"
-#include "Util/ColorConstants.h"
 #include "Operations/SmoothBoneWeights.h"
 #include "ContextObjectStore.h"
 #include "Editor/Persona/Public/IPersonaEditorModeManager.h"
@@ -747,7 +748,6 @@ void USkinWeightsPaintTool::Setup()
 
 	// configure preview mesh
 	PreviewMesh->SetTangentsMode(EDynamicMeshComponentTangentsMode::AutoCalculated);
-	PreviewMesh->EnableWireframe(true);
 	PreviewMesh->SetShadowsEnabled(false);
 	// enable vtx colors on preview mesh
 	PreviewMesh->EditMesh([](FDynamicMesh3& Mesh)
@@ -764,6 +764,10 @@ void USkinWeightsPaintTool::Setup()
 		PreviewMesh->SetOverrideRenderMaterial(VtxColorMaterial);
 	}
 
+	// modify viewport render settings to optimize for painting weights
+	PreviewProfileToRestore = PreviewProfileController.GetActiveProfile();
+	PreviewProfileController.SetActiveProfile(UDefaultEditorProfiles::EditingProfileName.ToString());
+
 	// build octree for vertices
 	VerticesOctree.Initialize(PreviewMesh->GetMesh(), true);
 
@@ -855,7 +859,6 @@ void USkinWeightsPaintTool::Render(IToolsContextRenderAPI* RenderAPI)
 	{
 		PolygonSelectionMechanic->Render(RenderAPI);
 	}
-	
 }
 
 FBox USkinWeightsPaintTool::GetWorldSpaceFocusBox()
@@ -1547,6 +1550,9 @@ void USkinWeightsPaintTool::OnShutdown(EToolShutdownType ShutdownType)
 		GetToolManager()->EndUndoTransaction();
 	}
 
+	// restore viewport show flags and preview settings
+	PreviewProfileController.SetActiveProfile(PreviewProfileToRestore);
+
 	if (EditorContext.IsValid())
 	{
 		EditorContext->UnbindFrom(this);
diff --git a/Engine/Plugins/Experimental/MeshModelingToolsetExp/Source/MeshModelingToolsEditorOnlyExp/Public/SkeletalMesh/SkinWeightsPaintTool.h b/Engine/Plugins/Experimental/MeshModelingToolsetExp/Source/MeshModelingToolsEditorOnlyExp/Public/SkeletalMesh/SkinWeightsPaintTool.h
--- a/Engine/Plugins/Experimental/MeshModelingToolsetExp/Source/MeshModelingToolsEditorOnlyExp/Public/SkeletalMesh/SkinWeightsPaintTool.h
+++ b/Engine/Plugins/Experimental/MeshModelingToolsetExp/Source/MeshModelingToolsEditorOnlyExp/Public/SkeletalMesh/SkinWeightsPaintTool.h
@@ -2,12 +2,14 @@
 
 #pragma once
 
+#include "AssetViewerSettings.h"
 #include "DynamicMeshBrushTool.h"
 #include "BaseTools/MeshSurfacePointMeshEditingTool.h"
 #include "MeshDescription.h"
 #include "DynamicMesh/DynamicVerticesOctree3.h"
 #include "BoneWeights.h"
 #include "GroupTopology.h"
+#include "PreviewProfileController.h"
 #include "Misc/Optional.h"
 #include "Containers/Map.h"
 #include "DynamicMesh/DynamicMeshOctree3.h"
@@ -507,5 +509,9 @@ protected:
 	UPROPERTY()
 	TWeakObjectPtr<UPersonaEditorModeManagerContext> PersonaModeManagerContext = nullptr;
 
+	// editor state to restore when exiting the paint tool
+	FPreviewProfileController PreviewProfileController;
+	FString PreviewProfileToRestore;
+
 	friend SkinPaintTool::FSkinToolDeformer;
 };
diff --git a/Engine/Source/Editor/AdvancedPreviewScene/Private/AssetViewerSettings.cpp b/Engine/Source/Editor/AdvancedPreviewScene/Private/AssetViewerSettings.cpp
--- a/Engine/Source/Editor/AdvancedPreviewScene/Private/AssetViewerSettings.cpp
+++ b/Engine/Source/Editor/AdvancedPreviewScene/Private/AssetViewerSettings.cpp
@@ -5,6 +5,21 @@
 #include "Editor/EditorPerProjectUserSettings.h"
 #include "Editor.h"
 
+FName UDefaultEditorProfiles::DefaultProfileName = FName("Preview");
+FName UDefaultEditorProfiles::EditingProfileName = FName("Edit");
+
+const FPreviewSceneProfile* UDefaultEditorProfiles::GetProfile(const FString& ProfileName)
+{
+	for (const FPreviewSceneProfile& Profile : Profiles)
+	{
+		if (Profile.ProfileName == ProfileName)
+		{
+			return &Profile;
+		}
+	}
+
+	return nullptr;
+}
 
 UAssetViewerSettings::UAssetViewerSettings()
 {
@@ -44,6 +59,11 @@ UAssetViewerSettings* UAssetViewerSettings::Get()
 	return DefaultSettings;
 }
 
+FPreviewSceneProfile& UAssetViewerSettings::GetCurrentUserProjectProfile()
+{
+	return Get()->Profiles[GetMutableDefault<UEditorPerProjectUserSettings>()->AssetViewerProfileIndex];
+}
+
 void UAssetViewerSettings::Save(bool bWarnIfFail)
 {
 	ULocalProfiles* LocalProfilesObject = GetMutableDefault<ULocalProfiles>();
@@ -116,13 +136,30 @@ void UAssetViewerSettings::PostInitProperties()
 	{
 		Profiles.Add(Profile);
 	}
+	
+	TArray<FPreviewSceneProfile>& DefaultEditorProfiles = GetMutableDefault<UDefaultEditorProfiles>()->Profiles;
+	for (FPreviewSceneProfile& ProfileToAdd : DefaultEditorProfiles)
+	{
+		// add the default profile if it's not already stored
+		// default editor profiles should be marked as "shared" profiles to allow user overrides
+		// therefore, they are only added once and maintained by the user thereafter
+		if (!Profiles.ContainsByPredicate([&ProfileToAdd](const FPreviewSceneProfile& Profile)
+			{
+				return Profile.ProfileName == ProfileToAdd.ProfileName;
+			}))
+		{
+			Profiles.Add(ProfileToAdd);
+		}
+	}
 
-	if (Profiles.Num() == 0)
+	// the UDefaultEditorProfiles should always add at least one default profile
+	if (!ensure(!Profiles.IsEmpty()))
 	{
 		// Make sure there always is one profile as default
 		Profiles.AddDefaulted(1);
 		Profiles[0].ProfileName = TEXT("Profile_0");
 	}
+
 	NumProfiles = Profiles.Num();
 
 	UEditorPerProjectUserSettings* ProjectSettings = GetMutableDefault<UEditorPerProjectUserSettings>();
diff --git a/Engine/Source/Editor/AdvancedPreviewScene/Private/PreviewProfileController.cpp b/Engine/Source/Editor/AdvancedPreviewScene/Private/PreviewProfileController.cpp
--- a/Engine/Source/Editor/AdvancedPreviewScene/Private/PreviewProfileController.cpp
+++ b/Engine/Source/Editor/AdvancedPreviewScene/Private/PreviewProfileController.cpp
@@ -52,7 +52,7 @@ void FPreviewProfileController::UpdateAssetViewerProfiles()
 		// Rebuild the profile list.
 		for (const FPreviewSceneProfile& Profile : AssetViewerSettings->Profiles)
 		{
-			AssetViewerProfileNames.Add(Profile.ProfileName + (Profile.bSharedProfile ? TEXT(" (Shared)") : TEXT("")));
+			AssetViewerProfileNames.Add(Profile.ProfileName);
 		}
 
 		CurrentProfileIndex = PerProjectSettings->AssetViewerProfileIndex;
diff --git a/Engine/Source/Editor/AdvancedPreviewScene/Private/SAdvancedPreviewDetailsTab.cpp b/Engine/Source/Editor/AdvancedPreviewScene/Private/SAdvancedPreviewDetailsTab.cpp
--- a/Engine/Source/Editor/AdvancedPreviewScene/Private/SAdvancedPreviewDetailsTab.cpp
+++ b/Engine/Source/Editor/AdvancedPreviewScene/Private/SAdvancedPreviewDetailsTab.cpp
@@ -115,9 +115,23 @@ void SAdvancedPreviewDetailsTab::Construct(const FArguments& InArgs, const TShar
 			.AutoWidth()
 			[
 				SNew(SButton)
-				.OnClicked(this, &SAdvancedPreviewDetailsTab::RemoveProfileButtonClick)
-				.Text(LOCTEXT("RemoveProfileButton", "Remove Profile"))
-				.ToolTipText(LOCTEXT("SceneProfileRemoveProfile", "Removes the currently selected profile."))
+				.OnClicked(this, &SAdvancedPreviewDetailsTab::RemoveOrResetProfileButtonClick)
+				.Text_Lambda([this]()
+				{
+					if (DefaultSettings->Profiles[ProfileIndex].bIsEngineDefaultProfile)
+					{
+						return LOCTEXT("ResetProfileButton", "Reset Profile");
+					}
+					return LOCTEXT("RemoveProfileButton", "Remove Profile");
+				})
+				.ToolTipText_Lambda([this]()
+				{
+					if (DefaultSettings->Profiles[ProfileIndex].bIsEngineDefaultProfile)
+					{
+						return LOCTEXT("SceneProfileResetProfile", "Resets this engine profile to default settings. Cannot delete engine profiles.");
+					}
+					return LOCTEXT("SceneProfileRemoveProfile", "Removes the currently selected profile.");
+				})
 				.IsEnabled_Lambda([this]()->bool
 				{
 					return ProfileNames.Num() > 1; 
@@ -130,7 +144,7 @@ void SAdvancedPreviewDetailsTab::Construct(const FArguments& InArgs, const TShar
 	UpdateSettingsView();
 }
 
-void SAdvancedPreviewDetailsTab::ComboBoxSelectionChanged(TSharedPtr<FString> NewSelection, ESelectInfo::Type /*SelectInfo*/)
+void SAdvancedPreviewDetailsTab::ComboBoxSelectionChanged(TSharedPtr<FString> NewSelection, ESelectInfo::Type SelectInfo)
 {
 	int32 NewSelectionIndex;
 
@@ -140,7 +154,10 @@ void SAdvancedPreviewDetailsTab::ComboBoxSelectionChanged(TSharedPtr<FString> Ne
 		PerProjectSettings->AssetViewerProfileIndex = ProfileIndex;
 		UpdateSettingsView();	
 		check(PreviewScenePtr.IsValid());
-		PreviewScenePtr.Pin()->SetProfileIndex(ProfileIndex);
+		if (SelectInfo == ESelectInfo::Type::OnMouseClick)
+		{
+			PreviewScenePtr.Pin()->SetProfileIndex(ProfileIndex);	
+		}
 	}
 }
 
@@ -162,7 +179,16 @@ void SAdvancedPreviewDetailsTab::UpdateProfileNames()
 	ProfileNames.Empty();
 	for (FPreviewSceneProfile& Profile : DefaultSettings->Profiles)
 	{
-		ProfileNames.Add(TSharedPtr<FString>(new FString(Profile.ProfileName + (Profile.bSharedProfile ? TEXT(" (Shared)") : TEXT("") ))));
+		FString Suffix = TEXT("");
+		if (Profile.bSharedProfile)
+		{
+			Suffix = TEXT(" (Shared)");
+		}
+		if (Profile.bIsEngineDefaultProfile)
+		{
+			Suffix = TEXT(" (Engine Default)");
+		}
+		ProfileNames.Add(TSharedPtr<FString>(new FString(Profile.ProfileName + Suffix)));
 	}
 
 	ProfileComboBox->RefreshOptions();
@@ -214,8 +240,28 @@ FReply SAdvancedPreviewDetailsTab::AddProfileButtonClick()
 	return FReply::Handled();
 }
 
-FReply SAdvancedPreviewDetailsTab::RemoveProfileButtonClick()
+FReply SAdvancedPreviewDetailsTab::RemoveOrResetProfileButtonClick()
 {
+	const FPreviewSceneProfile& CurrentProfile = DefaultSettings->Profiles[ProfileIndex];
+	
+	if (CurrentProfile.bIsEngineDefaultProfile)
+	{
+		const FScopedTransaction Transaction(LOCTEXT("ResetSceneProfile", "Reset Preview Scene Profile"));
+		DefaultSettings->Modify();
+
+		const FPreviewSceneProfile* DefaultEditorProfile = GetMutableDefault<UDefaultEditorProfiles>()->GetProfile(CurrentProfile.ProfileName);
+		if (DefaultEditorProfile)
+		{
+			// Reset currently selected profile 
+			DefaultSettings->Profiles[ProfileIndex] = *DefaultEditorProfile;
+			DefaultSettings->PostEditChange();
+			return FReply::Handled();
+		}
+
+		// if we get here, it means an engine-provided default profile was removed from the engine,
+		// in which case we should remove it...
+	}
+
 	const FScopedTransaction Transaction(LOCTEXT("RemoveSceneProfile", "Remove Preview Scene Profile"));
 	DefaultSettings->Modify();
 
diff --git a/Engine/Source/Editor/AdvancedPreviewScene/Public/AssetViewerSettings.h b/Engine/Source/Editor/AdvancedPreviewScene/Public/AssetViewerSettings.h
--- a/Engine/Source/Editor/AdvancedPreviewScene/Public/AssetViewerSettings.h
+++ b/Engine/Source/Editor/AdvancedPreviewScene/Public/AssetViewerSettings.h
@@ -10,6 +10,7 @@
 #include "Engine/Scene.h"
 #include "Engine/TextureCube.h"
 #include "EditorUndoClient.h"
+#include "ShowFlags.h"
 #include "AssetViewerSettings.generated.h"
 
 /**
@@ -23,6 +24,7 @@ struct FPreviewSceneProfile
 	FPreviewSceneProfile()
 	{		
 		bSharedProfile = false;
+		bIsEngineDefaultProfile = false;
 		bUseSkyLighting = true;
 		bShowFloor = true;
 		bShowEnvironment = true;
@@ -39,6 +41,8 @@ struct FPreviewSceneProfile
 		EnvironmentCubeMapPath = TEXT("/Engine/EditorMaterials/AssetViewer/EpicQuadPanorama_CC+EV1.EpicQuadPanorama_CC+EV1");
 		bPostProcessingEnabled = true;
 		DirectionalLightRotation = FRotator(-40.f, -67.5f, 0.f);
+		bEnableToneMapping = true;
+		bShowMeshEdges = false;
 	}
 
 	/** Name to identify the profile */
@@ -48,7 +52,11 @@ struct FPreviewSceneProfile
 	/** Whether or not this profile should be stored in the Project ini file */
 	UPROPERTY(EditAnywhere, config, Category = Profile)
 	bool bSharedProfile;
-
+	
+	/** Whether or not this profile is one of the default profiles included with the engine */
+	UPROPERTY(config)
+	bool bIsEngineDefaultProfile;
+	
 	/** Whether or not image based lighting is enabled for the environment cube map */
 	UPROPERTY(EditAnywhere, AdvancedDisplay, config, Category = Lighting)
 	bool bUseSkyLighting;
@@ -77,6 +85,10 @@ struct FPreviewSceneProfile
 	UPROPERTY(EditAnywhere, config, Category = Environment)
 	bool bShowFloor;
 
+	/** Toggle visibility of floor grid on/off */
+	UPROPERTY(EditAnywhere, config, Category = Environment)
+	bool bShowGrid;
+
 	/** The environment color, used if Show Environment is false. */
 	UPROPERTY(EditAnywhere, config, Category = Environment, meta=(EditCondition="!bShowEnvironment"))
 	FLinearColor EnvironmentColor;
@@ -113,6 +125,14 @@ struct FPreviewSceneProfile
 	UPROPERTY(config)
 	FRotator DirectionalLightRotation;
 
+	/** Useful when editing in an unlit view, prevents colors from being adjusted by the tonemapping */
+	UPROPERTY(EditAnywhere, config, Category = Editing)
+	bool bEnableToneMapping;
+
+	/** Show wireframes composited on top of the shaded view */
+	UPROPERTY(EditAnywhere, config, Category = Editing)
+	bool bShowMeshEdges;
+
 	/** Retrieve the environment map texture using the saved path */
 	void LoadEnvironmentMap()
 	{
@@ -131,8 +151,64 @@ struct FPreviewSceneProfile
 			}
 		}		
 	}
+
+	void SetShowFlags(FEngineShowFlags& ShowFlags) const
+	{
+		if (bPostProcessingEnabled)
+		{
+			ShowFlags.EnableAdvancedFeatures();
+			ShowFlags.SetBloom(true); // bloom not included in EnableAdvancedFeatures() for thumbnails (see func comments)
+		}
+		else
+		{
+			ShowFlags.DisableAdvancedFeatures();
+			ShowFlags.SetBloom(false);
+		}
+
+		ShowFlags.SetTonemapper(bEnableToneMapping);
+		ShowFlags.SetGrid(bShowGrid);
+		ShowFlags.SetMeshEdges(bShowMeshEdges);
+	}
 };
 
+UCLASS()
+class ADVANCEDPREVIEWSCENE_API UDefaultEditorProfiles : public UObject
+{
+	GENERATED_BODY()
+public:
+
+	UDefaultEditorProfiles()
+	{
+		FPreviewSceneProfile DefaultProfile;
+		DefaultProfile.bIsEngineDefaultProfile = true;
+		DefaultProfile.bSharedProfile = true;
+		DefaultProfile.ProfileName = DefaultProfileName.ToString();
+		Profiles.Add(DefaultProfile);
+		
+		FPreviewSceneProfile EditingProfile;
+		EditingProfile.ProfileName = EditingProfileName.ToString();
+		EditingProfile.bIsEngineDefaultProfile = true;
+		EditingProfile.bSharedProfile = true;
+		EditingProfile.bShowEnvironment = false;
+		EditingProfile.bShowFloor = false;
+		EditingProfile.bShowGrid = true;
+		EditingProfile.EnvironmentColor = FLinearColor::MakeFromHSV8(0,0,10);
+		EditingProfile.bUseSkyLighting = true;
+		EditingProfile.bPostProcessingEnabled = false;
+		EditingProfile.bShowMeshEdges = true;
+		EditingProfile.bEnableToneMapping = false;
+		Profiles.Add(EditingProfile);
+	}
+
+	const FPreviewSceneProfile* GetProfile(const FString& ProfileName);
+
+	static FName DefaultProfileName;
+	static FName EditingProfileName;
+	
+	/** Collection of default engine-provided profiles used in various editing environments*/
+	UPROPERTY()
+	TArray<FPreviewSceneProfile> Profiles;
+};
 
 UCLASS(config = Editor)
 class ULocalProfiles : public UObject
@@ -166,6 +242,8 @@ public:
 	virtual ~UAssetViewerSettings();
 
 	static UAssetViewerSettings* Get();
+
+	static FPreviewSceneProfile& GetCurrentUserProjectProfile();
 	
 	/**
 	 * Saves the config data out to the ini files
diff --git a/Engine/Source/Editor/AdvancedPreviewScene/Public/SAdvancedPreviewDetailsTab.h b/Engine/Source/Editor/AdvancedPreviewScene/Public/SAdvancedPreviewDetailsTab.h
--- a/Engine/Source/Editor/AdvancedPreviewScene/Public/SAdvancedPreviewDetailsTab.h
+++ b/Engine/Source/Editor/AdvancedPreviewScene/Public/SAdvancedPreviewDetailsTab.h
@@ -51,7 +51,7 @@ protected:
 	void UpdateSettingsView();
 	void UpdateProfileNames();
 	FReply AddProfileButtonClick();
-	FReply RemoveProfileButtonClick();
+	FReply RemoveOrResetProfileButtonClick();
 protected:
 	void OnAssetViewerSettingsRefresh(const FName& InPropertyName);
 	void OnAssetViewerSettingsPostUndo();
diff --git a/Engine/Source/Editor/DetailCustomizations/Private/AssetViewerSettingsCustomization.cpp b/Engine/Source/Editor/DetailCustomizations/Private/AssetViewerSettingsCustomization.cpp
--- a/Engine/Source/Editor/DetailCustomizations/Private/AssetViewerSettingsCustomization.cpp
+++ b/Engine/Source/Editor/DetailCustomizations/Private/AssetViewerSettingsCustomization.cpp
@@ -85,8 +85,11 @@ void FAssetViewerSettingsCustomization::CustomizeDetails(IDetailLayoutBuilder& D
 					[
 						// PropertyHandle->CreatePropertyValueWidget()
 						SAssignNew(NameEditTextBox, SEditableTextBox)
+						.IsEnabled_Lambda([this]()
+						{
+							return !ViewerSettings->Profiles[ProfileIndex].bIsEngineDefaultProfile;
+						})
 						.Text(this, &FAssetViewerSettingsCustomization::OnGetProfileName)
-						//.ToolTip(VarNameTooltip)
 						.OnTextChanged(this, &FAssetViewerSettingsCustomization::OnProfileNameChanged)
 						.OnTextCommitted(this, &FAssetViewerSettingsCustomization::OnProfileNameCommitted)
 						.Font(IDetailLayoutBuilder::GetDetailFont())
@@ -169,7 +172,9 @@ const bool FAssetViewerSettingsCustomization::IsProfileNameValid(const FString&
 
 bool FAssetViewerSettingsCustomization::CanSetSharedProfile() const
 {
-	return !FPlatformFileManager::Get().GetPlatformFile().IsReadOnly(*SharedProfileConfigFilePath());
+	const bool bIsEngineDefaultProfile = ViewerSettings->Profiles[ProfileIndex].bIsEngineDefaultProfile;
+	const bool bIsConfigWriteable = !FPlatformFileManager::Get().GetPlatformFile().IsReadOnly(*SharedProfileConfigFilePath());
+	return bIsConfigWriteable && !bIsEngineDefaultProfile;
 }
 
 EVisibility FAssetViewerSettingsCustomization::ShowFileWatcherWidget() const
diff --git a/Engine/Source/Editor/Persona/Private/AnimationEditorViewportClient.cpp b/Engine/Source/Editor/Persona/Private/AnimationEditorViewportClient.cpp
--- a/Engine/Source/Editor/Persona/Private/AnimationEditorViewportClient.cpp
+++ b/Engine/Source/Editor/Persona/Private/AnimationEditorViewportClient.cpp
@@ -226,8 +226,7 @@ void FAnimationViewportClient::Initialize()
 	// Register delegate to update the show flags when the post processing is turned on or off
 	UAssetViewerSettings::Get()->OnAssetViewerSettingsChanged().AddSP(this, &FAnimationViewportClient::OnAssetViewerSettingsChanged);
 	// Set correct flags according to current profile settings
-	SetAdvancedShowFlagsForScene(UAssetViewerSettings::Get()->Profiles[GetMutableDefault<UEditorPerProjectUserSettings>()->AssetViewerProfileIndex].bPostProcessingEnabled);
-
+	UAssetViewerSettings::GetCurrentUserProjectProfile().SetShowFlags(EngineShowFlags);
 }
 
 void FAnimationViewportClient::OnToggleAutoAlignFloor()
@@ -2451,23 +2450,11 @@ void FAnimationViewportClient::OnAssetViewerSettingsChanged(const FName& InPrope
 		const int32 ProfileIndex = GetPreviewScene()->GetCurrentProfileIndex();
 		if (Settings->Profiles.IsValidIndex(ProfileIndex))
 		{			
-			SetAdvancedShowFlagsForScene(Settings->Profiles[ProfileIndex].bPostProcessingEnabled);
+			Settings->Profiles[ProfileIndex].SetShowFlags(EngineShowFlags);
 		}
 	}
 }
 
-void FAnimationViewportClient::SetAdvancedShowFlagsForScene(const bool bAdvancedShowFlags)
-{	
-	if (bAdvancedShowFlags)
-	{
-		EngineShowFlags.EnableAdvancedFeatures();
-	}
-	else
-	{
-		EngineShowFlags.DisableAdvancedFeatures();
-	}
-}
-
 void FAnimationViewportClient::SetPlaybackSpeedMode(EAnimationPlaybackSpeeds::Type InMode)
 {
 	AnimationPlaybackSpeedMode = InMode;
diff --git a/Engine/Source/Editor/Persona/Private/SAnimViewportToolBar.cpp b/Engine/Source/Editor/Persona/Private/SAnimViewportToolBar.cpp
--- a/Engine/Source/Editor/Persona/Private/SAnimViewportToolBar.cpp
+++ b/Engine/Source/Editor/Persona/Private/SAnimViewportToolBar.cpp
@@ -51,6 +51,7 @@
 #include "Widgets/Text/SRichTextBlock.h"
 #include "ContentBrowserModule.h"
 #include "IContentBrowserSingleton.h"
+#include "SCommonEditorViewportToolbarBase.h"
 #include "Animation/MirrorDataTable.h"
 #include "ScopedTransaction.h"
 #include "SNameComboBox.h"
@@ -255,6 +256,7 @@ protected:
 
 void SAnimViewportToolBar::Construct(const FArguments& InArgs, TSharedPtr<class SAnimationEditorViewportTabBody> InViewport, TSharedPtr<class SEditorViewport> InRealViewport)
 {
+	PreviewProfileController = InArgs._PreviewProfileController;
 	bShowShowMenu = InArgs._ShowShowMenu;
 	bShowCharacterMenu= InArgs._ShowCharacterMenu;
 	bShowLODMenu = InArgs._ShowLODMenu;
@@ -326,6 +328,12 @@ void SAnimViewportToolBar::Construct(const FArguments& InArgs, TSharedPtr<class
 		+ SHorizontalBox::Slot()
 		.AutoWidth()
 		.Padding(ToolbarSlotPadding)
+		[
+			SNew(SPreviewSceneProfileSelector).PreviewProfileController(InArgs._PreviewProfileController)
+		]
+		+ SHorizontalBox::Slot()
+		.AutoWidth()
+		.Padding(ToolbarSlotPadding)
 		[
 			SNew(SEditorViewportToolbarMenu)
 			.ToolTipText(LOCTEXT("PhysicsMenuTooltip", "Physics Options. Use this to control the physics of the scene."))
diff --git a/Engine/Source/Editor/Persona/Private/SAnimViewportToolBar.h b/Engine/Source/Editor/Persona/Private/SAnimViewportToolBar.h
--- a/Engine/Source/Editor/Persona/Private/SAnimViewportToolBar.h
+++ b/Engine/Source/Editor/Persona/Private/SAnimViewportToolBar.h
@@ -15,7 +15,7 @@
 
 class FMenuBuilder;
 class SComboButton;
-
+class IPreviewProfileController;
 
 /**
  * A level viewport toolbar widget that is placed in a viewport
@@ -24,7 +24,8 @@ class SAnimViewportToolBar : public SViewportToolBar
 {
 public:
 	SLATE_BEGIN_ARGS( SAnimViewportToolBar )
-		: _ShowShowMenu(true)
+		: _PreviewProfileController(nullptr)
+		, _ShowShowMenu(true)
 		, _ShowCharacterMenu(true)
 		, _ShowLODMenu(true)
 		, _ShowPlaySpeedMenu(true)
@@ -35,6 +36,8 @@ public:
 
 	SLATE_ARGUMENT(FName, ContextName)
 
+	SLATE_ARGUMENT(TSharedPtr<IPreviewProfileController>, PreviewProfileController)
+
 	SLATE_ARGUMENT(bool, ShowShowMenu)
 
 	SLATE_ARGUMENT(bool, ShowCharacterMenu)
@@ -213,6 +216,9 @@ private:
 	/** Transaction to handle scoping updates for sliders we own that transact objects (i.e. floor mesh) */
 	TUniquePtr<FScopedTransaction> PendingTransaction;
 
+	/** The controller used to toggle Preview Scene Profiles */
+	TSharedPtr<IPreviewProfileController> PreviewProfileController;
+
 	/** Whether to show the 'Show' menu */
 	bool bShowShowMenu;
 
diff --git a/Engine/Source/Editor/Persona/Private/SAnimationEditorViewport.cpp b/Engine/Source/Editor/Persona/Private/SAnimationEditorViewport.cpp
--- a/Engine/Source/Editor/Persona/Private/SAnimationEditorViewport.cpp
+++ b/Engine/Source/Editor/Persona/Private/SAnimationEditorViewport.cpp
@@ -29,15 +29,13 @@
 #include "BufferVisualizationMenuCommands.h"
 #include "UICommandList_Pinnable.h"
 #include "IPersonaEditorModeManager.h"
-#include "AssetViewerSettings.h"
-#include "Editor/EditorPerProjectUserSettings.h"
+#include "PreviewProfileController.h"
 #include "Materials/Material.h"
 #include "EditorFontGlyphs.h"
 #include "EdModeInteractiveToolsContext.h"
 #include "ContextObjectStore.h"
 #include "IPersonaEditMode.h"
 
-#include "SkeletalMeshTypes.h"
 #include "IPersonaToolkit.h"
 #include "Subsystems/AssetEditorSubsystem.h"
 #include "SNameComboBox.h"
@@ -141,6 +139,7 @@ TSharedPtr<SWidget> SAnimationEditorViewport::MakeViewportToolbar()
 		.Cursor(EMouseCursor::Default)
 		.Extenders(Extenders)
 		.ContextName(ContextName)
+		.PreviewProfileController(MakeShared<FPreviewProfileController>())
 		.ShowShowMenu(bShowShowMenu)
 		.ShowLODMenu(bShowLODMenu)
 		.ShowPlaySpeedMenu(bShowPlaySpeedMenu)
diff --git a/Engine/Source/Editor/Persona/Public/AnimationEditorViewportClient.h b/Engine/Source/Editor/Persona/Public/AnimationEditorViewportClient.h
--- a/Engine/Source/Editor/Persona/Public/AnimationEditorViewportClient.h
+++ b/Engine/Source/Editor/Persona/Public/AnimationEditorViewportClient.h
@@ -474,9 +474,6 @@ private:
 
 	/** Delegate for preview profile is changed (used for updating show flags) */
 	void OnAssetViewerSettingsChanged(const FName& InPropertyName);
-
-	/** Sets up the ShowFlag according to the current preview scene profile */
-	void SetAdvancedShowFlagsForScene(const bool bAdvancedShowFlags);
 	
 	/** Computes a bounding box for the selected section of the preview mesh component.
 	    If there is no selected section, returns an empty box. */
diff --git a/Engine/Source/Editor/UnrealEd/Private/SCommonEditorViewportToolbarBase.cpp b/Engine/Source/Editor/UnrealEd/Private/SCommonEditorViewportToolbarBase.cpp
--- a/Engine/Source/Editor/UnrealEd/Private/SCommonEditorViewportToolbarBase.cpp
+++ b/Engine/Source/Editor/UnrealEd/Private/SCommonEditorViewportToolbarBase.cpp
@@ -32,9 +32,9 @@
 #define LOCTEXT_NAMESPACE "SCommonEditorViewportToolbarBase"
 
 //////////////////////////////////////////////////////////////////////////
-// SCommonEditorViewportToolbarBase
+// SPreviewSceneProfileSelector
 
-SCommonEditorViewportToolbarBase::~SCommonEditorViewportToolbarBase()
+SPreviewSceneProfileSelector::~SPreviewSceneProfileSelector()
 {
 	if (PreviewProfileController)
 	{
@@ -43,21 +43,99 @@
 	}
 }
 
+void SPreviewSceneProfileSelector::Construct(const FArguments& InArgs)
+{
+	PreviewProfileController = InArgs._PreviewProfileController;
+	if (PreviewProfileController)
+	{
+		PreviewProfileController->OnPreviewProfileListChanged().AddRaw(this, &SPreviewSceneProfileSelector::UpdateAssetViewerProfileList);
+		PreviewProfileController->OnPreviewProfileChanged().AddRaw(this, &SPreviewSceneProfileSelector::UpdateAssetViewerProfileSelection);
+		UpdateAssetViewerProfileList();
+	}
+	
+	ChildSlot
+	[
+		SNew(SVerticalBox)
+		+SVerticalBox::Slot()
+		.AutoHeight()
+		[
+			SAssignNew(AssetViewerProfileComboBox, STextComboBox)//SComboBox<TSharedPtr<FString>>)
+			.OptionsSource(&AssetViewerProfileNames)
+			.ButtonStyle(&FAppStyle::Get().GetWidgetStyle<FButtonStyle>("EditorViewportToolBar.Button"))
+			.ContentPadding(FMargin(2, 0))
+			.ToolTipText(LOCTEXT("AssetViewerProfile_ToolTip", "Select the Preview Scene Profile for this viewport."))
+			.OnSelectionChanged(this, &SPreviewSceneProfileSelector::OnSelectionChanged)
+			.Visibility_Lambda([this]() { return AssetViewerProfileNames.Num() > 1 ? EVisibility::Visible : EVisibility::Collapsed; })
+		]
+	];
 
+	UpdateAssetViewerProfileSelection();
+}
+
+void SPreviewSceneProfileSelector::OnSelectionChanged(TSharedPtr<FString> NewSelection, ESelectInfo::Type SelectInfo)
+{
+	int32 NewSelectionIndex;
+	if (AssetViewerProfileNames.Find(NewSelection, NewSelectionIndex))
+	{
+		// If that's the user changing the combo box, not an update coming from code to reflect a change that already occurred.
+		if (SelectInfo != ESelectInfo::Direct)
+		{
+			PreviewProfileController->SetActiveProfile(*NewSelection);
+		}
+	}
+}
+
+void SPreviewSceneProfileSelector::UpdateAssetViewerProfileList()
+{
+	if (PreviewProfileController)
+	{
+		// Pull the latest profile list.
+		int32 CurrProfileIndex = 0;
+		TArray<FString> ProfileNames = PreviewProfileController->GetPreviewProfiles(CurrProfileIndex);
+
+		// Rebuild the combo box list.
+		AssetViewerProfileNames.Empty();
+		for (const FString& Profile : ProfileNames)
+		{
+			AssetViewerProfileNames.Add(MakeShared<FString>(Profile));
+		}
+
+		// Select the current profile item.
+		if (AssetViewerProfileComboBox)
+		{
+			AssetViewerProfileComboBox->RefreshOptions();
+			AssetViewerProfileComboBox->SetSelectedItem(AssetViewerProfileNames[CurrProfileIndex]);
+		}
+	}
+}
+
+void SPreviewSceneProfileSelector::UpdateAssetViewerProfileSelection()
+{
+	if (PreviewProfileController)
+	{
+		FString ActiveProfileName = PreviewProfileController->GetActiveProfile();
+		if (TSharedPtr<FString>* Match = AssetViewerProfileNames.FindByPredicate(
+			[&ActiveProfileName](const TSharedPtr<FString>& Candidate) { return *Candidate == ActiveProfileName; }))
+		{
+			AssetViewerProfileComboBox->SetSelectedItem(*Match);
+		}
+		else // The profile was likely renamed.
+		{
+			UpdateAssetViewerProfileList();
+		}
+	}
+}
+
+//////////////////////////////////////////////////////////////////////////
+// SCommonEditorViewportToolbarBase
+
+
 void SCommonEditorViewportToolbarBase::Construct(const FArguments& InArgs, TSharedPtr<class ICommonEditorViewportToolbarInfoProvider> InInfoProvider)
 {
 	InfoProviderPtr = InInfoProvider;
-	PreviewProfileController = InArgs._PreviewProfileController;
 
 	TSharedRef<SEditorViewport> ViewportRef = GetInfoProvider().GetViewportWidget();
 	TSharedPtr<SHorizontalBox> MainBoxPtr;
 
-	if (PreviewProfileController)
-	{
-		PreviewProfileController->OnPreviewProfileListChanged().AddRaw(this, &SCommonEditorViewportToolbarBase::UpdateAssetViewerProfileList);
-		PreviewProfileController->OnPreviewProfileChanged().AddRaw(this, &SCommonEditorViewportToolbarBase::UpdateAssetViewerProfileSelection);
-		UpdateAssetViewerProfileList();
-	}
-
 	const FMargin ToolbarSlotPadding(4.0f, 1.0f);
 	const FMargin ToolbarButtonPadding(4.0f, 0.0f);
@@ -123,7 +201,7 @@
 			.AutoWidth()
 			.Padding(ToolbarSlotPadding)
 			[
-				MakeAssetViewerProfileComboBox()
+				SNew(SPreviewSceneProfileSelector).PreviewProfileController(InArgs._PreviewProfileController)
 			];
 	}
 
@@ -462,79 +540,6 @@
 	MenuBuilder.EndSection();
 }
 
-void SCommonEditorViewportToolbarBase::UpdateAssetViewerProfileList()
-{
-	if (PreviewProfileController)
-	{
-		// Pull the latest profile list.
-		int32 CurrProfileIndex = 0;
-		TArray<FString> ProfileNames = PreviewProfileController->GetPreviewProfiles(CurrProfileIndex);
-
-		// Rebuild the combo box list.
-		AssetViewerProfileNames.Empty();
-		for (const FString& Profile : ProfileNames)
-		{
-			AssetViewerProfileNames.Add(MakeShared<FString>(Profile));
-		}
-
-		// Select the current profile item.
-		if (AssetViewerProfileComboBox)
-		{
-			AssetViewerProfileComboBox->RefreshOptions();
-			AssetViewerProfileComboBox->SetSelectedItem(AssetViewerProfileNames[CurrProfileIndex]);
-		}
-	}
-}
-
-void SCommonEditorViewportToolbarBase::UpdateAssetViewerProfileSelection()
-{
-	if (PreviewProfileController)
-	{
-		FString ActiveProfileName = PreviewProfileController->GetActiveProfile();
-		if (TSharedPtr<FString>* Match = AssetViewerProfileNames.FindByPredicate(
-			[&ActiveProfileName](const TSharedPtr<FString>& Candidate) { return *Candidate == ActiveProfileName; }))
-		{
-			AssetViewerProfileComboBox->SetSelectedItem(*Match);
-		}
-		else // The profile was likely renamed.
-		{
-			UpdateAssetViewerProfileList();
-		}
-	}
-}
-
-void SCommonEditorViewportToolbarBase::OnAssetViewerProfileComboBoxSelectionChanged(TSharedPtr<FString> NewSelection, ESelectInfo::Type SelectInfo)
-{
-	int32 NewSelectionIndex;
-	if (AssetViewerProfileNames.Find(NewSelection, NewSelectionIndex))
-	{
-		// If that's the user changing the combo box, not an update coming from code to reflect a change that already occurred.
-		if (SelectInfo != ESelectInfo::Direct)
-		{
-			PreviewProfileController->SetActiveProfile(*NewSelection);
-		}
-	}
-}
-
-TSharedRef<SWidget> SCommonEditorViewportToolbarBase::MakeAssetViewerProfileComboBox()
-{
-	AssetViewerProfileComboBox = SNew(STextComboBox)
-		.OptionsSource(&AssetViewerProfileNames)
-		.ButtonStyle(&FAppStyle::Get().GetWidgetStyle<FButtonStyle>("EditorViewportToolBar.Button"))
-		.ContentPadding(FMargin(2, 0))
-		.ToolTipText(LOCTEXT("AssetViewerProfile_ToolTip", "Changes the asset viewer profile"))
-		.OnSelectionChanged(this, &SCommonEditorViewportToolbarBase::OnAssetViewerProfileComboBoxSelectionChanged)
-		.Visibility_Lambda([this]() { return AssetViewerProfileNames.Num() > 1 ? EVisibility::Visible : EVisibility::Collapsed; });
-
-	AssetViewerProfileComboBox->RefreshOptions();
-	if (!AssetViewerProfileNames.IsEmpty())
-	{
-		AssetViewerProfileComboBox->SetSelectedItem(AssetViewerProfileNames[0]);
-	}
-
-	return AssetViewerProfileComboBox.ToSharedRef();
-}
-
 FText SCommonEditorViewportToolbarBase::GetCameraMenuLabel() const
 {
 	return GetCameraMenuLabelFromViewportType( GetViewportClient().GetViewportType() );
diff --git a/Engine/Source/Editor/UnrealEd/Public/SCommonEditorViewportToolbarBase.h b/Engine/Source/Editor/UnrealEd/Public/SCommonEditorViewportToolbarBase.h
--- a/Engine/Source/Editor/UnrealEd/Public/SCommonEditorViewportToolbarBase.h
+++ b/Engine/Source/Editor/UnrealEd/Public/SCommonEditorViewportToolbarBase.h
@@ -76,6 +76,42 @@ public:
 	TWeakPtr<const class SCommonEditorViewportToolbarBase> ToolbarWidget;
 };
 
+class UNREALED_API SPreviewSceneProfileSelector : public SCompoundWidget
+{
+public:
+
+	virtual ~SPreviewSceneProfileSelector() override;
+	
+	SLATE_BEGIN_ARGS(SPreviewSceneProfileSelector)
+	{}
+	SLATE_ARGUMENT(TSharedPtr<IPreviewProfileController>, PreviewProfileController)
+	SLATE_END_ARGS()
+	
+	void Construct(const FArguments& InArgs);
+
+protected:
+	
+	/** Update the list of asset viewer profiles displayed by the combo box. */
+	void UpdateAssetViewerProfileList();
+	void UpdateAssetViewerProfileSelection();
+
+	/** Invoked when the asset viewer profile combo box selection changes. */
+	void OnSelectionChanged(TSharedPtr<FString> NewSelection, ESelectInfo::Type /*SelectInfo*/);
+
+	/** Creates and returns the asset viewer profile combo box.*/
+	TSharedRef<SWidget> MakeAssetViewerProfileComboBox();
+
+private:
+	/** Interface to set/get/list the preview profiles. */
+	TSharedPtr<IPreviewProfileController> PreviewProfileController;
+
+	/** List of advanced preview profiles to fill up the Profiles combo box. */
+	TArray<TSharedPtr<FString>> AssetViewerProfileNames;
+
+	/** Displays/Selects the active advanced viewer profile. */
+	TSharedPtr<STextComboBox> AssetViewerProfileComboBox;
+};
+
 /**
  * A viewport toolbar widget for an asset or level editor that is placed in a viewport
  */
@@ -91,7 +127,7 @@ public:
 		SLATE_ARGUMENT(TSharedPtr<IPreviewProfileController>, PreviewProfileController) // Should be null if the Preview doesn't require profile.
 	SLATE_END_ARGS()
 
-	UNREALED_API virtual ~SCommonEditorViewportToolbarBase();
+	UNREALED_API virtual ~SCommonEditorViewportToolbarBase(){};
 
 	UNREALED_API void Construct(const FArguments& InArgs, TSharedPtr<class ICommonEditorViewportToolbarInfoProvider> InInfoProvider);
 
@@ -219,27 +255,8 @@ protected:
 	/** Called when the ScreenPercentage slider is adjusted in the viewport */
 	UNREALED_API void OnScreenPercentageValueChanged(int32 NewValue);
 
-	/** Update the list of asset viewer profiles displayed by the combo box. */
-	UNREALED_API void UpdateAssetViewerProfileList();
-	UNREALED_API void UpdateAssetViewerProfileSelection();
-
-	/** Invoked when the asset viewer profile combo box selection changes. */
-	UNREALED_API void OnAssetViewerProfileComboBoxSelectionChanged(TSharedPtr<FString> NewSelection, ESelectInfo::Type /*SelectInfo*/);
-
-	/** Creates and returns the asset viewer profile combo box.*/
-	UNREALED_API TSharedRef<SWidget> MakeAssetViewerProfileComboBox();
-
 private:
 	/** The viewport that we are in */
 	TWeakPtr<class ICommonEditorViewportToolbarInfoProvider> InfoProviderPtr;
-
-	/** Interface to set/get/list the preview profiles. */
-	TSharedPtr<IPreviewProfileController> PreviewProfileController;
-
-	/** List of advanced preview profiles to fill up the Profiles combo box. */
-	TArray<TSharedPtr<FString>> AssetViewerProfileNames;
-
-	/** Displays/Selects the active advanced viewer profile. */
-	TSharedPtr<STextComboBox> AssetViewerProfileComboBox;
 };
 
