Cooker: Fix tracking of loads of external actor packages. External actor packages are renamed to _Trash when their map packages are cook-saved, so we need to record their name on creation, and not try to record the name at a later point that might occur after the cook-save of their map.
#rnx

diff --git a/Engine/Source/Editor/UnrealEd/Private/CookOnTheFlyServer.cpp b/Engine/Source/Editor/UnrealEd/Private/CookOnTheFlyServer.cpp
--- a/Engine/Source/Editor/UnrealEd/Private/CookOnTheFlyServer.cpp
+++ b/Engine/Source/Editor/UnrealEd/Private/CookOnTheFlyServer.cpp
@@ -4335,15 +4335,14 @@ void UCookOnTheFlyServer::ProcessUnsolicitedPackages(TArray<FName>* OutDiscovere
 
 	using namespace UE::Cook;
 
-	TMap<UPackage*, UE::Cook::FInstigator> NewPackages = PackageTracker->GetNewPackages();
+	TMap<FName, UE::Cook::FInstigator> NewPackages = PackageTracker->GetNewPackages();
 
-	for (auto& PackageWithInstigator : NewPackages)
+	for (TPair<FName, UE::Cook::FInstigator>& PackageWithInstigator : NewPackages)
 	{
-		UPackage* Package = PackageWithInstigator.Key;
-		check(Package != nullptr);
+		FName PackageName = PackageWithInstigator.Key;
 		FInstigator& Instigator = PackageWithInstigator.Value;
 
-		UE::Cook::FPackageData* PackageData = PackageDatas->TryAddPackageDataByPackageName(Package->GetFName());
+		UE::Cook::FPackageData* PackageData = PackageDatas->TryAddPackageDataByPackageName(PackageName);
 		if (!PackageData)
 		{
 			continue; // Getting the PackageData will fail if e.g. it is a script package
diff --git a/Engine/Source/Editor/UnrealEd/Private/Cooker/PackageTracker.cpp b/Engine/Source/Editor/UnrealEd/Private/Cooker/PackageTracker.cpp
--- a/Engine/Source/Editor/UnrealEd/Private/Cooker/PackageTracker.cpp
+++ b/Engine/Source/Editor/UnrealEd/Private/Cooker/PackageTracker.cpp
@@ -82,7 +82,7 @@ void FPackageTracker::InitializeTracking()
 	NewPackages.Reserve(LoadedPackages.Num());
 	for (UPackage* Package : LoadedPackages)
 	{
-		NewPackages.Add(Package, FInstigator(EInstigator::StartupPackage));
+		NewPackages.Add(Package->GetFName(), FInstigator(EInstigator::StartupPackage));
 	}
 
 	GUObjectArray.AddUObjectDeleteListener(this);
@@ -91,7 +91,7 @@ void FPackageTracker::InitializeTracking()
 	bTrackingInitialized = true;
 }
 
-TMap<UPackage*, FInstigator> FPackageTracker::GetNewPackages()
+TMap<FName, FInstigator> FPackageTracker::GetNewPackages()
 {
 	if (!bTrackingInitialized)
 	{
@@ -100,7 +100,7 @@ TMap<UPackage*, FInstigator> FPackageTracker::GetNewPackages()
 	}
 
 	FWriteScopeLock ScopeLock(Lock);
-	TMap<UPackage*, FInstigator> Result = MoveTemp(NewPackages);
+	TMap<FName, FInstigator> Result = MoveTemp(NewPackages);
 	NewPackages.Reset();
 	return Result;
 }
@@ -118,7 +118,7 @@ void FPackageTracker::MarkLoadedPackagesAsNew()
 	NewPackages.Reserve(LoadedPackages.Num());
 	for (UPackage* Package : LoadedPackages)
 	{
-		NewPackages.FindOrAdd(Package, FInstigator(EInstigator::StartupPackage));
+		NewPackages.FindOrAdd(Package->GetFName(), FInstigator(EInstigator::StartupPackage));
 	}
 }
 
@@ -134,8 +134,10 @@ void FPackageTracker::NotifyUObjectCreated(const class UObjectBase* Object, int3
 #if ENABLE_COOK_STATS
 			++DetailedCookStats::NumDetectedLoads;
 #endif
+			FName PackageName = Package->GetFName();
 #if UE_WITH_PACKAGE_ACCESS_TRACKING
-			PackageAccessTracking_Private::FTrackedData* AccumulatedScopeData = PackageAccessTracking_Private::FPackageAccessRefScope::GetCurrentThreadAccumulatedData();
+			PackageAccessTracking_Private::FTrackedData* AccumulatedScopeData =
+				PackageAccessTracking_Private::FPackageAccessRefScope::GetCurrentThreadAccumulatedData();
 			FName ReferencerName(AccumulatedScopeData ? AccumulatedScopeData->PackageName : NAME_None);
 #else
 			FName ReferencerName(NAME_None);
@@ -156,17 +158,22 @@ void FPackageTracker::NotifyUObjectCreated(const class UObjectBase* Object, int3
 			FInstigator Instigator(InstigatorType, ReferencerName);
 			if (InstigatorType == EInstigator::Unsolicited && COTFS.bHiddenDependenciesDebug)
 			{
-				COTFS.OnDiscoveredPackageDebug(Package->GetFName(), Instigator);
+				COTFS.OnDiscoveredPackageDebug(PackageName, Instigator);
 			}
 
 			FWriteScopeLock ScopeLock(Lock);
-			if (ExpectedNeverLoadPackages.Contains(Package->GetFName()))
+			if (ExpectedNeverLoadPackages.Contains(PackageName))
 			{
-				UE_LOG(LogCook, Verbose, TEXT("SoftGC PoorPerformance: Reloaded package %s."), *WriteToString<256>(Package->GetFName()));
+				UE_LOG(LogCook, Verbose, TEXT("SoftGC PoorPerformance: Reloaded package %s."),
+					*WriteToString<256>(PackageName));
 			}
-			LoadedPackages.Add(Package);
-			NewPackages.Add(Package, MoveTemp(Instigator));
 
+			LoadedPackages.Add(Package);
+			// We store packages by name rather than by pointer, because they might have their name changed. When
+			// external actors are moved out of their external package, we rename the package to <PackageName>_Trash.
+			// We want to report a load dependency on the package as it was originally loaded; we don't want to report
+			// the renamed packagename if it gets renamed after load.
+			NewPackages.Add(PackageName, MoveTemp(Instigator));
 		}
 	}
 }
@@ -179,7 +186,6 @@ void FPackageTracker::NotifyUObjectDeleted(const class UObjectBase* Object, int3
 
 		FWriteScopeLock ScopeLock(Lock);
 		LoadedPackages.Remove(Package);
-		NewPackages.Remove(Package);
 	}
 }
 
diff --git a/Engine/Source/Editor/UnrealEd/Private/Cooker/PackageTracker.h b/Engine/Source/Editor/UnrealEd/Private/Cooker/PackageTracker.h
--- a/Engine/Source/Editor/UnrealEd/Private/Cooker/PackageTracker.h
+++ b/Engine/Source/Editor/UnrealEd/Private/Cooker/PackageTracker.h
@@ -177,7 +177,7 @@ public:
 	~FPackageTracker();
 
 	/** Returns all packages that have been loaded since the last time GetNewPackages was called */
-	TMap<UPackage*, FInstigator> GetNewPackages();
+	TMap<FName, FInstigator> GetNewPackages();
 
 	/**
 	 * Copy all LoadedPackages into NewPackages. Called when reachability of all packages needs to be recalculated
@@ -240,7 +240,7 @@ private:
 	TSet<FName> ExpectedNeverLoadPackages;
 
 	// This list contains the UPackages loaded since last call to GetNewPackages
-	TMap<UPackage*, FInstigator> NewPackages;
+	TMap<FName, FInstigator> NewPackages;
 	bool bTrackingInitialized = false;
 };
 
