Fixed localization issue

diff --git a/Engine/Source/Editor/UnrealEd/Private/EditorBuildUtils.cpp b/Engine/Source/Editor/UnrealEd/Private/EditorBuildUtils.cpp
--- a/Engine/Source/Editor/UnrealEd/Private/EditorBuildUtils.cpp
+++ b/Engine/Source/Editor/UnrealEd/Private/EditorBuildUtils.cpp
@@ -1832,7 +1832,7 @@ bool FEditorBuildUtils::EditorBuildVirtualTexture(UWorld* InWorld)
 
 		// Restore world feature level
 		{
-			SubBuildTask.EnterProgressFrame(1.0f, FText::Format(LOCTEXT("VirtualTextureSwitchToAltFeatureLevel", "Switching back feature level to {0}"), FText::FromString(LexToString(CurFeatureLevel))));
+			SubBuildTask.EnterProgressFrame(1.0f, FText::Format(LOCTEXT("VirtualTextureSwitchBackToFeatureLevel", "Switching back feature level to {0}"), FText::FromString(LexToString(CurFeatureLevel))));
 			UMaterialInterface::SetGlobalRequiredFeatureLevel(CurFeatureLevel, /*bShouldCompile = */ false);
 			InWorld->ChangeFeatureLevel(CurFeatureLevel);
 		}
