- Remove UWorld specific behaviour from ObjectTools
- Deprecate old signature for OnAssetAddExtraDeletedAssets
- Made external package discovery (for deletion) recursive

diff --git a/Engine/Source/Editor/UnrealEd/Classes/Editor/EditorEngine.h b/Engine/Source/Editor/UnrealEd/Classes/Editor/EditorEngine.h
--- a/Engine/Source/Editor/UnrealEd/Classes/Editor/EditorEngine.h
+++ b/Engine/Source/Editor/UnrealEd/Classes/Editor/EditorEngine.h
@@ -3079,6 +3079,8 @@ protected:
 	TObjectPtr<class UActorGroupingUtils> ActorGroupingUtils;
 private:
 	FTimerHandle CleanupPIEOnlineSessionsTimerHandle;
+	
+	FDelegateHandle WorldAddExtraDeletionObjectsHandle;
 
 	/** Delegate handle for game viewport close requests in PIE sessions. */
 	FDelegateHandle ViewportCloseRequestedDelegateHandle;
diff --git a/Engine/Source/Editor/UnrealEd/Private/Editor.cpp b/Engine/Source/Editor/UnrealEd/Private/Editor.cpp
--- a/Engine/Source/Editor/UnrealEd/Private/Editor.cpp
+++ b/Engine/Source/Editor/UnrealEd/Private/Editor.cpp
@@ -154,6 +154,7 @@ FSimpleMulticastDelegate								FEditorDelegates::OnShutdownPostPackagesSaved;
 FEditorDelegates::FOnPackageDeleted						FEditorDelegates::OnPackageDeleted;
 FEditorDelegates::FOnAssetsCanDelete					FEditorDelegates::OnAssetsCanDelete;
 FEditorDelegates::FOnAssetsAddExtraObjectsToDelete		FEditorDelegates::OnAssetsAddExtraObjectsToDelete;
+FEditorDelegates::FOnAddExtraObjectsToDelete			FEditorDelegates::OnAddExtraObjectsToDelete;
 FEditorDelegates::FOnAssetsPreDelete					FEditorDelegates::OnAssetsPreDelete;
 FEditorDelegates::FOnAssetsDeleted						FEditorDelegates::OnAssetsDeleted;
 FEditorDelegates::FOnAssetDragStarted					FEditorDelegates::OnAssetDragStarted;
diff --git a/Engine/Source/Editor/UnrealEd/Private/EditorEngine.cpp b/Engine/Source/Editor/UnrealEd/Private/EditorEngine.cpp
--- a/Engine/Source/Editor/UnrealEd/Private/EditorEngine.cpp
+++ b/Engine/Source/Editor/UnrealEd/Private/EditorEngine.cpp
@@ -1074,6 +1074,8 @@ void UEditorEngine::InitEditor(IEngineLoop* InEngineLoop)
 	}
 
 	FAssetCompilingManager::Get().OnAssetPostCompileEvent().AddUObject(this, &UEditorEngine::OnAssetPostCompile);
+
+	WorldAddExtraDeletionObjectsHandle = FEditorDelegates::OnAddExtraObjectsToDelete.AddStatic(&UWorld::OnAddExtraObjectsToDelete);
 }
 
 bool UEditorEngine::HandleOpenAsset(UObject* Asset)
@@ -1525,6 +1527,7 @@ void UEditorEngine::FinishDestroy()
 		}
 		FAssetCompilingManager::Get().OnAssetPostCompileEvent().RemoveAll(this);
 
+		FEditorDelegates::OnAddExtraObjectsToDelete.Remove(WorldAddExtraDeletionObjectsHandle);
 
 		// Shut down transaction tracking system.
 		if( Trans )
diff --git a/Engine/Source/Editor/UnrealEd/Private/ObjectTools.cpp b/Engine/Source/Editor/UnrealEd/Private/ObjectTools.cpp
--- a/Engine/Source/Editor/UnrealEd/Private/ObjectTools.cpp
+++ b/Engine/Source/Editor/UnrealEd/Private/ObjectTools.cpp
@@ -2703,20 +2703,24 @@ namespace ObjectTools
 
 	void AddExtraObjectsToDelete(TArray< UObject* >& ObjectsToDelete)
 	{
-		const int32 OriginalNum = ObjectsToDelete.Num();
-		for (int32 i=0; i < OriginalNum; ++i)
+		// Allows to inject extra assets to delete without modifying the engine source.
+		TSet<UObject*> SecondaryObjects;
+		FEditorDelegates::OnAddExtraObjectsToDelete.Broadcast(ObjectsToDelete, SecondaryObjects);
+		for (UObject* Object : SecondaryObjects)
 		{
-			UObject* ObjectToDelete = ObjectsToDelete[i];
+			ObjectsToDelete.AddUnique(Object);
+		}
 
-			// Delete MapBuildData with maps & owned packages for map
-			if (UWorld* World = Cast<UWorld>(ObjectToDelete))
+		// Recursively include external packages
+		const int32 OriginalNum = ObjectsToDelete.Num();
+		TSet<const UPackage*> ProcessedOuterPackages;
+		for (int32 Index=0; Index < OriginalNum; ++Index)
+		{
+			const UObject* ObjectToDelete = ObjectsToDelete[Index];
+			const UPackage* OuterPackage = ObjectToDelete->GetPackage();
+			if (!ProcessedOuterPackages.Contains(OuterPackage))
 			{
-				if (World->PersistentLevel && World->PersistentLevel->MapBuildData)
-				{
-					ObjectsToDelete.AddUnique(World->PersistentLevel->MapBuildData);
-				}
-
-				for (UPackage* Package : World->GetOutermost()->GetExternalPackages())
+				for (UPackage* Package : OuterPackage->GetExternalPackages())
 				{
 					// Don't include newly created packages
 					if (!Package->HasAnyPackageFlags(PKG_NewlyCreated))
@@ -2724,16 +2728,11 @@ namespace ObjectTools
 						ObjectsToDelete.AddUnique(Package);
 					}
 				}
+				
+				ProcessedOuterPackages.Add(OuterPackage);
 			}
+			
 		}
-
-		// Allows to inject extra assets to delete without modifying the engine source.
-		FEditorDelegates::OnAssetsAddExtraObjectsToDelete.Broadcast(ObjectsToDelete);
-
-		//This method is called 2x in the deletion flow. Make sure there is no duplicates in the array as we can't rely on the methods registered to the delegate to uniquely add.
-		TSet<UObject*> CleanupDuplicatesSet(MoveTemp(ObjectsToDelete)); // Move items into the set to remove duplicate pointers.
-		ObjectsToDelete = CleanupDuplicatesSet.Array(); // Copy elements back again
-
 	}
 
 	bool ContainsWorldInUse(const TArray< UObject* >& ObjectsToDelete)
diff --git a/Engine/Source/Editor/UnrealEd/Public/Editor.h b/Engine/Source/Editor/UnrealEd/Public/Editor.h
--- a/Engine/Source/Editor/UnrealEd/Public/Editor.h
+++ b/Engine/Source/Editor/UnrealEd/Public/Editor.h
@@ -135,7 +135,11 @@ struct FEditorDelegates
 	/** delegate type for when a user requests to delete certain package */
 	DECLARE_MULTICAST_DELEGATE_OneParam(FOnPackageDeleted, UPackage*);
 	/** delegate type for when a user requests to delete certain assets... It allows the addition of secondary assets that should also be deleted */
+	UE_DEPRECATED(5.5, "FOnAssetsAddExtraObjectsToDelete has been deprecated, please use FOnAddExtraObjectsToDelete instead")
 	DECLARE_MULTICAST_DELEGATE_OneParam(FOnAssetsAddExtraObjectsToDelete, TArray<UObject*>&);
+	/** delegate type for when a user requests to delete certain assets... It allows the addition of secondary assets that should also be deleted */
+	DECLARE_MULTICAST_DELEGATE_TwoParams(FOnAddExtraObjectsToDelete, const TArray<UObject*>&, TSet<UObject*>&);
+	
 	/** delegate type for when a user requests to delete certain assets... DOES NOT mean the asset(s) will be deleted (the user could cancel) */
 	DECLARE_MULTICAST_DELEGATE_OneParam(FOnAssetsPreDelete, const TArray<UObject*>&);
 	/** delegate type for when a user requested force deleting objects. The objects(s) will be deleted (no possibility to cancel), so implementations should delete references */
@@ -323,7 +327,11 @@ struct FEditorDelegates
 	/** Called when the user requests assets to be deleted to determine if the operation is available.  */
 	static UNREALED_API FOnAssetsCanDelete OnAssetsCanDelete;
 	/** Called when the user requests certain assets be deletedand  allows the addition of secondary assets that should also be deleted */
+	PRAGMA_DISABLE_DEPRECATION_WARNINGS
+	UE_DEPRECATED(5.5, "Delegate signature for OnAssetsAddExtraObjectsToDelete has been deprecated please use new version OnAddExtraObjectsToDelete")
 	static UNREALED_API FOnAssetsAddExtraObjectsToDelete OnAssetsAddExtraObjectsToDelete;
+	PRAGMA_ENABLE_DEPRECATION_WARNINGS	
+	static UNREALED_API FOnAddExtraObjectsToDelete OnAddExtraObjectsToDelete;	
 	/** Called when the user requests certain assets be deleted (DOES NOT imply that the asset will be deleted... the user could cancel) */
 	static UNREALED_API FOnAssetsPreDelete OnAssetsPreDelete;
 	/** Called when one or more assets have been deleted */
diff --git a/Engine/Source/Runtime/Engine/Classes/Engine/World.h b/Engine/Source/Runtime/Engine/Classes/Engine/World.h
--- a/Engine/Source/Runtime/Engine/Classes/Engine/World.h
+++ b/Engine/Source/Runtime/Engine/Classes/Engine/World.h
@@ -1389,6 +1389,8 @@ public:
 	/** Restore the purged editor world FScene back to the proper GPU representation */
 	void RestoreScene();
 
+	static void OnAddExtraObjectsToDelete(const TArray<UObject*>& InObjectsToDelete, TSet<UObject*>& OutSecondaryObjects);
+
 #endif // WITH_EDITOR
 
 	/**
diff --git a/Engine/Source/Runtime/Engine/Classes/SparseVolumeTexture/SparseVolumeTexture.h b/Engine/Source/Runtime/Engine/Classes/SparseVolumeTexture/SparseVolumeTexture.h
--- a/Engine/Source/Runtime/Engine/Classes/SparseVolumeTexture/SparseVolumeTexture.h
+++ b/Engine/Source/Runtime/Engine/Classes/SparseVolumeTexture/SparseVolumeTexture.h
@@ -569,7 +569,9 @@ public:
 	//~ End USparseVolumeTexture Interface.
 
 #if WITH_EDITOR
+	UE_DEPRECATED(5.5, "OnAssetsAddExtraObjectsToDelete signature has been deprecated please use OnAddExtraObjectsToDelete instead")
 	ENGINE_API void OnAssetsAddExtraObjectsToDelete(TArray<UObject*>& ObjectsToDelete);
+	ENGINE_API void OnAddExtraObjectsToDelete(const TArray<UObject*>& InObjectsToDelete, TSet<UObject*>& OutSecondaryObjects);
 #endif
 
 protected:
diff --git a/Engine/Source/Runtime/Engine/Private/SparseVolumeTexture/SparseVolumeTexture.cpp b/Engine/Source/Runtime/Engine/Private/SparseVolumeTexture/SparseVolumeTexture.cpp
--- a/Engine/Source/Runtime/Engine/Private/SparseVolumeTexture/SparseVolumeTexture.cpp
+++ b/Engine/Source/Runtime/Engine/Private/SparseVolumeTexture/SparseVolumeTexture.cpp
@@ -1686,6 +1686,22 @@ void UStreamableSparseVolumeTexture::OnAssetsAddExtraObjectsToDelete(TArray<UObj
 		}
 	}
 }
+
+void UStreamableSparseVolumeTexture::OnAddExtraObjectsToDelete(const TArray<UObject*>& InObjectsToDelete, TSet<UObject*>& OutSecondaryObjects)
+{
+	if (InObjectsToDelete.Contains(this))
+	{
+		// When UStreamableSparseVolumeTexture is deleted, we also want all owned USparseVolumeTextureFrame objects to be deleted.
+		for (USparseVolumeTextureFrame* Frame : Frames)
+		{
+			if (Frame)
+			{
+				OutSecondaryObjects.Add(Frame);
+			}
+		}
+	}
+	
+}
 #endif
 
 #if WITH_EDITORONLY_DATA
@@ -1755,14 +1771,14 @@ void UStreamableSparseVolumeTexture::RegisterEditorDelegates()
 	if (ShouldRegisterDelegates())
 	{
 		UnregisterEditorDelegates();
-		FEditorDelegates::OnAssetsAddExtraObjectsToDelete.AddUObject(this, &UStreamableSparseVolumeTexture::OnAssetsAddExtraObjectsToDelete);
+		FEditorDelegates::OnAddExtraObjectsToDelete.AddUObject(this, &UStreamableSparseVolumeTexture::OnAddExtraObjectsToDelete);
 	}
 }
 void UStreamableSparseVolumeTexture::UnregisterEditorDelegates()
 {
 	if (ShouldRegisterDelegates())
 	{
-		FEditorDelegates::OnAssetsAddExtraObjectsToDelete.RemoveAll(this);
+		FEditorDelegates::OnAddExtraObjectsToDelete.RemoveAll(this);
 	}
 }
 #endif // WITH_EDITOR
diff --git a/Engine/Source/Runtime/Engine/Private/World.cpp b/Engine/Source/Runtime/Engine/Private/World.cpp
--- a/Engine/Source/Runtime/Engine/Private/World.cpp
+++ b/Engine/Source/Runtime/Engine/Private/World.cpp
@@ -9179,6 +9179,21 @@ void UWorld::RestoreScene()
 	}
 }
 
+void UWorld::OnAddExtraObjectsToDelete(const TArray<UObject*>& InObjectsToDelete, TSet<UObject*>& OutSecondaryObjects)
+{	
+	for (const UObject* Object : InObjectsToDelete)
+	{
+		if (const UWorld* World = Cast<UWorld>(Object))
+		{
+			if (World->PersistentLevel && World->PersistentLevel->MapBuildData)
+			{
+				// Delete MapBuildData together with maps
+				OutSecondaryObjects.Add(World->PersistentLevel->MapBuildData);
+			}
+		}
+	}
+}
+
 void UWorld::GetAssetRegistryTags(TArray<FAssetRegistryTag>& OutTags) const
 {
 	PRAGMA_DISABLE_DEPRECATION_WARNINGS;
