Cooker: Tweak full cook/incremental cook messages to reduce confusion with the keeping/clearing of DDC data.
#rnx

diff --git a/Engine/Source/Editor/UnrealEd/Private/CookOnTheFlyServer.cpp b/Engine/Source/Editor/UnrealEd/Private/CookOnTheFlyServer.cpp
--- a/Engine/Source/Editor/UnrealEd/Private/CookOnTheFlyServer.cpp
+++ b/Engine/Source/Editor/UnrealEd/Private/CookOnTheFlyServer.cpp
@@ -8124,7 +8124,7 @@ bool UCookOnTheFlyServer::ArePreviousCookSettingsCompatible(const TMap<FName, FS
 	const FConfigSection* CookSettings = ConfigFile.FindSection(TEXT_CookSettings);
 	if (CookSettings == nullptr)
 	{
-		UE_LOG(LogCook, Display, TEXT("Cook invalidated for CookSettings file %s is invalid. Clearing all cooked content."),
+		UE_LOG(LogCook, Display, TEXT("Cook invalidated for CookSettings file %s is invalid. Clearing previously cooked packages."),
 			*GetCookSettingsFileName(TargetPlatform));
 		return false;
 	}
@@ -8142,7 +8142,7 @@ bool UCookOnTheFlyServer::ArePreviousCookSettingsCompatible(const TMap<FName, FS
 		const FConfigValue* PreviousSetting = CookSettings->Find(CurrentSetting.Key);
 		if (!PreviousSetting || PreviousSetting->GetValue() != CurrentSetting.Value)
 		{
-			UE_LOG(LogCook, Display, TEXT("Cook invalidated for platform %s because %s has changed. Old: %s, New: %s. Clearing all cooked content."),
+			UE_LOG(LogCook, Display, TEXT("Cook invalidated for platform %s because %s has changed. Old: %s, New: %s. Clearing previously cooked packages."),
 				*TargetPlatform->PlatformName(), *CurrentSetting.Key.ToString(),
 				PreviousSetting ? *PreviousSetting->GetValue() : TEXT(""),
 				*CurrentSetting.Value);
@@ -8152,7 +8152,7 @@ bool UCookOnTheFlyServer::ArePreviousCookSettingsCompatible(const TMap<FName, FS
 
 	if (!bIterativeIgnoreIni && !bHybridIterativeEnabled && IniSettingsOutOfDate(TargetPlatform))
 	{
-		UE_LOG(LogCook, Display, TEXT("Cook invalidated for platform %s because ini settings have changed. Clearing all cooked content."),
+		UE_LOG(LogCook, Display, TEXT("Cook invalidated for platform %s because ini settings have changed. Clearing previously cooked packages."),
 			*TargetPlatform->PlatformName());
 		return false;
 	}
@@ -8162,7 +8162,7 @@ bool UCookOnTheFlyServer::ArePreviousCookSettingsCompatible(const TMap<FName, FS
 		const FString* CurrentHash = CurrentCookSettings.Find(ExecutableHashName);
 		if (!CurrentHash)
 		{
-			UE_LOG(LogCook, Display, TEXT("Cook invalidated for platform %s because current executable hash is invalid. Invalid module=%s. Clearing all cooked content."),
+			UE_LOG(LogCook, Display, TEXT("Cook invalidated for platform %s because current executable hash is invalid. Invalid module=%s. Clearing previously cooked packages."),
 				*TargetPlatform->PlatformName(), *CurrentCookSettings.FindRef(ExecutableHashInvalidModuleName));
 			return false;
 		}
@@ -8170,13 +8170,13 @@ bool UCookOnTheFlyServer::ArePreviousCookSettingsCompatible(const TMap<FName, FS
 		if (!PreviousHash)
 		{
 			const FConfigValue* InvalidModuleName = CookSettings->Find(ExecutableHashInvalidModuleName);
-			UE_LOG(LogCook, Display, TEXT("Cook invalidated for platform %s because old executable hash is invalid. Invalid module=%s. Clearing all cooked content."),
+			UE_LOG(LogCook, Display, TEXT("Cook invalidated for platform %s because old executable hash is invalid. Invalid module=%s. Clearing previously cooked packages."),
 				*TargetPlatform->PlatformName(), InvalidModuleName ? *InvalidModuleName->GetValue() : TEXT(""));
 			return false;
 		}
 		if (!CurrentHash->Equals(*PreviousHash->GetValue(), ESearchCase::CaseSensitive))
 		{
-			UE_LOG(LogCook, Display, TEXT("Cook invalidated for platform %s because executable hash has changed. Old: %s, New: %s. Clearing all cooked content."),
+			UE_LOG(LogCook, Display, TEXT("Cook invalidated for platform %s because executable hash has changed. Old: %s, New: %s. Clearing previously cooked packages."),
 				*TargetPlatform->PlatformName(), *PreviousHash->GetValue(), **CurrentHash);
 			return false;
 		}
@@ -11004,7 +11004,7 @@ void UCookOnTheFlyServer::LoadBeginCookIterativeFlagsLocal(FBeginCookContext& Be
 		if (bIsDiffOnly)
 		{
 			UE_LOG(LogCook, Display,
-				TEXT("INCREMENTAL COOK: cooking incrementally due to -DiffOnly flag. Keeping cooked content for platform %s."),
+				TEXT("INCREMENTAL COOK: cooking incrementally due to -DiffOnly flag. Keeping previously cooked packages for platform %s and cooking into memory buffers."),
 				*TargetPlatform->PlatformName());
 			// When looking for deterministic cooking differences in cooked packages, don't delete the packages on disk
 			PlatformContext.bFullBuild = false;
@@ -11019,14 +11019,14 @@ void UCookOnTheFlyServer::LoadBeginCookIterativeFlagsLocal(FBeginCookContext& Be
 			if (!bIterative && !PlatformData->bIsSandboxInitialized)
 			{
 				UE_LOG(LogCook, Display,
-					TEXT("FULL COOK: Neither -iterative nor -cookincremental were specified. Clearing all cooked content for platform %s and executing a full cook."),
+					TEXT("FULL COOK: Neither -iterative nor -cookincremental were specified. Deleting previously cooked packages for platform %s and executing a full cook."),
 					*TargetPlatform->PlatformName());
 				bIterativeAllowed = false;
 			}
 			else if (!ArePreviousCookSettingsCompatible(PlatformContext.CurrentCookSettings, TargetPlatform))
 			{
 				UE_LOG(LogCook, Display,
-					TEXT("FULL COOK: %s was specified, but globally invalidated. Clearing all cooked content for platform %s and executing a full cook."),
+					TEXT("FULL COOK: %s was specified, but global settings have changed and all previously cook packages are invalidated. Deleting previously cooked packages for platform %s and executing a full cook."),
 					bHybridIterativeEnabled ? TEXT("-cookincremental") : TEXT("-iterative"),
 					*TargetPlatform->PlatformName());
 				bIterativeAllowed = false;
@@ -11035,7 +11035,7 @@ void UCookOnTheFlyServer::LoadBeginCookIterativeFlagsLocal(FBeginCookContext& Be
 			if (bIterativeAllowed)
 			{
 				UE_LOG(LogCook, Display,
-					TEXT("INCREMENTAL COOK: %s was specified and not globally invalidated. Keeping cooked content for platform platform %s."),
+					TEXT("INCREMENTAL COOK: %s was specified and global settings are still valid. Keeping previously cooked packages for platform %s and cooking only packages that have been modified."),
 					bHybridIterativeEnabled ? TEXT("-cookincremental") : TEXT("-iterative"),
 					*TargetPlatform->PlatformName());
 				PlatformContext.bFullBuild = false;
