Make sure to load thumbnails before resaving packages using
- UWorldPartitionBuilder::SavePackages()
- FEditorFileUtils

This will force a call to FLinkerLoad::SerializeThumbnails(). Must be done before any calls to ResetLoaders(), otherwise thumbnails won't be loaded.

diff --git a/Engine/Source/Editor/UnrealEd/Private/FileHelpers.cpp b/Engine/Source/Editor/UnrealEd/Private/FileHelpers.cpp
--- a/Engine/Source/Editor/UnrealEd/Private/FileHelpers.cpp
+++ b/Engine/Source/Editor/UnrealEd/Private/FileHelpers.cpp
@@ -3316,6 +3316,12 @@ static void PrepareWorldsForExplicitSave(const TArray<UPackage*>& PackagesToPrep
 
 static void PrepareSavePackages(const TArray<UPackage*>& PackagesToSave)
 {
+	// Load existing thumbnails to be able to resave them properly
+	for (UPackage* PackageToSave : PackagesToSave)
+	{
+		EnsureLoadingComplete(PackageToSave);
+	}
+
 	// Don't call ResetLoaders on newly created world packages as this will prevent future loading of external actor packages to work propertly
 	// Linker will fail to resolve SourceLinker of external actor's world package import (see GetPackageLinker test for PKG_InMemoryOnly on TargetPackage's Package Flag)
 	TArray<UPackage*> PackagesToResetLoaders;
diff --git a/Engine/Source/Editor/UnrealEd/Private/WorldPartition/WorldPartitionBuilder.cpp b/Engine/Source/Editor/UnrealEd/Private/WorldPartition/WorldPartitionBuilder.cpp
--- a/Engine/Source/Editor/UnrealEd/Private/WorldPartition/WorldPartitionBuilder.cpp
+++ b/Engine/Source/Editor/UnrealEd/Private/WorldPartition/WorldPartitionBuilder.cpp
@@ -373,6 +373,12 @@ bool UWorldPartitionBuilder::SavePackages(const TArray<UPackage*>& Packages, FPa
 		}
 	}
 
+	// Load existing thumbnails to be able to resave them properly
+	for (UPackage* PackageToSave : Packages)
+	{
+		EnsureLoadingComplete(PackageToSave);
+	}
+
 	ResetLoaders(TArray<UObject*>(Packages));
 
 	for (int PackageIndex = 0; PackageIndex < Packages.Num(); ++PackageIndex)
