Cooker: Refactor FRequestCluster::FGraphSearch::ExploreVertexEdges into subfunctions to make it easier to read.
#rnx

diff --git a/Engine/Source/Editor/UnrealEd/Private/Cooker/CookRequestCluster.cpp b/Engine/Source/Editor/UnrealEd/Private/Cooker/CookRequestCluster.cpp
--- a/Engine/Source/Editor/UnrealEd/Private/Cooker/CookRequestCluster.cpp
+++ b/Engine/Source/Editor/UnrealEd/Private/Cooker/CookRequestCluster.cpp
@@ -625,6 +625,7 @@ void FRequestCluster::PumpExploration(const FCookerTimer& CookerTimer, bool& bOu
 FRequestCluster::FGraphSearch::FGraphSearch(FRequestCluster& InCluster, ETraversalTier InTraversalTier)
 	: Cluster(InCluster)
 	, TraversalTier(InTraversalTier)
+	, ExploreEdgesContext(InCluster, *this)
 	, AsyncResultsReadyEvent(EEventMode::ManualReset)
 {
 	AsyncResultsReadyEvent->Trigger();
@@ -799,14 +800,15 @@ void FRequestCluster::FGraphSearch::TickExploration(bool& bOutDone)
 	bool bHadActivity = false;
 	for (;;)
 	{
-		TOptional<FVertexData*> Vertex = AsyncQueueResults.Dequeue();
-		if (!Vertex.IsSet())
+		TOptional<FVertexData*> FrontVertex = AsyncQueueResults.Dequeue();
+		if (!FrontVertex.IsSet())
 		{
 			break;
 		}
-		ExploreVertexEdges(**Vertex);
-		FreeQueryData((**Vertex).QueryData);
-		(**Vertex).QueryData = nullptr;
+		FVertexData* Vertex = *FrontVertex;
+		ExploreEdgesContext.Explore(*Vertex);
+		FreeQueryData(Vertex->QueryData);
+		Vertex->QueryData = nullptr;
 		bHadActivity = true;
 	}
 
@@ -1032,205 +1034,185 @@ void FRequestCluster::FGraphSearch::VisitVertexForPlatform(FVertexData& Vertex,
 	PlatformData.SetVisitedByCluster(true);
 }
 
-void FRequestCluster::FGraphSearch::ExploreVertexEdges(FVertexData& Vertex)
+FRequestCluster::FGraphSearch::FExploreEdgesContext::FExploreEdgesContext(FRequestCluster& InCluster,
+	FGraphSearch& InGraphSearch)
+	: Cluster(InCluster)
+	, GraphSearch(InGraphSearch)
 {
-	// Only called from PumpExploration thread
-	using namespace UE::AssetRegistry;
-	using namespace UE::TargetDomain;
+}
+
+void FRequestCluster::FGraphSearch::FExploreEdgesContext::Explore(FVertexData& InVertex)
+{
+	// Only called from scheduler thread
 
 	// The PackageData will not exist if the package does not exist on disk or
 	// the PackageData was removed from the FRequestCluster due to changes in the PackageData's
 	// state elsewhere in the cooker.
-	if (!Vertex.PackageData)
+	if (!InVertex.PackageData)
 	{
 		return;
 	}
 
-	TArray<FName>& HardGameDependencies(Scratch.HardGameDependencies);
-	TArray<FName>& HardEditorDependencies(Scratch.HardEditorDependencies);
-	TArray<FName>& SoftGameDependencies(Scratch.SoftGameDependencies);
-	TSet<FName>& HardDependenciesSet(Scratch.HardDependenciesSet);
+	Initialize(InVertex);
+	CalculatePlatformsToExplore();
+	if (PlatformsToExplore.IsEmpty())
+	{
+		return;
+	}
+
+	CalculatePackageDataDependenciesPlatformAgnostic();
+	CalculateDependenciesAndIterativelySkippable();
+	QueueVisitsOfDependencies();
+}
+
+void FRequestCluster::FGraphSearch::FExploreEdgesContext::Initialize(FVertexData& InVertex)
+{
+	Vertex = &InVertex;
+	check(Vertex->PackageData);
+	check(Vertex->QueryData);
+	PackageData = Vertex->PackageData;
+	PackageName = Vertex->QueryData->PackageName;
+
 	HardGameDependencies.Reset();
 	HardEditorDependencies.Reset();
 	SoftGameDependencies.Reset();
+	CookerLoadingDependencies.Reset();
+	PlatformsToExplore.Reset();
+	PlatformDependencyMap.Reset();
 	HardDependenciesSet.Reset();
-	FPackageData& PackageData = *Vertex.PackageData;
-	FName PackageName = PackageData.GetPackageName();
-	bool bFetchAnyTargetPlatform = Vertex.QueryData->Platforms[PlatformAgnosticPlatformIndex].bActive;
-	TArray<FName>* DiscoveredDependencies = Cluster.COTFS.DiscoveredDependencies.Find(PackageName);
-	if (bFetchAnyTargetPlatform)
+	SkippedPackages.Reset();
+
+	LocalNumFetchPlatforms = GraphSearch.NumFetchPlatforms();
+	bFetchAnyTargetPlatform = false;
+
+	DiscoveredDependencies = Cluster.COTFS.DiscoveredDependencies.Find(PackageName);
+}
+
+void FRequestCluster::FGraphSearch::FExploreEdgesContext::CalculatePlatformsToExplore()
+{
+	bFetchAnyTargetPlatform = Vertex->QueryData->Platforms[PlatformAgnosticPlatformIndex].bActive;
+	for (int32 PlatformIndex = 0; PlatformIndex < LocalNumFetchPlatforms; ++PlatformIndex)
 	{
-		EDependencyQuery FlagsForHardDependencyQuery;
-		if (Cluster.COTFS.bSkipOnlyEditorOnly)
-		{
-			Cluster.AssetRegistry.GetDependencies(PackageName, HardGameDependencies, EDependencyCategory::Package,
-				EDependencyQuery::Game | EDependencyQuery::Hard);
-			HardDependenciesSet.Append(HardGameDependencies);
-		}
-		else
-		{
-			// We're not allowed to skip editoronly imports, so include all hard dependencies
-			FlagsForHardDependencyQuery = EDependencyQuery::Hard;
-			Cluster.AssetRegistry.GetDependencies(PackageName, HardGameDependencies, EDependencyCategory::Package,
-				EDependencyQuery::Game | EDependencyQuery::Hard);
-			Cluster.AssetRegistry.GetDependencies(PackageName, HardEditorDependencies, EDependencyCategory::Package,
-				EDependencyQuery::EditorOnly | EDependencyQuery::Hard);
-			HardDependenciesSet.Append(HardGameDependencies);
-			HardDependenciesSet.Append(HardEditorDependencies);
-		}
-		if (DiscoveredDependencies)
+		FQueryPlatformData& QueryPlatformData = Vertex->QueryData->Platforms[PlatformIndex];
+		if (!QueryPlatformData.bActive || PlatformIndex == PlatformAgnosticPlatformIndex)
 		{
-			HardDependenciesSet.Append(*DiscoveredDependencies);
+			continue;
 		}
-		if (Cluster.bAllowSoftDependencies)
-		{
-			// bSkipOnlyEditorOnly is always true for soft dependencies; skip editoronly soft dependencies
-			Cluster.AssetRegistry.GetDependencies(PackageName, SoftGameDependencies, EDependencyCategory::Package,
-				EDependencyQuery::Game | EDependencyQuery::Soft);
-
-			// Even if we're following soft references in general, we need to check with the SoftObjectPath registry
-			// for any startup packages that marked their softobjectpaths as excluded, and not follow those
-			TSet<FName>& SkippedPackages(Scratch.SkippedPackages);
-			if (GRedirectCollector.RemoveAndCopySoftObjectPathExclusions(PackageName, SkippedPackages))
-			{
-				SoftGameDependencies.RemoveAll([&SkippedPackages](FName SoftDependency)
-					{
-						return SkippedPackages.Contains(SoftDependency);
-					});
-			}
+		PlatformsToExplore.Add(PlatformIndex);
+	}
+}
 
-			// LocalizationReferences are a source of SoftGameDependencies that are not present in the AssetRegistry
-			SoftGameDependencies.Append(GetLocalizationReferences(PackageName, Cluster.COTFS));
+void FRequestCluster::FGraphSearch::FExploreEdgesContext::CalculateIterativelyUnmodified()
+{
+	using namespace UE::TargetDomain;
 
-			// The AssetManager can provide additional SoftGameDependencies
-			SoftGameDependencies.Append(GetAssetManagerReferences(PackageName));
-		}
+	if (!Cluster.IsIncrementalCook())
+	{
+		return;
 	}
 
-	int32 LocalNumFetchPlatforms = NumFetchPlatforms();
-	TMap<FName, FScratchPlatformDependencyBits>& PlatformDependencyMap(Scratch.PlatformDependencyMap);
-	PlatformDependencyMap.Reset();
-	auto AddPlatformDependency = [&PlatformDependencyMap, LocalNumFetchPlatforms]
-	(FName DependencyName, int32 PlatformIndex, EInstigator InstigatorType)
+	for (int32 PlatformIndex : PlatformsToExplore)
 	{
-		FScratchPlatformDependencyBits& PlatformDependencyBits = PlatformDependencyMap.FindOrAdd(DependencyName);
-		if (PlatformDependencyBits.HasPlatformByIndex.Num() != LocalNumFetchPlatforms)
+		if (PlatformIndex == CookerLoadingPlatformIndex)
 		{
-			PlatformDependencyBits.HasPlatformByIndex.Init(false, LocalNumFetchPlatforms);
-			PlatformDependencyBits.InstigatorType = EInstigator::SoftDependency;
+			continue;
 		}
-		PlatformDependencyBits.HasPlatformByIndex[PlatformIndex] = true;
 
-		// Calculate PlatformDependencyType.InstigatorType == 
-		// Max(InstigatorType, PlatformDependencyType.InstigatorType)
-		// based on the enum values, from least required to most: [ Soft, HardEditorOnly, Hard ]
-		switch (InstigatorType)
+		FQueryPlatformData& QueryPlatformData = Vertex->QueryData->Platforms[PlatformIndex];
+		FFetchPlatformData& FetchPlatformData = GraphSearch.FetchPlatforms[PlatformIndex];
+		const ITargetPlatform* TargetPlatform = FetchPlatformData.Platform;
+		FPackagePlatformData& PackagePlatformData = PackageData->FindOrAddPlatformData(TargetPlatform);
+
+		if (!PackagePlatformData.IsCookable())
 		{
-		case EInstigator::HardDependency:
-			PlatformDependencyBits.InstigatorType = InstigatorType;
-			break;
-		case EInstigator::HardEditorOnlyDependency:
-			if (PlatformDependencyBits.InstigatorType != EInstigator::HardDependency)
-			{
-				PlatformDependencyBits.InstigatorType = InstigatorType;
-			}
-			break;
-		case EInstigator::SoftDependency:
-			// New value is minimum, so keep the old value
-			break;
-		case EInstigator::InvalidCategory:
-			// Caller indicated they do not want to set the InstigatorType
-			break;
-		default:
-			checkNoEntry();
-			break;
+			QueryPlatformData.bIterativelyUnmodified = false;
+			continue;
 		}
-	};
-	auto AddPlatformDependencyRange = [&AddPlatformDependency]
-	(TConstArrayView<FName> Range, int32 PlatformIndex, EInstigator InstigatorType)
-	{
-		for (FName DependencyName : Range)
+
+		UE::TargetDomain::FCookDependencies& CookDependencies = QueryPlatformData.CookAttachments.Dependencies;
+		if (!CookDependencies.HasKeyMatch())
 		{
-			AddPlatformDependency(DependencyName, PlatformIndex, InstigatorType);
+			QueryPlatformData.bIterativelyUnmodified = false;
+			continue;
 		}
-	};
 
-	FQueryPlatformData& PlatformAgnosticQueryPlatformData = Vertex.QueryData->Platforms[PlatformAgnosticPlatformIndex];
+		if (!IsIterativeEnabled(PackageName, Cluster.COTFS.bHybridIterativeAllowAllClasses))
+		{
+			QueryPlatformData.bIterativelyUnmodified = false;
+			continue;
+		}
+
+		QueryPlatformData.bIterativelyUnmodified = true;
+		PackagePlatformData.SetIterativelyUnmodified(true);
+	}
+}
+
+void FRequestCluster::FGraphSearch::FExploreEdgesContext::CalculatePackageDataDependenciesPlatformAgnostic()
+{
+	using namespace UE::AssetRegistry;
 
-	auto ProcessPlatformAttachments =
-		[this, PackageName, &PackageData, &PlatformAgnosticQueryPlatformData, &AddPlatformDependencyRange]
-		(int32 PlatformIndex, const ITargetPlatform* TargetPlatform, FFetchPlatformData& FetchPlatformData,
-			FPackagePlatformData& PackagePlatformData, FCookAttachments& PlatformAttachments,
-			bool bExploreDependencies)
+	if (!bFetchAnyTargetPlatform)
 	{
-		bool bFoundBuildDefinitions = false;
-		ICookedPackageWriter* PackageWriter = FetchPlatformData.Writer;
+		return;
+	}
 
-		if (Cluster.IsIncrementalCook() && PackagePlatformData.IsCookable())
-		{
-			bool bIterativelyUnmodified = false;
-			UE::TargetDomain::FCookDependencies& CookDependencies = PlatformAttachments.Dependencies;
-			if (CookDependencies.HasKeyMatch())
-			{
-				if (IsIterativeEnabled(PackageName, Cluster.COTFS.bHybridIterativeAllowAllClasses))
-				{
-					bIterativelyUnmodified = true;
-					PackagePlatformData.SetIterativelyUnmodified(true);
-				}
-				if (bExploreDependencies && Cluster.bAllowSoftDependencies)
-				{
-					AddPlatformDependencyRange(CookDependencies.GetRuntimePackageDependencies(), PlatformIndex,
-						EInstigator::SoftDependency);
-				}
+	EDependencyQuery FlagsForHardDependencyQuery;
+	if (Cluster.COTFS.bSkipOnlyEditorOnly)
+	{
+		Cluster.AssetRegistry.GetDependencies(PackageName, HardGameDependencies, EDependencyCategory::Package,
+			EDependencyQuery::Game | EDependencyQuery::Hard);
+		HardDependenciesSet.Append(HardGameDependencies);
+	}
+	else
+	{
+		// We're not allowed to skip editoronly imports, so include all hard dependencies
+		FlagsForHardDependencyQuery = EDependencyQuery::Hard;
+		Cluster.AssetRegistry.GetDependencies(PackageName, HardGameDependencies, EDependencyCategory::Package,
+			EDependencyQuery::Game | EDependencyQuery::Hard);
+		Cluster.AssetRegistry.GetDependencies(PackageName, HardEditorDependencies, EDependencyCategory::Package,
+			EDependencyQuery::EditorOnly | EDependencyQuery::Hard);
+		HardDependenciesSet.Append(HardGameDependencies);
+		HardDependenciesSet.Append(HardEditorDependencies);
+	}
+	if (DiscoveredDependencies)
+	{
+		HardDependenciesSet.Append(*DiscoveredDependencies);
+	}
+	if (Cluster.bAllowSoftDependencies)
+	{
+		// bSkipOnlyEditorOnly is always true for soft dependencies; skip editoronly soft dependencies
+		Cluster.AssetRegistry.GetDependencies(PackageName, SoftGameDependencies, EDependencyCategory::Package,
+			EDependencyQuery::Game | EDependencyQuery::Soft);
 
-				if (Cluster.bPreQueueBuildDefinitions)
-				{
-					bFoundBuildDefinitions = true;
-					Cluster.BuildDefinitions.AddBuildDefinitionList(PackageName, TargetPlatform,
-						PlatformAttachments.BuildDefinitions.Definitions);
-				}
-			}
-			bool bShouldIterativelySkip = bIterativelyUnmodified;
-			PackageWriter->UpdatePackageModificationStatus(PackageName, bIterativelyUnmodified,
-				bShouldIterativelySkip);
-			if (bShouldIterativelySkip)
-			{
-				// Call SetPlatformCooked instead of just PackagePlatformData.SetCookResults because we might also need
-				// to set OnFirstCookedPlatformAdded
-				PackageData.SetPlatformCooked(TargetPlatform, ECookResult::Succeeded);
-				Cluster.SetPackageDataWasMarkedCooked(PackageData, true);
-				if (PlatformIndex == FirstSessionPlatformIndex)
+		// Even if we're following soft references in general, we need to check with the SoftObjectPath registry
+		// for any startup packages that marked their softobjectpaths as excluded, and not follow those
+		if (GRedirectCollector.RemoveAndCopySoftObjectPathExclusions(PackageName, SkippedPackages))
+		{
+			SoftGameDependencies.RemoveAll([this](FName SoftDependency)
 				{
-					COOK_STAT(++DetailedCookStats::NumPackagesIterativelySkipped);
-				}
-				// Declare the package to the EDLCookInfo verification so we don't warn about missing exports from it
-				UE::SavePackageUtilities::EDLCookInfoAddIterativelySkippedPackage(PackageName);
-			}
+					return SkippedPackages.Contains(SoftDependency);
+				});
 		}
 
-		if (Cluster.bPreQueueBuildDefinitions && !bFoundBuildDefinitions)
-		{
-			if (PlatformAgnosticQueryPlatformData.bActive &&
-				PlatformAgnosticQueryPlatformData.CookAttachments.Dependencies.HasKeyMatch())
-			{
-				Cluster.BuildDefinitions.AddBuildDefinitionList(PackageName, TargetPlatform,
-					PlatformAgnosticQueryPlatformData.CookAttachments.BuildDefinitions.Definitions);
-			}
-		}
-	};
+		// LocalizationReferences are a source of SoftGameDependencies that are not present in the AssetRegistry
+		SoftGameDependencies.Append(GetLocalizationReferences(PackageName, Cluster.COTFS));
 
-	for (int32 PlatformIndex = 0; PlatformIndex < LocalNumFetchPlatforms; ++PlatformIndex)
-	{
-		FQueryPlatformData& QueryPlatformData = Vertex.QueryData->Platforms[PlatformIndex];
-		if (!QueryPlatformData.bActive || PlatformIndex == PlatformAgnosticPlatformIndex)
-		{
-			continue;
-		}
+		// The AssetManager can provide additional SoftGameDependencies
+		SoftGameDependencies.Append(GetAssetManagerReferences(PackageName));
+	}
+}
 
-		FFetchPlatformData& FetchPlatformData = FetchPlatforms[PlatformIndex];
+void FRequestCluster::FGraphSearch::FExploreEdgesContext::CalculateDependenciesAndIterativelySkippable()
+{
+	using namespace UE::AssetRegistry;
+
+	for (int32 PlatformIndex : PlatformsToExplore)
+	{
+		FQueryPlatformData& QueryPlatformData = Vertex->QueryData->Platforms[PlatformIndex];
+		FFetchPlatformData& FetchPlatformData = GraphSearch.FetchPlatforms[PlatformIndex];
 		const ITargetPlatform* TargetPlatform = FetchPlatformData.Platform;
-		FPackagePlatformData& PackagePlatformData = PackageData.FindOrAddPlatformData(TargetPlatform);
-		if ((TraversalTier < ETraversalTier::FollowDependencies) || !PackagePlatformData.IsExplorable())
+		FPackagePlatformData& PackagePlatformData = PackageData->FindOrAddPlatformData(TargetPlatform);
+		if ((GraphSearch.TraversalTier < ETraversalTier::FollowDependencies) || !PackagePlatformData.IsExplorable())
 		{
 			// ExploreVertexEdges is responsible for updating package modification status so we might
 			// have been called for this platform even if not explorable. If not explorable, just update
@@ -1246,9 +1228,6 @@ void FRequestCluster::FGraphSearch::ExploreVertexEdges(FVertexData& Vertex)
 
 		if (PlatformIndex == CookerLoadingPlatformIndex)
 		{
-			TArray<FName>& CookerLoadingDependencies(Scratch.CookerLoadingDependencies);
-			CookerLoadingDependencies.Reset();
-
 			Cluster.AssetRegistry.GetDependencies(PackageName, CookerLoadingDependencies, EDependencyCategory::Package,
 				EDependencyQuery::Hard);
 
@@ -1285,6 +1264,10 @@ void FRequestCluster::FGraphSearch::ExploreVertexEdges(FVertexData& Vertex)
 			AddPlatformDependencyRange(*DiscoveredDependencies, PlatformIndex, EInstigator::HardDependency);
 		}
 	}
+}
+
+void FRequestCluster::FGraphSearch::FExploreEdgesContext::QueueVisitsOfDependencies()
+{
 	if (PlatformDependencyMap.IsEmpty())
 	{
 		return;
@@ -1302,7 +1285,7 @@ void FRequestCluster::FGraphSearch::ExploreVertexEdges(FVertexData& Vertex)
 			FCoreRedirectObjectName(NAME_None, NAME_None, DependencyName)).PackageName;
 		DependencyName = Redirected;
 
-		FVertexData& DependencyVertex = FindOrAddVertex(DependencyName);
+		FVertexData& DependencyVertex = GraphSearch.FindOrAddVertex(DependencyName);
 		if (!DependencyVertex.PackageData)
 		{
 			continue;
@@ -1316,7 +1299,7 @@ void FRequestCluster::FGraphSearch::ExploreVertexEdges(FVertexData& Vertex)
 			{
 				continue;
 			}
-			FFetchPlatformData& FetchPlatformData = FetchPlatforms[PlatformIndex];
+			FFetchPlatformData& FetchPlatformData = GraphSearch.FetchPlatforms[PlatformIndex];
 			const ITargetPlatform* TargetPlatform = FetchPlatformData.Platform;
 			FPackagePlatformData& PlatformData = DependencyPackageData.FindOrAddPlatformData(TargetPlatform);
 
@@ -1324,7 +1307,7 @@ void FRequestCluster::FGraphSearch::ExploreVertexEdges(FVertexData& Vertex)
 			{
 				if (!Edges)
 				{
-					Edges = &GraphEdges.FindOrAdd(&PackageData);
+					Edges = &GraphSearch.GraphEdges.FindOrAdd(PackageData);
 					Edges->Reset(PlatformDependencyMap.Num());
 				}
 				Edges->Add(&DependencyPackageData);
@@ -1345,7 +1328,112 @@ void FRequestCluster::FGraphSearch::ExploreVertexEdges(FVertexData& Vertex)
 		}
 		if (bAddToFrontier)
 		{
-			AddToFrontier(DependencyVertex);
+			GraphSearch.AddToFrontier(DependencyVertex);
+		}
+	}
+}
+
+void FRequestCluster::FGraphSearch::FExploreEdgesContext::AddPlatformDependency(FName DependencyName,
+	int32 PlatformIndex, EInstigator InstigatorType)
+{
+	FScratchPlatformDependencyBits& PlatformDependencyBits = PlatformDependencyMap.FindOrAdd(DependencyName);
+	if (PlatformDependencyBits.HasPlatformByIndex.Num() != LocalNumFetchPlatforms)
+	{
+		PlatformDependencyBits.HasPlatformByIndex.Init(false, LocalNumFetchPlatforms);
+		PlatformDependencyBits.InstigatorType = EInstigator::SoftDependency;
+	}
+	PlatformDependencyBits.HasPlatformByIndex[PlatformIndex] = true;
+
+	// Calculate PlatformDependencyType.InstigatorType == 
+	// Max(InstigatorType, PlatformDependencyType.InstigatorType)
+	// based on the enum values, from least required to most: [ Soft, HardEditorOnly, Hard ]
+	switch (InstigatorType)
+	{
+	case EInstigator::HardDependency:
+		PlatformDependencyBits.InstigatorType = InstigatorType;
+		break;
+	case EInstigator::HardEditorOnlyDependency:
+		if (PlatformDependencyBits.InstigatorType != EInstigator::HardDependency)
+		{
+			PlatformDependencyBits.InstigatorType = InstigatorType;
+		}
+		break;
+	case EInstigator::SoftDependency:
+		// New value is minimum, so keep the old value
+		break;
+	case EInstigator::InvalidCategory:
+		// Caller indicated they do not want to set the InstigatorType
+		break;
+	default:
+		checkNoEntry();
+		break;
+	}
+}
+
+void FRequestCluster::FGraphSearch::FExploreEdgesContext::AddPlatformDependencyRange(TConstArrayView<FName> Range,
+	int32 PlatformIndex, EInstigator InstigatorType)
+{
+	for (FName DependencyName : Range)
+	{
+		AddPlatformDependency(DependencyName, PlatformIndex, InstigatorType);
+	}
+}
+
+void FRequestCluster::FGraphSearch::FExploreEdgesContext::ProcessPlatformAttachments(int32 PlatformIndex,
+	const ITargetPlatform* TargetPlatform, FFetchPlatformData& FetchPlatformData,
+	FPackagePlatformData& PackagePlatformData, UE::TargetDomain::FCookAttachments& PlatformAttachments,
+	bool bExploreDependencies)
+{
+	bool bFoundBuildDefinitions = false;
+	ICookedPackageWriter* PackageWriter = FetchPlatformData.Writer;
+	FQueryPlatformData& QueryPlatformData = Vertex->QueryData->Platforms[PlatformIndex];
+
+	if (Cluster.IsIncrementalCook() && PackagePlatformData.IsCookable())
+	{
+		bool bIterativelyUnmodified = QueryPlatformData.bIterativelyUnmodified;
+		if (bIterativelyUnmodified)
+		{
+			UE::TargetDomain::FCookDependencies& CookDependencies = PlatformAttachments.Dependencies;
+			if (bExploreDependencies && Cluster.bAllowSoftDependencies)
+			{
+				AddPlatformDependencyRange(CookDependencies.GetRuntimePackageDependencies(), PlatformIndex,
+					EInstigator::SoftDependency);
+			}
+
+			if (Cluster.bPreQueueBuildDefinitions)
+			{
+				bFoundBuildDefinitions = true;
+				Cluster.BuildDefinitions.AddBuildDefinitionList(PackageName, TargetPlatform,
+					PlatformAttachments.BuildDefinitions.Definitions);
+			}
+		}
+		bool bShouldIterativelySkip = bIterativelyUnmodified;
+		PackageWriter->UpdatePackageModificationStatus(PackageName, bIterativelyUnmodified,
+			bShouldIterativelySkip);
+		if (bShouldIterativelySkip)
+		{
+			// Call SetPlatformCooked instead of just PackagePlatformData.SetCookResults because we might also need
+			// to set OnFirstCookedPlatformAdded
+			PackageData->SetPlatformCooked(TargetPlatform, ECookResult::Succeeded);
+			Cluster.SetPackageDataWasMarkedCooked(*PackageData, true);
+			if (PlatformIndex == FirstSessionPlatformIndex)
+			{
+				COOK_STAT(++DetailedCookStats::NumPackagesIterativelySkipped);
+			}
+			// Declare the package to the EDLCookInfo verification so we don't warn about missing exports from it
+			UE::SavePackageUtilities::EDLCookInfoAddIterativelySkippedPackage(PackageName);
+		}
+	}
+
+	if (Cluster.bPreQueueBuildDefinitions && !bFoundBuildDefinitions)
+	{
+		FQueryPlatformData& PlatformAgnosticQueryData = Vertex->QueryData->Platforms[PlatformAgnosticPlatformIndex];
+
+		if (PlatformAgnosticQueryData.bActive &&
+			PlatformAgnosticQueryData.CookAttachments.Dependencies.HasKeyMatch())
+		{
+			Cluster.BuildDefinitions.AddBuildDefinitionList(PackageName, TargetPlatform,
+				PlatformAgnosticQueryData.CookAttachments.BuildDefinitions.Definitions);
 		}
 	}
 }
diff --git a/Engine/Source/Editor/UnrealEd/Private/Cooker/CookRequestCluster.h b/Engine/Source/Editor/UnrealEd/Private/Cooker/CookRequestCluster.h
--- a/Engine/Source/Editor/UnrealEd/Private/Cooker/CookRequestCluster.h
+++ b/Engine/Source/Editor/UnrealEd/Private/Cooker/CookRequestCluster.h
@@ -108,6 +108,7 @@ private:
 		 * bActive=true for the ones that are present.
 		 */
 		bool bActive = false;
+		bool bIterativelyUnmodified = false;
 	};
 
 	/** Input/output data for an active query for a vertex's dependencies/previous incrementa results. */
@@ -239,15 +240,45 @@ private:
 			TBitArray<> HasPlatformByIndex;
 			EInstigator InstigatorType = EInstigator::SoftDependency;
 		};
-		struct FScratch
+		struct FExploreEdgesContext
 		{
+		public:
+			FExploreEdgesContext(FRequestCluster& InCluster, FGraphSearch& InGraphSearch);
+
+			/** Process the results from async edges fetch and queue the found dependencies-for-visiting. */
+			void Explore(FVertexData& InVertex);
+
+		private:
+			void Initialize(FVertexData& InVertex);
+			void CalculatePlatformsToExplore();
+			void CalculateIterativelyUnmodified();
+			void CalculatePackageDataDependenciesPlatformAgnostic();
+			void CalculateDependenciesAndIterativelySkippable();
+			void QueueVisitsOfDependencies();
+
+			void AddPlatformDependency(FName DependencyName, int32 PlatformIndex, EInstigator InstigatorType);
+			void AddPlatformDependencyRange(TConstArrayView<FName> Range, int32 PlatformIndex, EInstigator InstigatorType);
+			void ProcessPlatformAttachments(int32 PlatformIndex, const ITargetPlatform* TargetPlatform,
+				FFetchPlatformData& FetchPlatformData, FPackagePlatformData& PackagePlatformData,
+				UE::TargetDomain::FCookAttachments& PlatformAttachments, bool bExploreDependencies);
+
+		private:
+			FRequestCluster& Cluster;
+			FGraphSearch& GraphSearch;
+			FVertexData* Vertex = nullptr;
+			FPackageData* PackageData = nullptr;
+			TArray<FName>* DiscoveredDependencies = nullptr;
 			TArray<FName> HardGameDependencies;
 			TArray<FName> HardEditorDependencies;
 			TArray<FName> SoftGameDependencies;
 			TArray<FName> CookerLoadingDependencies;
+			TArray<int32, TInlineAllocator<10>> PlatformsToExplore;
 			TMap<FName, FScratchPlatformDependencyBits> PlatformDependencyMap;
 			TSet<FName> HardDependenciesSet;
 			TSet<FName> SkippedPackages;
+			FName PackageName;
+			int32 LocalNumFetchPlatforms = 0;
+			bool bFetchAnyTargetPlatform = false;
 		};
 		friend struct FQueryVertexBatch;
 		friend struct FVertexData;
@@ -263,8 +294,6 @@ private:
 		/** Calculate and store the vertex's PackageData's cookability for the platform. */
 		void VisitVertexForPlatform(FVertexData& VertexData, const ITargetPlatform* Platform,
 			FPackagePlatformData& PlatformData, ESuppressCookReason& AccumulatedSuppressCookReason);
-		 /** Process the results from async edges fetch and queue the found dependencies-for-visiting. */
-		void ExploreVertexEdges(FVertexData& VertexData);
 
 		/** Find or add a Vertex for PackageName. If PackageData is provided, use it, otherwise look it up. */
 		FVertexData& FindOrAddVertex(FName PackageName);
@@ -294,9 +323,9 @@ private:
 		void OnVertexCompleted();
 
 		/** Total number of platforms known to the cluster, including the special cases. */
-		int32 NumFetchPlatforms() const { return FetchPlatforms.Num(); }
+		int32 NumFetchPlatforms() const;
 		/** Total number of non-special-case platforms known to the cluster.Identical to COTFS's session platforms */
-		int32 NumSessionPlatforms() const { return FetchPlatforms.Num() - 2; }
+		int32 NumSessionPlatforms() const;
 
 	private:
 		// Variables that are read-only during multithreading
@@ -305,8 +334,8 @@ private:
 		ETraversalTier TraversalTier = ETraversalTier::All;
 
 		// Variables that are accessible only from the Process thread
-		/** Scratch-space variables that are reused to avoid allocations. */
-		FScratch Scratch;
+		/** A set of stack and scratch variables used when calculating and exploring the edges of a vertex. */
+		FExploreEdgesContext ExploreEdgesContext;
 		TMap<FPackageData*, TArray<FPackageData*>> GraphEdges;
 		TMap<FName, FVertexData*> Vertices;
 		TSet<FVertexData*> Frontier;
@@ -315,7 +344,7 @@ private:
 		/** Vertices queued for async processing that are not yet numerous enough to fill a batch. */
 		TRingBuffer<FVertexData*> PreAsyncQueue;
 		/** Time-tracker for timeout warnings in Poll */
-		double LastActivityTime;
+		double LastActivityTime = 0.;
 
 		// Variables that are accessible from multiple threads, guarded by Lock
 		FCriticalSection Lock;
@@ -441,4 +470,14 @@ inline void FRequestCluster::FProcessingFlags::SetWasMarkedCooked(bool bValue)
 	bWasMarkedCooked = bValue;
 }
 
+inline int32 FRequestCluster::FGraphSearch::NumFetchPlatforms() const
+{
+	return FetchPlatforms.Num();
+}
+
+inline int32 FRequestCluster::FGraphSearch::NumSessionPlatforms() const
+{
+	return FetchPlatforms.Num() - 2;
+}
+
 }
