Cooker: OnDemandShaderCompilation: Disable the unnecessary and unrelated clearing of cooked packages in most cases when using -odsc by piggybacking on iterative cooks. Add log messages describing what's happening when that is not possible.
#rnx

diff --git a/Engine/Source/Editor/UnrealEd/Classes/CookOnTheSide/CookOnTheFlyServer.h b/Engine/Source/Editor/UnrealEd/Classes/CookOnTheSide/CookOnTheFlyServer.h
--- a/Engine/Source/Editor/UnrealEd/Classes/CookOnTheSide/CookOnTheFlyServer.h
+++ b/Engine/Source/Editor/UnrealEd/Classes/CookOnTheSide/CookOnTheFlyServer.h
@@ -1533,8 +1533,11 @@ private:
 	bool bIterativeIgnoreExe = false;
 	/** Whether we should calculate the exe's hash; might be true even if bIterativeIgnoreExe is true. */
 	bool bIterativeCalculateExe = true;
-	/** If true this will ignore cooking unsolicited packages.  This is only useful if launching the COTF server to only process shader recompile requests. */
-	bool bIgnoreUnsolicitedPackages = false;
+	/**
+	 * When running as a shader server (-odsc) we avoid cooking packages; the cooker does not queue them and the game does not request them. 
+	 * This mode is used to respond to the shader requests and shares that implementation with cookonthefly.
+	 */
+	bool bRunningAsShaderServer = false;
 	/** Whether to skip saving packages that are cooked. When true, the cook will only load and process packages but not write them to disk */
 	bool bSkipSave = false;
 	/** Timers for tracking how long we have been busy, to manage retries and warnings of deadlock */
diff --git a/Engine/Source/Editor/UnrealEd/Private/Commandlets/CookCommandlet.cpp b/Engine/Source/Editor/UnrealEd/Private/Commandlets/CookCommandlet.cpp
--- a/Engine/Source/Editor/UnrealEd/Private/Commandlets/CookCommandlet.cpp
+++ b/Engine/Source/Editor/UnrealEd/Private/Commandlets/CookCommandlet.cpp
@@ -269,6 +269,12 @@ bool UCookCommandlet::CookByTheBook(const TArray<ITargetPlatform*>& Platforms)
 
 		bIterativeCooking = true;
 	}
+	if (Switches.Contains(TEXT("ODSC")))
+	{
+		// ODSC piggybacks on the cook commandlet but does not cook any packages. Turn -iterative on to prevent
+		// an unnecessary clearing of the cook results.
+		bIterativeCooking = true;
+	}
 
 	ECookInitializationFlags CookFlags = ECookInitializationFlags::IncludeServerMaps;
 	CookFlags |= bIterativeCooking ? IterateFlags : ECookInitializationFlags::None;
diff --git a/Engine/Source/Editor/UnrealEd/Private/CookOnTheFlyServer.cpp b/Engine/Source/Editor/UnrealEd/Private/CookOnTheFlyServer.cpp
--- a/Engine/Source/Editor/UnrealEd/Private/CookOnTheFlyServer.cpp
+++ b/Engine/Source/Editor/UnrealEd/Private/CookOnTheFlyServer.cpp
@@ -4328,7 +4328,7 @@ void UCookOnTheFlyServer::DumpCrashContext(FCrashContextExtendedWriter& Writer)
 void UCookOnTheFlyServer::ProcessUnsolicitedPackages(TArray<FName>* OutDiscoveredPackageNames,
 	TMap<FName, UE::Cook::FInstigator>* OutInstigators)
 {
-	if (bIgnoreUnsolicitedPackages)
+	if (bRunningAsShaderServer)
 	{
 		return;
 	}
@@ -7198,7 +7198,7 @@ void UCookOnTheFlyServer::SetInitializeConfigSettings(UE::Cook::FInitializeConfi
 	// It would be better to always calculate it, but we want to avoid the performance cost until it becomes more widely used
 	bIterativeCalculateExe = !bIterativeIgnoreExe || !bConfigSettingSetIterativeIgnoreExe;
 
-	bIgnoreUnsolicitedPackages = FParse::Param(FCommandLine::Get(), TEXT("odsc"));
+	bRunningAsShaderServer = FParse::Param(FCommandLine::Get(), TEXT("odsc"));;
 	bSkipSave = FParse::Param(FCommandLine::Get(), TEXT("CookSkipSave"));
 
 	FString Severity;
@@ -11022,6 +11022,12 @@ void UCookOnTheFlyServer::LoadBeginCookIterativeFlagsLocal(FBeginCookContext& Be
 					TEXT("FULL COOK: Neither -iterative nor -cookincremental were specified. Deleting previously cooked packages for platform %s and executing a full cook."),
 					*TargetPlatform->PlatformName());
 				bIterativeAllowed = false;
+				if (bRunningAsShaderServer)
+				{
+					UE_LOG(LogCook, Display,
+						TEXT("'-odsc' was passed on commandline, but '-iterative' was not, so the cooker as a side effect is clearing cook results. The build will need to be recooked before it can be staged. Add the commandline argument '-iterative' to avoid this unnecessary clear.")
+					);
+				}
 			}
 			else if (!ArePreviousCookSettingsCompatible(PlatformContext.CurrentCookSettings, TargetPlatform))
 			{
@@ -11029,6 +11035,12 @@ void UCookOnTheFlyServer::LoadBeginCookIterativeFlagsLocal(FBeginCookContext& Be
 					TEXT("FULL COOK: %s was specified, but global settings have changed and all previously cook packages are invalidated. Deleting previously cooked packages for platform %s and executing a full cook."),
 					bHybridIterativeEnabled ? TEXT("-cookincremental") : TEXT("-iterative"),
 					*TargetPlatform->PlatformName());
+				if (bRunningAsShaderServer)
+				{
+					UE_LOG(LogCook, Display,
+						TEXT("'-odsc -iterative' was passed on commandline, but due to unrelated changes in global settings the cooker has to clear cook results. The build will need to be recooked before it can be staged.")
+					);
+				}
 				bIterativeAllowed = false;
 			}
 
