Handle nested components when looking for components that were declared natively but no longer are. Because we did not handle this we would usually remove nested components.

diff --git a/Engine/Source/Editor/UnrealEd/Private/Kismet2/Kismet2.cpp b/Engine/Source/Editor/UnrealEd/Private/Kismet2/Kismet2.cpp
--- a/Engine/Source/Editor/UnrealEd/Private/Kismet2/Kismet2.cpp
+++ b/Engine/Source/Editor/UnrealEd/Private/Kismet2/Kismet2.cpp
@@ -21,6 +21,7 @@
 #include "UObject/UObjectHash.h"
 #include "Serialization/FindObjectReferencers.h"
 #include "Serialization/ArchiveReplaceObjectRef.h"
+#include "String/ParseTokens.h"
 #include "Misc/PackageName.h"
 #include "Widgets/DeclarativeSyntaxSupport.h"
 #include "Textures/SlateIcon.h"
@@ -960,7 +961,25 @@ static void ConformComponentsUtils::ConformRemovedNativeComponents(UObject* BpCd
 
 static UObject* ConformComponentsUtils::FindNativeArchetype(const UObject* NativeCDO, UActorComponent* Component)
 {
-	UObject* NativeSubobject = StaticFindObjectFast(UObject::StaticClass(), const_cast<UObject*>(NativeCDO), Component->GetFName());
+	FSoftObjectPath Path(Component);
+	const FString& SubPathString = Path.GetSubPathString();
+	const FStringView SubobjectDelim = FStringView(TEXT("."));
+	const UE::String::EParseTokensOptions ParseOptions = UE::String::EParseTokensOptions::None;
+	UObject* Iterator = const_cast<UObject*>(NativeCDO);
+	UE::String::ParseTokens(SubPathString, SubobjectDelim,
+		[&Iterator](FStringView Token)
+		{
+			if (Iterator == nullptr)
+			{
+				return;
+			}
+
+			Iterator = StaticFindObjectFast(UObject::StaticClass(), Iterator, FName(Token));
+		},
+		ParseOptions);
+
+	UObject* NativeSubobject = const_cast<UObject*>(Iterator);
+
 	if (!NativeSubobject)
 	{
 		return Component->GetClass()->ClassDefaultObject;
