remove unused dangerous texture import APIs

diff --git a/Engine/Plugins/Interchange/Runtime/Source/Import/Private/Texture/InterchangePSDTranslator.cpp b/Engine/Plugins/Interchange/Runtime/Source/Import/Private/Texture/InterchangePSDTranslator.cpp
--- a/Engine/Plugins/Interchange/Runtime/Source/Import/Private/Texture/InterchangePSDTranslator.cpp
+++ b/Engine/Plugins/Interchange/Runtime/Source/Import/Private/Texture/InterchangePSDTranslator.cpp
@@ -451,10 +451,12 @@ TOptional<UE::Interchange::FImportImage> UInterchangePSDTranslator::GetTexturePa
 	UE::Interchange::FImportImage PayloadData;
 
 	// The PSD is supported. Load it up.        
-	PayloadData.Init2DWithOneMip(
+	PayloadData.Init2DWithParams(
 		PSDHeader.Width,
 		PSDHeader.Height,
-		TextureFormat
+		TextureFormat,
+		true,
+		true
 	);
 
 	FMutableMemoryView Output(PayloadData.RawData.GetData(), PayloadData.RawData.GetSize());
diff --git a/Engine/Plugins/Interchange/Runtime/Source/Import/Private/Texture/InterchangeTexturePayloadData.cpp b/Engine/Plugins/Interchange/Runtime/Source/Import/Private/Texture/InterchangeTexturePayloadData.cpp
--- a/Engine/Plugins/Interchange/Runtime/Source/Import/Private/Texture/InterchangeTexturePayloadData.cpp
+++ b/Engine/Plugins/Interchange/Runtime/Source/Import/Private/Texture/InterchangeTexturePayloadData.cpp
@@ -28,30 +28,6 @@ void UE::Interchange::FImportImage::Init2DWithParams(int32 InSizeX, int32 InSize
 	}
 }
 
-void UE::Interchange::FImportImage::Init2DWithOneMip(int32 InSizeX, int32 InSizeY, ETextureSourceFormat InFormat, const void* InData)
-{
-	Init2DWithParams(InSizeX, InSizeY, 1, InFormat, bSRGB);
-
-	if (InData)
-	{
-		// @@ dangerous: InData has no size
-		FMemory::Memcpy(RawData.GetData(), InData, RawData.GetSize());
-	}
-}
-
-void UE::Interchange::FImportImage::Init2DWithMips(int32 InSizeX, int32 InSizeY, int32 InNumMips, ETextureSourceFormat InFormat, const void* InData)
-{
-	// @@ not used, delete me
-	// this is not used; import with mips is via DDS which uses a different code path
-	Init2DWithParams(InSizeX, InSizeY, InNumMips, InFormat, bSRGB);
-
-	if (InData)
-	{
-		// @@ dangerous: InData has no size
-		FMemory::Memcpy(RawData.GetData(), InData, RawData.GetSize());
-	}
-}
-
 int64 UE::Interchange::FImportImage::GetMipSize(int32 InMipIndex) const
 {
 	check(InMipIndex >= 0);
diff --git a/Engine/Plugins/Interchange/Runtime/Source/Import/Public/Texture/InterchangeTexturePayloadData.h b/Engine/Plugins/Interchange/Runtime/Source/Import/Public/Texture/InterchangeTexturePayloadData.h
--- a/Engine/Plugins/Interchange/Runtime/Source/Import/Public/Texture/InterchangeTexturePayloadData.h
+++ b/Engine/Plugins/Interchange/Runtime/Source/Import/Public/Texture/InterchangeTexturePayloadData.h
@@ -36,8 +36,6 @@ namespace UE
 
 			void Init2DWithParams(int32 InSizeX, int32 InSizeY, ETextureSourceFormat InFormat, bool bInSRGB, bool bShouldAllocateRawData = true);
 			void Init2DWithParams(int32 InSizeX, int32 InSizeY, int32 InNumMips, ETextureSourceFormat InFormat, bool bInSRGB, bool bShouldAllocateRawData = true);
-			void Init2DWithOneMip(int32 InSizeX, int32 InSizeY, ETextureSourceFormat InFormat, const void* InData = nullptr);
-			void Init2DWithMips(int32 InSizeX, int32 InSizeY, int32 InNumMips, ETextureSourceFormat InFormat, const void* InData = nullptr);
 
 			virtual int64 GetMipSize(int32 InMipIndex) const;
 			virtual int64 ComputeBufferSize() const;
diff --git a/Engine/Source/Editor/UnrealEd/Classes/Factories/TextureFactory.h b/Engine/Source/Editor/UnrealEd/Classes/Factories/TextureFactory.h
--- a/Engine/Source/Editor/UnrealEd/Classes/Factories/TextureFactory.h
+++ b/Engine/Source/Editor/UnrealEd/Classes/Factories/TextureFactory.h
@@ -25,8 +25,6 @@ struct FImportImage
 	ETextureSourceCompressionFormat RawDataCompressionFormat = TSCF_None;
 
 	void Init2DWithParams(int32 InSizeX, int32 InSizeY, ETextureSourceFormat InFormat, bool InSRGB);
-	void Init2DWithOneMip(int32 InSizeX, int32 InSizeY, ETextureSourceFormat InFormat, const void* InData = nullptr);
-	void Init2DWithMips(int32 InSizeX, int32 InSizeY, int32 InNumMips, ETextureSourceFormat InFormat, const void* InData = nullptr);
 
 	int64 GetMipSize(int32 InMipIndex) const;
 	void* GetMipData(int32 InMipIndex);
diff --git a/Engine/Source/Editor/UnrealEd/Private/Factories/EditorFactories.cpp b/Engine/Source/Editor/UnrealEd/Private/Factories/EditorFactories.cpp
--- a/Engine/Source/Editor/UnrealEd/Private/Factories/EditorFactories.cpp
+++ b/Engine/Source/Editor/UnrealEd/Private/Factories/EditorFactories.cpp
@@ -2849,43 +2849,6 @@ void FImportImage::Init2DWithParams(int32 InSizeX, int32 InSizeY, ETextureSource
 	SRGB = InSRGB;
 }
 
-void FImportImage::Init2DWithOneMip(int32 InSizeX, int32 InSizeY, ETextureSourceFormat InFormat, const void* InData)
-{
-	SizeX = InSizeX;
-	SizeY = InSizeY;
-	NumMips = 1;
-	Format = InFormat;
-	RawData.AddUninitialized((int64)SizeX * SizeY * FTextureSource::GetBytesPerPixel(Format));
-	if (InData)
-	{
-		// @@ dangerous: InData has no size
-		FMemory::Memcpy(RawData.GetData(), InData, RawData.Num());
-	}
-}
-
-void FImportImage::Init2DWithMips(int32 InSizeX, int32 InSizeY, int32 InNumMips, ETextureSourceFormat InFormat, const void* InData)
-{
-	// @@ not used, delete me
-
-	SizeX = InSizeX;
-	SizeY = InSizeY;
-	NumMips = InNumMips;
-	Format = InFormat;
-
-	int64 TotalSize = 0;
-	for (int32 MipIndex = 0; MipIndex < InNumMips; ++MipIndex)
-	{
-		TotalSize += GetMipSize(MipIndex);
-	}
-	RawData.AddUninitialized(TotalSize);
-
-	if (InData)
-	{
-		// @@ dangerous: InData has no size
-		FMemory::Memcpy(RawData.GetData(), InData, RawData.Num());
-	}
-}
-
 int64 FImportImage::GetMipSize(int32 InMipIndex) const
 {
 	check(InMipIndex >= 0);
@@ -3106,13 +3069,21 @@ bool UTextureFactory::ImportImage(const uint8* Buffer, int64 Length, FFeedbackCo
 		}
 
 		// The psd is supported. Load it up.        
-		OutImage.Init2DWithOneMip(
+		bool bSRGB = true;
+		OutImage.Init2DWithParams(
 			psdhdr.Width,
 			psdhdr.Height,
-			TextureFormat
+			TextureFormat,
+			bSRGB
 		);
+		
+		OutImage.RawData.SetNumUninitialized((int64)psdhdr.Width * psdhdr.Height * FTextureSource::GetBytesPerPixel(TextureFormat));
+
 		uint8* Dst = (uint8*)OutImage.RawData.GetData();
 
+		// @todo : dangerous : psd_ReadData doesn't take Dst end, may overrun
+		//	-> just delete the PSD reader from here entirely, it should be in ImageWrapper instead
+
 		if (!psd_ReadData(Dst, Buffer, psdhdr))
 		{
 			Warn->Logf(TEXT("Failed to read this PSD"));
