Add override for details view creation to BaseAssetToolkit

diff --git a/Engine/Source/Editor/UnrealEd/Private/Tools/BaseAssetToolkit.cpp b/Engine/Source/Editor/UnrealEd/Private/Tools/BaseAssetToolkit.cpp
--- a/Engine/Source/Editor/UnrealEd/Private/Tools/BaseAssetToolkit.cpp
+++ b/Engine/Source/Editor/UnrealEd/Private/Tools/BaseAssetToolkit.cpp
@@ -115,6 +115,15 @@ TSharedRef<SDockTab>  FBaseAssetToolkit::SpawnTab_Details(const FSpawnTabArgs& A
 	return DetailsTab.ToSharedRef();
 }
 
+TSharedRef<IDetailsView> FBaseAssetToolkit::CreateDetailsView()
+{
+	FPropertyEditorModule& PropertyEditorModule = FModuleManager::GetModuleChecked<FPropertyEditorModule>("PropertyEditor");
+	FDetailsViewArgs DetailsViewArgs;
+	DetailsViewArgs.NameAreaSettings = FDetailsViewArgs::HideNameArea;
+	DetailsViewArgs.bHideSelectionTip = true;
+	return PropertyEditorModule.CreateDetailView(DetailsViewArgs);
+}
+
 void FBaseAssetToolkit::RegisterToolbar()
 {
 }
@@ -148,11 +157,8 @@ void FBaseAssetToolkit::CreateWidgets()
 	ViewportTabContent = MakeShareable(new FEditorViewportTabContent());
 	LayoutExtender = MakeShared<FLayoutExtender>();
 
-	FPropertyEditorModule& PropertyEditorModule = FModuleManager::GetModuleChecked<FPropertyEditorModule>("PropertyEditor");
-	FDetailsViewArgs DetailsViewArgs;
-	DetailsViewArgs.NameAreaSettings = FDetailsViewArgs::HideNameArea;
-	DetailsViewArgs.bHideSelectionTip = true;
-	DetailsView = PropertyEditorModule.CreateDetailView(DetailsViewArgs);
+	DetailsView = CreateDetailsView();
+	check(DetailsView.IsValid());
 }
 
 void FBaseAssetToolkit::SetEditingObject(class UObject* InObject)
diff --git a/Engine/Source/Editor/UnrealEd/Public/Tools/BaseAssetToolkit.h b/Engine/Source/Editor/UnrealEd/Public/Tools/BaseAssetToolkit.h
--- a/Engine/Source/Editor/UnrealEd/Public/Tools/BaseAssetToolkit.h
+++ b/Engine/Source/Editor/UnrealEd/Public/Tools/BaseAssetToolkit.h
@@ -59,6 +59,7 @@ protected:
 	UNREALED_API virtual TSharedPtr<FEditorViewportClient> CreateEditorViewportClient() const;
 	UNREALED_API virtual TSharedRef<SDockTab> SpawnTab_Viewport(const FSpawnTabArgs& Args);
 	UNREALED_API virtual TSharedRef<SDockTab> SpawnTab_Details(const FSpawnTabArgs& Args);
+	UNREALED_API virtual TSharedRef<IDetailsView> CreateDetailsView();
 
 protected:
 	/** Property View */
