IterativeCook: TransitiveBuildDependencies
UObjects can declare a transitiveBuildDependency for their package on another package. Only used in Incremental cooks, not used in legacy iterative cooks. If the target package of a TransitiveBuildDependency is invalidated by any of its dependencies, the source package is also invalidated. This is the same behavior that all package dependencies have in legacy iterative cooks, but in incremental cook the behavior of a package dependency is that only the direct holder of the dependency is invalidated. TransitiveBuildDependencies add the capability to have in incremental cooks that transitive invalidation from legacy iterative cooks.
The first use case is UMaterialInstanceConstant, which needs to be invalidated if any of the UFunction shader files used by its UMaterial are modified.
#rnx

diff --git a/Engine/Source/Editor/UnrealEd/Private/Cooker/CookRequestCluster.cpp b/Engine/Source/Editor/UnrealEd/Private/Cooker/CookRequestCluster.cpp
--- a/Engine/Source/Editor/UnrealEd/Private/Cooker/CookRequestCluster.cpp
+++ b/Engine/Source/Editor/UnrealEd/Private/Cooker/CookRequestCluster.cpp
@@ -9,6 +9,7 @@
 #include "Algo/Unique.h"
 #include "AssetRegistry/IAssetRegistry.h"
 #include "Async/Async.h"
+#include "Cooker/CookDependency.h"
 #include "Cooker/CookPackageData.h"
 #include "Cooker/CookPlatformManager.h"
 #include "Cooker/CookProfiling.h"
@@ -382,11 +383,6 @@ void FRequestCluster::RemovePackageData(FPackageData* PackageData)
 	{
 		--PackagesToMarkNotInProgressCount;
 	}
-
-	if (GraphSearch)
-	{
-		GraphSearch->RemovePackageData(PackageData);
-	}
 }
 
 void FRequestCluster::SetPackageDataSuppressReason(FPackageData& PackageData, ESuppressCookReason Reason, bool* bOutExisted)
@@ -632,8 +628,6 @@ FRequestCluster::FGraphSearch::FGraphSearch(FRequestCluster& InCluster, ETravers
 	LastActivityTime = FPlatformTime::Seconds();
 	VertexAllocator.SetMaxBlockSize(1024);
 	VertexAllocator.SetMaxBlockSize(65536);
-	VertexQueryAllocator.SetMinBlockSize(1024);
-	VertexQueryAllocator.SetMaxBlockSize(1024);
 	BatchAllocator.SetMaxBlockSize(16);
 	BatchAllocator.SetMaxBlockSize(16);
 
@@ -663,22 +657,20 @@ void FRequestCluster::FGraphSearch::VisitWithoutDependencies()
 	// If we're skipping the dependencies search, handle that responsibility for the initial requests and return.
 	for (TPair<FPackageData*, FProcessingFlags>& Pair : Cluster.OwnedPackageDatas)
 	{
-		FVertexData Vertex;
-		Vertex.PackageData = Pair.Key;
+		check(Pair.Key);
+		FVertexData Vertex(Pair.Key->GetPackageName(), Pair.Key, *this);
 		VisitVertex(Vertex);
 	}
 }
 
 void FRequestCluster::FGraphSearch::StartSearch()
 {
-	Frontier.Reserve(Cluster.OwnedPackageDatas.Num());
+	VisitVertexQueue.Reserve(Cluster.OwnedPackageDatas.Num());
 	for (TPair<FPackageData*, FProcessingFlags>& Pair : Cluster.OwnedPackageDatas)
 	{
 		FVertexData& Vertex = FindOrAddVertex(Pair.Key->GetPackageName(), *Pair.Key);
 		check(Vertex.PackageData);
-		// We're iterating over OwnedPackageDatas and the Vertex is already in the Cluster so we don't need to call
-		// AddToFrontier; just add it directly.
-		Frontier.Add(&Vertex);
+		AddToVisitVertexQueue(Vertex);
 	}
 }
 
@@ -700,14 +692,14 @@ FRequestCluster::FGraphSearch::~FGraphSearch()
 		}
 		for (;;)
 		{
-			TOptional<FVertexData*> Vertex = AsyncQueueResults.Dequeue();
-			if (!Vertex)
+			if (AsyncQueueResults.Dequeue())
+			{
+				bHadActivity = true;
+			}
+			else
 			{
 				break;
 			}
-			FreeQueryData((**Vertex).QueryData);
-			(**Vertex).QueryData = nullptr;
-			bHadActivity = true;
 		}
 		if (bAsyncBatchesEmpty)
 		{
@@ -726,22 +718,6 @@ FRequestCluster::FGraphSearch::~FGraphSearch()
 	}
 }
 
-void FRequestCluster::FGraphSearch::RemovePackageData(FPackageData* PackageData)
-{
-	check(PackageData);
-	FVertexData** Vertex = Vertices.Find(PackageData->GetPackageName());
-	if (Vertex)
-	{
-		(**Vertex).PackageData = nullptr;
-	}
-
-	GraphEdges.Remove(PackageData);
-	for (TPair<FPackageData*, TArray<FPackageData*>>& Pair : GraphEdges)
-	{
-		Pair.Value.Remove(PackageData);
-	}
-}
-
 void FRequestCluster::FGraphSearch::OnNewReachablePlatforms(FPackageData* PackageData)
 {
 	FVertexData** VertexPtr = Vertices.Find(PackageData->GetPackageName());
@@ -749,44 +725,46 @@ void FRequestCluster::FGraphSearch::OnNewReachablePlatforms(FPackageData* Packag
 	{
 		return;
 	}
-	// Already in OwnedPackageDatas, so just add to Frontier directly
-	Frontier.Add(*VertexPtr);
+	AddToVisitVertexQueue(**VertexPtr);
 }
 
-void FRequestCluster::FGraphSearch::QueueEdgesFetch(FVertexData& Vertex,
-	TConstArrayView<const ITargetPlatform*> Platforms)
+void FRequestCluster::FGraphSearch::QueueEdgesFetch(FVertexData& Vertex, TConstArrayView<int32> PlatformIndexes)
 {
-	check(!Vertex.QueryData);
-	FVertexQueryData& QueryData = *AllocateQueryData();
-	Vertex.QueryData = &QueryData;
-	
-	QueryData.PackageName = Vertex.PackageData->GetPackageName();
-	QueryData.Platforms.SetNum(FetchPlatforms.Num());
+	check(Vertex.PackageData); // Caller must not call without a PackageData; doing so serves no purpose
+
+	bool bAnyRequestedNeedsPlatformAgnostic = false;
+	bool bAnyRequested = false;
+	bool bAllHaveAlreadyCompletedFetch = false;
 
-	// Store Platforms in QueryData->Platforms.bActive. All bActive values start false from constructor or from Reset
-	bool bHasPlatformAgnostic = false;
-	for (const ITargetPlatform* Platform : Platforms)
+	for (int32 PlatformIndex : PlatformIndexes)
 	{
-		int32 Index = Algo::BinarySearchBy(FetchPlatforms, Platform, [](const FFetchPlatformData& D) 
+		// The platform data may have already been requested; request it only if current status is NotRequested
+		FQueryPlatformData& QueryData = Vertex.PlatformData[PlatformIndex];
+		if (!QueryData.bSchedulerThreadFetchCompleted)
 			{
-				return D.Platform;
-			});
-		check(Index != INDEX_NONE);
-		QueryData.Platforms[Index].bActive = true;
-		if (Platform != CookerLoadingPlatformKey)
+			bAllHaveAlreadyCompletedFetch = false;
+			EAsyncQueryStatus ExpectedStatus = EAsyncQueryStatus::NotRequested;
+			if (QueryData.CompareExchangeAsyncQueryStatus(ExpectedStatus, EAsyncQueryStatus::SchedulerRequested))
 		{
-			bHasPlatformAgnostic = true;
+				bAnyRequested = true;
 		}
 	}
-	if (bHasPlatformAgnostic)
+	}
+
+	if (bAnyRequested)
 	{
-		QueryData.Platforms[PlatformAgnosticPlatformIndex].bActive = true;
+		PreAsyncQueue.Add(&Vertex);
+		CreateAvailableBatches(false /* bAllowIncompleteBatch */);
 	}
-	int32 NumPendingPlatforms = Platforms.Num() + (bHasPlatformAgnostic ? 1 : 0);
-	QueryData.PendingPlatforms.store(NumPendingPlatforms, std::memory_order_release);
 
-	PreAsyncQueue.Add(&Vertex);
-	CreateAvailableBatches(false /* bAllowIncompleteBatch */);
+	if (bAllHaveAlreadyCompletedFetch)
+	{
+		// We are contractually obligated to kick the vertex. Normally we would put it into PreAsyncQueue and that
+		// queue would take responsibility for kicking it. Also, it might still be in the AsyncQueueResults for one
+		// of the platforms so it will be kicked by TickExplore pulling it out of the AsyncQueueResults. But if all
+		// requested platforms already previously pulled it out of AsyncQueueResults, then we need to kick it again.
+		KickVertex(&Vertex);
+	}
 }
 
 void FRequestCluster::FGraphSearch::WaitForAsyncQueue(double WaitTimeSeconds)
@@ -798,6 +776,8 @@ void FRequestCluster::FGraphSearch::WaitForAsyncQueue(double WaitTimeSeconds)
 void FRequestCluster::FGraphSearch::TickExploration(bool& bOutDone)
 {
 	bool bHadActivity = false;
+
+	int32 RunawayLoopCount = 0;
 	for (;;)
 	{
 		TOptional<FVertexData*> FrontVertex = AsyncQueueResults.Dequeue();
@@ -806,35 +786,58 @@ void FRequestCluster::FGraphSearch::TickExploration(bool& bOutDone)
 			break;
 		}
 		FVertexData* Vertex = *FrontVertex;
+		for (FQueryPlatformData& PlatformData : GetPlatformDataArray(*Vertex))
+		{
+			if (!PlatformData.bSchedulerThreadFetchCompleted)
+			{
+				PlatformData.bSchedulerThreadFetchCompleted =
+					PlatformData.GetAsyncQueryStatus() >= EAsyncQueryStatus::Complete;
+				// Note that AsyncQueryStatus might change immediately after we read it, so we might have set
+				// FetchCompleted=false but now AsyncQueryStatus is complete. In that case, whatever async thread
+				// changed the AsyncQueryStatus will also kick the vertex again and we will detect the new value when
+				// we reach the new value of the vertexdata later in AsyncQueueResults
+			}
+		}
+
 		ExploreEdgesContext.Explore(*Vertex);
-		FreeQueryData(Vertex->QueryData);
-		Vertex->QueryData = nullptr;
 		bHadActivity = true;
+
+		if (RunawayLoopCount++ > 2 * Vertices.Num())
+		{
+			UE_LOG(LogCook, Fatal, TEXT("Infinite loop detected in FRequestCluster::TickExploration's AsyncQueueResults."));
+		}
 	}
 
-	if (!Frontier.IsEmpty())
+	RunawayLoopCount = 0;
+	while (!VisitVertexQueue.IsEmpty())
 	{
-		TArray<FVertexData*> BusyVertices;
-		for (FVertexData* Vertex : Frontier)
+		bHadActivity = true;
+		// VisitVertex might try to add other vertices onto VisitVertexQueue, so move it into a snapshot and process
+		// the snapshot. After snapshot processing is done, add on anything that was added and then move it back.
+		// We move it back even if it is empty so we can avoid reallocating when we add to it again later.
+		TSet<FVertexData*> Snapshot = MoveTemp(VisitVertexQueue);
+		VisitVertexQueue.Reset();
+		for (FVertexData* Vertex : Snapshot)
 		{
-			if (Vertex->QueryData)
-			{
-				// Vertices that are already in the AsyncQueue can not be added again; we would clobber their
-				// QueryData. Postpone them.
-				BusyVertices.Add(Vertex);
-			}
-			else
-			{
-				VisitVertex(*Vertex);
-			}
+			VisitVertex(*Vertex);
+		}
+		Snapshot.Reset();
+		Snapshot.Append(VisitVertexQueue);
+		VisitVertexQueue = MoveTemp(Snapshot);
+
+		if (RunawayLoopCount++ > 2 * Vertices.Num())
+		{
+			UE_LOG(LogCook, Fatal, TEXT("Infinite loop detected in FRequestCluster::TickExploration's VisitVertexQueue."));
 		}
-		bHadActivity |= BusyVertices.Num() != Frontier.Num();
-		Frontier.Reset();
-		Frontier.Append(BusyVertices);
 	}
 
 	if (bHadActivity)
 	{
+		++RunAwayTickLoopCount;
+		if (RunAwayTickLoopCount++ > 2 * Vertices.Num()*NumFetchPlatforms())
+		{
+			UE_LOG(LogCook, Fatal, TEXT("Infinite loop detected in reentrant calls to FRequestCluster::TickExploration."));
+		}
 		LastActivityTime = FPlatformTime::Seconds();
 		bOutDone = false;
 		return;
@@ -876,46 +879,114 @@ void FRequestCluster::FGraphSearch::TickExploration(bool& bOutDone)
 		return;
 	}
 
-	// Frontier was reset above, and it cannot be modified between there and here.
-	// If it were non-empty we would not be done.
-	check(Frontier.IsEmpty());
+	if (!VisitVertexQueue.IsEmpty() || !bAsyncQueueEmpty || !PreAsyncQueue.IsEmpty())
+	{
+		// A container ticked earlier was populated by the tick of a later container; restart tick from beginning
+		bOutDone = false;
+		return;
+	}
+
+	// We are out of direct dependency work to do, but there could be a cycle in the graph of
+	// TransitiveBuildDependencies. If so, resolve the cycle and allow those vertices' edges to be explored.
+	if (!PendingTransitiveBuildDependencyVertices.IsEmpty())
+	{
+		ResolveTransitiveBuildDependencyCycle();
+		bOutDone = false;
+		++RunAwayTickLoopCount;
+		if (RunAwayTickLoopCount++ > 2 * Vertices.Num() * NumFetchPlatforms())
+		{
+			UE_LOG(LogCook, Fatal, TEXT("Infinite loop detected in FRequestCluster::PendingTransitiveBuildDependencyVertices."));
+		}
+		return;
+	}
+
 	bOutDone = true;
 }
 
+void FRequestCluster::FGraphSearch::ResolveTransitiveBuildDependencyCycle()
+{
+	// We interpret cycles in the transitive build dependency graph to mean that every vertex in the cycle is
+	// invalidated if and only if any dependency from any vertex that points outside the cycle is invalidated (the
+	// dependency pointing outside the cycle might be either a transitive build dependency on a package outside of the
+	// cycle or a direct dependency).
+
+	// Using this definition, we can resolve as not iteratively modified, with no further calculation needed, all
+	// elements in the PendingTransitiveBuildDependencyVertices graph, when we run out of direct dependency work to do.
+	// Proof:
+
+	// Every package in the PendingTransitiveBuildDependencyVertices set is one that is not invalidated by any of its
+	// direct dependencies, but it has transitive build dependencies that might be invalidated.
+	// If we have run out of direct dependency work to do, then all there are no transitive build dependencies on any
+	// vertex not in the set.
+	// No direct dependency invalidations and no transitive build dependency invalidations, by our interpretation of a
+	// cycle above, mean that the package is not invalidated.
+
+	// Mark all of the currently fetched platforms of all packages in the PendingTransitiveBuildDependencyVertices as
+	// ignore transitive build dependencies and kick them.
+
+	FVertexData* FirstVertex = nullptr;
+	for (FVertexData* CycleVert : PendingTransitiveBuildDependencyVertices)
+	{
+		if (!FirstVertex)
+		{
+			FirstVertex = CycleVert;
+		}
+		for (FQueryPlatformData& PlatformData : GetPlatformDataArray(*CycleVert))
+		{
+			if (PlatformData.bIterativelyUnmodifiedRequested || PlatformData.bExploreRequested)
+			{
+				PlatformData.bTransitiveBuildDependenciesResolvedAsNotModified = true;
+			}
+		}
+		// We can also empty the IterativelyModifiedListeners since any remaining listeners must be in
+		// PendingTransitiveBuildDependencyVertices. Empting the list here avoids the expense of kicking
+		// for a second time each of the listeners.
+		CycleVert->IterativelyModifiedListeners.Empty();
+		KickVertex(CycleVert);
+	}
+	check(FirstVertex); // This function should not be called if PendingTransitiveBuildDependencyVertices is empty.
+	PendingTransitiveBuildDependencyVertices.Empty();
+	UE_LOG(LogCook, Display,
+		TEXT("Cycle detected in the graph of transitive build dependencies.")
+		TEXT(" No vertices in the cycle are invalidated by their direct dependencies, so marking them all as iteratively skippable.")
+		TEXT("\n\tVertex in the cycle: %s"),
+		*FirstVertex->PackageName.ToString());
+}
+
 void FRequestCluster::FGraphSearch::UpdateDisplay()
 {
 	constexpr double WarningTimeout = 10.0;
 	if (FPlatformTime::Seconds() > LastActivityTime + WarningTimeout && Cluster.IsIncrementalCook())
 	{
 		FScopeLock ScopeLock(&Lock);
-		int32 NumVertices = 0;
+		int32 NumPendingRequestsInBatches = 0;
 		int32 NumBatches = AsyncQueueBatches.Num();
 		for (FQueryVertexBatch* Batch : AsyncQueueBatches)
 		{
-			NumVertices += Batch->PendingVertices;
+			NumPendingRequestsInBatches += Batch->NumPendingRequests;
 		}
 
 		UE_LOG(LogCook, Warning,
 			TEXT("FRequestCluster waited more than %.0lfs for previous build results from the oplog. ")
-			TEXT("NumPendingBatches == %d, NumPendingVertices == %d. Continuing to wait..."),
-			WarningTimeout, NumBatches, NumVertices);
+			TEXT("NumPendingBatches == %d, NumPendingRequestsInBatches == %d. Continuing to wait..."),
+			WarningTimeout, NumBatches, NumPendingRequestsInBatches);
 		LastActivityTime = FPlatformTime::Seconds();
 	}
 }
 
 void FRequestCluster::FGraphSearch::VisitVertex(FVertexData& Vertex)
 {
-	// Only called from PumpExploration thread
+	// Only called from scheduler thread
 
-	// The PackageData will not exist if the package does not exist on disk or
-	// the PackageData was removed from the FRequestCluster due to changes in the PackageData's
-	// state elsewhere in the cooker.
+	// The PackageData will not exist if the package does not exist on disk
 	if (!Vertex.PackageData)
 	{
 		return;
 	}
 
-	TArray<const ITargetPlatform*, TInlineAllocator<1>> ExplorePlatforms;
+	int32 LocalNumFetchPlatforms = NumFetchPlatforms();
+	TBitArray<> ShouldFetchPlatforms(false, LocalNumFetchPlatforms);
+	
 	FPackagePlatformData* CookerLoadingPlatform = nullptr;
 	const ITargetPlatform* FirstReachableSessionPlatform = nullptr;
 	ESuppressCookReason SuppressCookReason = ESuppressCookReason::Invalid;
@@ -924,24 +995,35 @@ void FRequestCluster::FGraphSearch::VisitVertex(FVertexData& Vertex)
 		Vertex.PackageData->GetPlatformDatasConstKeysMutableValues())
 	{
 		FPackagePlatformData& PlatformData = Pair.Value;
-		if (Pair.Key == CookerLoadingPlatformKey)
+		const ITargetPlatform* TargetPlatform = Pair.Key;
+		if (TargetPlatform == CookerLoadingPlatformKey)
 		{
-			CookerLoadingPlatform = &Pair.Value;
+			CookerLoadingPlatform = &PlatformData;
 		}
 		else if (PlatformData.IsReachable())
 		{
+			int32 PlatformIndex = Algo::BinarySearchBy(FetchPlatforms, TargetPlatform, [](const FFetchPlatformData& D)
+				{
+					return D.Platform;
+				});
+			check(PlatformIndex != INDEX_NONE);
+
 			if (!FirstReachableSessionPlatform)
 			{
-				FirstReachableSessionPlatform = Pair.Key;
+				FirstReachableSessionPlatform = TargetPlatform;
 			}
 			if (!PlatformData.IsVisitedByCluster())
 			{
-				VisitVertexForPlatform(Vertex, Pair.Key, PlatformData, SuppressCookReason);
+				VisitVertexForPlatform(Vertex, TargetPlatform, PlatformData, SuppressCookReason);
+
 				if ((TraversalTier >= ETraversalTier::FetchEdgeData) && 
 					(((TraversalTier >= ETraversalTier::FollowDependencies) && PlatformData.IsExplorable())
 						|| Cluster.IsIncrementalCook()))
 				{
-					ExplorePlatforms.Add(Pair.Key);
+					ShouldFetchPlatforms[PlatformIndex] = true;
+					Vertex.PlatformData[PlatformIndex].bExploreRequested = true;
+					// Exploration of any session platform also requires exploration of PlatformAgnosticPlatform
+					Vertex.PlatformData[PlatformAgnosticPlatformIndex].bExploreRequested = true;
 				}
 			}
 			if (PlatformData.IsCookable())
@@ -997,13 +1079,45 @@ void FRequestCluster::FGraphSearch::VisitVertex(FVertexData& Vertex)
 		CookerLoadingPlatform->SetVisitedByCluster(true);
 		if (TraversalTier >= ETraversalTier::FollowDependencies)
 		{
-			ExplorePlatforms.Add(CookerLoadingPlatformKey);
+			ShouldFetchPlatforms[CookerLoadingPlatformIndex] = true;
+			Vertex.PlatformData[CookerLoadingPlatformIndex].bExploreRequested = true;
 		}
 	}
 
-	if (!ExplorePlatforms.IsEmpty() && TraversalTier >= ETraversalTier::FetchEdgeData)
+	if (TraversalTier >= ETraversalTier::FetchEdgeData)
 	{
-		QueueEdgesFetch(Vertex, ExplorePlatforms);
+		for (int32 PlatformIndex = 0; PlatformIndex < LocalNumFetchPlatforms; ++PlatformIndex)
+		{
+			FQueryPlatformData& PlatformData = Vertex.PlatformData[PlatformIndex];
+
+			// Add on the fetch (but not the explore) of bIterativelyUnmodifiedRequested platforms
+			if (PlatformData.bIterativelyUnmodifiedRequested)
+			{
+				ShouldFetchPlatforms[PlatformIndex] = true;
+			}
+
+			// Also add the fetch (but not necessarily the explore) of PlatformAgnosticPlatform if a
+			// SessionPlatform is fetched.
+			if (ShouldFetchPlatforms[PlatformIndex] &&
+				PlatformIndex != CookerLoadingPlatformIndex && PlatformIndex != PlatformAgnosticPlatformIndex)
+			{
+				ShouldFetchPlatforms[PlatformAgnosticPlatformIndex] = true;
+			}
+		}
+
+		// Convert Bit Array to an array of indexes and fetch them if non empty
+		TArray<int32, TInlineAllocator<10>> FetchPlatformIndexes;
+		for (int32 PlatformIndex = 0; PlatformIndex < LocalNumFetchPlatforms; ++PlatformIndex)
+		{
+			if (ShouldFetchPlatforms[PlatformIndex])
+			{
+				FetchPlatformIndexes.Add(PlatformIndex);
+			}
+		}
+		if (!FetchPlatformIndexes.IsEmpty())
+		{
+			QueueEdgesFetch(Vertex, FetchPlatformIndexes);
+		}
 	}
 }
 
@@ -1045,106 +1159,232 @@ void FRequestCluster::FGraphSearch::FExploreEdgesContext::Explore(FVertexData& I
 {
 	// Only called from scheduler thread
 
-	// The PackageData will not exist if the package does not exist on disk or
-	// the PackageData was removed from the FRequestCluster due to changes in the PackageData's
-	// state elsewhere in the cooker.
-	if (!InVertex.PackageData)
+	Initialize(InVertex);
+	CalculatePlatformsToProcess();
+	if (PlatformsToProcess.IsEmpty())
 	{
 		return;
 	}
 
-	Initialize(InVertex);
-	CalculatePlatformsToExplore();
+	if (!TryCalculateIterativelyUnmodified())
+	{
+		// The vertex was added as a listener to the pending data it needs. Exit from explore
+		// for now and we will reenter it later when the data becomes available.
+		return;
+	}
 	if (PlatformsToExplore.IsEmpty())
 	{
+		// We had platforms we needed to test for iteratively unmodified (for e.g. TransitiveBuildDependencies), but
+		// nothing to explore. No more work to do until/unless they become marked for explore later.
 		return;
 	}
 
 	CalculatePackageDataDependenciesPlatformAgnostic();
 	CalculateDependenciesAndIterativelySkippable();
 	QueueVisitsOfDependencies();
+	MarkExploreComplete();
 }
 
 void FRequestCluster::FGraphSearch::FExploreEdgesContext::Initialize(FVertexData& InVertex)
 {
 	Vertex = &InVertex;
+	// Vertices without a package data are never queued for fetch
 	check(Vertex->PackageData);
-	check(Vertex->QueryData);
 	PackageData = Vertex->PackageData;
-	PackageName = Vertex->QueryData->PackageName;
+	PackageName = Vertex->PackageName;
 
 	HardGameDependencies.Reset();
 	HardEditorDependencies.Reset();
 	SoftGameDependencies.Reset();
 	CookerLoadingDependencies.Reset();
+	PlatformsToProcess.Reset();
 	PlatformsToExplore.Reset();
 	PlatformDependencyMap.Reset();
 	HardDependenciesSet.Reset();
 	SkippedPackages.Reset();
+	UnreadyTransitiveBuildVertices.Reset();
 
 	LocalNumFetchPlatforms = GraphSearch.NumFetchPlatforms();
 	bFetchAnyTargetPlatform = false;
 
 	DiscoveredDependencies = Cluster.COTFS.DiscoveredDependencies.Find(PackageName);
+
+	GraphSearch.PendingTransitiveBuildDependencyVertices.Remove(Vertex);
 }
 
-void FRequestCluster::FGraphSearch::FExploreEdgesContext::CalculatePlatformsToExplore()
+void FRequestCluster::FGraphSearch::FExploreEdgesContext::CalculatePlatformsToProcess()
 {
-	bFetchAnyTargetPlatform = Vertex->QueryData->Platforms[PlatformAgnosticPlatformIndex].bActive;
+	FQueryPlatformData& PlatformAgnosticQueryData = Vertex->PlatformData[PlatformAgnosticPlatformIndex];
 	for (int32 PlatformIndex = 0; PlatformIndex < LocalNumFetchPlatforms; ++PlatformIndex)
 	{
-		FQueryPlatformData& QueryPlatformData = Vertex->QueryData->Platforms[PlatformIndex];
-		if (!QueryPlatformData.bActive || PlatformIndex == PlatformAgnosticPlatformIndex)
+		if (PlatformIndex == PlatformAgnosticPlatformIndex)
+		{
+			continue;
+		}
+		FQueryPlatformData& QueryPlatformData = Vertex->PlatformData[PlatformIndex];
+		if (!QueryPlatformData.bSchedulerThreadFetchCompleted)
+		{
+			continue;
+		}
+		bool bIterativelyUnmodifiedNeeded = !QueryPlatformData.bIterativelyUnmodified.IsSet();
+		bool bExploreNeeded = !QueryPlatformData.bExploreCompleted && QueryPlatformData.bExploreRequested;
+		if (!bIterativelyUnmodifiedNeeded && !bExploreNeeded)
 		{
 			continue;
 		}
-		PlatformsToExplore.Add(PlatformIndex);
+		if (bExploreNeeded && PlatformIndex != CookerLoadingPlatformIndex)
+		{
+			if (!PlatformAgnosticQueryData.bSchedulerThreadFetchCompleted)
+			{
+				continue;
+			}
+			// bExploreNeeded implies bExploreRequested, and wherever bExploreRequested is set to true we also set it
+			// to true for PlatformAgnosticQueryData.
+			check(PlatformAgnosticQueryData.bExploreRequested);
+			bFetchAnyTargetPlatform = true;
+		}
+		PlatformsToProcess.Add(PlatformIndex);
+		if (bExploreNeeded)
+		{
+			PlatformsToExplore.Add(PlatformIndex);
+		}
 	}
 }
 
-void FRequestCluster::FGraphSearch::FExploreEdgesContext::CalculateIterativelyUnmodified()
+bool FRequestCluster::FGraphSearch::FExploreEdgesContext::TryCalculateIterativelyUnmodified()
 {
 	using namespace UE::TargetDomain;
 
 	if (!Cluster.IsIncrementalCook())
 	{
-		return;
+		return true;
 	}
 
-	for (int32 PlatformIndex : PlatformsToExplore)
+	bool bAllPlatformsAreReady = true;
+	for (int32 PlatformIndex : PlatformsToProcess)
 	{
 		if (PlatformIndex == CookerLoadingPlatformIndex)
 		{
 			continue;
 		}
 
-		FQueryPlatformData& QueryPlatformData = Vertex->QueryData->Platforms[PlatformIndex];
+		FQueryPlatformData& QueryPlatformData = Vertex->PlatformData[PlatformIndex];
+		if (QueryPlatformData.bIterativelyUnmodified.IsSet())
+		{
+			continue;
+		}
+
 		FFetchPlatformData& FetchPlatformData = GraphSearch.FetchPlatforms[PlatformIndex];
 		const ITargetPlatform* TargetPlatform = FetchPlatformData.Platform;
 		FPackagePlatformData& PackagePlatformData = PackageData->FindOrAddPlatformData(TargetPlatform);
 
 		if (!PackagePlatformData.IsCookable())
 		{
-			QueryPlatformData.bIterativelyUnmodified = false;
+			SetIsIterativelyUnmodified(PlatformIndex, false);
 			continue;
 		}
 
 		UE::TargetDomain::FCookDependencies& CookDependencies = QueryPlatformData.CookAttachments.Dependencies;
 		if (!CookDependencies.HasKeyMatch())
 		{
-			QueryPlatformData.bIterativelyUnmodified = false;
+			SetIsIterativelyUnmodified(PlatformIndex, false);
 			continue;
 		}
 
 		if (!IsIterativeEnabled(PackageName, Cluster.COTFS.bHybridIterativeAllowAllClasses))
 		{
-			QueryPlatformData.bIterativelyUnmodified = false;
+			SetIsIterativelyUnmodified(PlatformIndex, false);
 			continue;
 		}
 
-		QueryPlatformData.bIterativelyUnmodified = true;
+		if (!QueryPlatformData.bTransitiveBuildDependenciesResolvedAsNotModified)
+		{
+			bool bAnyTransitiveBuildDependencyIsModified = false;
+			UnreadyTransitiveBuildVertices.Reset();
+			for (const UE::Cook::FCookDependency& TransitiveBuildDependency :
+				CookDependencies.GetTransitiveBuildDependencies())
+			{
+				FName TransitiveBuildPackageName = TransitiveBuildDependency.GetPackageName();
+				FVertexData& TransitiveBuildVertex = GraphSearch.FindOrAddVertex(TransitiveBuildPackageName);
+				if (!TransitiveBuildVertex.PackageData)
+				{
+					// A build dependency on a non-existent package can occur e.g. if the package is in an
+					// unmounted plugin. If the package does not exist we count the transitivebuilddependency
+					// as not iteratively unmodified, the same as any package that is not cooked, so mark this
+					// package as not iteratively unmodified.
+					// This is an unexpected data layout however, so log it as a warning.
+					UE_LOG(LogCook, Warning,
+						TEXT("TransitiveBuildDependency to non-existent package.")
+						TEXT(" Package %s has a transitive build dependency on package %s, which does not exist or is not mounted.")
+						TEXT(" Package %s will be marked as not iteratively skippable and will be recooked."),
+						*Vertex->PackageName.ToString(), *TransitiveBuildPackageName.ToString(),
+						*Vertex->PackageName.ToString());
+					bAnyTransitiveBuildDependencyIsModified = true;
+					break;
+				}
+
+				FQueryPlatformData& TransitivePlatformData = TransitiveBuildVertex.PlatformData[PlatformIndex];
+				if (!TransitivePlatformData.bIterativelyUnmodified.IsSet())
+				{
+					UnreadyTransitiveBuildVertices.Add(&TransitiveBuildVertex);
+					continue;
+				}
+				if (!TransitivePlatformData.bIterativelyUnmodified.GetValue())
+				{
+					bAnyTransitiveBuildDependencyIsModified = true;
+					break;
+				}
+			}
+
+			if (bAnyTransitiveBuildDependencyIsModified)
+			{
+				SetIsIterativelyUnmodified(PlatformIndex, false);
+				continue;
+			}
+			if (!UnreadyTransitiveBuildVertices.IsEmpty())
+			{
+				// Add this vertex as a listener to the TransitiveBuildVertices' TryCalculateIterativelyUnmodified
+				for (FVertexData* TransitiveBuildVertex : UnreadyTransitiveBuildVertices)
+				{
+					FQueryPlatformData& TransitivePlatformData = TransitiveBuildVertex->PlatformData[PlatformIndex];
+
+					// Do not kick the vertex again if it has already been fetched; doing so will create busy work
+					// in the case of a cycle and prevent us from detecting the cycle.
+					if (!TransitivePlatformData.bSchedulerThreadFetchCompleted)
+					{
+						TransitivePlatformData.bIterativelyUnmodifiedRequested = true;
+						GraphSearch.AddToVisitVertexQueue(*TransitiveBuildVertex);
+					}
+					// It's okay to add duplicates to IterativelyModifiedListeners; we remove them when broadcasting
+					TransitiveBuildVertex->IterativelyModifiedListeners.Add(this->Vertex);
+				}
+
+				bAllPlatformsAreReady = false;
+				continue;
+			}
+		}
+
+		SetIsIterativelyUnmodified(PlatformIndex, true);
 		PackagePlatformData.SetIterativelyUnmodified(true);
 	}
+
+	if (!bAllPlatformsAreReady)
+	{
+		GraphSearch.PendingTransitiveBuildDependencyVertices.Add(Vertex);
+		return false;
+	}
+
+	if (!Vertex->IterativelyModifiedListeners.IsEmpty())
+	{
+		Algo::Sort(Vertex->IterativelyModifiedListeners);
+		Vertex->IterativelyModifiedListeners.SetNum(Algo::Unique(Vertex->IterativelyModifiedListeners));
+		for (FVertexData* ListenerVertex : Vertex->IterativelyModifiedListeners)
+		{
+			GraphSearch.KickVertex(ListenerVertex);
+		}
+		Vertex->IterativelyModifiedListeners.Empty();
+	}
+	return true;
 }
 
 void FRequestCluster::FGraphSearch::FExploreEdgesContext::CalculatePackageDataDependenciesPlatformAgnostic()
@@ -1208,7 +1448,7 @@ void FRequestCluster::FGraphSearch::FExploreEdgesContext::CalculateDependenciesA
 
 	for (int32 PlatformIndex : PlatformsToExplore)
 	{
-		FQueryPlatformData& QueryPlatformData = Vertex->QueryData->Platforms[PlatformIndex];
+		FQueryPlatformData& QueryPlatformData = Vertex->PlatformData[PlatformIndex];
 		FFetchPlatformData& FetchPlatformData = GraphSearch.FetchPlatforms[PlatformIndex];
 		const ITargetPlatform* TargetPlatform = FetchPlatformData.Platform;
 		FPackagePlatformData& PackagePlatformData = PackageData->FindOrAddPlatformData(TargetPlatform);
@@ -1291,7 +1531,7 @@ void FRequestCluster::FGraphSearch::FExploreEdgesContext::QueueVisitsOfDependenc
 			continue;
 		}
 		FPackageData& DependencyPackageData(*DependencyVertex.PackageData);
-		bool bAddToFrontier = false;
+		bool bAddToVisitVertexQueue = false;
 
 		for (int32 PlatformIndex = 0; PlatformIndex < LocalNumFetchPlatforms; ++PlatformIndex)
 		{
@@ -1323,16 +1563,34 @@ void FRequestCluster::FGraphSearch::FExploreEdgesContext::QueueVisitsOfDependenc
 			}
 			if (!PlatformData.IsVisitedByCluster())
 			{
-				bAddToFrontier = true;
+				bAddToVisitVertexQueue = true;
 			}
 		}
-		if (bAddToFrontier)
+		if (bAddToVisitVertexQueue)
 		{
-			GraphSearch.AddToFrontier(DependencyVertex);
+			if (DependencyVertex.PackageData)
+			{
+				// Only pull the vertex into the cluster if it has not already been pulled into the cluster.
+				// This prevents us from trying to readd a packagedata after COTFS called Cluster->RemovePackageData.
+				if (!DependencyVertex.bPulledIntoCluster)
+				{
+					DependencyVertex.bPulledIntoCluster = true;
+					Cluster.PullIntoCluster(*DependencyVertex.PackageData);
+				}
+			}
+			GraphSearch.AddToVisitVertexQueue(DependencyVertex);
 		}
 	}
 }
 
+void FRequestCluster::FGraphSearch::FExploreEdgesContext::MarkExploreComplete()
+{
+	for (int32 PlatformIndex : PlatformsToExplore)
+	{
+		Vertex->PlatformData[PlatformIndex].bExploreCompleted = true;
+	}
+}
+
 void FRequestCluster::FGraphSearch::FExploreEdgesContext::AddPlatformDependency(FName DependencyName,
 	int32 PlatformIndex, EInstigator InstigatorType)
 {
@@ -1386,11 +1644,12 @@ void FRequestCluster::FGraphSearch::FExploreEdgesContext::ProcessPlatformAttachm
 {
 	bool bFoundBuildDefinitions = false;
 	ICookedPackageWriter* PackageWriter = FetchPlatformData.Writer;
-	FQueryPlatformData& QueryPlatformData = Vertex->QueryData->Platforms[PlatformIndex];
+	FQueryPlatformData& QueryPlatformData = Vertex->PlatformData[PlatformIndex];
 
 	if (Cluster.IsIncrementalCook() && PackagePlatformData.IsCookable())
 	{
-		bool bIterativelyUnmodified = QueryPlatformData.bIterativelyUnmodified;
+		check(QueryPlatformData.bIterativelyUnmodified.IsSet());
+		bool bIterativelyUnmodified = QueryPlatformData.bIterativelyUnmodified.GetValue();
 		if (bIterativelyUnmodified)
 		{
 			UE::TargetDomain::FCookDependencies& CookDependencies = PlatformAttachments.Dependencies;
@@ -1427,9 +1686,9 @@ void FRequestCluster::FGraphSearch::FExploreEdgesContext::ProcessPlatformAttachm
 
 	if (Cluster.bPreQueueBuildDefinitions && !bFoundBuildDefinitions)
 	{
-		FQueryPlatformData& PlatformAgnosticQueryData = Vertex->QueryData->Platforms[PlatformAgnosticPlatformIndex];
+		FQueryPlatformData& PlatformAgnosticQueryData = Vertex->PlatformData[PlatformAgnosticPlatformIndex];
 
-		if (PlatformAgnosticQueryData.bActive &&
+		if (PlatformAgnosticQueryData.bSchedulerThreadFetchCompleted &&
 			PlatformAgnosticQueryData.CookAttachments.Dependencies.HasKeyMatch())
 		{
 			Cluster.BuildDefinitions.AddBuildDefinitionList(PackageName, TargetPlatform,
@@ -1438,58 +1697,52 @@ void FRequestCluster::FGraphSearch::FExploreEdgesContext::ProcessPlatformAttachm
 	}
 }
 
-void FRequestCluster::FVertexQueryData::Reset()
+void FRequestCluster::FGraphSearch::FExploreEdgesContext::SetIsIterativelyUnmodified(int32 PlatformIndex, bool bIterativelyUnmodified)
 {
-	for (FQueryPlatformData& PlatformData : Platforms)
-	{
-		PlatformData.CookAttachments.Reset();
-		PlatformData.bActive = false;
-	}
-}
-
-FRequestCluster::FVertexData* FRequestCluster::FGraphSearch::AllocateVertex()
-{
-	return VertexAllocator.NewElement();
+	Vertex->PlatformData[PlatformIndex].bIterativelyUnmodified.Emplace(bIterativelyUnmodified);
 }
 
-FRequestCluster::FVertexQueryData* FRequestCluster::FGraphSearch::AllocateQueryData()
+FRequestCluster::FVertexData* FRequestCluster::FGraphSearch::AllocateVertex(FName PackageName, FPackageData* PackageData)
 {
-	// VertexQueryAllocator uses DeferredDestruction, so this might be a resused Batch, but we don't need to Reset it
-	// during allocation because Batches are Reset during Free.
-	return VertexQueryAllocator.NewElement();
+	// TODO: Change TypeBlockedAllocator to have an optional Size and Align argument,
+	// and use it to allocate the array of PlatformData, to reduce cpu time of allocating the array.
+	return VertexAllocator.NewElement(PackageName, PackageData, *this);
 }
 
-void FRequestCluster::FGraphSearch::FreeQueryData(FVertexQueryData* QueryData)
+FRequestCluster::FVertexData::FVertexData(FName InPackageName, UE::Cook::FPackageData* InPackageData,
+	FGraphSearch& GraphSearch)
+	: PackageName(InPackageName)
+	, PackageData(InPackageData)
 {
-	QueryData->Reset();
-	VertexQueryAllocator.Free(QueryData);
+	PlatformData.Reset(new FQueryPlatformData[GraphSearch.NumFetchPlatforms()]);
 }
 
 FRequestCluster::FVertexData&
 FRequestCluster::FGraphSearch::FindOrAddVertex(FName PackageName)
 {
-	// Only called from PumpExploration thread
+	// Only called from scheduler thread
 	FVertexData*& ExistingVertex = Vertices.FindOrAdd(PackageName);
 	if (ExistingVertex)
 	{
 		return *ExistingVertex;
 	}
 
-	ExistingVertex = AllocateVertex();
+	FPackageData* PackageData = nullptr;
 	TStringBuilder<256> NameBuffer;
 	PackageName.ToString(NameBuffer);
-	ExistingVertex->PackageData = nullptr;
 	if (!FPackageName::IsScriptPackage(NameBuffer))
 	{
-		ExistingVertex->PackageData = Cluster.COTFS.PackageDatas->TryAddPackageDataByPackageName(PackageName);
+		PackageData = Cluster.COTFS.PackageDatas->TryAddPackageDataByPackageName(PackageName);
 	}
+
+	ExistingVertex = AllocateVertex(PackageName, PackageData);
 	return *ExistingVertex;
 }
 
 FRequestCluster::FVertexData&
 FRequestCluster::FGraphSearch::FindOrAddVertex(FName PackageName, FPackageData& PackageData)
 {
-	// Only called from PumpExploration thread
+	// Only called from scheduler thread
 	FVertexData*& ExistingVertex = Vertices.FindOrAdd(PackageName);
 	if (ExistingVertex)
 	{
@@ -1497,18 +1750,13 @@ FRequestCluster::FGraphSearch::FindOrAddVertex(FName PackageName, FPackageData&
 		return *ExistingVertex;
 	}
 
-	ExistingVertex = AllocateVertex();
-	ExistingVertex->PackageData = &PackageData;
+	ExistingVertex = AllocateVertex(PackageName, &PackageData);
 	return *ExistingVertex;
 }
 
-void FRequestCluster::FGraphSearch::AddToFrontier(FVertexData& Vertex)
+void FRequestCluster::FGraphSearch::AddToVisitVertexQueue(FVertexData& Vertex)
 {
-	if (Vertex.PackageData)
-	{
-		Cluster.PullIntoCluster(*Vertex.PackageData);
-	}
-	Frontier.Add(&Vertex);
+	VisitVertexQueue.Add(&Vertex);
 }
 
 void FRequestCluster::FGraphSearch::CreateAvailableBatches(bool bAllowIncompleteBatch)
@@ -1562,8 +1810,12 @@ FRequestCluster::FQueryVertexBatch* FRequestCluster::FGraphSearch::CreateBatchOf
 	for (int32 BatchIndex = 0; BatchIndex < BatchSize; ++BatchIndex)
 	{
 		FVertexData* Vertex = PreAsyncQueue.PopFrontValue();
-		FVertexData*& ExistingVert = BatchData->Vertices.FindOrAdd(Vertex->QueryData->PackageName);
-		check(!ExistingVert); // We should not have any duplicate names in PreAsyncQueue
+		FVertexData*& ExistingVert = BatchData->Vertices.FindOrAdd(Vertex->PackageName);
+		// Each PackageName should be used by just a single vertex.
+		check(!ExistingVert || ExistingVert == Vertex);
+		// If the vertex was already previously added to the batch that's okay, just ignore the new add.
+		// A batch size of 0 is a problem but that can't happen just because a vertex is in the batch twice.
+		// A batch size smaller than the expected `BatchSize` parameter is a minor performance issue but not a problem.
 		ExistingVert = Vertex;
 	}
 	AsyncQueueBatches.Add(BatchData);
@@ -1578,13 +1830,14 @@ void FRequestCluster::FGraphSearch::OnBatchCompleted(FQueryVertexBatch* Batch)
 	AsyncResultsReadyEvent->Trigger();
 }
 
-void FRequestCluster::FGraphSearch::OnVertexCompleted()
+void FRequestCluster::FGraphSearch::KickVertex(FVertexData* Vertex)
 {
 	// The trigger occurs outside of the lock, and might get clobbered and incorrectly ignored by a call from the
-	// consumer thread if the consumer tried to consume and found the vertices empty before our caller added a vertex
-	// but then pauses and calls AsyncResultsReadyEvent->Reset after this AsyncResultsReadyEvent->Trigger.
-	// This clobbering will not cause a deadlock, because eventually DestroyBatch will be called which triggers it
-	// inside the lock. Doing the per-vertex trigger outside the lock is good for performance.
+	// scheduler thread if the scheduler tried to pop the AsyncQueueResults and found it empty before KickVertex calls
+	// Enqueue but then pauses and calls AsyncResultsReadyEvent->Reset after KicKVertex calls Trigger. This clobbering
+	// will not cause a deadlock, because eventually DestroyBatch will be called which triggers it inside the lock. Doing
+	// the per-vertex trigger outside the lock is good for performance.
+	AsyncQueueResults.Enqueue(Vertex);
 	AsyncResultsReadyEvent->Trigger();
 }
 
@@ -1605,24 +1858,46 @@ void FRequestCluster::FQueryVertexBatch::Reset()
 
 void FRequestCluster::FQueryVertexBatch::Send()
 {
+	int32 NumAddedRequests = 0;
 	for (const TPair<FName, FVertexData*>& Pair : Vertices)
 	{
 		FVertexData* Vertex = Pair.Value;
-		TArray<FQueryPlatformData>& QueryPlatforms = Vertex->QueryData->Platforms;
-		bool bAtLeastOnePlatform = false;
+		bool bAnyRequested = false;
+		bool bAllHaveAlreadyCompletedFetch = false;
 		for (int32 PlatformIndex = 0; PlatformIndex < PlatformDatas.Num(); ++PlatformIndex)
 		{
-			if (QueryPlatforms[PlatformIndex].bActive)
+			// The platform data may have already been requested; request it only if current status is NotRequested
+			FQueryPlatformData& PlatformData = Vertex->PlatformData[PlatformIndex];
+			if (!PlatformData.bSchedulerThreadFetchCompleted)
 			{
-				PlatformDatas[PlatformIndex].PackageNames.Add(Pair.Key);
+				bAllHaveAlreadyCompletedFetch = false;
+				EAsyncQueryStatus ExpectedStatus = EAsyncQueryStatus::SchedulerRequested;
+				if (PlatformData.CompareExchangeAsyncQueryStatus(ExpectedStatus,
+					EAsyncQueryStatus::AsyncRequested))
+				{
+					PlatformDatas[PlatformIndex].PackageNames.Add(Pair.Key);
+					++NumAddedRequests;
+				}
 			}
-			bAtLeastOnePlatform = true;
 		}
-		// We only check for the vertex's completion when the vertex receives a callback from the completion of a
-		// platform. Therefore we do not support Vertices in the batch that have no platforms.
-		check(bAtLeastOnePlatform);
+		if (bAllHaveAlreadyCompletedFetch)
+		{
+			// We are contractually obligated to kick the vertex. Normally we would call FCookAttachments::Fetch with it
+			// and would then kick the vertex in our callback. Also, it might still be in the AsyncQueueResults for one
+			// of the platforms so it will be kicked by TickExplore pulling it out of the AsyncQueueResults. But if all
+			// requested platforms already previously pulled it out of AsyncQueueResults, then we need to kick it again.
+			ThreadSafeOnlyVars.KickVertex(Vertex);
+		}
+	}
+	if (NumAddedRequests == 0)
+	{
+		// We turned out not to need to send any from this batch. Report that the batch is complete.
+		ThreadSafeOnlyVars.OnBatchCompleted(this);
+		// *this is no longer accessible
+		return;
 	}
-	PendingVertices.store(Vertices.Num(), std::memory_order_release);
+
+	NumPendingRequests.store(NumAddedRequests, std::memory_order_release);
 
 	for (int32 PlatformIndex = 0; PlatformIndex < PlatformDatas.Num(); ++PlatformIndex)
 	{
@@ -1672,23 +1947,46 @@ void FRequestCluster::FQueryVertexBatch::RecordCacheResults(FName PackageName, i
 	UE::TargetDomain::FCookAttachments&& CookAttachments)
 {
 	FVertexData* Vertex = Vertices.FindChecked(PackageName);
-	check(Vertex->QueryData);
-	FVertexQueryData& QueryData = *Vertex->QueryData;
-	QueryData.Platforms[PlatformIndex].CookAttachments = MoveTemp(CookAttachments);
-	if (QueryData.PendingPlatforms.fetch_sub(1, std::memory_order_acq_rel) == 1)
-	{
-		ThreadSafeOnlyVars.AsyncQueueResults.Enqueue(Vertex);
-		bool bBatchComplete = PendingVertices.fetch_sub(1, std::memory_order_relaxed) == 1;
-		if (!bBatchComplete)
-		{
-			ThreadSafeOnlyVars.OnVertexCompleted();
+	FQueryPlatformData& PlatformData = Vertex->PlatformData[PlatformIndex];
+	PlatformData.CookAttachments = MoveTemp(CookAttachments);
+
+	EAsyncQueryStatus Expected = EAsyncQueryStatus::AsyncRequested;
+	if (PlatformData.CompareExchangeAsyncQueryStatus(Expected, EAsyncQueryStatus::Complete))
+	{
+		// Kick the vertex if it has no more platforms in pending. Otherwise keep waiting and the later
+		// call to RecordCacheResults will kick the vertex. Note that the "later call" might be another
+		// call to RecordCacheResults on a different thread executing at the same time, and we are racing.
+		// The last one to set CompareExchangeAsyncQueryStatus(EAsyncQueryStatus::Complete) will definitely
+		// see all other values as complete, because we are using std::memory_order_release. It is possible
+		// that both calls to RecordCacheResults will see all values complete, and we will kick it twice.
+		// Kicking twice is okay; it is supported and is a noop.
+		bool bAllPlatformsComplete = true;
+		int32 LocalNumFetchPlatforms = ThreadSafeOnlyVars.NumFetchPlatforms();
+		for (int32 OtherPlatformIndex = 0; OtherPlatformIndex < LocalNumFetchPlatforms; ++OtherPlatformIndex)
+		{
+			if (OtherPlatformIndex == PlatformIndex)
+			{
+				continue;
+			}
+			FQueryPlatformData& OtherPlatformData = Vertex->PlatformData[OtherPlatformIndex];
+			EAsyncQueryStatus OtherStatus = OtherPlatformData.GetAsyncQueryStatus();
+			if (EAsyncQueryStatus::AsyncRequested <= OtherStatus && OtherStatus < EAsyncQueryStatus::Complete)
+			{
+				bAllPlatformsComplete = false;
+				break;
+			}
 		}
-		else
+		if (bAllPlatformsComplete)
 		{
-			ThreadSafeOnlyVars.OnBatchCompleted(this);
-			// *this is no longer accessible
+			ThreadSafeOnlyVars.KickVertex(Vertex);
 		}
 	}
+
+	if (NumPendingRequests.fetch_sub(1, std::memory_order_relaxed) == 1)
+	{
+		ThreadSafeOnlyVars.OnBatchCompleted(this);
+		// *this is no longer accessible
+	}
 }
 
 TMap<FPackageData*, TArray<FPackageData*>>& FRequestCluster::FGraphSearch::GetGraphEdges()
diff --git a/Engine/Source/Editor/UnrealEd/Private/Cooker/CookRequestCluster.h b/Engine/Source/Editor/UnrealEd/Private/Cooker/CookRequestCluster.h
--- a/Engine/Source/Editor/UnrealEd/Private/Cooker/CookRequestCluster.h
+++ b/Engine/Source/Editor/UnrealEd/Private/Cooker/CookRequestCluster.h
@@ -98,29 +98,37 @@ private:
 		bool bVisited = false;
 	};
 
-	/** Per-platform data in an active query for a vertex's dependencies/previous incremental results. */
-	struct FQueryPlatformData
+	/** Status for where a vertex is on the journey through having its CookDependency information fetched from DDC. */
+	enum class EAsyncQueryStatus : uint8
 	{
-		/** Data looked up about the package's dependencies from the PackageWriter's previous cook of the package */
-		UE::TargetDomain::FCookAttachments CookAttachments;
-		/**
-		 * Platforms present in the FVertexQueryData are recorded as a an array of all possible platforms, with
-		 * bActive=true for the ones that are present.
-		 */
-		bool bActive = false;
-		bool bIterativelyUnmodified = false;
+		NotRequested,
+		SchedulerRequested,
+		AsyncRequested,
+		Complete,
 	};
 
-	/** Input/output data for an active query for a vertex's dependencies/previous incrementa results. */
-	struct FVertexQueryData
+	/** Per-platform data in an active query for a vertex's dependencies/previous incremental results. */
+	struct FQueryPlatformData
 	{
-		void Reset();
+		EAsyncQueryStatus GetAsyncQueryStatus();
+		bool CompareExchangeAsyncQueryStatus(EAsyncQueryStatus& Expected, EAsyncQueryStatus Desired);
 
-		FName PackageName;
-		/** Settings and Results for each of the GraphSearch's FetchPlatforms. Element n corresponds to FetchPlatform n. */
-		TArray<FQueryPlatformData> Platforms;
-		/** Number of asynchronous results from PackageWriter still being waited on. */
-		std::atomic<uint32> PendingPlatforms;
+	public:
+		// All fields other than CookAttachments and AsyncQueryStatus are read/write on Scheduler thread only
+		/**
+		 * Data looked up about the package's dependencies from the PackageWriter's previous cook of the package.
+		 * Thread synchronization: this field is write-once from the async thread and is not readable until
+		 * bSchedulerThreadFetchCompleted.
+		 */
+		UE::TargetDomain::FCookAttachments CookAttachments;
+		bool bSchedulerThreadFetchCompleted = false;
+		bool bExploreRequested = false;
+		bool bExploreCompleted = false;
+		bool bIterativelyUnmodifiedRequested = false;
+		bool bTransitiveBuildDependenciesResolvedAsNotModified = false;
+		TOptional<bool> bIterativelyUnmodified;
+	private:
+		std::atomic<EAsyncQueryStatus> AsyncQueryStatus;
 	};
 
 	/**
@@ -129,10 +137,16 @@ private:
 	 */
 	struct FVertexData
 	{
+		FVertexData(FName InPackageName, UE::Cook::FPackageData* InPackageData, FGraphSearch& GraphSearch);
+
+		/* Async thread is not allowed to access PackageData, so store its name.The name is immutable for vertex lifetime. */
+		FName PackageName;
+		TArray<FVertexData*> IterativelyModifiedListeners;
 		UE::Cook::FPackageData* PackageData = nullptr;
-		/** Non-null if there as a query active for the vertex's dependencies/previous incremental results. */
-		FVertexQueryData* QueryData = nullptr;
 		bool bAnyCookable = true;
+		bool bPulledIntoCluster = false;
+		/** Settings and Results for each of the GraphSearch's FetchPlatforms. Element n corresponds to FetchPlatform n. */
+		TUniquePtr<FQueryPlatformData[]> PlatformData;
 	};
 
 	/**
@@ -165,8 +179,8 @@ private:
 		TMap<FName, FVertexData*> Vertices;
 		/** Accessor for the GraphSearch; only thread-safe functions and variables should be accessed. */
 		FGraphSearch& ThreadSafeOnlyVars;
-		/** The number of vertices that are still awaiting results. This batch is closed when PendingVertices == 0. */
-		std::atomic<uint32> PendingVertices;
+		/** Number of vertex*platform requests that still await results. Batch is done when NumPendingRequests == 0. */
+		std::atomic<uint32> NumPendingRequests;
 	};
 
 	/** Platform information that is constant (usually, some events can change it) during the cluster's lifetime. */
@@ -216,7 +230,6 @@ private:
 		void VisitWithoutDependencies();
 		/** Start a search from the Cluster's current OwnedPackageDatas. */
 		void StartSearch();
-		void RemovePackageData(FPackageData* PackageData);
 		void OnNewReachablePlatforms(FPackageData* PackageData);
 
 		/**
@@ -245,16 +258,20 @@ private:
 		public:
 			FExploreEdgesContext(FRequestCluster& InCluster, FGraphSearch& InGraphSearch);
 
-			/** Process the results from async edges fetch and queue the found dependencies-for-visiting. */
+			/**
+			 * Process the results from async edges fetch and queue the found dependencies-for-visiting. Only does
+			 * portions of the work for each FQueryPlatformData that were requested by the flags on the PlatformData.
+			 */
 			void Explore(FVertexData& InVertex);
 
 		private:
 			void Initialize(FVertexData& InVertex);
-			void CalculatePlatformsToExplore();
-			void CalculateIterativelyUnmodified();
+			void CalculatePlatformsToProcess();
+			bool TryCalculateIterativelyUnmodified();
 			void CalculatePackageDataDependenciesPlatformAgnostic();
 			void CalculateDependenciesAndIterativelySkippable();
 			void QueueVisitsOfDependencies();
+			void MarkExploreComplete();
 
 			void AddPlatformDependency(FName DependencyName, int32 PlatformIndex, EInstigator InstigatorType);
 			void AddPlatformDependencyRange(TConstArrayView<FName> Range, int32 PlatformIndex, EInstigator InstigatorType);
@@ -262,6 +279,8 @@ private:
 				FFetchPlatformData& FetchPlatformData, FPackagePlatformData& PackagePlatformData,
 				UE::TargetDomain::FCookAttachments& PlatformAttachments, bool bExploreDependencies);
 
+			void SetIsIterativelyUnmodified(int32 PlatformIndex, bool bIterativelyUnmodified);
+
 		private:
 			FRequestCluster& Cluster;
 			FGraphSearch& GraphSearch;
@@ -272,10 +291,12 @@ private:
 			TArray<FName> HardEditorDependencies;
 			TArray<FName> SoftGameDependencies;
 			TArray<FName> CookerLoadingDependencies;
+			TArray<int32, TInlineAllocator<10>> PlatformsToProcess;
 			TArray<int32, TInlineAllocator<10>> PlatformsToExplore;
 			TMap<FName, FScratchPlatformDependencyBits> PlatformDependencyMap;
 			TSet<FName> HardDependenciesSet;
 			TSet<FName> SkippedPackages;
+			TArray<FVertexData*> UnreadyTransitiveBuildVertices;
 			FName PackageName;
 			int32 LocalNumFetchPlatforms = 0;
 			bool bFetchAnyTargetPlatform = false;
@@ -288,23 +309,21 @@ private:
 		void UpdateDisplay();
 
 		/** Asynchronously fetch the dependencies and previous incremental results for a vertex */
-		void QueueEdgesFetch(FVertexData& Vertex, TConstArrayView<const ITargetPlatform*> Platforms);
+		void QueueEdgesFetch(FVertexData& Vertex, TConstArrayView<int32> PlatformIndexes);
 		/** Calculate and store the vertex's PackageData's cookability for each reachable platform. Kick off edges fetch. */
 		void VisitVertex(FVertexData& VertexData);
 		/** Calculate and store the vertex's PackageData's cookability for the platform. */
 		void VisitVertexForPlatform(FVertexData& VertexData, const ITargetPlatform* Platform,
 			FPackagePlatformData& PlatformData, ESuppressCookReason& AccumulatedSuppressCookReason);
+		void ResolveTransitiveBuildDependencyCycle();
 
 		/** Find or add a Vertex for PackageName. If PackageData is provided, use it, otherwise look it up. */
 		FVertexData& FindOrAddVertex(FName PackageName);
 		FVertexData& FindOrAddVertex(FName PackageName, FPackageData& PackageData);
 		/** Batched allocation for vertices. */
-		FVertexData* AllocateVertex();
-		/** Batched allocation and free for QueryData. */
-		FVertexQueryData* AllocateQueryData();
-		void FreeQueryData(FVertexQueryData* QueryData);
+		FVertexData* AllocateVertex(FName PackageName, FPackageData* PackageData);
 		/** Queue a vertex for visiting and dependency traversal */
-		void AddToFrontier(FVertexData& Vertex);
+		void AddToVisitVertexQueue(FVertexData& Vertex);
 
 		// Functions that must be called only within the Lock
 		/** Allocate memory for a new batch; returned batch is not yet constructed. */
@@ -320,13 +339,15 @@ private:
 		/** Notify process thread of batch completion and deallocate it. */
 		void OnBatchCompleted(FQueryVertexBatch* Batch);
 		/** Notify process thread of vertex completion. */
-		void OnVertexCompleted();
+		void KickVertex(FVertexData* Vertex);
 
 		/** Total number of platforms known to the cluster, including the special cases. */
 		int32 NumFetchPlatforms() const;
 		/** Total number of non-special-case platforms known to the cluster.Identical to COTFS's session platforms */
 		int32 NumSessionPlatforms() const;
 
+		TArrayView<FQueryPlatformData> GetPlatformDataArray(FVertexData& Vertex);
+
 	private:
 		// Variables that are read-only during multithreading
 		TArray<FFetchPlatformData> FetchPlatforms;
@@ -338,13 +359,14 @@ private:
 		FExploreEdgesContext ExploreEdgesContext;
 		TMap<FPackageData*, TArray<FPackageData*>> GraphEdges;
 		TMap<FName, FVertexData*> Vertices;
-		TSet<FVertexData*> Frontier;
+		TSet<FVertexData*> VisitVertexQueue;
+		TSet<FVertexData*> PendingTransitiveBuildDependencyVertices;
 		TTypedBlockAllocatorFreeList<FVertexData> VertexAllocator;
-		TTypedBlockAllocatorResetList<FVertexQueryData> VertexQueryAllocator;
 		/** Vertices queued for async processing that are not yet numerous enough to fill a batch. */
 		TRingBuffer<FVertexData*> PreAsyncQueue;
 		/** Time-tracker for timeout warnings in Poll */
 		double LastActivityTime = 0.;
+		int32 RunAwayTickLoopCount = 0;
 
 		// Variables that are accessible from multiple threads, guarded by Lock
 		FCriticalSection Lock;
@@ -480,4 +502,29 @@ inline int32 FRequestCluster::FGraphSearch::NumSessionPlatforms() const
 	return FetchPlatforms.Num() - 2;
 }
 
+inline TArrayView<FRequestCluster::FQueryPlatformData> FRequestCluster::FGraphSearch::GetPlatformDataArray(
+	FVertexData& Vertex)
+{
+	return TArrayView< FQueryPlatformData>(Vertex.PlatformData.Get(), NumFetchPlatforms());
+}
+
+inline FRequestCluster::EAsyncQueryStatus FRequestCluster::FQueryPlatformData::GetAsyncQueryStatus()
+{
+	return AsyncQueryStatus.load(std::memory_order_acquire);
+}
+
+inline bool FRequestCluster::FQueryPlatformData::CompareExchangeAsyncQueryStatus(EAsyncQueryStatus& Expected,
+	EAsyncQueryStatus Desired)
+{
+	return AsyncQueryStatus.compare_exchange_strong(Expected, Desired,
+		// For the read operation to see whether we should set it, we need only relaxed memory order;
+		// we don't care about the values of other related variables that depend on it when deciding whether
+		// it is our turn to set it.
+		// For the write operation if we decide to set it, we need release memory order to guard reads of
+		// the variables that depend on it (e.g. CookAttachments).
+		std::memory_order_release /* success memory order */,
+		std::memory_order_relaxed /* failure memory order */
+	);
+}
+
 }
diff --git a/Engine/Source/Editor/UnrealEd/Private/TargetDomain/TargetDomainUtils.cpp b/Engine/Source/Editor/UnrealEd/Private/TargetDomain/TargetDomainUtils.cpp
--- a/Engine/Source/Editor/UnrealEd/Private/TargetDomain/TargetDomainUtils.cpp
+++ b/Engine/Source/Editor/UnrealEd/Private/TargetDomain/TargetDomainUtils.cpp
@@ -281,9 +281,6 @@ FCookDependencies FCookDependencies::Collect(UPackage* Package, const ITargetPla
 	AssetRegistry->GetDependencies(Result.PackageName, AssetDependencies,
 		UE::AssetRegistry::EDependencyCategory::Package, UE::AssetRegistry::EDependencyQuery::Game);
 	RuntimeDependencies.Append(MoveTemp(AssetDependencies));
-	RuntimeDependencies.RemoveAllSwap(IsTransientPackageName, EAllowShrinking::No);
-	RuntimeDependencies.Sort(FNameLexicalLess());
-	RuntimeDependencies.SetNum(Algo::Unique(RuntimeDependencies), EAllowShrinking::Yes);
 
 	FPackageBuildDependencyTracker& Tracker = FPackageBuildDependencyTracker::Get();
 
@@ -312,8 +309,6 @@ FCookDependencies FCookDependencies::Collect(UPackage* Package, const ITargetPla
 	Result.BuildPackageDependencies.RemoveAllSwap(IsTransientPackageName, EAllowShrinking::Yes);
 	Result.BuildPackageDependencies.Sort(FNameLexicalLess());
 
-	Result.RuntimePackageDependencies = MoveTemp(RuntimeDependencies);
-
 #if UE_WITH_CONFIG_TRACKING
 	{
 		using namespace UE::ConfigAccessTracking;
@@ -331,10 +326,43 @@ FCookDependencies FCookDependencies::Collect(UPackage* Package, const ITargetPla
 #endif
 	if (SaveResult)
 	{
-		Result.CookDependencies = MoveTemp(SaveResult->CookDependencies);
+		// Move most types of FCookDependency into this->CookDependencies, but separate out
+		// TransitiveBuildDependencies; they require different processing than other types.
+		Result.CookDependencies.Reserve(SaveResult->CookDependencies.Num());
+		for (UE::Cook::FCookDependency& CookDependency : SaveResult->CookDependencies)
+		{
+			if (CookDependency.GetType() == UE::Cook::ECookDependency::TransitiveBuild)
+			{
+				// Build dependencies from a package to itself have a performance cost and serve no
+				// purpose, so remove them. They can occur in some systems that naively add a build dependency
+				// from one object to another without checking whether the second object is in the same package.
+				FName DependencyPackageName = CookDependency.GetPackageName();
+				if (DependencyPackageName == Result.PackageName)
+				{
+					continue;
+				}
+				if (CookDependency.IsAlsoAddRuntimeDependency())
+				{
+					RuntimeDependencies.Add(DependencyPackageName);
+				}
+				Result.TransitiveBuildDependencies.Add(MoveTemp(CookDependency));
+			}
+			else
+			{
+				Result.CookDependencies.Add(MoveTemp(CookDependency));
+			}
+		}
+
 		Algo::Sort(Result.CookDependencies);
+		Algo::Sort(Result.TransitiveBuildDependencies);
 	}
 
+	RuntimeDependencies.RemoveAllSwap(IsTransientPackageName, EAllowShrinking::No);
+	RuntimeDependencies.Sort(FNameLexicalLess());
+	RuntimeDependencies.SetNum(Algo::Unique(RuntimeDependencies), EAllowShrinking::Yes);
+
+	Result.RuntimePackageDependencies = MoveTemp(RuntimeDependencies);
+
 	if (!Result.TryCalculateCurrentKey(OutErrorMessage))
 	{
 		return FCookDependencies();
@@ -401,6 +429,13 @@ bool LoadFromCompactBinary(FCbObjectView ObjectView, UE::TargetDomain::FCookDepe
 				return false;
 			}
 		}
+		if (FieldView.GetName().Equals(UTF8TEXTVIEW("TransitiveBuildDependencies")))
+		{
+			if (!LoadFromCompactBinary(FieldView++, Dependencies.TransitiveBuildDependencies))
+			{
+				return false;
+			}
+		}
 		if (FieldView == Last)
 		{
 			++FieldView;
@@ -437,6 +472,10 @@ FCbWriter& operator<<(FCbWriter& Writer, const UE::TargetDomain::FCookDependenci
 	{
 		Writer << "CookDependencies" << CookDependencies.CookDependencies;
 	}
+	if (!CookDependencies.TransitiveBuildDependencies.IsEmpty())
+	{
+		Writer << "TransitiveBuildDependencies" << CookDependencies.TransitiveBuildDependencies;
+	}
 
 	Writer.EndObject();
 	return Writer;
diff --git a/Engine/Source/Editor/UnrealEd/Private/TargetDomain/TargetDomainUtils.h b/Engine/Source/Editor/UnrealEd/Private/TargetDomain/TargetDomainUtils.h
--- a/Engine/Source/Editor/UnrealEd/Private/TargetDomain/TargetDomainUtils.h
+++ b/Engine/Source/Editor/UnrealEd/Private/TargetDomain/TargetDomainUtils.h
@@ -63,6 +63,7 @@ public:
 
 	// Cook Dependencies
 	const TArray<UE::Cook::FCookDependency>& GetCookDependencies() const;
+	const TArray<UE::Cook::FCookDependency>& GetTransitiveBuildDependencies() const;
 
 	// Data about the Key and Package that are not dependencies
 	/**
@@ -96,6 +97,7 @@ private:
 	TArray<FString> ConfigDependencies;
 	TArray<FName> RuntimePackageDependencies;
 	TArray<UE::Cook::FCookDependency> CookDependencies;
+	TArray<UE::Cook::FCookDependency> TransitiveBuildDependencies;
 	FName PackageName;
 	FIoHash StoredKey;
 	FIoHash CurrentKey;
@@ -193,6 +195,11 @@ inline const TArray<UE::Cook::FCookDependency>& FCookDependencies::GetCookDepend
 	return CookDependencies;
 }
 
+inline const TArray<UE::Cook::FCookDependency>& FCookDependencies::GetTransitiveBuildDependencies() const
+{
+	return TransitiveBuildDependencies;
+}
+
 inline FName FCookDependencies::GetPackageName() const
 {
 	return PackageName;
diff --git a/Engine/Source/Runtime/CoreUObject/Private/Cooker/CookDependency.cpp b/Engine/Source/Runtime/CoreUObject/Private/Cooker/CookDependency.cpp
--- a/Engine/Source/Runtime/CoreUObject/Private/Cooker/CookDependency.cpp
+++ b/Engine/Source/Runtime/CoreUObject/Private/Cooker/CookDependency.cpp
@@ -50,6 +50,14 @@ FCookDependency FCookDependency::Function(FName InFunctionName, FCbFieldIterator
 	return Result;
 }
 
+FCookDependency FCookDependency::TransitiveBuildAndRuntime(FName PackageName)
+{
+	FCookDependency Result(ECookDependency::TransitiveBuild);
+	Result.TransitiveBuildData.PackageName = PackageName;
+	Result.TransitiveBuildData.bAlsoAddRuntimeDependency = true;
+	return Result;
+}
+
 FCookDependency::FCookDependency()
 	: Type(ECookDependency::None)
 {
@@ -93,6 +101,9 @@ FCookDependency& FCookDependency::operator=(const FCookDependency& Other)
 	case ECookDependency::Function:
 		FunctionData = Other.FunctionData;
 		break;
+	case ECookDependency::TransitiveBuild:
+		TransitiveBuildData = Other.TransitiveBuildData;
+		break;
 	default:
 		checkNoEntry();
 		break;
@@ -115,6 +126,9 @@ FCookDependency& FCookDependency::operator=(FCookDependency&& Other)
 	case ECookDependency::Function:
 		FunctionData = MoveTemp(Other.FunctionData);
 		break;
+	case ECookDependency::TransitiveBuild:
+		TransitiveBuildData = MoveTemp(Other.TransitiveBuildData);
+		break;
 	default:
 		checkNoEntry();
 		break;
@@ -176,6 +190,11 @@ void FCookDependency::UpdateHash(FCookDependencyContext& Context) const
 		(*Function)(GetFunctionArgs(), Context);
 		return;
 	}
+	case ECookDependency::TransitiveBuild:
+		// Transitive Build dependencies do not impact the hash; they instead operate by marking the package
+		// as invalidated based on the invalidation of other packages, in a separate pass after its hash is compared
+		return;
+
 	default:
 		checkNoEntry();
 		return;
@@ -194,6 +213,9 @@ void FCookDependency::Construct()
 	case ECookDependency::Function:
 		new(&FunctionData) FFunctionData();
 		break;
+	case ECookDependency::TransitiveBuild:
+		new(&TransitiveBuildData) FTransitiveBuildData();
+		break;
 	default:
 		checkNoEntry();
 		break;
@@ -212,6 +234,9 @@ void FCookDependency::Destruct()
 	case ECookDependency::Function:
 		FunctionData.~FFunctionData();
 		break;
+	case ECookDependency::TransitiveBuild:
+		TransitiveBuildData.~FTransitiveBuildData();
+		break;
 	default:
 		checkNoEntry();
 		break;
@@ -236,6 +261,10 @@ void FCookDependency::Save(FCbWriter& Writer) const
 			Writer << Iter;
 		}
 		break;
+	case ECookDependency::TransitiveBuild:
+		Writer << TransitiveBuildData.PackageName;
+		Writer << TransitiveBuildData.bAlsoAddRuntimeDependency;
+		break;
 	default:
 		checkNoEntry();
 		break;
@@ -283,6 +312,22 @@ bool FCookDependency::Load(FCbFieldView Value)
 		FunctionData.Args.MakeOwned();
 		return true;
 	}
+	case ECookDependency::TransitiveBuild:
+	{
+		FName LocalPackageName;
+		if (!LoadFromCompactBinary(Field++, LocalPackageName))
+		{
+			return false;
+		}
+		bool bLocalAlsoAddRuntimeDependency = Field.AsBool();
+		if ((Field++).HasError())
+		{
+			return false;
+		}
+		*this = FCookDependency::TransitiveBuildAndRuntime(LocalPackageName);
+		TransitiveBuildData.bAlsoAddRuntimeDependency = bLocalAlsoAddRuntimeDependency;
+		return true;
+	}
 	default:
 		break;
 	}
diff --git a/Engine/Source/Runtime/CoreUObject/Public/Cooker/CookDependency.h b/Engine/Source/Runtime/CoreUObject/Public/Cooker/CookDependency.h
--- a/Engine/Source/Runtime/CoreUObject/Public/Cooker/CookDependency.h
+++ b/Engine/Source/Runtime/CoreUObject/Public/Cooker/CookDependency.h
@@ -78,6 +78,7 @@ enum class ECookDependency : uint8
 	None = 0x00,
 	File = 0x01,
 	Function = 0x02,
+	TransitiveBuild = 0x03,
 };
 
 /**
@@ -106,6 +107,18 @@ public:
 	 * UE_COOK_DEPENDENCY_FUNCTION_CALL(CppTokenUsedAsName).
 	 */
 	COREUOBJECT_API static FCookDependency Function(FName InFunctionName, FCbFieldIterator&& InArgs);
+
+	/**
+	 * Create a transitive build dependency on another package. In an incremental cook if the other package was not
+	 * cooked in a previous cook session, or its previous cook result was invalidated, the current package will also
+	 * have its cook result invalidated.
+	 *
+	 * This version of the function also adds a runtime dependency - the requested package will be staged for the
+	 * current platform. Adding a transitive build dependency without adding a runtime dependency is not yet supported
+	 * due to limitations in the cooker.
+	 */
+	COREUOBJECT_API static FCookDependency TransitiveBuildAndRuntime(FName PackageName);
+
 	/** Construct an empty dependency; it will never be invalidated. */
 	COREUOBJECT_API FCookDependency();
 
@@ -126,6 +139,11 @@ public:
 	/** FunctionArgs if GetType() == Function, else FCbFieldViewIterator(). */
 	FCbFieldViewIterator GetFunctionArgs() const;
 
+	/** PackageName if GetType() == TransitiveBuild, else NAME_None. */
+	FName GetPackageName() const;
+	/** If GetType() == TransitiveBuild, whether AlsoAddRuntimeDependency was selected, otherwise false. */
+	bool IsAlsoAddRuntimeDependency() const;
+
 	/**
 	 * Comparison operator for e.g. deterministic ordering of dependencies.
 	 * Uses persistent comparison data and is somewhat expensive.
@@ -161,10 +179,16 @@ private:
 		FName Name;
 		FCbFieldIterator Args;
 	};
+	struct FTransitiveBuildData
+	{
+		FName PackageName;
+		bool bAlsoAddRuntimeDependency = true;
+	};
 	union
 	{
 		FString FileName;
 		FFunctionData FunctionData;
+		FTransitiveBuildData TransitiveBuildData;
 	};
 };
 
@@ -262,6 +286,16 @@ inline FCbFieldViewIterator FCookDependency::GetFunctionArgs() const
 	return Type == ECookDependency::Function ? FunctionData.Args : FCbFieldViewIterator();
 }
 
+inline FName FCookDependency::GetPackageName() const
+{
+	return Type == ECookDependency::TransitiveBuild ? TransitiveBuildData.PackageName : NAME_None;
+}
+
+inline bool FCookDependency::IsAlsoAddRuntimeDependency() const
+{
+	return Type == ECookDependency::TransitiveBuild ? TransitiveBuildData.bAlsoAddRuntimeDependency : false;
+}
+
 inline bool FCookDependency::operator<(const FCookDependency& Other) const
 {
 	if (static_cast<uint8>(Type) != static_cast<uint8>(Other.Type))
@@ -278,7 +312,10 @@ inline bool FCookDependency::operator<(const FCookDependency& Other) const
 	case ECookDependency::Function:
 	{
 		int32 Compare = FunctionData.Name.Compare(Other.FunctionData.Name);
-		if (Compare != 0) return Compare < 0;
+		if (Compare != 0)
+		{
+			return Compare < 0;
+		}
 		FMemoryView ViewA;
 		FMemoryView ViewB;
 		bool bHasViewA = FunctionData.Args.TryGetRangeView(ViewA);
@@ -289,6 +326,20 @@ inline bool FCookDependency::operator<(const FCookDependency& Other) const
 		}
 		return ViewA.CompareBytes(ViewB) < 0;
 	}
+	case ECookDependency::TransitiveBuild:
+	{
+		// FName.Compare is lexical and case-insensitive, which is what we want
+		int32 Compare = TransitiveBuildData.PackageName.Compare(Other.TransitiveBuildData.PackageName) < 0;
+		if (Compare != 0)
+		{
+			return Compare < 0;
+		}
+		if (TransitiveBuildData.bAlsoAddRuntimeDependency != Other.TransitiveBuildData.bAlsoAddRuntimeDependency)
+		{
+			return TransitiveBuildData.bAlsoAddRuntimeDependency == false;
+		}
+		return false;
+	}
 	default:
 		checkNoEntry();
 		return false;
diff --git a/Engine/Source/Runtime/Engine/Private/Materials/MaterialInstance.cpp b/Engine/Source/Runtime/Engine/Private/Materials/MaterialInstance.cpp
--- a/Engine/Source/Runtime/Engine/Private/Materials/MaterialInstance.cpp
+++ b/Engine/Source/Runtime/Engine/Private/Materials/MaterialInstance.cpp
@@ -62,6 +62,10 @@
 #include "VT/RuntimeVirtualTexture.h"
 #include "LocalVertexFactory.h"
 
+#if WITH_EDITOR
+#include "Cooker/CookDependency.h"
+#endif
+
 #include UE_INLINE_GENERATED_CPP_BY_NAME(MaterialInstance)
 
 DECLARE_CYCLE_STAT(TEXT("MaterialInstance CopyMatInstParams"), STAT_MaterialInstance_CopyMatInstParams, STATGROUP_Shaders);
@@ -4804,6 +4808,13 @@ void UMaterialInstance::PreSave(FObjectPreSaveContext ObjectSaveContext)
 		TArray<FMaterialResourceForCooking>* Resources = CachedMaterialResourcesForCooking.Find(TargetPlatform);
 		UE::MaterialInterface::Private::RecordMaterialDependenciesForCook(ObjectSaveContext,
 			Resources ? *Resources : TArray<FMaterialResourceForCooking>());
+
+		UMaterialInterface* EffectiveParent = Parent ? Parent : UMaterial::GetDefaultMaterial(MD_Surface);
+		if (EffectiveParent && EffectiveParent->GetPackage() != GetPackage())
+		{
+			ObjectSaveContext.AddCookBuildDependency(UE::Cook::FCookDependency::TransitiveBuildAndRuntime(
+				EffectiveParent->GetPackage()->GetFName()));
+		}
 	}
 #endif
 }
