VolumeTexturePreview use SizeZ from actual texture

Not only is there no good reason to pass it in separately when SizeX/Y are taken from
the texture, the way it was computed didn't actually match the real SizeZ in case LOD
bias is clamped, which is very likely for volumes.

diff --git a/Engine/Source/Editor/TextureEditor/Private/Models/TextureEditorViewportClient.cpp b/Engine/Source/Editor/TextureEditor/Private/Models/TextureEditorViewportClient.cpp
--- a/Engine/Source/Editor/TextureEditor/Private/Models/TextureEditorViewportClient.cpp
+++ b/Engine/Source/Editor/TextureEditor/Private/Models/TextureEditorViewportClient.cpp
@@ -265,7 +265,6 @@ void FTextureEditorViewportClient::Draw(FViewport* Viewport, FCanvas* Canvas)
 		{
 			BatchedElementParameters = new FBatchedElementVolumeTexturePreviewParameters(
 				TextureEditorPinned->GetVolumeViewMode() == TextureEditorVolumeViewMode_DepthSlices,
-				FMath::Max<int32>(VolumeTexture->GetSizeZ() >> VolumeTexture->GetCachedLODBias(), 1), 
 				MipLevel, 
 				(float)TextureEditorPinned->GetVolumeOpacity(),
 				true, 
@@ -276,7 +275,6 @@ void FTextureEditorViewportClient::Draw(FViewport* Viewport, FCanvas* Canvas)
 		{
 			BatchedElementParameters = new FBatchedElementVolumeTexturePreviewParameters(
 				TextureEditorPinned->GetVolumeViewMode() == TextureEditorVolumeViewMode_DepthSlices,
-				FMath::Max<int32>(RTTextureVolume->SizeZ >> RTTextureVolume->GetCachedLODBias(), 1),
 				MipLevel,
 				(float)TextureEditorPinned->GetVolumeOpacity(),
 				true,
diff --git a/Engine/Source/Editor/UnrealEd/Private/VolumeTexturePreview.cpp b/Engine/Source/Editor/UnrealEd/Private/VolumeTexturePreview.cpp
--- a/Engine/Source/Editor/UnrealEd/Private/VolumeTexturePreview.cpp
+++ b/Engine/Source/Editor/UnrealEd/Private/VolumeTexturePreview.cpp
@@ -78,7 +78,7 @@ public:
 		return IsFeatureLevelSupported(Parameters.Platform, ERHIFeatureLevel::SM5) && !IsConsolePlatform(Parameters.Platform);
 	}
 	
-	void SetParameters(FRHIBatchedShaderParameters& BatchedParameters, const FTexture* TextureValue, int32 SizeZ, const FMatrix44f& ColorWeightsValue, float GammaValue, float MipLevel, float Opacity, const FRotator& TraceOrientation, bool bUsePointSampling)
+	void SetParameters(FRHIBatchedShaderParameters& BatchedParameters, const FTexture* TextureValue, const FMatrix44f& ColorWeightsValue, float GammaValue, float MipLevel, float Opacity, const FRotator& TraceOrientation, bool bUsePointSampling)
 	{
 		FRHISamplerState* SamplerState = bUsePointSampling ? TStaticSamplerState<SF_Point>::GetRHI() : TextureValue->SamplerStateRHI.GetReference();
 		SetTextureParameter(BatchedParameters, InTexture, InTextureSampler, SamplerState, TextureValue->TextureRHI);
@@ -93,6 +93,7 @@ public:
 		}
 		SetShaderValue(BatchedParameters,ColorWeights,ColorWeightsValue);
 
+		const int32 SizeZ = TextureValue->GetSizeZ();
 		const int32 MipSizeZ = MipLevel >= 0 ? FMath::Max<int32>(SizeZ >> FMath::FloorToInt(MipLevel), 1) : SizeZ;
 		FVector4f PackedParametersValue(GammaValue, MipLevel, (float)MipSizeZ, Opacity);
 		SetShaderValue(BatchedParameters, PackedParameters, PackedParametersValue);
@@ -206,5 +207,5 @@ void FBatchedElementVolumeTexturePreviewParameters::BindShaders(
 	SetGraphicsPipelineState(RHICmdList, GraphicsPSOInit, 0);
 
 	SetShaderParametersLegacyVS(RHICmdList, VertexShader, InTransform);
-	SetShaderParametersLegacyPS(RHICmdList, PixelShader, Texture, SizeZ, ColorWeights, InGamma, MipLevel, Opacity, TraceOrientation, bUsePointSampling);
+	SetShaderParametersLegacyPS(RHICmdList, PixelShader, Texture, ColorWeights, InGamma, MipLevel, Opacity, TraceOrientation, bUsePointSampling);
 }
diff --git a/Engine/Source/Editor/UnrealEd/Public/VolumeTexturePreview.h b/Engine/Source/Editor/UnrealEd/Public/VolumeTexturePreview.h
--- a/Engine/Source/Editor/UnrealEd/Public/VolumeTexturePreview.h
+++ b/Engine/Source/Editor/UnrealEd/Public/VolumeTexturePreview.h
@@ -10,14 +10,13 @@
 #include "BatchedElements.h"
 
 /**
- * Batched element parameters for previewing 2d textures.
+ * Batched element parameters for previewing 3d textures.
  */
 class FBatchedElementVolumeTexturePreviewParameters : public FBatchedElementParameters
 {
 public:
-	FBatchedElementVolumeTexturePreviewParameters(bool InViewModeAsDepthSlices, int32 InSizeZ, float InMipLevel, float InOpacity, bool InShowSlices, const FRotator& InTraceOrientation, bool bInUsePointSampling)
+	FBatchedElementVolumeTexturePreviewParameters(bool InViewModeAsDepthSlices, float InMipLevel, float InOpacity, bool InShowSlices, const FRotator& InTraceOrientation, bool bInUsePointSampling)
 		: bViewModeAsDepthSlices(InViewModeAsDepthSlices)
-		, SizeZ(InSizeZ)
 		, MipLevel(InMipLevel)
 		, Opacity(InOpacity)
 		, bShowSlices(InShowSlices)
@@ -34,15 +33,12 @@ private:
 	/** Whether to render depth slices or trace into the volume */
 	bool bViewModeAsDepthSlices;
 	
-	/** The size Z of the texture */
-	int32 SizeZ;
-
 	/** The mip level to visualize */
 	float MipLevel;
 
 	float Opacity;
 
-	/** Whether to show each depth slize of the volume */
+	/** Whether to show each depth slice of the volume */
 	 bool bShowSlices;
 
 	/** The orientation when tracing */
