Cooker: CookRequestCluster: add assertion to silence static analyzer warning.
#rnx

diff --git a/Engine/Source/Editor/UnrealEd/Private/Cooker/CookRequestCluster.cpp b/Engine/Source/Editor/UnrealEd/Private/Cooker/CookRequestCluster.cpp
--- a/Engine/Source/Editor/UnrealEd/Private/Cooker/CookRequestCluster.cpp
+++ b/Engine/Source/Editor/UnrealEd/Private/Cooker/CookRequestCluster.cpp
@@ -927,6 +927,7 @@ void FRequestCluster::FGraphSearch::ResolveTransitiveBuildDependencyCycle()
 	FVertexData* FirstVertex = nullptr;
 	for (FVertexData* CycleVert : PendingTransitiveBuildDependencyVertices)
 	{
+		check(CycleVert != nullptr); // Required hint for static analyzers.
 		if (!FirstVertex)
 		{
 			FirstVertex = CycleVert;
