Factoring out duplicated ShouldTextureBeVirtualByAutoImportSize logic in texture import code

diff --git a/Engine/Plugins/Interchange/Runtime/Source/Import/Private/Texture/InterchangeDDSTranslator.cpp b/Engine/Plugins/Interchange/Runtime/Source/Import/Private/Texture/InterchangeDDSTranslator.cpp
--- a/Engine/Plugins/Interchange/Runtime/Source/Import/Private/Texture/InterchangeDDSTranslator.cpp
+++ b/Engine/Plugins/Interchange/Runtime/Source/Import/Private/Texture/InterchangeDDSTranslator.cpp
@@ -218,6 +218,8 @@ bool UInterchangeDDSTranslator::CanImportSourceData(const UInterchangeSourceData
 		return false;
 	}
 
+	// @@ check IsImageImportPossible here?
+
 	return DDS->IsValidTexture2D() || DDS->IsValidTextureCube() || DDS->IsValidTextureArray() || DDS->IsValidTextureVolume();
 }
 
diff --git a/Engine/Plugins/Interchange/Runtime/Source/Import/Private/Texture/InterchangeTextureFactory.cpp b/Engine/Plugins/Interchange/Runtime/Source/Import/Private/Texture/InterchangeTextureFactory.cpp
--- a/Engine/Plugins/Interchange/Runtime/Source/Import/Private/Texture/InterchangeTextureFactory.cpp
+++ b/Engine/Plugins/Interchange/Runtime/Source/Import/Private/Texture/InterchangeTextureFactory.cpp
@@ -1457,31 +1457,9 @@ namespace UE::Interchange::Private::InterchangeTextureFactory
 		// The texture has been imported and has no editor specific changes applied so we clear the painted flag.
 		Texture2D->bHasBeenPaintedInEditor = false;
 
-		// If the texture is larger than a certain threshold make it VT. This is explicitly done after the
-		// application of the existing settings above, so if a texture gets reimported at a larger size it will
-		// still be properly flagged as a VT (note: What about reimporting at a lower resolution?)
-		static const TConsoleVariableData<int32>* CVarVirtualTexturesEnabled = IConsoleManager::Get().FindTConsoleVariableDataInt(TEXT("r.VirtualTextures"));
-		check(CVarVirtualTexturesEnabled);
-
-		static const auto CVarVirtualTexturesAutoImportEnabled = IConsoleManager::Get().FindTConsoleVariableDataInt(TEXT("r.VT.EnableAutoImport"));
-		check(CVarVirtualTexturesAutoImportEnabled);
-
-		if (CVarVirtualTexturesEnabled->GetValueOnGameThread() && CVarVirtualTexturesAutoImportEnabled->GetValueOnGameThread())
+		if ( UE::TextureUtilitiesCommon::ShouldTextureBeVirtualByAutoImportSize(Texture2D) )
 		{
-			const int64 VirtualTextureAutoEnableThreshold = GetDefault<UTextureImportSettings>()->AutoVTSize;
-			const int64 VirtualTextureAutoEnableThresholdPixels = VirtualTextureAutoEnableThreshold * VirtualTextureAutoEnableThreshold;
-
-			// We do this in pixels so a 8192 x 128 texture won't get VT enabled 
-			// We use the Source size instead of simple Texture2D->GetSizeX() as this uses the size of the platform data
-			// however for a new texture platform data may not be generated yet, and for an reimport of a texture this is the size of the
-			// old texture. 
-			// Using source size gives one small caveat. It looks at the size before mipmap power of two padding adjustment.
-			if ( (int64) Texture2D->Source.GetSizeX() * Texture2D->Source.GetSizeY() >= VirtualTextureAutoEnableThresholdPixels ||
-				Texture2D->Source.GetSizeX() > UTexture::GetMaximumDimensionOfNonVT() ||
-				Texture2D->Source.GetSizeY() > UTexture::GetMaximumDimensionOfNonVT())
-			{
-				Texture2D->VirtualTextureStreaming = true;
-			}
+			Texture2D->VirtualTextureStreaming = true;
 		}
 #endif // WITH_EDITORONLY_DATA
 	}
diff --git a/Engine/Source/Editor/UnrealEd/Private/Factories/EditorFactories.cpp b/Engine/Source/Editor/UnrealEd/Private/Factories/EditorFactories.cpp
--- a/Engine/Source/Editor/UnrealEd/Private/Factories/EditorFactories.cpp
+++ b/Engine/Source/Editor/UnrealEd/Private/Factories/EditorFactories.cpp
@@ -4041,59 +4041,40 @@ UObject* UTextureFactory::FactoryCreateBinary
 			}
 		}
 	}
-			
-	static const auto CVarVirtualTexturesEnabled = IConsoleManager::Get().FindTConsoleVariableDataInt(TEXT("r.VirtualTextures"));
-	check(CVarVirtualTexturesEnabled);
-	
-	// If the texture is larger than a certain threshold make it VT.
-	// Note that previously for re-imports we still checked size and potentially changed the VT status.
-	// But that was unintuitive for many users so now for re-imports we will end up ignoring this and respecting the existing setting below.
-
-	static const auto CVarVirtualTexturesAutoImportEnabled = IConsoleManager::Get().FindTConsoleVariableDataInt(TEXT("r.VT.EnableAutoImport"));
-	check(CVarVirtualTexturesAutoImportEnabled);
-
-	if (CVarVirtualTexturesEnabled->GetValueOnAnyThread() && CVarVirtualTexturesAutoImportEnabled->GetValueOnAnyThread())
-	{
-		const int64 VirtualTextureAutoEnableThreshold = GetDefault<UTextureImportSettings>()->AutoVTSize;
-		const int64 VirtualTextureAutoEnableThresholdPixels = VirtualTextureAutoEnableThreshold * VirtualTextureAutoEnableThreshold;
-
-		// We do this in pixels so a 8192 x 128 texture won't get VT enabled 
-		// We use the Source size instead of simple Texture2D->GetSizeX() as this uses the size of the platform data
-		// however for a new texture platform data may not be generated yet, and for an reimport of a texture this is the size of the
-		// old texture. 
-		// Using source size gives one small caveat. It looks at the size before mipmap power of two padding adjustment.
-		// Textures with more than 1 block (UDIM textures) must be imported as VT
-		if (Texture->Source.GetNumBlocks() > 1 ||
-			( (int64) Texture->Source.GetSizeX() * Texture->Source.GetSizeY() ) >= VirtualTextureAutoEnableThresholdPixels ||
-			Texture->Source.GetSizeX() > UTexture::GetMaximumDimensionOfNonVT() ||
-			Texture->Source.GetSizeY() > UTexture::GetMaximumDimensionOfNonVT() )
-		{
-			// only UTexture2D can be VT
-			if ( Texture->GetTextureClass() == ETextureClass::TwoD )
-			{
-				Texture->VirtualTextureStreaming = true;
-			}
-			else
-			{
-				UE_LOG(LogEditorFactories, Warning, TEXT("Texture is too large for non-VT (%d x %d) but is not a UTexture2D."),
-					Texture->Source.GetSizeX() , Texture->Source.GetSizeY());
-			}
-		}
-	}
-
-	// if Texture is too large and should be VT but VT is not enabled, warn about that	
-	if ( ! CVarVirtualTexturesEnabled->GetValueOnAnyThread() )
+				
+	if ( UE::TextureUtilitiesCommon::ShouldTextureBeVirtualByAutoImportSize(Texture) )
 	{
-		int32 MaxDimension = FMath::Max( Texture->Source.GetSizeX() , Texture->Source.GetSizeY() );
-		bool bLargeTextureMustBeVT = MaxDimension > UTexture::GetMaximumDimensionOfNonVT();
-		if ( bLargeTextureMustBeVT )
+		// only UTexture2D can be VT
+		if ( Texture->GetTextureClass() == ETextureClass::TwoD )
 		{
-			UE_LOG(LogEditorFactories, Warning, TEXT("Texture is too large for non-VT (%d x %d) but VT is not enabled in this project."),
+			Texture->VirtualTextureStreaming = true;
+		}
+		else
+		{
+			UE_LOG(LogEditorFactories, Warning, TEXT("Texture [%s] is too large for non-VT (%d x %d) but is not a UTexture2D."),
+				*Texture->GetName(),
 				Texture->Source.GetSizeX() , Texture->Source.GetSizeY());
 		}
-		else if ( Texture->Source.GetNumBlocks() > 1 )
+	}
+	else
+	{
+		static const auto CVarVirtualTexturesEnabled = IConsoleManager::Get().FindTConsoleVariableDataInt(TEXT("r.VirtualTextures"));
+		check(CVarVirtualTexturesEnabled);
+
+		// if Texture is too large and should be VT but VT is not enabled, warn about that	
+		if ( ! CVarVirtualTexturesEnabled->GetValueOnAnyThread() )
 		{
-			UE_LOG(LogEditorFactories, Warning, TEXT("Texture has blocks which requires VT but VT is not enabled in this project."));
+			int32 MaxDimension = FMath::Max( Texture->Source.GetSizeX() , Texture->Source.GetSizeY() );
+			bool bLargeTextureMustBeVT = MaxDimension > UTexture::GetMaximumDimensionOfNonVT();
+			if ( bLargeTextureMustBeVT )
+			{
+				UE_LOG(LogEditorFactories, Warning, TEXT("Texture is too large for non-VT (%d x %d) but VT is not enabled in this project."),
+					Texture->Source.GetSizeX() , Texture->Source.GetSizeY());
+			}
+			else if ( Texture->Source.GetNumBlocks() > 1 )
+			{
+				UE_LOG(LogEditorFactories, Warning, TEXT("Texture has blocks which requires VT but VT is not enabled in this project."));
+			}
 		}
 	}
 
diff --git a/Engine/Source/Runtime/ImageCore/Private/ImageCoreUtils.cpp b/Engine/Source/Runtime/ImageCore/Private/ImageCoreUtils.cpp
--- a/Engine/Source/Runtime/ImageCore/Private/ImageCoreUtils.cpp
+++ b/Engine/Source/Runtime/ImageCore/Private/ImageCoreUtils.cpp
@@ -17,7 +17,7 @@ IMAGECORE_API bool FImageCoreUtils::IsImageImportPossible(int64 Width,int64 Heig
 		return false;
 	}
 
-	int64 PixelCount = Width * Height; // mutliply is safe because of INT32_MAX check
+	int64 PixelCount = Width * Height; // multiply is safe because of INT32_MAX check
 
 	// for non-VT limitation is 16384 for W/H , and 2 GB for output surface bytes
 	// for VT the limitation is that the built VT must fit in 4 GB
diff --git a/Engine/Source/Runtime/TextureUtilitiesCommon/Private/TextureImportSettings.cpp b/Engine/Source/Runtime/TextureUtilitiesCommon/Private/TextureImportSettings.cpp
--- a/Engine/Source/Runtime/TextureUtilitiesCommon/Private/TextureImportSettings.cpp
+++ b/Engine/Source/Runtime/TextureUtilitiesCommon/Private/TextureImportSettings.cpp
@@ -146,6 +146,47 @@ namespace UE::TextureUtilitiesCommon
 			}
 		}
 	}
+	
+	bool ShouldTextureBeVirtualByAutoImportSize(const UTexture * Texture)
+	{
+		// If the texture is larger than a certain threshold make it VT.
+		// Note that previously for re-imports we still checked size and potentially changed the VT status.
+		// But that was unintuitive for many users so now for re-imports we will end up ignoring this and respecting the existing setting below.
+		
+		static const TConsoleVariableData<int32>* CVarVirtualTexturesEnabled = IConsoleManager::Get().FindTConsoleVariableDataInt(TEXT("r.VirtualTextures"));
+		check(CVarVirtualTexturesEnabled);
+
+		static const auto CVarVirtualTexturesAutoImportEnabled = IConsoleManager::Get().FindTConsoleVariableDataInt(TEXT("r.VT.EnableAutoImport"));
+		check(CVarVirtualTexturesAutoImportEnabled);
+
+		if (CVarVirtualTexturesEnabled->GetValueOnAnyThread() && CVarVirtualTexturesAutoImportEnabled->GetValueOnAnyThread())
+		{
+			const int64 VirtualTextureAutoEnableThreshold = GetDefault<UTextureImportSettings>()->AutoVTSize;
+
+			if ( VirtualTextureAutoEnableThreshold == 0 )
+			{
+				return false;
+			}
+
+			const int64 VirtualTextureAutoEnableThresholdPixels = VirtualTextureAutoEnableThreshold * VirtualTextureAutoEnableThreshold;
+
+			// We do this in pixels so a 8192 x 128 texture won't get VT enabled 
+			// We use the Source size instead of simple Texture2D->GetSizeX() as this uses the size of the platform data
+			// however for a new texture platform data may not be generated yet, and for an reimport of a texture this is the size of the
+			// old texture. 
+			// Using source size gives one small caveat. It looks at the size before mipmap power of two padding adjustment.
+			// Textures with more than 1 block (UDIM textures) must be imported as VT
+			if (Texture->Source.GetNumBlocks() > 1 ||
+				( (int64) Texture->Source.GetSizeX() * Texture->Source.GetSizeY() ) >= VirtualTextureAutoEnableThresholdPixels ||
+				Texture->Source.GetSizeX() > UTexture::GetMaximumDimensionOfNonVT() ||
+				Texture->Source.GetSizeY() > UTexture::GetMaximumDimensionOfNonVT() )
+			{
+				return true;
+			}
+		}
+
+		return false;
+	}
 	#endif // WITH_EDITOR
 
 	/* Get the default value for Texture->SRGB
diff --git a/Engine/Source/Runtime/TextureUtilitiesCommon/Private/TextureImportUtils.cpp b/Engine/Source/Runtime/TextureUtilitiesCommon/Private/TextureImportUtils.cpp
--- a/Engine/Source/Runtime/TextureUtilitiesCommon/Private/TextureImportUtils.cpp
+++ b/Engine/Source/Runtime/TextureUtilitiesCommon/Private/TextureImportUtils.cpp
@@ -579,7 +579,9 @@ namespace UE
 
 bool UE::TextureUtilitiesCommon::IsImportResolutionValid(int64 Width, int64 Height, bool bAllowNonPowerOfTwo, FText* OutErrorMessage)
 {
-	// code dupe from: UTextureFactory::IsImportResolutionValid
+	// note: stricter than IsImageImportPossible
+	//	IsImageImportPossible is a first check that can be done early in the loading to bail on totally impossible sizes
+	//	this is done late and uses project-specific config
 
 	// MaximumSupportedResolutionNonVT is only a popup/warning , not a hard limit
 	// Get the non-VT size limit :
diff --git a/Engine/Source/Runtime/TextureUtilitiesCommon/Public/TextureImportSettings.h b/Engine/Source/Runtime/TextureUtilitiesCommon/Public/TextureImportSettings.h
--- a/Engine/Source/Runtime/TextureUtilitiesCommon/Public/TextureImportSettings.h
+++ b/Engine/Source/Runtime/TextureUtilitiesCommon/Public/TextureImportSettings.h
@@ -103,6 +103,12 @@ namespace UE::TextureUtilitiesCommon
 	*  Should be called after all texture properties are set, before PostEditChange() 
 	*/
 	TEXTUREUTILITIESCOMMON_API void ApplyDefaultsForNewlyImportedTextures(UTexture * Texture, bool bIsReimport);
+	
+	/* Query if texture will be made VT automatically on import due to exceeding AutoVTSize 
+		or other conditions which require it to be VT
+		only returns true if VT is enabled and auto-VT is enabled
+	*/
+	TEXTUREUTILITIESCOMMON_API bool ShouldTextureBeVirtualByAutoImportSize(const UTexture * Texture);
 #endif
 
 	/* Get the default value for Texture->SRGB
diff --git a/Engine/Source/Runtime/TextureUtilitiesCommon/Public/TextureImportUtils.h b/Engine/Source/Runtime/TextureUtilitiesCommon/Public/TextureImportUtils.h
--- a/Engine/Source/Runtime/TextureUtilitiesCommon/Public/TextureImportUtils.h
+++ b/Engine/Source/Runtime/TextureUtilitiesCommon/Public/TextureImportUtils.h
@@ -36,6 +36,8 @@ namespace UE::TextureUtilitiesCommon
 		* @param OutErrorMessage Optional output for an error message
 		*
 		* @return bool true if the given height/width represent a supported texture resolution, false if not
+		*
+		* NOTE: may open a dialog box to ask user if large non-VT import is wanted!
 		*/
 	TEXTUREUTILITIESCOMMON_API bool IsImportResolutionValid(int64 Width, int64 Height, bool bAllowNonPowerOfTwo, FText* OutErrorMessage);
 }
\ No newline at end of file
