Reword confirmation dialog for existing assets.

diff --git a/Engine/Source/Developer/AssetTools/AssetTools.Build.cs b/Engine/Source/Developer/AssetTools/AssetTools.Build.cs
--- a/Engine/Source/Developer/AssetTools/AssetTools.Build.cs
+++ b/Engine/Source/Developer/AssetTools/AssetTools.Build.cs
@@ -40,6 +40,7 @@ public class AssetTools : ModuleRules
 				"RHI",
 				"MaterialEditor",
 				"ToolMenus",
+				"ToolWidgets",
 				"PhysicsCore",
 				"DeveloperSettings",
 				"ClassViewer",
diff --git a/Engine/Source/Developer/AssetTools/Private/AssetTools.cpp b/Engine/Source/Developer/AssetTools/Private/AssetTools.cpp
--- a/Engine/Source/Developer/AssetTools/Private/AssetTools.cpp
+++ b/Engine/Source/Developer/AssetTools/Private/AssetTools.cpp
@@ -86,6 +86,7 @@
 #include "Dialogs/Dialogs.h"
 #include "Widgets/Input/SEditableTextBox.h"
 #include "Widgets/Layout/SSpacer.h"
+#include "Dialog/SMessageDialog.h"
 #include "Interfaces/IPluginManager.h"
 #include "Settings/ContentBrowserSettings.h"
 #include "Algo/Count.h"
@@ -4740,13 +4741,21 @@ bool UAssetToolsImpl::CanCreateAsset(const FString& AssetName, const FString& Pa
 	{
 		// Object already exists in either the specified package or another package.  Check to see if the user wants
 		// to replace the object.
-		bool bWantReplace =
-			EAppReturnType::Yes == FMessageDialog::Open(
-				EAppMsgType::YesNo,
-				EAppReturnType::No,
-				FText::Format(
-					NSLOCTEXT("UnrealEd", "ReplaceExistingObjectInPackage_F", "An object [{0}] of class [{1}] already exists in file [{2}].  Do you want to replace the existing object?  If you click 'Yes', the existing object will be deleted.  Otherwise, click 'No' and choose a unique name for your new object." ),
-					FText::FromString(AssetName), FText::FromString(ExistingObject->GetClass()->GetName()), FText::FromString(PackageName) ) );
+		TSharedRef<SMessageDialog> ConfirmDialog = SNew(SMessageDialog)
+			.Icon(FAppStyle::Get().GetBrush("Icons.WarningWithColor.Large"))
+			.Title(FText(NSLOCTEXT("UnrealEd", "ReplaceExistingObjectInPackageConfirmation_Title", "Overwrite Existing Object")))
+			.Message(FText::Format(NSLOCTEXT("UnrealEd", "ReplaceExistingObjectInPackageConfirmation_Message", "An object already exists with this name.\n\n\tName: {0}\n\tClass: {1}\n\tAsset path: {2}\n\nOverwrite the existing object?"),
+				FText::FromString(AssetName),
+				FText::FromString(ExistingObject->GetClass()->GetName()),
+				FText::FromString(PackageName)))
+			.Buttons({
+				SCustomDialog::FButton(NSLOCTEXT("UnrealEd", "ReplaceExistingObjectInPackageConfirmation_ButtonOverwrite", "Overwrite")).SetPrimary(true),
+				SCustomDialog::FButton(NSLOCTEXT("UnrealEd", "ReplaceExistingObjectInPackageConfirmation_ButtonCancel", "Cancel")),
+				})
+			.ContentMinWidth(300.0f);
+		uint32 ConfirmationResult = ConfirmDialog->ShowModal();
+
+		bool bWantReplace = ConfirmationResult == 0;
 
 		if( bWantReplace )
 		{
diff --git a/Engine/Source/Developer/ToolWidgets/Private/Dialog/SMessageDialog.cpp b/Engine/Source/Developer/ToolWidgets/Private/Dialog/SMessageDialog.cpp
--- a/Engine/Source/Developer/ToolWidgets/Private/Dialog/SMessageDialog.cpp
+++ b/Engine/Source/Developer/ToolWidgets/Private/Dialog/SMessageDialog.cpp
@@ -21,6 +21,7 @@ void SMessageDialog::Construct(const FArguments& InArgs)
 		.Content()
 		[
 			SAssignNew(RichTextBlock, SRichTextBlock)
+			.MinDesiredWidth(InArgs._ContentMinWidth)
 			.Text(Message)
 			.WrapTextAt(InArgs._WrapMessageAt)
 			.Decorators(InArgs._Decorators)
diff --git a/Engine/Source/Developer/ToolWidgets/Public/Dialog/SMessageDialog.h b/Engine/Source/Developer/ToolWidgets/Public/Dialog/SMessageDialog.h
--- a/Engine/Source/Developer/ToolWidgets/Public/Dialog/SMessageDialog.h
+++ b/Engine/Source/Developer/ToolWidgets/Public/Dialog/SMessageDialog.h
@@ -64,6 +64,9 @@ public:
 		/** When to wrap the message text (default: 512) */
 		SLATE_ATTRIBUTE(float, WrapMessageAt)
 	
+		/** Minimum width for the text part of the message box. Optional for very short message text to add breathing space to the layout. */
+		SLATE_ATTRIBUTE(float, ContentMinWidth)
+
 	SLATE_END_ARGS()
 	
 	void Construct(const FArguments& InArgs);
diff --git a/Engine/Source/Editor/UnrealEd/Private/Dialogs/Dialogs.cpp b/Engine/Source/Editor/UnrealEd/Private/Dialogs/Dialogs.cpp
--- a/Engine/Source/Editor/UnrealEd/Private/Dialogs/Dialogs.cpp
+++ b/Engine/Source/Editor/UnrealEd/Private/Dialogs/Dialogs.cpp
@@ -26,6 +26,7 @@
 #include "HAL/PlatformApplicationMisc.h"
 #include "HAL/PlatformMisc.h"
 #include "SPrimaryButton.h"
+#include "Dialog/SMessageDialog.h"
 
 DEFINE_LOG_CATEGORY_STATIC(LogDialogs, Log, All);
 
@@ -987,14 +988,21 @@ bool PromptUserIfExistingObject(const FString& Name, const FString& Package, UPa
 	{
 		// Object already exists in either the specified package or another package.  Check to see if the user wants
 		// to replace the object.
-		bool bWantReplace =
-			EAppReturnType::Yes == FMessageDialog::Open(
-									EAppMsgType::YesNo,
-									FText::Format(
-									NSLOCTEXT("UnrealEd", "ReplaceExistingObjectInPackage_F", "An object [{0}] of class [{1}] already exists in file [{2}].  Do you want to replace the existing object?  If you click 'Yes', the existing object will be deleted.  Otherwise, click 'No' and choose a unique name for your new object." ),
-									FText::FromString(Name),
-									FText::FromString(ExistingObject->GetClass()->GetName()),
-									FText::FromString(Package) ) );
+		TSharedRef<SMessageDialog> ConfirmDialog = SNew(SMessageDialog)
+			.Icon(FAppStyle::Get().GetBrush("Icons.WarningWithColor.Large"))
+			.Title(FText(NSLOCTEXT("UnrealEd", "ReplaceExistingObjectInPackageConfirmation_Title", "Overwrite Existing Object")))
+			.Message(FText::Format(NSLOCTEXT("UnrealEd", "ReplaceExistingObjectInPackageConfirmation_Message", "An object already exists with this name.\n\n\tName: {0}\n\tClass: {1}\n\tAsset path: {2}\n\nOverwrite the existing object?"),
+				FText::FromString(Name),
+				FText::FromString(ExistingObject->GetClass()->GetName()),
+				FText::FromString(Package)))
+			.Buttons({
+				SCustomDialog::FButton(NSLOCTEXT("UnrealEd", "ReplaceExistingObjectInPackageConfirmation_ButtonOverwrite", "Overwrite")).SetPrimary(true),
+				SCustomDialog::FButton(NSLOCTEXT("UnrealEd", "ReplaceExistingObjectInPackageConfirmation_ButtonCancel", "Cancel")),
+				})
+			.ContentMinWidth(300.0f);
+		uint32 ConfirmationResult = ConfirmDialog->ShowModal();
+
+		bool bWantReplace = ConfirmationResult == 0;
 
 		if( bWantReplace )
 		{
diff --git a/Engine/Source/Editor/UnrealEd/Private/ObjectTools.cpp b/Engine/Source/Editor/UnrealEd/Private/ObjectTools.cpp
--- a/Engine/Source/Editor/UnrealEd/Private/ObjectTools.cpp
+++ b/Engine/Source/Editor/UnrealEd/Private/ObjectTools.cpp
@@ -63,6 +63,7 @@
 #include "EditorDirectories.h"
 #include "FileHelpers.h"
 #include "Dialogs/Dialogs.h"
+#include "Dialog/SMessageDialog.h"
 #include "UnrealEdGlobals.h"
 #include "PackageTools.h"
 #include "Internationalization/TextPackageNamespaceUtil.h"
@@ -858,15 +859,21 @@ namespace ObjectTools
 		// object. This will delete the object so the new one can be created in its place.
 		if(bPromptToOverwrite && ObjectsToOverwriteName.Len() > 0 )
 		{
-			bool bOverwriteExistingObjects =
-				EAppReturnType::Yes == FMessageDialog::Open(
-				EAppMsgType::YesNo,
-				EAppReturnType::No,
-				FText::Format(
-				NSLOCTEXT("UnrealEd", "ReplaceExistingObjectInPackage_F", "An object [{0}] of class [{1}] already exists in file [{2}].  Do you want to replace the existing object?  If you click 'Yes', the existing object will be deleted.  Otherwise, click 'No' and choose a unique name for your new object." ),
-				FText::FromString(ObjectsToOverwriteName),
-				FText::FromString(ObjectsToOverwriteClass),
-				FText::FromString(ObjectsToOverwritePackage) ) );
+			TSharedRef<SMessageDialog> ConfirmDialog = SNew(SMessageDialog)
+				.Icon(FAppStyle::Get().GetBrush("Icons.WarningWithColor.Large"))
+				.Title(FText(NSLOCTEXT("UnrealEd", "ReplaceExistingObjectInPackageConfirmation_Title", "Overwrite Existing Object")))
+				.Message(FText::Format(NSLOCTEXT("UnrealEd", "ReplaceExistingObjectInPackageConfirmation_Message", "An object already exists with this name.\n\n\tName: {0}\n\tClass: {1}\n\tAsset path: {2}\n\nOverwrite the existing object?"),
+					FText::FromString(ObjectsToOverwriteName),
+					FText::FromString(ObjectsToOverwriteClass),
+					FText::FromString(ObjectsToOverwritePackage)))
+				.Buttons({
+					SCustomDialog::FButton(NSLOCTEXT("UnrealEd", "ReplaceExistingObjectInPackageConfirmation_ButtonOverwrite", "Overwrite")).SetPrimary(true),
+					SCustomDialog::FButton(NSLOCTEXT("UnrealEd", "ReplaceExistingObjectInPackageConfirmation_ButtonCancel", "Cancel")),
+					})
+				.ContentMinWidth(300.0f);
+			uint32 ConfirmationResult = ConfirmDialog->ShowModal();
+
+			bool bOverwriteExistingObjects = ConfirmationResult == 0;
 
 			// The user didn't want to overwrite the existing options, so bail out of the duplicate operation.
 			if( !bOverwriteExistingObjects )
