TextureAssetActions : make it possible to convert TSF G16 to 8 on Grayscale textures

diff --git a/Engine/Source/Editor/UnrealEd/Private/TextureAssetActions.cpp b/Engine/Source/Editor/UnrealEd/Private/TextureAssetActions.cpp
--- a/Engine/Source/Editor/UnrealEd/Private/TextureAssetActions.cpp
+++ b/Engine/Source/Editor/UnrealEd/Private/TextureAssetActions.cpp
@@ -776,7 +776,8 @@ static ETextureSourceFormat GetReducedTextureSourceFormat(const TextureCompressi
 	case TC_Grayscale				: //"Grayscale (G8/16, RGB8 sRGB)"),
 	case TC_Displacementmap			: //"Displacementmap (G8/16)"),
 		// Gray and Displacement pass through G16 ; note they do not do that for RGBA16 (see GetDefaultTextureFormatName)
-		if ( InTSF == TSF_G16 ) return InTSF;
+		//	do this conditional on NormalMapsKeep16bits just so we have a way to toggle it, even though it's not really a normal map
+		if ( InTSF == TSF_G16 && NormalMapsKeep16bits ) return InTSF;
 		// otherwise we will convert to G8
 		// [[fallthrough]];
 	case TC_DistanceFieldFont		: //"DistanceFieldFont (G8)"),
