Fix Mobile SVT commandlet used in the nigthly build. Some shaderJobs can be created when we recreate the render state contexts, so we need to wait on shader compilation after. Reverse the logic of r.Mobile.VirtualTextures cvar so that when we call the commandlet in pc mobile, the cvar as the correct value.

diff --git a/Engine/Source/Editor/UnrealEd/Private/EditorBuildUtils.cpp b/Engine/Source/Editor/UnrealEd/Private/EditorBuildUtils.cpp
--- a/Engine/Source/Editor/UnrealEd/Private/EditorBuildUtils.cpp
+++ b/Engine/Source/Editor/UnrealEd/Private/EditorBuildUtils.cpp
@@ -1825,6 +1825,11 @@ bool FEditorBuildUtils::EditorBuildVirtualTexture(UWorld* InWorld)
 				FAssetCompilingManager::Get().ProcessAsyncTasks();
 				FGlobalComponentRecreateRenderStateContext Context;
 			}
+
+			// FGlobalComponentRecreateRenderStateContext can create new shaderJobs, make sure to wait on them.
+			FAssetCompilingManager::Get().FinishAllCompilation();
+			FAssetCompilingManager::Get().ProcessAsyncTasks();
+
 		}
 
 		SubBuildTask.EnterProgressFrame(1.0f);
diff --git a/Engine/Source/Editor/UnrealEd/Private/WorldPartition/WorldPartitionRuntimeVirtualTextureBuilder.cpp b/Engine/Source/Editor/UnrealEd/Private/WorldPartition/WorldPartitionRuntimeVirtualTextureBuilder.cpp
--- a/Engine/Source/Editor/UnrealEd/Private/WorldPartition/WorldPartitionRuntimeVirtualTextureBuilder.cpp
+++ b/Engine/Source/Editor/UnrealEd/Private/WorldPartition/WorldPartitionRuntimeVirtualTextureBuilder.cpp
@@ -121,6 +121,10 @@ bool UWorldPartitionRuntimeVirtualTextureBuilder::RunInternal(UWorld* World, con
 			FAssetCompilingManager::Get().ProcessAsyncTasks();
 			FGlobalComponentRecreateRenderStateContext Context;
 		}
+
+		// FGlobalComponentRecreateRenderStateContext can create new shaderJobs, make sure to wait on them.
+		FAssetCompilingManager::Get().FinishAllCompilation();
+		FAssetCompilingManager::Get().ProcessAsyncTasks();
 				
 		for (URuntimeVirtualTextureComponent* Component : Components[1])
 		{
