Remove deleted FNaniteMultiViewMaterialVS name from shader pipeline cache tools commandlet

diff --git a/Engine/Source/Editor/UnrealEd/Private/Commandlets/ShaderPipelineCacheToolsCommandlet.cpp b/Engine/Source/Editor/UnrealEd/Private/Commandlets/ShaderPipelineCacheToolsCommandlet.cpp
--- a/Engine/Source/Editor/UnrealEd/Private/Commandlets/ShaderPipelineCacheToolsCommandlet.cpp
+++ b/Engine/Source/Editor/UnrealEd/Private/Commandlets/ShaderPipelineCacheToolsCommandlet.cpp
@@ -603,14 +603,12 @@ bool CouldBeUsedTogether(const FStableShaderKeyAndValue& A, const FStableShaderK
 	static FName NAME_FWriteToSliceVS("FWriteToSliceVS");
 	static FName NAME_FScreenPassVS("FScreenPassVS");
 	static FName NAME_FWriteToSliceGS("FWriteToSliceGS");
-	static FName NAME_FNaniteMultiViewMaterialVS("FNaniteMultiViewMaterialVS");
 	if (
 		A.ShaderType == NAME_FDeferredDecalVS || B.ShaderType == NAME_FDeferredDecalVS ||
 		A.ShaderType == NAME_FDeferredLightVS || B.ShaderType == NAME_FDeferredLightVS ||
 		A.ShaderType == NAME_FWriteToSliceVS || B.ShaderType == NAME_FWriteToSliceVS ||
 		A.ShaderType == NAME_FScreenPassVS || B.ShaderType == NAME_FScreenPassVS ||
-		A.ShaderType == NAME_FWriteToSliceGS || B.ShaderType == NAME_FWriteToSliceGS ||
-		A.ShaderType == NAME_FNaniteMultiViewMaterialVS || B.ShaderType == NAME_FNaniteMultiViewMaterialVS
+		A.ShaderType == NAME_FWriteToSliceGS || B.ShaderType == NAME_FWriteToSliceGS
 		)
 	{
 		// oddball mix and match with any material shader.
