Add a cvar to enable cut/copy/duplicate for temporary actors (e.g PIE).
Note: This is still default disabled because the behavior is non-standard and has potential issues.

diff --git a/Engine/Source/Editor/UnrealEd/Private/EditorServer.cpp b/Engine/Source/Editor/UnrealEd/Private/EditorServer.cpp
--- a/Engine/Source/Editor/UnrealEd/Private/EditorServer.cpp
+++ b/Engine/Source/Editor/UnrealEd/Private/EditorServer.cpp
@@ -3242,15 +3242,20 @@ bool UEditorEngine::CanCopySelectedActorsToClipboard( UWorld* InWorld, FCopySele
 				CopySelected.LevelAllActorsAreIn = Actor->GetLevel();
 			}
 
-			if (Actor->GetLevel())
+			IConsoleVariable* TemporaryActorCopyCvar = IConsoleManager::Get().FindConsoleVariable(TEXT("TypedElements.EnableTemporaryActorCopy"));
+
+			if (TemporaryActorCopyCvar && !TemporaryActorCopyCvar->GetBool())
 			{
-				if (UWorld* ActorWorld = Actor->GetLevel()->GetWorld())
+				if (Actor->GetLevel())
 				{
-					// If the actor is in a PIE world but doesn't have an editor counterpart it means it's a temporary
-					// actor spawned to the world. These actors can cause issues when copied so have been disabled.
-					if (ActorWorld->WorldType == EWorldType::PIE && !GEditor->ObjectsThatExistInEditorWorld.Get(Actor))
+					if (UWorld* ActorWorld = Actor->GetLevel()->GetWorld())
 					{
-						return false;
+						// If the actor is in a PIE world but doesn't have an editor counterpart it means it's a temporary
+						// actor spawned to the world. These actors can cause issues when copied so have been disabled.
+						if (ActorWorld->WorldType == EWorldType::PIE && !GEditor->ObjectsThatExistInEditorWorld.Get(Actor))
+						{
+							return false;
+						}
 					}
 				}
 			}
diff --git a/Engine/Source/Editor/UnrealEd/Private/Elements/Actor/ActorElementEditorWorldInterface.cpp b/Engine/Source/Editor/UnrealEd/Private/Elements/Actor/ActorElementEditorWorldInterface.cpp
--- a/Engine/Source/Editor/UnrealEd/Private/Elements/Actor/ActorElementEditorWorldInterface.cpp
+++ b/Engine/Source/Editor/UnrealEd/Private/Elements/Actor/ActorElementEditorWorldInterface.cpp
@@ -14,6 +14,13 @@
 #include "UnrealEdGlobals.h"
 #include "Editor/UnrealEdEngine.h"
 
+namespace UE::ActorElementEditorWorldInterface::Local
+{
+	bool bEnableTemporaryActorCopy = false;
+	static FAutoConsoleVariableRef CVarEnableTemporaryActorCopy(TEXT("TypedElements.EnableTemporaryActorCopy"), bEnableTemporaryActorCopy,
+		TEXT("Enable cut/copy/duplicate for temporary actors created during PIE. Experimental and can cause issues."));
+}
+
 bool UActorElementEditorWorldInterface::GetPivotOffset(const FTypedElementHandle& InElementHandle, FVector& OutPivotOffset)
 {
 	if (AActor* Actor = ActorElementDataUtil::GetActorFromHandle(InElementHandle))
@@ -81,15 +88,19 @@ bool UActorElementEditorWorldInterface::DeleteElements(TArrayView<const FTypedEl
 bool UActorElementEditorWorldInterface::CanDuplicateElement(const FTypedElementHandle& InElementHandle)
 {
 	AActor* Actor = ActorElementDataUtil::GetActorFromHandle(InElementHandle);
-	if (Actor && Actor->GetLevel())
+
+	if (!UE::ActorElementEditorWorldInterface::Local::bEnableTemporaryActorCopy)
 	{
-		if (UWorld* ActorWorld = Actor->GetLevel()->GetWorld())
+		if (Actor && Actor->GetLevel())
 		{
-			// If the actor is in a PIE world but doesn't have an editor counterpart it means it's a temporary
-			// actor spawned to the world. These actors can cause issues when copied so have been disabled.
-			if (ActorWorld->WorldType == EWorldType::PIE && !GEditor->ObjectsThatExistInEditorWorld.Get(Actor))
+			if (UWorld* ActorWorld = Actor->GetLevel()->GetWorld())
 			{
-				return false;
+				// If the actor is in a PIE world but doesn't have an editor counterpart it means it's a temporary
+				// actor spawned to the world. These actors can cause issues when copied so have been disabled.
+				if (ActorWorld->WorldType == EWorldType::PIE && !GEditor->ObjectsThatExistInEditorWorld.Get(Actor))
+				{
+					return false;
+				}
 			}
 		}
 	}
@@ -138,19 +149,23 @@ void UActorElementEditorWorldInterface::DuplicateElements(TArrayView<const FType
 
 bool UActorElementEditorWorldInterface::CanCopyElement(const FTypedElementHandle& InElementHandle)
 {
-	AActor* Actor = ActorElementDataUtil::GetActorFromHandle(InElementHandle);
-	if (Actor && Actor->GetLevel())
+	if (!UE::ActorElementEditorWorldInterface::Local::bEnableTemporaryActorCopy)
 	{
-		if (UWorld* ActorWorld = Actor->GetLevel()->GetWorld())
+		AActor* Actor = ActorElementDataUtil::GetActorFromHandle(InElementHandle);
+		if (Actor && Actor->GetLevel())
 		{
-			// If the actor is in a PIE world but doesn't have an editor counterpart it means it's a temporary
-			// actor spawned to the world. These actors can cause issues when copied so have been disabled.
-			if (ActorWorld->WorldType == EWorldType::PIE && !GEditor->ObjectsThatExistInEditorWorld.Get(Actor))
+			if (UWorld* ActorWorld = Actor->GetLevel()->GetWorld())
 			{
-				return false;
+				// If the actor is in a PIE world but doesn't have an editor counterpart it means it's a temporary
+				// actor spawned to the world. These actors can cause issues when copied so have been disabled.
+				if (ActorWorld->WorldType == EWorldType::PIE && !GEditor->ObjectsThatExistInEditorWorld.Get(Actor))
+				{
+					return false;
+				}
 			}
 		}
 	}
+	
 	return true;
 }
 
