Fixes for Copy and paste in prefab edtior

diff --git a/Engine/Source/Editor/UnrealEd/Private/EditorObject.cpp b/Engine/Source/Editor/UnrealEd/Private/EditorObject.cpp
--- a/Engine/Source/Editor/UnrealEd/Private/EditorObject.cpp
+++ b/Engine/Source/Editor/UnrealEd/Private/EditorObject.cpp
@@ -1427,7 +1427,7 @@ static const TCHAR* ImportPropertiesStep(
 					Params.ObjectRemapper = ObjectRemapper;
 					Params.PropertiesToSkip = PropertiesToSkip;
 
-					ImportObjectsPropertiesStep(Params);
+					CurrentSourceText = ImportObjectsPropertiesStep(Params);
 				}
 			}
 			else if( FParse::Command(&Str,TEXT("CustomProperties")))
diff --git a/Engine/Source/Editor/UnrealEd/Private/Elements/Actor/ActorElementEditorWorldInterface.cpp b/Engine/Source/Editor/UnrealEd/Private/Elements/Actor/ActorElementEditorWorldInterface.cpp
--- a/Engine/Source/Editor/UnrealEd/Private/Elements/Actor/ActorElementEditorWorldInterface.cpp
+++ b/Engine/Source/Editor/UnrealEd/Private/Elements/Actor/ActorElementEditorWorldInterface.cpp
@@ -182,7 +182,7 @@ void UActorElementEditorWorldInterface::CopyElements(TArrayView<const FTypedElem
 	ActorElementsCopy->ActorsToCopy = Actors;
 
 	constexpr int32 Indent = 3;
-	UExporter::ExportToOutputDevice(nullptr, ActorElementsCopy, nullptr, Out, TEXT("copy"), Indent, PPF_DeepCompareInstances);
+	UExporter::ExportToOutputDevice(nullptr, ActorElementsCopy, nullptr, Out, TEXT("copy"), Indent, PPF_DeepCompareInstances | PPF_ExportsNotFullyQualified);
 }
 
 TSharedPtr<FWorldElementPasteImporter> UActorElementEditorWorldInterface::GetPasteImporter()
