[DiffUtils]
 - Fix crash caused by recursive loading of external actors when diffing

diff --git a/Engine/Source/Editor/UnrealEd/Private/DiffUtils.cpp b/Engine/Source/Editor/UnrealEd/Private/DiffUtils.cpp
--- a/Engine/Source/Editor/UnrealEd/Private/DiffUtils.cpp
+++ b/Engine/Source/Editor/UnrealEd/Private/DiffUtils.cpp
@@ -8,6 +8,7 @@
 #include "EditorCategoryUtils.h"
 #include "IAssetTools.h"
 #include "Engine/Blueprint.h"
+#include "Engine/Level.h"
 #include "IAssetTypeActions.h"
 #include "ISourceControlModule.h"
 #include "ISourceControlProvider.h"
@@ -1189,6 +1190,7 @@ UPackage* DiffUtils::LoadPackageForDiff(const FPackagePath& InTempPackagePath, c
 	FLinkerInstancingContext Context;
 	if (!InOriginalPackagePath.GetLocalFullPath().IsEmpty())
 	{
+		Context.AddTag(ULevel::DontLoadExternalObjectsTag);
 		Context.AddPackageMapping(InOriginalPackagePath.GetPackageFName(), InTempPackagePath.GetPackageFName());
 	}
 	
