Add manual port binding to ODSC

Change Boolean FCookOnTheFlyStartupOptions::bBindAnyPort to integral FCookOnTheFlyStartupOptions::Port.
Add constexpr DefaultPort and AnyPort to the struct, corresponding to the previous false and true value behaviors, respectively, of bBindAnyPort.
Add logic to UCookCommandlet to parse COTF listening port from command line, e.g., -port=12345.
(Undocumented: -port=0 allows binding to any port; chosen port will be logged once the server starts up.)

Clients can connect to a nonstandard port by specifying it on the command line with the host address, e.g., -odschost=127.0.0.1:12345

diff --git a/Engine/Source/Editor/UnrealEd/Classes/Commandlets/CookCommandlet.h b/Engine/Source/Editor/UnrealEd/Classes/Commandlets/CookCommandlet.h
--- a/Engine/Source/Editor/UnrealEd/Classes/Commandlets/CookCommandlet.h
+++ b/Engine/Source/Editor/UnrealEd/Classes/Commandlets/CookCommandlet.h
@@ -77,7 +77,7 @@ protected:
 	 *
 	 * @return true on success, false otherwise.
 	 */
-	bool CookOnTheFly( FGuid InstanceId, int32 Timeout = 180, bool bForceClose = false, const TArray<ITargetPlatform*>& TargetPlatforms=TArray<ITargetPlatform*>() );
+	bool CookOnTheFly( FGuid InstanceId, int32 Port, int32 Timeout = 180, bool bForceClose = false, const TArray<ITargetPlatform*>& TargetPlatforms=TArray<ITargetPlatform*>() );
 
 	/** Cooks for specified targets */
 	bool CookByTheBook(const TArray<ITargetPlatform*>& Platforms);
diff --git a/Engine/Source/Editor/UnrealEd/Classes/CookOnTheSide/CookOnTheFlyServer.h b/Engine/Source/Editor/UnrealEd/Classes/CookOnTheSide/CookOnTheFlyServer.h
--- a/Engine/Source/Editor/UnrealEd/Classes/CookOnTheSide/CookOnTheFlyServer.h
+++ b/Engine/Source/Editor/UnrealEd/Classes/CookOnTheSide/CookOnTheFlyServer.h
@@ -606,8 +606,11 @@ public:
 
 	struct FCookOnTheFlyStartupOptions
 	{
-		/** Wether the network file server or the I/O store connection server should bind to any port */
-		bool bBindAnyPort = false;
+		static constexpr int32 AnyPort = 0;
+		static constexpr int32 DefaultPort = -1;
+
+		/** What port the network file server or the I/O store connection server should bind to */
+		int32 Port = DefaultPort;
 		/** Whether to save the cooked output to the Zen storage server. */
 		bool bZenStore = false;
 		/**
@@ -653,7 +656,7 @@ public:
 	 *
 	 * @return true on success, false otherwise.
 	 */
-	UNREALED_API bool StartCookOnTheFly(FCookOnTheFlyStartupOptions InCookOnTheFlyOptions); 
+	UNREALED_API bool StartCookOnTheFly(FCookOnTheFlyStartupOptions InCookOnTheFlyOptions);
 
 	/** Broadcast the fileserver's presence on the network */
 	UNREALED_API bool BroadcastFileserverPresence( const FGuid &InstanceId );
diff --git a/Engine/Source/Editor/UnrealEd/Private/Commandlets/CookCommandlet.cpp b/Engine/Source/Editor/UnrealEd/Private/Commandlets/CookCommandlet.cpp
--- a/Engine/Source/Editor/UnrealEd/Private/Commandlets/CookCommandlet.cpp
+++ b/Engine/Source/Editor/UnrealEd/Private/Commandlets/CookCommandlet.cpp
@@ -81,7 +81,7 @@ UCookCommandlet::UCookCommandlet( const FObjectInitializer& ObjectInitializer )
 	LogToConsole = false;
 }
 
-bool UCookCommandlet::CookOnTheFly( FGuid InstanceId, int32 Timeout, bool bForceClose, const TArray<ITargetPlatform*>& TargetPlatforms)
+bool UCookCommandlet::CookOnTheFly( FGuid InstanceId, int32 Port, int32 Timeout, bool bForceClose, const TArray<ITargetPlatform*>& TargetPlatforms)
 {
 	UCookOnTheFlyServer *CookOnTheFlyServer = NewObject<UCookOnTheFlyServer>();
 
@@ -100,7 +100,7 @@ bool UCookCommandlet::CookOnTheFly( FGuid InstanceId, int32 Timeout, bool bForce
 	CookOnTheFlyServer->Initialize( ECookMode::CookOnTheFly, CookFlags );
 
 	UCookOnTheFlyServer::FCookOnTheFlyStartupOptions CookOnTheFlyStartupOptions;
-	CookOnTheFlyStartupOptions.bBindAnyPort = InstanceId.IsValid();
+	CookOnTheFlyStartupOptions.Port = Port;
 	CookOnTheFlyStartupOptions.bZenStore = Switches.Contains(TEXT("ZenStore"));
 	CookOnTheFlyStartupOptions.bPlatformProtocol = Switches.Contains(TEXT("PlatformProtocol"));
 	CookOnTheFlyStartupOptions.TargetPlatforms = TargetPlatforms;
@@ -220,10 +220,20 @@ int32 UCookCommandlet::Main(const FString& CmdLineParams)
 		FString InstanceIdString;
 		bool bForceClose = Switches.Contains(TEXT("FORCECLOSE"));
 
+		int32 Port = UCookOnTheFlyServer::FCookOnTheFlyStartupOptions::DefaultPort;
+		if (!FParse::Value(*Params, TEXT("port="), Port))
+		{
+			Port = UCookOnTheFlyServer::FCookOnTheFlyStartupOptions::DefaultPort;
+		}
+
 		FGuid InstanceId;
 		if (FParse::Value(*Params, TEXT("InstanceId="), InstanceIdString))
 		{
-			if (!FGuid::Parse(InstanceIdString, InstanceId))
+			if (FGuid::Parse(InstanceIdString, InstanceId) && InstanceId.IsValid())
+			{
+				Port = UCookOnTheFlyServer::FCookOnTheFlyStartupOptions::AnyPort;
+			}
+			else
 			{
 				UE_LOG(LogCookCommandlet, Warning, TEXT("Invalid InstanceId on command line: %s"), *InstanceIdString);
 			}
@@ -235,7 +245,7 @@ int32 UCookCommandlet::Main(const FString& CmdLineParams)
 			Timeout = 180;
 		}
 
-		CookOnTheFly( InstanceId, Timeout, bForceClose, TargetPlatforms);
+		CookOnTheFly( InstanceId, Port, Timeout, bForceClose, TargetPlatforms);
 	}
 	else if (Switches.Contains(TEXT("COOKWORKER")))
 	{
diff --git a/Engine/Source/Editor/UnrealEd/Private/CookOnTheFlyServer.cpp b/Engine/Source/Editor/UnrealEd/Private/CookOnTheFlyServer.cpp
--- a/Engine/Source/Editor/UnrealEd/Private/CookOnTheFlyServer.cpp
+++ b/Engine/Source/Editor/UnrealEd/Private/CookOnTheFlyServer.cpp
@@ -619,7 +619,7 @@ bool UCookOnTheFlyServer::StartCookOnTheFly(FCookOnTheFlyStartupOptions InCookOn
 
 	FCookOnTheFlyNetworkServerOptions NetworkServerOptions;
 	NetworkServerOptions.Protocol = CookOnTheFlyOptions->bPlatformProtocol ? ECookOnTheFlyNetworkServerProtocol::Platform : ECookOnTheFlyNetworkServerProtocol::Tcp;
-	NetworkServerOptions.Port = CookOnTheFlyOptions->bBindAnyPort ? 0 : -1;
+	NetworkServerOptions.Port = CookOnTheFlyOptions->Port;
 	if (!InCookOnTheFlyOptions.TargetPlatforms.IsEmpty())
 	{
 		NetworkServerOptions.TargetPlatforms = InCookOnTheFlyOptions.TargetPlatforms;
@@ -11508,7 +11508,7 @@ FBeginCookContext UCookOnTheFlyServer::CreateBeginCookByTheBookContext(const FCo
 FBeginCookContext UCookOnTheFlyServer::CreateBeginCookOnTheFlyContext(const FCookOnTheFlyStartupOptions& Options)
 {
 	bZenStore = Options.bZenStore;
-	CookOnTheFlyOptions->bBindAnyPort = Options.bBindAnyPort;
+	CookOnTheFlyOptions->Port = Options.Port;
 	CookOnTheFlyOptions->bPlatformProtocol = Options.bPlatformProtocol;
 	return FBeginCookContext(*this);
 }
diff --git a/Engine/Source/Editor/UnrealEd/Private/Cooker/CookTypes.cpp b/Engine/Source/Editor/UnrealEd/Private/Cooker/CookTypes.cpp
--- a/Engine/Source/Editor/UnrealEd/Private/Cooker/CookTypes.cpp
+++ b/Engine/Source/Editor/UnrealEd/Private/Cooker/CookTypes.cpp
@@ -881,7 +881,7 @@ bool LoadFromCompactBinary(FCbFieldView Field, UE::Cook::FCookByTheBookOptions&
 FCbWriter& operator<<(FCbWriter& Writer, const UE::Cook::FCookOnTheFlyOptions& Value)
 {
 	Writer.BeginObject();
-	Writer << "BindAnyPort" << Value.bBindAnyPort;
+	Writer << "Port" << Value.Port;
 	Writer << "PlatformProtocol" << Value.bPlatformProtocol;
 	Writer.EndObject();
 	return Writer;
@@ -890,7 +890,7 @@ FCbWriter& operator<<(FCbWriter& Writer, const UE::Cook::FCookOnTheFlyOptions& V
 bool LoadFromCompactBinary(FCbFieldView Field, UE::Cook::FCookOnTheFlyOptions& OutValue)
 {
 	bool bOk = Field.IsObject();
-	bOk = LoadFromCompactBinary(Field["BindAnyPort"], OutValue.bBindAnyPort) & bOk;
+	bOk = LoadFromCompactBinary(Field["Port"], OutValue.Port) & bOk;
 	bOk = LoadFromCompactBinary(Field["PlatformProtocol"], OutValue.bPlatformProtocol) & bOk;
 	return bOk;
 }
diff --git a/Engine/Source/Editor/UnrealEd/Private/Cooker/CookTypes.h b/Engine/Source/Editor/UnrealEd/Private/Cooker/CookTypes.h
--- a/Engine/Source/Editor/UnrealEd/Private/Cooker/CookTypes.h
+++ b/Engine/Source/Editor/UnrealEd/Private/Cooker/CookTypes.h
@@ -468,8 +468,8 @@ public:
 // Cook on the fly startup options
 struct FCookOnTheFlyOptions
 {
-	/** Wether the network file server or the I/O store connection server should bind to any port */
-	bool bBindAnyPort = false;
+	/** What port the network file server or the I/O store connection server should bind to */
+	int32 Port = -1;
 	/**
 	 * Whether the network file server should use a platform-specific communication protocol instead of TCP (used when
 	 * bZenStore == false)
diff --git a/Engine/Source/Editor/UnrealEd/Private/UnrealEdEngine.cpp b/Engine/Source/Editor/UnrealEd/Private/UnrealEdEngine.cpp
--- a/Engine/Source/Editor/UnrealEd/Private/UnrealEdEngine.cpp
+++ b/Engine/Source/Editor/UnrealEd/Private/UnrealEdEngine.cpp
@@ -266,7 +266,7 @@ void UUnrealEdEngine::Init(IEngineLoop* InEngineLoop)
 			CookServer->Initialize(ECookMode::CookOnTheFlyFromTheEditor, BaseCookingFlags);
 
 			UCookOnTheFlyServer::FCookOnTheFlyStartupOptions CookOnTheFlyStartupOptions;
-			CookOnTheFlyStartupOptions.bBindAnyPort = false;
+			CookOnTheFlyStartupOptions.Port = UCookOnTheFlyServer::FCookOnTheFlyStartupOptions::DefaultPort;
 			CookOnTheFlyStartupOptions.bZenStore = GetDefault<UProjectPackagingSettings>()->bUseZenStore;
 			CookServer->StartCookOnTheFly(CookOnTheFlyStartupOptions);
 		}
