MPCook: CookPackageSplitter: Add documentation and enforcement of a new detail to the contract required by WorldPartition: PopulateGeneratedPackage and PopulateGeneratorPackage will not be called twice without a garbage collect in between. This was accidentally guaranteed before retraction of generated packages, but when retraction was possible it could be broken by a worker calling PopulateGeneratorPackage once for one assigned generated package, deleting the splitter when done, and then calling it again for a new assigned generated package later.
Remove handling for FCookGenerationInfo.PackageData == nullptr; it is set to non-null in constructor so we don't need to check for it.
#rnx

diff --git a/Engine/Source/Editor/UnrealEd/Private/CookOnTheFlyServer.cpp b/Engine/Source/Editor/UnrealEd/Private/CookOnTheFlyServer.cpp
--- a/Engine/Source/Editor/UnrealEd/Private/CookOnTheFlyServer.cpp
+++ b/Engine/Source/Editor/UnrealEd/Private/CookOnTheFlyServer.cpp
@@ -3378,7 +3378,7 @@ UE::Cook::EPollStatus UCookOnTheFlyServer::PrepareSaveGenerationPackage(UE::Cook
 			}
 			for (FCookGenerationInfo& GeneratedInfo : GenerationHelper.GetPackagesToGenerate())
 			{
-				if (GeneratedInfo.PackageData && GeneratedInfo.PackageData->IsInProgress())
+				if (GeneratedInfo.PackageData->IsInProgress())
 				{
 					return EPollStatus::Incomplete;
 				}
@@ -3504,16 +3504,12 @@ UE::Cook::EPollStatus UCookOnTheFlyServer::BeginCacheObjectsToMove(UE::Cook::FGe
 {
 	using namespace UE::Cook;
 
-	check(Info.PackageData); // Caller validated this
 	FPackageData& PackageData(*Info.PackageData);
 	UPackage* Package = PackageData.GetPackage();
-	ICookPackageSplitter* Splitter = GenerationHelper.GetCookPackageSplitterInstance();
-	UObject* SplitDataObject = GenerationHelper.FindOrLoadSplitDataObject();
-	if (!Package || !Splitter || !SplitDataObject)
+	if (!Package)
 	{
 		UE_LOG(LogCook, Error,
-			TEXT("CookPackageSplitter is missing %s during BeginCacheObjectsToMove. PackageName: %s."),
-			(!Package ? TEXT("Package") : (!Splitter ? TEXT("Splitter") : TEXT("SplitDataObject"))),
+			TEXT("CookPackageSplitter is missing package during BeginCacheObjectsToMove. PackageName: %s."),
 			*PackageData.GetPackageName().ToString());
 		return EPollStatus::Error;
 	}
@@ -3544,24 +3540,10 @@ UE::Cook::EPollStatus UCookOnTheFlyServer::BeginCacheObjectsToMove(UE::Cook::FGe
 		}
 		else
 		{
-			TArray<UPackage*> KeepReferencedPackages;
-			ICookPackageSplitter::FGeneratedPackageForPopulate SplitterInfo{ Info.RelativePath,
-				Info.GeneratedRootPath, Package, Info.IsCreateAsMap() };
-			FScopedActivePackage ScopedActivePackage(*this, GenerationHelper.GetOwner().GetPackageName(),
-				PackageAccessTrackingOps::NAME_CookerBuildObject);
-			bPopulateSucceeded = Splitter->PopulateGeneratedPackage(Package, SplitDataObject, SplitterInfo,
-				ObjectsToMove, KeepReferencedPackages);
-			if (!bPopulateSucceeded)
-			{
-				UE_LOG(LogCook, Error,
-					TEXT("CookPackageSplitter returned false from PopulateGeneratedPackage. Splitter=%s")
-					TEXT("\nGeneratedPackage: %s"),
-					*GenerationHelper.GetSplitDataObjectName().ToString(),
-					*PackageData.GetPackageName().ToString());
+			if (!GenerationHelper.TryCallPopulateGeneratedPackage(Info, ObjectsToMove))
+			{
 				return EPollStatus::Error;
 			}
-
-			Info.AddKeepReferencedPackages(GenerationHelper, KeepReferencedPackages);
 		}
 
 		Info.TakeOverCachedObjectsAndAddMoved(GenerationHelper, PackageData.GetCachedObjectsInOuter(), ObjectsToMove);
@@ -3653,7 +3635,6 @@ UE::Cook::EPollStatus UCookOnTheFlyServer::BeginCachePostMove(UE::Cook::FGenerat
 {
 	using namespace UE::Cook;
 
-	check(Info.PackageData); // Caller has validated
 	UE::Cook::FPackageData& PackageData(*Info.PackageData);
 	UPackage* Package = PackageData.GetPackage();
 	ICookPackageSplitter* Splitter = GenerationHelper.GetCookPackageSplitterInstance();
@@ -3724,7 +3705,6 @@ UE::Cook::EPollStatus UCookOnTheFlyServer::TryPopulateGeneratedPackage(UE::Cook:
 {
 	using namespace UE::Cook;
 
-	check(GeneratedInfo.PackageData); // Caller already checked this
 	UE::Cook::FPackageData& GeneratedPackageData = *GeneratedInfo.PackageData;
 	const FString GeneratedPackageName = GeneratedPackageData.GetPackageName().ToString();
 	UPackage* OwnerPackage = GenerationHelper.FindOrLoadOwnerPackage(*this);
@@ -5260,37 +5240,33 @@ void UCookOnTheFlyServer::PreGarbageCollect()
 		GCKeepPackageDatas.Add(SavingPackageData);
 	}
 
-
-	// Demote any Generated/Generator packages we called PreSave on so they call their PostSave before the GC
-	// or prevent them from being garbage collected if the splitter wants to keep them referenced
-	for (FPackageData* PackageData : PackageDatas->GetSaveQueue())
-	{
-		FGenerationHelper* GenerationHelper = PackageData->GetGenerationHelper();
-		if (!GenerationHelper)
-		{
-			GenerationHelper = PackageData->GetParentGenerationHelper();
-		}
-		if (GenerationHelper && GenerationHelper->IsInitialized())
+	// Notify every FGenerationHelper of the garbage collect
+	PackageDatas->LockAndEnumeratePackageDatas([this, &GCKeepPackages, &GCKeepPackageDatas](FPackageData* PackageData)
 		{
-			FCookGenerationInfo* Info = GenerationHelper->FindInfo(*PackageData);
-			if (Info)
+			TRefCountPtr<FGenerationHelper> GenerationHelper = PackageData->GetGenerationHelper();
+			if (!GenerationHelper)
+			{
+				GenerationHelper = PackageData->GetParentGenerationHelper();
+			}
+			if (GenerationHelper)
 			{
 				bool bShouldDemote;
-				GenerationHelper->PreGarbageCollect(*Info, GCKeepObjects, GCKeepPackages,
+				GenerationHelper->PreGarbageCollect(GenerationHelper, *PackageData, GCKeepObjects, GCKeepPackages,
 					GCKeepPackageDatas, bShouldDemote);
-				if (bShouldDemote)
+				if (bShouldDemote && PackageData->GetState() >= EPackageState::Save)
 				{
+					// Demote any Generated/Generator packages we called PreSave on so they call their PostSave before the GC
+					// or prevent them from being garbage collected if the splitter wants to keep them referenced
 					ReleaseCookedPlatformData(*PackageData, UE::Cook::EStateChangeReason::GeneratorPreGarbageCollected,
 						EPackageState::Request);
 				}
 			}
-		}
-		if (PackageData->GetIsCookLast())
-		{
-			GCKeepPackages.Add(PackageData->GetPackage());
-			GCKeepPackageDatas.Add(PackageData);
-		}
-	}
+			if (PackageData->GetIsCookLast() && PackageData->GetState() >= EPackageState::Save)
+			{
+				GCKeepPackages.Add(PackageData->GetPackage());
+				GCKeepPackageDatas.Add(PackageData);
+			}
+		});
 	
 	// Find the packages that are waiting on async jobs to finish cooking data
 	// and make sure that they are not garbage collected until the jobs have
@@ -5523,7 +5499,7 @@ void UCookOnTheFlyServer::PostGarbageCollect()
 	{
 		if (TRefCountPtr<FGenerationHelper> GenerationHelper = PackageData->GetGenerationHelper())
 		{
-			GenerationHelper->PostGarbageCollect();
+			GenerationHelper->PostGarbageCollect(GenerationHelper);
 		}
 	});
 	PackageDatas->LockAndEnumeratePackageDatas([](FPackageData* PackageData)
diff --git a/Engine/Source/Editor/UnrealEd/Private/Cooker/CookDirector.cpp b/Engine/Source/Editor/UnrealEd/Private/Cooker/CookDirector.cpp
--- a/Engine/Source/Editor/UnrealEd/Private/Cooker/CookDirector.cpp
+++ b/Engine/Source/Editor/UnrealEd/Private/Cooker/CookDirector.cpp
@@ -592,10 +592,14 @@ void FCookDirector::RemoveFromWorker(FPackageData& PackageData)
 	OwningWorker->AbortAssignment(PackageData, ECookDirectorThread::SchedulerThread);
 }
 
-void FCookDirector::BroadcastGeneratorFencePassed(FGenerationHelper& GenerationHelper)
+void FCookDirector::BroadcastGeneratorMessage(FGeneratorEventMessage&& Message)
 {
-	FName PackageName = GenerationHelper.GetOwner().GetPackageName();
-	FGeneratorEventMessage Message(EGeneratorEvent::QueuedGeneratedPackagesFencePassed, PackageName);
+	if (bReceivingMessages)
+	{
+		QueuedGeneratorBroadcasts.Add(MoveTemp(Message));
+		return;
+	}
+
 	{
 		FScopeLock CommunicationScopeLock(&CommunicationLock);
 		InitializeWorkers();
@@ -670,10 +674,21 @@ void FCookDirector::TickFromSchedulerThread()
 
 	bool bIsStalled = bLocalWorkerIdle && !COTFS.PackageDatas->GetAssignedToWorkerSet().IsEmpty()
 		&& WorkersWithMessage.IsEmpty();
-	for (TRefCountPtr<FCookWorkerServer>& Worker : WorkersWithMessage)
 	{
-		Worker->HandleReceiveMessages(ECookDirectorThread::SchedulerThread);
+		bReceivingMessages = true;
+		for (TRefCountPtr<FCookWorkerServer>& Worker : WorkersWithMessage)
+		{
+			Worker->HandleReceiveMessages(ECookDirectorThread::SchedulerThread);
+		}
+		bReceivingMessages = false;
+	}
+	for (FGeneratorEventMessage& Message : QueuedGeneratorBroadcasts)
+	{
+		BroadcastGeneratorMessage(MoveTemp(Message));
 	}
+	QueuedGeneratorBroadcasts.Empty();
+
+	// Process any queued messages
 	WorkersWithMessage.Empty();
 
 	SetWorkersStalled(bIsStalled);
diff --git a/Engine/Source/Editor/UnrealEd/Private/Cooker/CookDirector.h b/Engine/Source/Editor/UnrealEd/Private/Cooker/CookDirector.h
--- a/Engine/Source/Editor/UnrealEd/Private/Cooker/CookDirector.h
+++ b/Engine/Source/Editor/UnrealEd/Private/Cooker/CookDirector.h
@@ -30,6 +30,7 @@ class UCookOnTheFlyServer;
 namespace UE::Cook { class FCookWorkerServer; }
 namespace UE::Cook { struct FAssignPackageExtraData; }
 namespace UE::Cook { struct FCookWorkerProfileData; }
+namespace UE::Cook { struct FGeneratorEventMessage; }
 namespace UE::Cook { struct FGenerationHelper; }
 namespace UE::Cook { struct FHeartbeatMessage; }
 namespace UE::Cook { struct FInitialConfigMessage; }
@@ -80,7 +81,7 @@ public:
 	 * Report EGeneratorEvent::QueuedGeneratedPackagesFencePassed to all CookWorkers so any of them
 	 * waiting on the fence to be passed before clearing their data can clear their data.
 	 */
-	void BroadcastGeneratorFencePassed(FGenerationHelper& GenerationHelper);
+	void BroadcastGeneratorMessage(FGeneratorEventMessage&& Message);
 
 	/** Periodic tick function. Sends/Receives messages to CookWorkers. */
 	void TickFromSchedulerThread();
@@ -244,6 +245,7 @@ private:
 	TArray<FCookWorkerProfileData> RemoteWorkerProfileDatas;
 	TArray<FPendingConnection> PendingConnections;
 	TUniquePtr<FCookWorkerProfileData> LocalWorkerProfileData;
+	TArray<FGeneratorEventMessage> QueuedGeneratorBroadcasts;
 	UCookOnTheFlyServer& COTFS;
 	double WorkersStalledStartTimeSeconds = 0.;
 	double WorkersStalledWarnTimeSeconds = 0.;
@@ -257,6 +259,7 @@ private:
 	bool bCookCompleteSent = false;
 	bool bWorkersStalled = false;
 	bool bMultiprocessAvailable = false;
+	bool bReceivingMessages = false;
 
 	// Data that is read-only while the CommunicationThread is active and is readable from any thread
 	FBeginCookContextForWorker BeginCookContext;
diff --git a/Engine/Source/Editor/UnrealEd/Private/Cooker/CookGenerationHelper.cpp b/Engine/Source/Editor/UnrealEd/Private/Cooker/CookGenerationHelper.cpp
--- a/Engine/Source/Editor/UnrealEd/Private/Cooker/CookGenerationHelper.cpp
+++ b/Engine/Source/Editor/UnrealEd/Private/Cooker/CookGenerationHelper.cpp
@@ -5,6 +5,7 @@
 #include "Algo/Unique.h"
 #include "AssetRegistry/AssetData.h"
 #include "Cooker/CookDirector.h"
+#include "Cooker/CookWorkerServer.h"
 #include "Cooker/IWorkerRequests.h"
 #include "Cooker/PackageTracker.h"
 #include "Misc/CommandLine.h"
@@ -91,6 +92,8 @@ void FGenerationHelper::Initialize(const UObject* InSplitDataObject,
 	SplitDataObjectName = FName(FStringView(InSplitDataObject->GetFullName()));
 	bUseInternalReferenceToAvoidGarbageCollect =
 		CookPackageSplitterInstance->UseInternalReferenceToAvoidGarbageCollect();
+	bRequiresGeneratorPackageDestructBeforeResplit =
+		CookPackageSplitterInstance->RequiresGeneratorPackageDestructBeforeResplit();
 	DoesGeneratedRequireGeneratorValue =
 		CookPackageSplitterInstance->DoesGeneratedRequireGenerator();
 }
@@ -104,6 +107,79 @@ void FGenerationHelper::InitializeAsInvalid()
 	InitializeStatus = EInitializeStatus::Invalid;
 }
 
+void FGenerationHelper::Uninitialize()
+{
+	if (InitializeStatus != EInitializeStatus::Valid)
+	{
+		return;
+	}
+	NotifyCompletion(ICookPackageSplitter::ETeardown::Complete);
+	check(!CookPackageSplitterInstance);
+
+	InitializeStatus = EInitializeStatus::Uninitialized;
+
+	OwnerInfo.Uninitialize();
+	SplitDataObject.Reset();
+	SplitDataObjectName = NAME_None;
+	RegisteredSplitterType = nullptr;
+	// CookPackageSplitterInstance was set to null above
+	for (FCookGenerationInfo& Info : PackagesToGenerate)
+	{
+		Info.Uninitialize();
+	}
+	OwnerPackage.Reset();
+	// Keep PreviousGeneratedPackages; they are allowed in the uninitialized state
+	check(ExternalActorDependencies.IsEmpty()); // We can not still be in the save state, so this should be empty
+	check(OwnerObjectsToMove.IsEmpty()); // We can not still be in the save state, so this should be empty
+	// Do not modify the reference tracking variables
+	// ReferenceFromKeepForIterative
+	// ReferenceFromKeepForQueueResults
+	// ReferenceFromKeepForGeneratorSave
+	// ReferenceFromKeepForAllSavedOrGC
+	check(MPCookNextAssignmentIndex == 0); // We can not still be in the save state, so should have been cleared
+	// Keep NumSaved; it is allowed in the uninitialized state
+	// Keep WorkerIdThatSavedGenerator; it is allowed in the uninitialized state
+	// InitializeStatus was modified above
+	// Keep DoesGeneratedRequireGeneratorValue; it is allowed in the uninitialized state
+	// Keep bUseInternalReferenceToAvoidGarbageCollect; it is allowed in the uninitialized state
+	// Keep bRequiresGeneratorPackageDestructBeforeResplit; it is allowed in the uninitialized state
+	bGeneratedList = false;
+	bCurrentGCHasKeptGeneratorPackage = false;
+	bCurrentGCHasKeptGeneratorKeepPackages = false;
+	// Keep bKeepForAllSavedOrGC ; it is allowed in the uninitialized state
+	// Keep bKeepForCompletedAllSavesMessage; it is allowed in the uninitialized state
+	// Keep bNeedConfirmGeneratorPackageDestroyed; it is allowed in the uninitialized state
+}
+
+void FGenerationHelper::ModifyNumSaved(int32 Delta)
+{
+	NumSaved += Delta;
+	int32 NumAllSaved = PackagesToGenerate.Num() + 1;
+	check(0 <= NumSaved && NumSaved <= NumAllSaved);
+	if (NumSaved == NumAllSaved)
+	{
+		UCookOnTheFlyServer& COTFS = GetOwner().GetPackageDatas().GetCookOnTheFlyServer();
+		if (!COTFS.CookWorkerClient)
+		{
+			if (COTFS.CookDirector)
+			{
+				FName PackageName = GetOwner().GetPackageName();
+				FGeneratorEventMessage Message(EGeneratorEvent::AllSavesCompleted, PackageName);
+				COTFS.CookDirector->BroadcastGeneratorMessage(MoveTemp(Message));
+			}
+			OnAllSavesCompleted(COTFS);
+		}
+	}
+}
+
+void FGenerationHelper::OnAllSavesCompleted(UCookOnTheFlyServer& COTFS)
+{
+	// Caller is responsible for holding a reference that keeps *this from destructing if it clears
+	// these references
+	ClearKeepForCompletedAllSavesMessage();
+	ClearKeepForAllSavedOrGC();
+}
+
 UPackage* FGenerationHelper::FindOrLoadPackage(UCookOnTheFlyServer& COTFS, FPackageData& OwnerPackageData)
 {
 	// This is the static helper function on FGenerationHelper that loads the package for any FPackageData; for the
@@ -226,6 +302,8 @@ void FGenerationHelper::ClearSelfReferences()
 	ClearKeepForIterative();
 	ClearKeepForGeneratorSave();
 	ClearKeepForQueueResults();
+	ClearKeepForAllSavedOrGC();
+	ClearKeepForCompletedAllSavesMessage();
 }
 
 FCookGenerationInfo* FGenerationHelper::FindInfo(const FPackageData& PackageData)
@@ -350,10 +428,18 @@ bool FGenerationHelper::TryGenerateList()
 			PackageAccessTrackingOps::NAME_CookerBuildObject);
 		GeneratorDatas = GetCookPackageSplitterInstance()->GetGenerateList(LocalOwnerPackage, OwnerObject);
 	}
-	PackagesToGenerate.Reset(GeneratorDatas.Num());
+
+	TMap<FPackageData*, int32> AlreadyExistingInfoPackageToIndex;
+	for (int32 ExistingIndex = 0; ExistingIndex < PackagesToGenerate.Num(); ++ExistingIndex)
+	{
+		AlreadyExistingInfoPackageToIndex.Add(PackagesToGenerate[ExistingIndex].PackageData, ExistingIndex);
+	}
+	PackagesToGenerate.Reserve(GeneratorDatas.Num());
+
 	TArray<const ITargetPlatform*, TInlineAllocator<1>> PlatformsToCook;
 	OwnerPackageData.GetPlatformsNeedingCooking(PlatformsToCook);
 
+	NumSaved = 0;
 	for (ICookPackageSplitter::FGeneratedPackage& SplitterData : GeneratorDatas)
 	{
 		if (!SplitterData.GetCreateAsMap().IsSet())
@@ -390,11 +476,23 @@ bool FGenerationHelper::TryGenerateList()
 			return false;
 		}
 
-		FCookGenerationInfo& GeneratedInfo = PackagesToGenerate.Emplace_GetRef(*PackageData, false /* bInGenerator */);
-		GeneratedInfo.RelativePath = MoveTemp(SplitterData.RelativePath);
-		GeneratedInfo.GeneratedRootPath = MoveTemp(SplitterData.GeneratedRootPath);
-		GeneratedInfo.PackageDependencies = MoveTemp(SplitterData.PackageDependencies);
-		for (TArray<FAssetDependency>::TIterator Iter(GeneratedInfo.PackageDependencies); Iter; ++Iter)
+		FCookGenerationInfo* GeneratedInfo = nullptr;
+		if (!AlreadyExistingInfoPackageToIndex.IsEmpty())
+		{
+			int32 ExistingIndex;
+			if (AlreadyExistingInfoPackageToIndex.RemoveAndCopyValue(PackageData, ExistingIndex))
+			{
+				GeneratedInfo = &PackagesToGenerate[ExistingIndex];
+			}
+		}
+		if (!GeneratedInfo)
+		{
+			GeneratedInfo = &PackagesToGenerate.Emplace_GetRef(*PackageData, false /* bInGenerator */);
+		}
+		GeneratedInfo->RelativePath = MoveTemp(SplitterData.RelativePath);
+		GeneratedInfo->GeneratedRootPath = MoveTemp(SplitterData.GeneratedRootPath);
+		GeneratedInfo->PackageDependencies = MoveTemp(SplitterData.PackageDependencies);
+		for (TArray<FAssetDependency>::TIterator Iter(GeneratedInfo->PackageDependencies); Iter; ++Iter)
 		{
 			if (Iter->Category != UE::AssetRegistry::EDependencyCategory::Package)
 			{
@@ -404,10 +502,10 @@ bool FGenerationHelper::TryGenerateList()
 				Iter.RemoveCurrent();
 			}
 		}
-		Algo::Sort(GeneratedInfo.PackageDependencies,
+		Algo::Sort(GeneratedInfo->PackageDependencies,
 			[](const FAssetDependency& A, const FAssetDependency& B) { return A.LexicalLess(B); });
-		GeneratedInfo.PackageDependencies.SetNum(Algo::Unique(GeneratedInfo.PackageDependencies));
-		GeneratedInfo.SetIsCreateAsMap(bCreateAsMap);
+		GeneratedInfo->PackageDependencies.SetNum(Algo::Unique(GeneratedInfo->PackageDependencies));
+		GeneratedInfo->SetIsCreateAsMap(bCreateAsMap);
 		if (DoesGeneratedRequireGenerator() >= ICookPackageSplitter::EGeneratedRequiresGenerator::Save ||
 			COTFS.MPCookGeneratorSplit == EMPCookGeneratorSplit::AllOnSameWorker)
 		{
@@ -415,8 +513,30 @@ bool FGenerationHelper::TryGenerateList()
 		}
 
 		// Create the Hash from the GenerationHash and Dependencies
-		GeneratedInfo.CreatePackageHash();
+		GeneratedInfo->CreatePackageHash();
+
+		NumSaved += GeneratedInfo->HasSaved() ? 1 : 0;
+	}
+	NumSaved += OwnerInfo.HasSaved() ? 1 : 0;
+
+	if (!AlreadyExistingInfoPackageToIndex.IsEmpty())
+	{
+		TArray<int32> UnusedExistingIndexes;
+		for (TPair<FPackageData*, int32>& Pair : AlreadyExistingInfoPackageToIndex)
+		{
+			UnusedExistingIndexes.Add(Pair.Value);
+			UE_LOG(LogCook, Warning, TEXT("Unexpected generated package (local TryGenerateList). A remote cookworker reported generated package %s for generator %s,")
+				TEXT(" but when TryGenerateList was called on the CookDirector, this package was not listed.")
+				TEXT(" This is unexpected and causes minor performance problems in the cook."),
+				*Pair.Key->GetPackageName().ToString(), *OwnerPackageData.GetPackageName().ToString());
+		}
+		Algo::Sort(UnusedExistingIndexes);
+		for (int32 UnusedIndex : ReverseIterate(UnusedExistingIndexes))
+		{
+			PackagesToGenerate.RemoveAt(UnusedIndex);
+		}
 	}
+	ModifyNumSaved(0);
 
 	bGeneratedList = true;
 	return true;
@@ -425,7 +545,7 @@ bool FGenerationHelper::TryGenerateList()
 bool FGenerationHelper::TryCallPopulateGeneratorPackage(
 	TArray<ICookPackageSplitter::FGeneratedPackageForPreSave>& InOutGeneratedPackagesForPresave)
 {
-	if (bCalledPopulateGeneratorPackage)
+	if (OwnerInfo.HasCalledPopulate())
 	{
 		return true;
 	}
@@ -480,7 +600,68 @@ bool FGenerationHelper::TryCallPopulateGeneratorPackage(
 		}
 	}
 
-	bCalledPopulateGeneratorPackage = true;
+	// Contract Point 1: We will not call populate again until the splitter has been destroyed
+	// Contract Point 2: We will not call populate again without garbage collecting the generator package
+	OwnerInfo.SetHasCalledPopulate(true);
+	SetKeepForAllSavedOrGC();
+	return true;
+}
+
+bool FGenerationHelper::TryCallPopulateGeneratedPackage(UE::Cook::FCookGenerationInfo& Info,
+	TArray<UObject*>& OutObjectsToMove)
+{
+	if (Info.HasCalledPopulate())
+	{
+		return true;
+	}
+	FPackageData& OwnerPackageData = GetOwner();
+	FName OwnerPackageName = OwnerPackageData.GetPackageName();
+	UCookOnTheFlyServer& COTFS = OwnerPackageData.GetPackageDatas().GetCookOnTheFlyServer();
+
+	if (!bGeneratedList)
+	{
+		// Unexpected, caller should not call in this case
+		UE_LOG(LogCook, Error, TEXT("TryCallPopulateGeneratedPackage called for package %s without a previous successful call to TryGenerateList."),
+			*Info.GetPackageName());
+		FDebug::DumpStackTraceToLog(ELogVerbosity::Warning);
+		return false;
+	}
+	check(IsValid()); // Could not have set bGeneratedList=true without being valid. 
+	UObject* LocalSplitDataObject = this->FindOrLoadSplitDataObject();
+	if (!LocalSplitDataObject)
+	{
+		UE_LOG(LogCook, Error,
+			TEXT("Failed to call TryCallPopulateGeneratedPackage, CookPackageSplitter missing. Splitter=%s"),
+			*GetSplitDataObjectName().ToString());
+		return false;
+	}
+
+	UPackage* Package = Info.PackageData->GetPackage();
+	check(Package); // Caller checked this
+	ICookPackageSplitter::FGeneratedPackageForPopulate SplitterInfo{ Info.RelativePath,
+		Info.GeneratedRootPath, Package, Info.IsCreateAsMap() };
+
+	UCookOnTheFlyServer::FScopedActivePackage ScopedActivePackage(COTFS, OwnerPackageName,
+		PackageAccessTrackingOps::NAME_CookerBuildObject);
+
+	TArray<UPackage*> KeepReferencedPackages;
+	bool bPopulateSucceeded = CookPackageSplitterInstance->PopulateGeneratedPackage(Package,
+		LocalSplitDataObject, SplitterInfo, OutObjectsToMove, KeepReferencedPackages);
+	if (!bPopulateSucceeded)
+	{
+		UE_LOG(LogCook, Error,
+			TEXT("CookPackageSplitter returned false from PopulateGeneratedPackage. Splitter=%s")
+			TEXT("\nGeneratedPackage: %s"),
+			*GetSplitDataObjectName().ToString(), *Info.GetPackageName());
+		return false;
+	}
+
+	Info.AddKeepReferencedPackages(*this, KeepReferencedPackages);
+
+	// Contract Point 1: We will not call populate again until the splitter has been destroyed
+	// Contract Point 2: We will not call populate again without garbage collecting the generator package
+	Info.SetHasCalledPopulate(true);
+	SetKeepForAllSavedOrGC();
 	return true;
 }
 
@@ -521,10 +702,6 @@ void FGenerationHelper::StartQueueGeneratedPackages(UCookOnTheFlyServer& COTFS)
 
 		for (FCookGenerationInfo& GeneratedInfo : PackagesToGenerate)
 		{
-			if (!GeneratedInfo.PackageData)
-			{
-				continue;
-			}
 			FIoHash PreviousHash;
 			if (RemainingPreviousPackages.RemoveAndCopyValue(GeneratedInfo.PackageData->GetPackageName(), PreviousHash)
 				&& !bHybridIterativeEnabled)
@@ -566,6 +743,7 @@ void FGenerationHelper::NotifyStartQueueGeneratedPackages(UCookOnTheFlyServer& C
 	{
 		COTFS.PackageDatas->GetRequestQueue().AddRequestFenceListener(GetOwner().GetPackageName());
 		WorkerIdThatSavedGenerator = SourceWorkerId;
+		SetKeepForCompletedAllSavesMessage();
 	}
 	SetKeepForQueueResults();
 }
@@ -593,8 +771,9 @@ void FGenerationHelper::OnRequestFencePassedBroadcast(UCookOnTheFlyServer& COTFS
 {
 	if (COTFS.CookDirector)
 	{
-		// Broadcast GenerationMessage API
-		COTFS.CookDirector->BroadcastGeneratorFencePassed(*this);
+		FName PackageName = GetOwner().GetPackageName();
+		FGeneratorEventMessage Message(EGeneratorEvent::QueuedGeneratedPackagesFencePassed, PackageName);
+		COTFS.CookDirector->BroadcastGeneratorMessage(MoveTemp(Message));
 	}
 	OnRequestFencePassed(COTFS);
 }
@@ -616,7 +795,6 @@ UPackage* FGenerationHelper::TryCreateGeneratedPackage(FCookGenerationInfo& Gene
 		return nullptr;
 	}
 
-	check(GeneratedInfo.PackageData); // Caller is responsible for validating
 	FPackageData& GeneratedPackageData = *GeneratedInfo.PackageData;
 	const FString GeneratedPackageName = GeneratedPackageData.GetPackageName().ToString();
 	UPackage* LocalOwnerPackage = FindOrLoadOwnerPackage(GetOwner().GetPackageDatas().GetCookOnTheFlyServer());
@@ -734,10 +912,7 @@ void FGenerationHelper::ResetSaveState(FCookGenerationInfo& Info, UPackage* Pack
 	{
 		// The save is completed and we will not come back to it; set state back to initial
 		// state and drop our reference keeping this GenerationHelper in memory for the save.
-		if (ReleaseSaveReason == EStateChangeReason::Completed)
-		{
-			Info.SetHasSaved(true);
-		}
+		Info.SetHasSaved(*this, true);
 
 		if (Info.IsGenerator())
 		{
@@ -753,16 +928,13 @@ void FGenerationHelper::ResetSaveState(FCookGenerationInfo& Info, UPackage* Pack
 		else
 		{
 			Info.SetSaveState(FCookGenerationInfo::ESaveState::StartPopulate);
-			if (Info.PackageData)
-			{
-				Info.PackageData->SetParentGenerationHelper(nullptr);
-			}
+			Info.PackageData->SetParentGenerationHelper(nullptr);
 		}
 	}
 
 	if (Info.HasTakenOverCachedCookedPlatformData())
 	{
-		if (NewState != EPackageState::Idle && Info.PackageData &&
+		if (NewState != EPackageState::Idle &&
 			Info.PackageData->GetCachedObjectsInOuter().Num() != 0 && IsUseInternalReferenceToAvoidGarbageCollect() &&
 			(ReleaseSaveReason != EStateChangeReason::Completed && ReleaseSaveReason != EStateChangeReason::DoneForNow
 				&& ReleaseSaveReason != EStateChangeReason::SaveError
@@ -859,16 +1031,28 @@ static void AppendWeakPtrsToObjectPtrArray(TArray<T*>& Out, TArray<TWeakObjectPt
 		}
 	}
 }
-void FGenerationHelper::PreGarbageCollect(FCookGenerationInfo& Info, TArray<TObjectPtr<UObject>>& GCKeepObjects,
+
+void FGenerationHelper::PreGarbageCollect(const TRefCountPtr<FGenerationHelper>& RefcountHeldByCaller,
+	FPackageData& PackageData, TArray<TObjectPtr<UObject>>& GCKeepObjects,
 	TArray<UPackage*>& GCKeepPackages, TArray<FPackageData*>& GCKeepPackageDatas, bool& bOutShouldDemote)
 {
+	bOutShouldDemote = false;
+	if (&PackageData == &GetOwner())
+	{
+		PreGarbageCollectGCLifetimeData();
+	}
 	if (!IsInitialized())
 	{
 		return;
 	}
 
-	bOutShouldDemote = false;
-	check(Info.PackageData); // Caller validates this is non-null
+	FCookGenerationInfo* InfoPtr = FindInfo(PackageData);
+	if (!InfoPtr)
+	{
+		return;
+	}
+	FCookGenerationInfo& Info = *InfoPtr;
+
 	if (!IsUseInternalReferenceToAvoidGarbageCollect() && !Info.PackageData->GetIsCookLast())
 	{
 		// If we don't have a contract to keep the packagedata referenced during GC, don't report
@@ -882,10 +1066,7 @@ void FGenerationHelper::PreGarbageCollect(FCookGenerationInfo& Info, TArray<TObj
 
 	// When we have a contract to keep the packagedata referenced, keep its various object pointers referenced.
 
-	// We have a contract that KeepReferencedPackages in any Info are kept referenced. COTFS calls PreGarbageCollect
-	// for any info in the save state; with one exception only packages in the savestate can have
-	// non-empty KeepReferencedPackages. The one exception is the generator, which can keep 
-	// KeepReferencedPackages beyond its save state for the benefit of generated packages that rely on it.
+	// We have a contract that KeepReferencedPackages in any Info are kept referenced.
 	bool bKeepingAnyObjects = false;
 	bool bNeedsGeneratorPackage = false;
 	if (&Info == &OwnerInfo)
@@ -949,8 +1130,83 @@ void FGenerationHelper::PreGarbageCollect(FCookGenerationInfo& Info, TArray<TObj
 	}
 }
 
-void FGenerationHelper::PostGarbageCollect()
+void FGenerationHelper::PreGarbageCollectGCLifetimeData()
+{
+	// Starts at one because the caller of PreGarbageCollect has a ref
+	uint32 HoldForGCRefCounts = 1;
+	HoldForGCRefCounts += ReferenceFromKeepForAllSavedOrGC.IsValid() ? 1 : 0;
+
+	check(GetRefCount() >= HoldForGCRefCounts);
+	if (GetRefCount() > HoldForGCRefCounts)
+	{
+		// Something else (generator save or generated package save, etc) is keeping us referenced
+		// and we need to not allow destruction. Nothing further to do.
+		return;
+	}
+
+	// We should either uninitialize or destroy after the garbage collect.
+	// We should not uninitialize unless the Generator package is going to be collected, but we are in a state
+	// where nothing in the cooker is depending on the package anymore (all generator and generated packages are
+	// not in the save state) so we do expect the generator package to be garbage collected by the upcoming GC
+	// But for that to happen we have to drop our references to it from this FGenerationHelper, so we need
+	// to uninitialize. Also mark that we should check for generator garbage collect in PostGarbageCollect.
+	// Depending on the Splitter class, it may tolerate failure to GC the Generator package, in which case we
+	// should not log this error.
+	Uninitialize();
+	bNeedConfirmGeneratorPackageDestroyed = IsRequiresGeneratorPackageDestructBeforeResplit();
+}
+
+void FGenerationHelper::PostGarbageCollectGCLifetimeData()
+{
+	if (bNeedConfirmGeneratorPackageDestroyed)
+	{
+		VerifyGeneratorPackageGarbageCollected();
+		bNeedConfirmGeneratorPackageDestroyed = false;
+	}
+
+	if (!IsInitialized())
+	{
+		// ClearKeepForAllSavedOrGC is no longer required when Uninitialized after a GC
+		// Note that this keep flag might be the last persistent reference to *this and *this will be deleted when the
+		// caller of PostGarbageCollect drops its reference.
+		ClearKeepForAllSavedOrGC();
+	}
+}
+
+void FGenerationHelper::TrackGeneratedPackageListedRemotely(UCookOnTheFlyServer& COTFS, FPackageData& PackageData)
+{
+	if (bGeneratedList)
+	{
+		if (!FindInfo(PackageData))
+		{
+			UE_LOG(LogCook, Warning, TEXT("Unexpected generated package (discovery replication). A remote cookworker reported generated package %s for generator %s,")
+				TEXT(" but when TryGenerateList was called on the CookDirector, this package was not listed.")
+				TEXT(" This is unexpected and causes minor performance problems in the cook."),
+				*PackageData.GetPackageName().ToString(), *GetOwner().GetPackageName().ToString());
+		}
+	}
+	else
+	{
+		if (!FindInfo(PackageData))
+		{
+			bool bGenerator = false; // Cannot be the generator, if it were we would have found it
+			PackagesToGenerate.Emplace(PackageData, bGenerator);
+		}
+	}
+}
+
+void FGenerationHelper::MarkPackageSavedRemotely(UCookOnTheFlyServer& COTFS, FPackageData& PackageData)
+{
+	FCookGenerationInfo* Info = FindInfo(PackageData);
+	if (Info)
+	{
+		Info->SetHasSaved(*this, true);
+	}
+}
+
+void FGenerationHelper::PostGarbageCollect(const TRefCountPtr<FGenerationHelper>& RefcountHeldByCaller)
 {
+	PostGarbageCollectGCLifetimeData();
 	if (!IsInitialized())
 	{
 		return;
@@ -978,25 +1234,14 @@ void FGenerationHelper::PostGarbageCollect()
 		// The splitter can opt-out of this contract and keep it referenced itself if it desires.
 		if (!Owner.IsInProgress() && !Owner.IsKeepReferencedDuringGC())
 		{
-			UPackage* LocalOwnerPackage = FindObject<UPackage>(nullptr, *Owner.GetPackageName().ToString());
-			if (LocalOwnerPackage)
-			{
-				UE_LOG(LogCook, Error,
-					TEXT("PackageSplitter found the Generator package still in memory after it should have been deleted by GC.")
-					TEXT("\n\tThis is unexpected since garbage has been collected and the package should have been unreferenced so it should have been collected, and will break population of Generated packages.")
-					TEXT("\n\tSplitter=%s"), *GetSplitDataObjectName().ToString());
-				EReferenceChainSearchMode SearchMode = EReferenceChainSearchMode::Shortest
-					| EReferenceChainSearchMode::PrintAllResults
-					| EReferenceChainSearchMode::FullChain;
-				FReferenceChainSearch RefChainSearch(LocalOwnerPackage, SearchMode);
-			}
+			VerifyGeneratorPackageGarbageCollected();
 		}
 	}
 
 	bool bHasIssuedWarning = false;
 	for (FCookGenerationInfo& Info : PackagesToGenerate)
 	{
-		if (Info.PackageData && FindObject<UPackage>(nullptr, *Info.PackageData->GetPackageName().ToString()))
+		if (FindObject<UPackage>(nullptr, *Info.PackageData->GetPackageName().ToString()))
 		{
 			if (!Info.PackageData->IsKeepReferencedDuringGC() && !Info.HasSaved() && !bHasIssuedWarning)
 			{
@@ -1037,6 +1282,34 @@ void FGenerationHelper::PostGarbageCollect()
 	}
 }
 
+void FGenerationHelper::VerifyGeneratorPackageGarbageCollected()
+{
+	FString GeneratorPackageName = GetOwner().GetPackageName().ToString();
+	UPackage* LocalOwnerPackage = FindObject<UPackage>(nullptr, *GeneratorPackageName);
+	if (LocalOwnerPackage)
+	{
+		// Might be called when uninitialized, so do not call GetSplitDataObjectNameIfAvailable
+		FString Identifier;
+		if (!SplitDataObjectName.IsNone())
+		{
+			Identifier = FString::Printf(TEXT("Splitter=%s"), *SplitDataObjectName.ToString());
+		}
+		else
+		{
+			Identifier = FString::Printf(TEXT("GeneratorPackage=%s"), *GeneratorPackageName);
+		}
+		UE_LOG(LogCook, Error,
+			TEXT("PackageSplitter found the Generator package still in memory after it should have been deleted by GC.")
+			TEXT("\n\tThis is unexpected since garbage has been collected and the package should have been unreferenced so it should have been collected, and will break population of Generated packages.")
+			TEXT("\n\tSplitter=%s"), *Identifier);
+		EReferenceChainSearchMode SearchMode = EReferenceChainSearchMode::Shortest
+			| EReferenceChainSearchMode::PrintAllResults
+			| EReferenceChainSearchMode::FullChain;
+		FReferenceChainSearch RefChainSearch(LocalOwnerPackage, SearchMode);
+	}
+
+}
+
 void FGenerationHelper::UpdateSaveAfterGarbageCollect(const FPackageData& PackageData, bool& bInOutDemote)
 {
 	if (!IsInitialized())
@@ -1143,8 +1416,29 @@ FCookGenerationInfo::FCookGenerationInfo(FPackageData& InPackageData, bool bInGe
 	: PackageData(&InPackageData)
 	, GeneratorSaveState(bInGenerator ? ESaveState::StartSave : ESaveState::StartPopulate)
 	, bCreateAsMap(false), bHasCreatedPackage(false), bHasSaved(false), bTakenOverCachedCookedPlatformData(false)
-	, bIssuedUndeclaredMovedObjectsWarning(false), bGenerator(bInGenerator)
+	, bIssuedUndeclaredMovedObjectsWarning(false), bGenerator(bInGenerator), bHasCalledPopulate(false)
+{
+}
+
+void FCookGenerationInfo::Uninitialize()
 {
+	PackageHash.Reset();
+	RelativePath.Empty();
+	GeneratedRootPath.Empty();
+	GenerationHash.Reset();
+	PackageDependencies.Empty();
+	// Keep PackageData; it is allowed in the uninitialized state
+	KeepReferencedPackages.Empty();
+	check(CachedObjectsInOuterInfo.IsEmpty()); // We can not still be in the save state, so this should be empty
+	// We can not still be in the save state, so this should have been cleared
+	check(GeneratorSaveState == (bGenerator ? ESaveState::StartSave : ESaveState::StartPopulate));
+	bCreateAsMap = false;
+	bHasCreatedPackage = false;
+	bHasSaved = false;
+	bTakenOverCachedCookedPlatformData = false;
+	bIssuedUndeclaredMovedObjectsWarning = false;
+	// Keep bGenerator; it is allowed in the uninitialized state
+	bHasCalledPopulate = false;
 }
 
 void FCookGenerationInfo::SetSaveStateComplete(ESaveState CompletedState)
@@ -1235,7 +1529,6 @@ EPollStatus FCookGenerationInfo::RefreshPackageObjects(FGenerationHelper& Genera
 	GetObjectsWithOuter(Package, CurrentObjectsInOuter, true /* bIncludeNestedObjects */, RF_NoFlags,
 		EInternalObjectFlags::Garbage);
 
-	check(PackageData); // RefreshPackageObjects is only called when there is a PackageData
 	TArray<FCachedObjectInOuter>& CachedObjectsInOuter = PackageData->GetCachedObjectsInOuter();
 	UObject* FirstNewObject = nullptr;
 	for (UObject* Object : CurrentObjectsInOuter)
diff --git a/Engine/Source/Editor/UnrealEd/Private/Cooker/CookGenerationHelper.h b/Engine/Source/Editor/UnrealEd/Private/Cooker/CookGenerationHelper.h
--- a/Engine/Source/Editor/UnrealEd/Private/Cooker/CookGenerationHelper.h
+++ b/Engine/Source/Editor/UnrealEd/Private/Cooker/CookGenerationHelper.h
@@ -92,13 +92,15 @@ public:
 	bool HasCreatedPackage() const;
 	void SetHasCreatedPackage(bool bValue);
 	bool HasSaved() const;
-	void SetHasSaved(bool bValue);
+	void SetHasSaved(FGenerationHelper& GenerationHelper, bool bValue);
 	bool HasTakenOverCachedCookedPlatformData() const;
 	void SetHasTakenOverCachedCookedPlatformData(bool bValue);
 	bool HasIssuedUndeclaredMovedObjectsWarning() const;
 	void SetHasIssuedUndeclaredMovedObjectsWarning(bool bValue);
 	bool IsGenerator() const;
 	void SetIsGenerator(bool bValue);
+	bool HasCalledPopulate() const;
+	void SetHasCalledPopulate(bool bValue);
 
 	/**
 	 * Steal the list of cached objects to call BeginCacheForCookedPlatformData on from the PackageData,
@@ -135,13 +137,21 @@ public:
 	/** Return the packagename, for use in debug messages. Handles PackageData==nullptr by returning RelativePath. */
 	FString GetPackageName() const;
 
+	/**
+	 * Reset this info to a state appropriate for an uninitialized GenerationHelper; all references to UObjects
+	 * are dropped, but information necessary even when in the uninitialized GenerationHelper state is kept.
+	 */
+	void Uninitialize();
+
 public:
+	// When adding a new variable, add it to Uninitialize as well
 	FIoHash PackageHash;
 	FString RelativePath;
 	FString GeneratedRootPath;
 	FBlake3Hash GenerationHash;
 	TArray<FAssetDependency> PackageDependencies;
-	FPackageData* PackageData = nullptr;
+	/** Cannot be null, set in constructor */
+	FPackageData* PackageData;
 	TArray<TWeakObjectPtr<UPackage>> KeepReferencedPackages;
 	TMap<UObject*, FCachedObjectInOuterGeneratorInfo> CachedObjectsInOuterInfo;
 private:
@@ -152,6 +162,7 @@ private:
 	bool bTakenOverCachedCookedPlatformData : 1;
 	bool bIssuedUndeclaredMovedObjectsWarning : 1;
 	bool bGenerator : 1;
+	bool bHasCalledPopulate : 1;
 };
 
 /**
@@ -224,6 +235,11 @@ public:
 	const FName GetSplitDataObjectName() const;
 	/** Return the Splitter's value for virtual bool UseInternalReferenceToAvoidGarbageCollect(). */
 	bool IsUseInternalReferenceToAvoidGarbageCollect() const;
+	/**
+	 * Return the Splitter's value for virtual bool RequiresGeneratorPackageDestructBeforeResplit(). Returns false if not
+	 * initialized. Does not call Initialize.
+	 */
+	bool IsRequiresGeneratorPackageDestructBeforeResplit() const;
 	/** Return the Splitter's value for virtual bool DoesGeneratedRequireGenerator(). */
 	ICookPackageSplitter::EGeneratedRequiresGenerator DoesGeneratedRequireGenerator() const;
 	/** Return the cached pointer to the SplitDataObject. Returns null if no longer in memory or marked as garbage. */
@@ -259,6 +275,11 @@ public:
 	 */
 	bool TryCallPopulateGeneratorPackage(
 		TArray<ICookPackageSplitter::FGeneratedPackageForPreSave>& InOutGeneratedPackagesForPresave);
+	/**
+	 * Call the Splitter's PopulateGeneratedPackage if not yet called. Assumes GenerateList has been called. Logs
+	 * errors and returns false on failure.
+	 */
+	bool TryCallPopulateGeneratedPackage(UE::Cook::FCookGenerationInfo& Info, TArray<UObject*>& OutObjectsToMove);
 	/**
 	 * Mark that the SavePackage of the Owner is starting. Keeps a reference to keep the generator alive until save
 	 * is finished.
@@ -302,11 +323,18 @@ public:
 	/** Return the information set by SetPreviousGeneratedPackages if not yet cleared. Does not call Initialize. */
 	const TMap<FName, FIoHash>& GetPreviousGeneratedPackages() const;
 
-	/** Callback during garbage collection. Does not call initialize. */
-	void PreGarbageCollect(FCookGenerationInfo& Info, TArray<TObjectPtr<UObject>>& GCKeepObjects,
+	/**
+	 * Callback during garbage collection. Does not call initialize. Caller must pass in a refcount to show
+	 * a guarantee that clearing internal references will not delete before function return.
+	 */
+	void PreGarbageCollect(const TRefCountPtr<FGenerationHelper>& RefcountHeldByCaller,
+		FPackageData& PackageData, TArray<TObjectPtr<UObject>>& GCKeepObjects,
 		TArray<UPackage*>& GCKeepPackages, TArray<FPackageData*>& GCKeepPackageDatas, bool& bOutShouldDemote);
-	/** Callback during garbage collection. Does not call initialize. */
-	void PostGarbageCollect();
+	/**
+	 * Callback during garbage collection. Does not call initialize. Caller must pass in a refcount to show
+	 * a guarantee that clearing internal references will not delete before function return.
+	 */
+	void PostGarbageCollect(const TRefCountPtr<FGenerationHelper>& RefcountHeldByCaller);
 	/**
 	 * Called from PackageData function of the same name to decide whether to demote the package out of save.
 	 * Does not call Initialize.
@@ -322,6 +350,10 @@ public:
 	bool IsWaitingForQueueResults() const;
 	void SetKeepForGeneratorSave();
 	void ClearKeepForGeneratorSave();
+	void SetKeepForAllSavedOrGC();
+	void ClearKeepForAllSavedOrGC();
+	void SetKeepForCompletedAllSavesMessage();
+	void ClearKeepForCompletedAllSavesMessage();
 
 	/**
 	 * Helper for assignment of generated packages in MPCook. Return the id of the CookWorker that saved the
@@ -334,6 +366,23 @@ public:
 	 */
 	int32& GetMPCookNextAssignmentIndex();
 
+	/**
+	 * Called for each of the generated packages that were discovered when TryGenerateList was called on a
+	 * remote CookWorker. Does not call Initialize.
+	 */
+	void TrackGeneratedPackageListedRemotely(UCookOnTheFlyServer& COTFS, FPackageData& PackageData);
+	/**
+	 * Called on the director when the generator or one of the generated packages was saved on a remote worker.
+	 * Used to manage KeepForGCOrAllSaved lifetime. Does not call Initialize.
+	 */
+	void MarkPackageSavedRemotely(UCookOnTheFlyServer& COTFS, FPackageData& PackageData);
+
+	/**
+	 * Called on the director and every CookWorker when all saves have been completed; this indicates that
+	 * some of our contract points are complete and the splitter can be destroyed. Does not call Initialize.
+	 */
+	void OnAllSavesCompleted(UCookOnTheFlyServer& COTFS);
+
 	/** Helper function for Initialize and for TryCreateValidGenerationHelper. */
 	static void SearchForRegisteredSplitDataObject(UCookOnTheFlyServer& COTFS, FName PackageName, UPackage* Package,
 		TOptional<TConstArrayView<FCachedObjectInOuter>> CachedObjectsInOuter, UObject*& OutSplitDataObject,
@@ -360,8 +409,15 @@ private:
 private:
 	void ConditionalInitialize() const;
 	void NotifyCompletion(ICookPackageSplitter::ETeardown Status);
+	void PreGarbageCollectGCLifetimeData();
+	void PostGarbageCollectGCLifetimeData();
+	void Uninitialize();
+	void ModifyNumSaved(int32 Delta);
+	void VerifyGeneratorPackageGarbageCollected();
 
 private:
+	// When adding a new variable, add it to Uninitialize as well
+
 	/** PackageData for the package that is being split */
 	FCookGenerationInfo OwnerInfo;
 	FWeakObjectPtr SplitDataObject;
@@ -378,16 +434,23 @@ private:
 	TRefCountPtr<FGenerationHelper> ReferenceFromKeepForIterative;
 	TRefCountPtr<FGenerationHelper> ReferenceFromKeepForQueueResults;
 	TRefCountPtr<FGenerationHelper> ReferenceFromKeepForGeneratorSave;
+	TRefCountPtr<FGenerationHelper> ReferenceFromKeepForAllSavedOrGC;
 	int32 MPCookNextAssignmentIndex = 0;
+	int32 NumSaved = 0;
 	FWorkerId WorkerIdThatSavedGenerator = FWorkerId::Invalid();
 	EInitializeStatus InitializeStatus = EInitializeStatus::Uninitialized;
 	ICookPackageSplitter::EGeneratedRequiresGenerator DoesGeneratedRequireGeneratorValue =
 		ICookPackageSplitter::EGeneratedRequiresGenerator::None;
 	bool bUseInternalReferenceToAvoidGarbageCollect = false;
+	bool bRequiresGeneratorPackageDestructBeforeResplit = false;
 	bool bGeneratedList = false;
-	bool bCalledPopulateGeneratorPackage = false;
 	bool bCurrentGCHasKeptGeneratorPackage = false;
 	bool bCurrentGCHasKeptGeneratorKeepPackages = false;
+	bool bKeepForAllSavedOrGC = false;
+	bool bKeepForCompletedAllSavesMessage = false;
+	bool bNeedConfirmGeneratorPackageDestroyed = false;
+	
+	friend FCookGenerationInfo;
 };
 
 
@@ -430,9 +493,13 @@ inline bool FCookGenerationInfo::HasSaved() const
 	return bHasSaved;
 }
 
-inline void FCookGenerationInfo::SetHasSaved(bool bValue)
+inline void FCookGenerationInfo::SetHasSaved(FGenerationHelper& GenerationHelper, bool bValue)
 {
-	bHasSaved = bValue;
+	if (bValue != bHasSaved)
+	{
+		bHasSaved = bValue;
+		GenerationHelper.ModifyNumSaved(bValue ? 1 : -1);
+	}
 }
 
 inline bool FCookGenerationInfo::HasTakenOverCachedCookedPlatformData() const
@@ -465,6 +532,16 @@ inline void FCookGenerationInfo::SetIsGenerator(bool bValue)
 	bGenerator = bValue;
 }
 
+inline bool FCookGenerationInfo::HasCalledPopulate() const
+{
+	return bHasCalledPopulate;
+}
+
+inline void FCookGenerationInfo::SetHasCalledPopulate(bool bValue)
+{
+	bHasCalledPopulate = bValue;
+}
+
 inline TConstArrayView<FAssetDependency> FCookGenerationInfo::GetDependencies() const
 {
 	return PackageDependencies;
@@ -472,7 +549,7 @@ inline TConstArrayView<FAssetDependency> FCookGenerationInfo::GetDependencies()
 
 inline FString FCookGenerationInfo::GetPackageName() const
 {
-	return PackageData ? *PackageData->GetPackageName().ToString() : *RelativePath;
+	return PackageData->GetPackageName().ToString();
 }
 
 inline bool FGenerationHelper::IsInitialized() const
@@ -535,6 +612,11 @@ inline bool FGenerationHelper::IsUseInternalReferenceToAvoidGarbageCollect() con
 	return bUseInternalReferenceToAvoidGarbageCollect;
 }
 
+inline bool FGenerationHelper::IsRequiresGeneratorPackageDestructBeforeResplit() const
+{
+	return bRequiresGeneratorPackageDestructBeforeResplit;
+}
+
 inline ICookPackageSplitter::EGeneratedRequiresGenerator FGenerationHelper::DoesGeneratedRequireGenerator() const
 {
 	ConditionalInitialize();
@@ -603,6 +685,36 @@ inline void FGenerationHelper::ClearKeepForGeneratorSave()
 	ReferenceFromKeepForGeneratorSave.SafeRelease();
 }
 
+inline void FGenerationHelper::SetKeepForAllSavedOrGC()
+{
+	ReferenceFromKeepForAllSavedOrGC = this;
+	bKeepForAllSavedOrGC = true;
+}
+
+inline void FGenerationHelper::ClearKeepForAllSavedOrGC()
+{
+	bKeepForAllSavedOrGC = false;
+	if (!bKeepForAllSavedOrGC && !bKeepForCompletedAllSavesMessage)
+	{
+		ReferenceFromKeepForAllSavedOrGC.SafeRelease();
+	}
+}
+
+inline void FGenerationHelper::SetKeepForCompletedAllSavesMessage()
+{
+	ReferenceFromKeepForAllSavedOrGC = this;
+	bKeepForCompletedAllSavesMessage = true;
+}
+
+inline void FGenerationHelper::ClearKeepForCompletedAllSavesMessage()
+{
+	bKeepForCompletedAllSavesMessage = false;
+	if (!bKeepForAllSavedOrGC && !bKeepForCompletedAllSavesMessage)
+	{
+		ReferenceFromKeepForAllSavedOrGC.SafeRelease();
+	}
+}
+
 inline FWorkerId FGenerationHelper::GetWorkerIdThatSavedGenerator() const
 {
 	return WorkerIdThatSavedGenerator;
diff --git a/Engine/Source/Editor/UnrealEd/Private/Cooker/CookWorkerClient.cpp b/Engine/Source/Editor/UnrealEd/Private/Cooker/CookWorkerClient.cpp
--- a/Engine/Source/Editor/UnrealEd/Private/Cooker/CookWorkerClient.cpp
+++ b/Engine/Source/Editor/UnrealEd/Private/Cooker/CookWorkerClient.cpp
@@ -340,6 +340,9 @@ void FCookWorkerClient::HandleGeneratorMessage(FGeneratorEventMessage&& Generato
 			case EGeneratorEvent::QueuedGeneratedPackagesFencePassed:
 				GenerationHelper->OnRequestFencePassed(COTFS);
 				break;
+			case EGeneratorEvent::AllSavesCompleted:
+				GenerationHelper->OnAllSavesCompleted(COTFS);
+				break;
 			default:
 				// We do not handle the remaining GeneratorEvents on clients
 				break;
diff --git a/Engine/Source/Editor/UnrealEd/Private/Cooker/CookWorkerServer.cpp b/Engine/Source/Editor/UnrealEd/Private/Cooker/CookWorkerServer.cpp
--- a/Engine/Source/Editor/UnrealEd/Private/Cooker/CookWorkerServer.cpp
+++ b/Engine/Source/Editor/UnrealEd/Private/Cooker/CookWorkerServer.cpp
@@ -921,6 +921,28 @@ void FCookWorkerServer::RecordResults(FPackageResultsMessage& Message)
 		bRetiredAnyPackages = true;
 		PackageData->SetWorkerAssignment(FWorkerId::Invalid(), ESendFlags::QueueNone);
 
+		if (PackageData->IsGenerated())
+		{
+			TRefCountPtr<FGenerationHelper> ParentGenerationHelper = PackageData->GetOrFindParentGenerationHelper();
+			if (!ParentGenerationHelper)
+			{
+				UE_LOG(LogCook, Warning,
+					TEXT("RecordResults received for generated package %s, but its ParentGenerationHelper has already been destructed so we can not update the save flag. Leaving the save flag unupdated; this might cause workers to run out of memory due to keeping the Generator referenced."),
+					*PackageData->GetPackageName().ToString());
+			}
+			else
+			{
+				ParentGenerationHelper->MarkPackageSavedRemotely(COTFS, *PackageData);
+				PackageData->SetParentGenerationHelper(nullptr);
+			}
+		}
+		TRefCountPtr<FGenerationHelper> GenerationHelper = PackageData->GetGenerationHelper();
+		if (GenerationHelper)
+		{
+			GenerationHelper->MarkPackageSavedRemotely(COTFS, *PackageData);
+			GenerationHelper.SafeRelease();
+		}
+
 		// MPCOOKTODO: Refactor FSaveCookedPackageContext::FinishPlatform and ::FinishPackage so we can call them from
 		// here to reduce duplication
 		if (Result.GetSuppressCookReason() == ESuppressCookReason::NotSuppressed)
@@ -1006,6 +1028,7 @@ void FCookWorkerServer::QueueDiscoveredPackage(FDiscoveredPackageReplication&& D
 			TRefCountPtr<FGenerationHelper> GenerationHelper =
 				GeneratorPackageData->CreateUninitializedGenerationHelper();
 			GenerationHelper->NotifyStartQueueGeneratedPackages(COTFS, WorkerId);
+			GenerationHelper->TrackGeneratedPackageListedRemotely(COTFS, PackageData);
 		}
 	}
 
@@ -1059,28 +1082,28 @@ void FCookWorkerServer::QueueDiscoveredPackage(FDiscoveredPackageReplication&& D
 void FCookWorkerServer::HandleGeneratorMessage(FGeneratorEventMessage& GeneratorMessage)
 {
 	FPackageData* PackageData = COTFS.PackageDatas->FindPackageDataByPackageName(GeneratorMessage.PackageName);
-	TRefCountPtr<FGenerationHelper> GenerationHelper;
-	if (PackageData)
-	{
-		GenerationHelper = PackageData->CreateUninitializedGenerationHelper();
-	}
-	if (!GenerationHelper)
+	if (!PackageData)
 	{
 		// This error should be impossible because GeneratorMessages are only sent in response to assignment from the server.
 		UE_LOG(LogCook, Error,
-			TEXT("CookWorkerServer received unexpected GeneratorMessage for package %s. The PackageData %s on the CookDirector. ")
+			TEXT("CookWorkerServer received unexpected GeneratorMessage for package %s. The PackageData %s does not exist on the CookDirector. ")
 			TEXT("\n\tCook of this generator package and its generated packages will be invalid."),
 			*GeneratorMessage.PackageName.ToString(),
 			(!PackageData ? TEXT("does not exist") : TEXT("is not a valid generator")));
 		return;
 	}
+
+	TRefCountPtr<FGenerationHelper> GenerationHelper;
+	GenerationHelper = PackageData->CreateUninitializedGenerationHelper();
+	check(GenerationHelper);
+
 	switch (GeneratorMessage.Event)
 	{
 	case EGeneratorEvent::QueuedGeneratedPackages:
 		GenerationHelper->EndQueueGeneratedPackagesOnDirector(COTFS, GetWorkerId());
 		break;
 	default:
-		// We do not handle the remaining GeneratorEvents on clients
+		// We do not handle the remaining GeneratorEvents on the server
 		break;
 	}
 }
diff --git a/Engine/Source/Editor/UnrealEd/Private/Cooker/CookWorkerServer.h b/Engine/Source/Editor/UnrealEd/Private/Cooker/CookWorkerServer.h
--- a/Engine/Source/Editor/UnrealEd/Private/Cooker/CookWorkerServer.h
+++ b/Engine/Source/Editor/UnrealEd/Private/Cooker/CookWorkerServer.h
@@ -449,6 +449,7 @@ enum class EGeneratorEvent
 	Invalid,
 	QueuedGeneratedPackages,
 	QueuedGeneratedPackagesFencePassed,
+	AllSavesCompleted,
 	Num,
 };
 
diff --git a/Engine/Source/Editor/UnrealEd/Public/CookPackageSplitter.h b/Engine/Source/Editor/UnrealEd/Public/CookPackageSplitter.h
--- a/Engine/Source/Editor/UnrealEd/Public/CookPackageSplitter.h
+++ b/Engine/Source/Editor/UnrealEd/Public/CookPackageSplitter.h
@@ -63,6 +63,21 @@ public:
 		return false;
 	}
 
+	/**
+	 * An ICookPackageSplitter for a single generator package normally is constructed only once and handles
+	 * all generated packages for that generator, but during MPCook in cases of load balancing between CookWorkers,
+	 * it is possible that the original splitter is destructed but then recreated later. This is guaranteed not
+	 * to happen without a GarbageCollection pass in between, but that GarbageCollection may fail to destruct the
+	 * generator package if it is still referenced from other packages or systems. Depending on the ICookPackageSplitter's
+	 * implemenation, this failure to GC might cause an error, because changes made from the previous splitter are not 
+	 * handled in the next splitter. If RequiresGeneratorPackageDestructBeforeResplit is true, the cooker will log this failure
+	 * to GC the generator package as an error.
+	 */
+	virtual bool RequiresGeneratorPackageDestructBeforeResplit()
+	{
+		return false;
+	}
+
 	/**
 	 * Return value for the DoesGeneratedRequireGenerator function. All levels behave correctly, but provide
 	 * different tradeoffs of guarantees to the splitter versus performance.
@@ -154,9 +169,12 @@ public:
 		bool bCreatedAsMap = false;
 	};
 	/**
-	 * Called before presaving the parent generator package, to give the generator a chance to inform the cooker which objects will
-	 * be moved into the generator package that are not already present in it.
+	 * Called before presaving the parent generator package, to give the generator a chance to inform the cooker which
+	 * objects will be moved into the generator package that are not already present in it.
 	 * 
+	 * PopulateGeneratorPackage is guaranteed to not be called again until the splitter has been destroyed and the
+	 * generator package has been garbage collected.
+	 *
 	 * @param OwnerPackage				The generator package being split
 	 * @param OwnerObject				The SplitDataClass instance that this CookPackageSplitter instance was created for
 	 * @param GeneratedPackages			Placeholder UPackage and relative path information for all packages that will be generated
@@ -225,6 +243,9 @@ public:
 	 * Return a list of all the objects that will be moved into the Generated package during its save, so the cooker
 	 * can call BeginCacheForCookedPlatformData on them before the move
 	 * After returning, the given package will be queued for saving into the TargetDomain
+	 * 
+	 * PopulateGeneratedPackage is guaranteed to not be called again on the same generated package until the splitter
+	 * has been destroyed and the generator package has been garbage collected.
 	 *
 	 * @param OwnerPackage				The parent package being split
 	 * @param OwnerObject				The SplitDataClass instance that this CookPackageSplitter instance was created for
diff --git a/Engine/Source/Runtime/Engine/Private/WorldPartition/WorldPartitionCookPackageSplitter.h b/Engine/Source/Runtime/Engine/Private/WorldPartition/WorldPartitionCookPackageSplitter.h
--- a/Engine/Source/Runtime/Engine/Private/WorldPartition/WorldPartitionCookPackageSplitter.h
+++ b/Engine/Source/Runtime/Engine/Private/WorldPartition/WorldPartitionCookPackageSplitter.h
@@ -28,6 +28,10 @@ public:
 	{
 		return EGeneratedRequiresGenerator::Populate;
 	}
+	virtual bool RequiresGeneratorPackageDestructBeforeResplit() override
+	{
+		return true;
+	}
 
 	virtual TArray<ICookPackageSplitter::FGeneratedPackage> GetGenerateList(const UPackage* OwnerPackage, const UObject* OwnerObject) override;
 	virtual bool PopulateGeneratedPackage(UPackage* OwnerPackage, UObject* OwnerObject,
