Add uncompressed size to asset size query commandlet

diff --git a/Engine/Source/Editor/UnrealEd/Private/Commandlets/AssetSizeQueryCommandlet.cpp b/Engine/Source/Editor/UnrealEd/Private/Commandlets/AssetSizeQueryCommandlet.cpp
--- a/Engine/Source/Editor/UnrealEd/Private/Commandlets/AssetSizeQueryCommandlet.cpp
+++ b/Engine/Source/Editor/UnrealEd/Private/Commandlets/AssetSizeQueryCommandlet.cpp
@@ -142,6 +142,8 @@ int32 UAssetSizeQueryCommandlet::Main(const FString& FullCommandLine)
 		int64 OptionalSize = 0;
 		int64 InstalledSize = 0;
 		int64 StreamingSize = 0;
+		int64 UncompressedSize = 0;
+
 	};
 	TMap<FTopLevelAssetPath /* AssetClass */, int64> FilteredClassCompressedSizes;
 	TMap<FTopLevelAssetPath /* AssetClass */, TArray<FMatchedAssetInfo>> FilteredClassMatchedAssets;
@@ -189,6 +191,7 @@ int32 UAssetSizeQueryCommandlet::Main(const FString& FullCommandLine)
 			AssetData.GetTagValue(UE::AssetRegistry::Stage_ChunkInstalledSizeFName, Info.InstalledSize);
 			AssetData.GetTagValue(UE::AssetRegistry::Stage_ChunkOptionalSizeFName, Info.OptionalSize);
 			AssetData.GetTagValue(UE::AssetRegistry::Stage_ChunkStreamingSizeFName, Info.StreamingSize);
+			AssetData.GetTagValue(UE::AssetRegistry::Stage_ChunkSizeFName, Info.UncompressedSize);
 			
 			FilteredCompressedSize += AssetCompressedSize;
 			int64& FilteredClassCompressedSize = FilteredClassCompressedSizes.FindOrAdd(AssetData.AssetClassPath);
@@ -249,13 +252,13 @@ int32 UAssetSizeQueryCommandlet::Main(const FString& FullCommandLine)
 		}
 		else if (OutputCSVType == EOutputCSVType::Assets)
 		{
-			Lines.Add(TEXT("AssetName,AssetType,CompressedSize,InstalledSize,OptionalSize,StreamingSize"));
+			Lines.Add(TEXT("AssetName,AssetType,CompressedSize,InstalledSize,OptionalSize,StreamingSize,UncompressedSize"));
 			for (TPair<FTopLevelAssetPath, TArray<FMatchedAssetInfo>>& ClassAssetsPair : FilteredClassMatchedAssets)
 			{
 				// we add to both maps at the same time to we know the lookup succeeds.
 				for (const FMatchedAssetInfo& AssetInfo : ClassAssetsPair.Value)
 				{
-					Lines.Add(FString::Printf(TEXT("%s,%s,%lld,%lld,%lld,%lld"), *AssetInfo.ObjectPath.ToString(), *ClassAssetsPair.Key.ToString(), AssetInfo.CompressedSize, AssetInfo.InstalledSize, AssetInfo.OptionalSize, AssetInfo.StreamingSize));
+					Lines.Add(FString::Printf(TEXT("%s,%s,%lld,%lld,%lld,%lld,%lld"), *AssetInfo.ObjectPath.ToString(), *ClassAssetsPair.Key.ToString(), AssetInfo.CompressedSize, AssetInfo.InstalledSize, AssetInfo.OptionalSize, AssetInfo.StreamingSize, AssetInfo.UncompressedSize));
 				}
 			}
 		}
