Keep materials alive during an ODSC session so that the permutations processed by CompileODSCMaterialsForRemoteRecompile can be aggregated between requests from the client.
This will enable the client to request a subset of all the shader permutations for a material and receive all the permutations encountered since the start of the session

#rnx

diff --git a/Engine/Source/Editor/UnrealEd/Classes/CookOnTheSide/CookOnTheFlyServer.h b/Engine/Source/Editor/UnrealEd/Classes/CookOnTheSide/CookOnTheFlyServer.h
--- a/Engine/Source/Editor/UnrealEd/Classes/CookOnTheSide/CookOnTheFlyServer.h
+++ b/Engine/Source/Editor/UnrealEd/Classes/CookOnTheSide/CookOnTheFlyServer.h
@@ -175,6 +175,7 @@ namespace UE::Cook
 	class FCookWorkerClient;
 	class FCookWorkerServer;
 	class FDiagnostics;
+	class FODSCClientData; 
 	class FPackageWriterMPCollector;
 	class FRequestCluster;
 	class FRequestQueue;
@@ -1607,6 +1608,7 @@ private:
 	float WaitForAsyncSleepSeconds = 0.0f;
 	float DisplayUpdatePeriodSeconds = 0.0f;
 	EIdleStatus IdleStatus = EIdleStatus::Done;
+	TUniquePtr<UE::Cook::FODSCClientData> ODSCClientData;
 
 	friend FAssetRegistryGenerator;
 	friend UE::Cook::FAssetRegistryMPCollector;
diff --git a/Engine/Source/Editor/UnrealEd/Private/CookOnTheFlyServer.cpp b/Engine/Source/Editor/UnrealEd/Private/CookOnTheFlyServer.cpp
--- a/Engine/Source/Editor/UnrealEd/Private/CookOnTheFlyServer.cpp
+++ b/Engine/Source/Editor/UnrealEd/Private/CookOnTheFlyServer.cpp
@@ -39,6 +39,7 @@
 #include "Cooker/LooseCookedPackageWriter.h"
 #include "Cooker/MPCollector.h"
 #include "Cooker/NetworkFileCookOnTheFlyRequestManager.h"
+#include "Cooker/OnDemandShaderCompilation.h"
 #include "Cooker/PackageTracker.h"
 #include "Cooker/WorkerRequestsLocal.h"
 #include "Cooker/WorkerRequestsRemote.h"
@@ -639,6 +640,18 @@ bool UCookOnTheFlyServer::StartCookOnTheFly(FCookOnTheFlyStartupOptions InCookOn
 				bool bAlreadyInitialized = false;
 				CookOnTheFlyServerInterface->AddPlatform(Connection.GetPlatformName(), bAlreadyInitialized);
 			}
+			if (ODSCClientData)
+			{
+				ODSCClientData->OnClientConnected(&Connection);
+			}
+		});
+
+	CookOnTheFlyNetworkServer->OnClientDisconnected().AddLambda([this](ICookOnTheFlyClientConnection& Connection)
+		{
+			if (ODSCClientData)
+			{
+				ODSCClientData->OnClientDisconnected(&Connection);
+			}
 		});
 
 	CookOnTheFlyNetworkServer->OnRequest(ECookOnTheFlyMessage::RecompileShaders).BindLambda([this](ICookOnTheFlyClientConnection& Connection, const FCookOnTheFlyRequest& Request)
@@ -655,6 +668,7 @@ bool UCookOnTheFlyServer::StartCookOnTheFly(FCookOnTheFlyStartupOptions InCookOn
 				TArray<FString> RecompileModifiedFiles;
 				TArray<uint8> MeshMaterialMaps;
 				TArray<uint8> GlobalShaderMap;
+				TArray<TStrongObjectPtr<UMaterialInterface>> LoadedMaterialsToRecompile;
 
 				FShaderRecompileData RecompileData(Connection.GetTargetPlatform()->PlatformName(), &RecompileModifiedFiles, &MeshMaterialMaps, &GlobalShaderMap);
 				{
@@ -662,9 +676,18 @@ bool UCookOnTheFlyServer::StartCookOnTheFly(FCookOnTheFlyStartupOptions InCookOn
 					*Ar << RecompileData;
 				}
 
+				RecompileData.LoadedMaterialsToRecompile = &LoadedMaterialsToRecompile;
+
+				const void* ConnectionPtr = &Connection;
 				FEventRef RecompileCompletedEvent;
-				UE::Cook::FRecompileShaderCompletedCallback RecompileCompleted = [this, &RecompileCompletedEvent]()
+				UE::Cook::FRecompileShaderCompletedCallback RecompileCompleted = [this, &RecompileCompletedEvent, ConnectionPtr, &LoadedMaterialsToRecompile]()
 				{
+					if (ODSCClientData)
+					{
+						ODSCClientData->KeepClientPersistentData(ConnectionPtr, LoadedMaterialsToRecompile);
+					}
+					LoadedMaterialsToRecompile.Empty();
+
 					RecompileCompletedEvent->Trigger();
 				};
 
@@ -7174,7 +7197,13 @@ void UCookOnTheFlyServer::SetInitializeConfigSettings(UE::Cook::FInitializeConfi
 	// It would be better to always calculate it, but we want to avoid the performance cost until it becomes more widely used
 	bIterativeCalculateExe = !bIterativeIgnoreExe || !bConfigSettingSetIterativeIgnoreExe;
 
-	bRunningAsShaderServer = FParse::Param(FCommandLine::Get(), TEXT("odsc"));;
+	bRunningAsShaderServer = FParse::Param(FCommandLine::Get(), TEXT("odsc"));
+    ODSCClientData = nullptr;
+	if (bRunningAsShaderServer)
+	{
+		ODSCClientData = MakeUnique<UE::Cook::FODSCClientData>();
+	}
+
 	bSkipSave = FParse::Param(FCommandLine::Get(), TEXT("CookSkipSave"));
 
 	FString Severity;
@@ -10418,6 +10447,8 @@ void UCookOnTheFlyServer::CookByTheBookFinishedInternal()
 
 void UCookOnTheFlyServer::ShutdownCookSession()
 {
+	ODSCClientData.Reset();
+
 	if (CookDirector)
 	{
 		CookDirector->ShutdownCookSession();
diff --git a/Engine/Source/Editor/UnrealEd/Private/Cooker/OnDemandShaderCompilation.cpp b/Engine/Source/Editor/UnrealEd/Private/Cooker/OnDemandShaderCompilation.cpp
new file mode 100644
--- /dev/null
+++ b/Engine/Source/Editor/UnrealEd/Private/Cooker/OnDemandShaderCompilation.cpp
@@ -0,0 +1,50 @@
+// Copyright Epic Games, Inc. All Rights Reserved.
+
+#include "Cooker/OnDemandShaderCompilation.h"
+
+#include "Async/TaskGraphInterfaces.h"
+#include "Async/Async.h"
+
+
+namespace UE::Cook
+{
+
+void FODSCClientData::OnClientConnected(const void* ConnectionPtr)
+{
+	FScopeLock PollablesScopeLock(&ODSCClientPersistentDataMapLock);
+	ODSCClientPersistentDataMap.FindOrAdd(ConnectionPtr);
+}
+
+void FODSCClientData::OnClientDisconnected(const void* ConnectionPtr)
+{
+	FScopeLock PollablesScopeLock(&ODSCClientPersistentDataMapLock);
+
+	FODSCClientPersistentData* ODSCClientPersistentData = ODSCClientPersistentDataMap.Find(ConnectionPtr);
+	if (ODSCClientPersistentData)
+	{
+		// We can delete TStrongObjectPtr only on the game thread
+		TSet<TStrongObjectPtr<UMaterialInterface>>* MaterialsToRemove = new TSet<TStrongObjectPtr<UMaterialInterface>>;
+		*MaterialsToRemove = MoveTemp(ODSCClientPersistentData->MaterialsKeptAlive);
+		AsyncTask(ENamedThreads::GameThread, [MaterialsToRemove]()
+				  {
+					  delete MaterialsToRemove;
+				  });
+
+		ODSCClientPersistentDataMap.Remove(ConnectionPtr);
+	}
+}
+
+void FODSCClientData::KeepClientPersistentData(const void* ConnectionPtr, const TArray<TStrongObjectPtr<UMaterialInterface>>& LoadedMaterialsToRecompile)
+{
+	FScopeLock PollablesScopeLock(&ODSCClientPersistentDataMapLock);
+	FODSCClientPersistentData* ODSCClientPersistentData = ODSCClientPersistentDataMap.Find(ConnectionPtr);
+	if (ODSCClientPersistentData)
+	{
+		for (const TStrongObjectPtr<UMaterialInterface>& MaterialInterface : LoadedMaterialsToRecompile)
+		{
+			ODSCClientPersistentData->MaterialsKeptAlive.Add(MaterialInterface);
+		}
+	}
+}
+
+}
\ No newline at end of file
diff --git a/Engine/Source/Editor/UnrealEd/Private/Cooker/OnDemandShaderCompilation.h b/Engine/Source/Editor/UnrealEd/Private/Cooker/OnDemandShaderCompilation.h
new file mode 100644
--- /dev/null
+++ b/Engine/Source/Editor/UnrealEd/Private/Cooker/OnDemandShaderCompilation.h
@@ -0,0 +1,30 @@
+// Copyright Epic Games, Inc. All Rights Reserved.
+
+#pragma once
+
+#include "UObject/StrongObjectPtr.h"
+#include "Materials/MaterialInterface.h"
+
+class UMaterialInterface;
+
+namespace UE::Cook
+{
+
+class FODSCClientData
+{
+public:
+	void OnClientConnected(const void* ConnectionPtr);
+	void OnClientDisconnected(const void* ConnectionPtr); 
+	void KeepClientPersistentData(const void* ConnectionPtr, const TArray<TStrongObjectPtr<UMaterialInterface>>& LoadedMaterialsToRecompile);
+
+private:
+	struct FODSCClientPersistentData
+	{
+		TSet<TStrongObjectPtr<UMaterialInterface> > MaterialsKeptAlive;
+	};
+
+	TMap<const void*, FODSCClientPersistentData> ODSCClientPersistentDataMap;
+	FCriticalSection ODSCClientPersistentDataMapLock;
+};
+
+}
\ No newline at end of file
diff --git a/Engine/Source/Runtime/Engine/Private/ShaderCompiler/ShaderCompiler.cpp b/Engine/Source/Runtime/Engine/Private/ShaderCompiler/ShaderCompiler.cpp
--- a/Engine/Source/Runtime/Engine/Private/ShaderCompiler/ShaderCompiler.cpp
+++ b/Engine/Source/Runtime/Engine/Private/ShaderCompiler/ShaderCompiler.cpp
@@ -64,6 +64,7 @@
 
 #if WITH_EDITOR
 #include "UObject/ArchiveCookContext.h"
+#include "UObject/UObjectGlobals.h"
 #include "DerivedDataCache.h"
 #include "DerivedDataRequestOwner.h"
 #include "TextureCompiler.h"
@@ -10251,7 +10252,23 @@ void RecompileShadersForRemote(
 	if (Args.ShadersToRecompile.Num() && (Args.MeshMaterialMaps != nullptr))
 	{
 		TMap<FString, TArray<TRefCountPtr<FMaterialShaderMap>>> CompiledShaderMaps;
-		UMaterial::CompileODSCMaterialsForRemoteRecompile(Args.ShadersToRecompile, CompiledShaderMaps);
+		{
+			// UMaterial::CompileODSCMaterialsForRemoteRecompile will call LoadObjects on the material names but doesn't keep them around. Add a GC guard to ensure we can still get them
+			// before they get unloaded, so that the whole chain UMaterial->FMaterial->FMaterialShaderMap is kept intact, and we can merge the next batch of ODSC requests
+			FGCScopeGuard NoGCScopeGuard;
+			UMaterial::CompileODSCMaterialsForRemoteRecompile(Args.ShadersToRecompile, CompiledShaderMaps);
+			if (Args.LoadedMaterialsToRecompile)
+			{
+				for (auto Iter : CompiledShaderMaps)
+				{
+					TStrongObjectPtr<UMaterialInterface> MaterialInterface = TStrongObjectPtr<UMaterialInterface>(FindObject<UMaterialInterface>(nullptr, *Iter.Key));
+					if (MaterialInterface)
+					{
+						Args.LoadedMaterialsToRecompile->Add(MaterialInterface);
+					}
+				}
+			}
+		}
 		SaveShaderMapsForRemote(TargetPlatform, CompiledShaderMaps, Args.MeshMaterialMaps);
 	}
 	else
diff --git a/Engine/Source/Runtime/Engine/Public/ShaderCompiler.h b/Engine/Source/Runtime/Engine/Public/ShaderCompiler.h
--- a/Engine/Source/Runtime/Engine/Public/ShaderCompiler.h
+++ b/Engine/Source/Runtime/Engine/Public/ShaderCompiler.h
@@ -27,6 +27,7 @@
 #include "Containers/Deque.h"
 #include "Hash/Blake3.h"
 #include "SceneTypes.h"
+#include "UObject/StrongObjectPtr.h"
 
 class FAsyncCompilationNotification;
 class FCbObjectView;
@@ -38,6 +39,7 @@ class FShaderCompileJob;
 class FShaderCompilerStats;
 class FShaderKeyGenerator;
 class FShaderPipelineCompileJob;
+class UMaterialInterface;
 struct FAnalyticsEventAttribute;
 
 DECLARE_LOG_CATEGORY_EXTERN(LogShaderCompilers, Log, All);
@@ -1246,6 +1248,9 @@ struct FShaderRecompileData
 	/** On-demand shader compiler payload.  */
 	TArray<FODSCRequestPayload> ShadersToRecompile;
 
+	/** Optional Array of the loaded materials  */
+	TArray<TStrongObjectPtr<UMaterialInterface>>* LoadedMaterialsToRecompile = nullptr;
+
 	/** Default constructor. */
 	FShaderRecompileData() {};
 
