Added DumpMaterialInfo commandlet, which writes a CSV with many material properties to disk.
Also added LWCUsagesCS to the material info panels.

diff --git a/Engine/Source/Editor/MaterialEditor/Private/MaterialEditor.cpp b/Engine/Source/Editor/MaterialEditor/Private/MaterialEditor.cpp
--- a/Engine/Source/Editor/MaterialEditor/Private/MaterialEditor.cpp
+++ b/Engine/Source/Editor/MaterialEditor/Private/MaterialEditor.cpp
@@ -2316,18 +2316,24 @@ void FMaterialEditor::DrawMaterialInfoStrings(
 			DrawPositionY += SpacingBetweenLines;
 		}
 
-		TStaticArray<uint16, (int)ELWCFunctionKind::Max> LWCFuncUsagesVS;
-		TStaticArray<uint16, (int)ELWCFunctionKind::Max> LWCFuncUsagesPS;
-		MaterialResource->GetEstimatedLWCFuncUsages(LWCFuncUsagesVS, LWCFuncUsagesPS);
+		FMaterialResource::FLWCUsagesArray LWCFuncUsagesVS;
+		FMaterialResource::FLWCUsagesArray LWCFuncUsagesPS;
+		FMaterialResource::FLWCUsagesArray LWCFuncUsagesCS;
+		MaterialResource->GetEstimatedLWCFuncUsages(LWCFuncUsagesVS, LWCFuncUsagesPS, LWCFuncUsagesCS);
 		for (int KindIndex = 0; KindIndex < (int)ELWCFunctionKind::Max; ++KindIndex)
 		{
-			int Usages = LWCFuncUsagesVS[KindIndex] + LWCFuncUsagesPS[KindIndex];
+			int Usages = LWCFuncUsagesVS[KindIndex] + LWCFuncUsagesPS[KindIndex] + LWCFuncUsagesCS[KindIndex];
 			if (Usages > 0)
 			{
 				Canvas->DrawShadowedString(
 					5,
 					DrawPositionY,
-					*FString::Printf(TEXT("LWC %s usages (Est.): %u"), *UEnum::GetDisplayValueAsText((ELWCFunctionKind)KindIndex).ToString(), Usages),
+					*FString::Printf(
+						TEXT("LWC %s usages (Est.): %u (VS), %u (PS), %u (CS)"), 
+						*UEnum::GetDisplayValueAsText((ELWCFunctionKind)KindIndex).ToString(), 
+						LWCFuncUsagesVS[KindIndex],
+						LWCFuncUsagesPS[KindIndex],
+						LWCFuncUsagesCS[KindIndex]),
 					FontToUse,
 					FLinearColor(1,1,0)
 				);
@@ -2935,247 +2941,253 @@ void FMaterialEditor::UpdateMaterialinfoList_Old()
 	// Always show basic features so that errors aren't hidden
 	FeatureLevelsToDisplay[NumFeatureLevels++] = GMaxRHIFeatureLevel;
 
-		UMaterial* MaterialForStats = bStatsFromPreviewMaterial ? Material : OriginalMaterial;
+	UMaterial* MaterialForStats = bStatsFromPreviewMaterial ? Material : OriginalMaterial;
 
-		for (int32 i = 0; i < NumFeatureLevels; ++i)
-		{
-			TArray<FString>				 CompileErrors;
-			TArray<UMaterialExpression*> FailingExpression;
+	for (int32 i = 0; i < NumFeatureLevels; ++i)
+	{
+		TArray<FString>				 CompileErrors;
+		TArray<UMaterialExpression*> FailingExpression;
 
-			ERHIFeatureLevel::Type FeatureLevel = FeatureLevelsToDisplay[i];
-			const FMaterialResource* MaterialResource = MaterialForStats->GetMaterialResource(FeatureLevel);
+		ERHIFeatureLevel::Type FeatureLevel = FeatureLevelsToDisplay[i];
+		const FMaterialResource* MaterialResource = MaterialForStats->GetMaterialResource(FeatureLevel);
 
-			if (MaterialResource == nullptr)
-			{
-				continue;
-			}
+		if (MaterialResource == nullptr)
+		{
+			continue;
+		}
 
-			if (MaterialFunction && ExpressionPreviewMaterial)
-			{
-				bool bHasValidOutput = true;
-				int32 NumInputs = 0;
-				int32 NumOutputs = 0;
-				// For Material Layers
+		if (MaterialFunction && ExpressionPreviewMaterial)
+		{
+			bool bHasValidOutput = true;
+			int32 NumInputs = 0;
+			int32 NumOutputs = 0;
+			// For Material Layers
 
-				if (MaterialFunction->GetMaterialFunctionUsage() == EMaterialFunctionUsage::MaterialLayer)
+			if (MaterialFunction->GetMaterialFunctionUsage() == EMaterialFunctionUsage::MaterialLayer)
+			{
+				// Material layers must have a single MA input and output only
+				for (UMaterialExpression* Expression : MaterialFunction->GetExpressions())
 				{
-					// Material layers must have a single MA input and output only
-					for (UMaterialExpression* Expression : MaterialFunction->GetExpressions())
+					if (UMaterialExpressionFunctionInput* InputExpression = Cast<UMaterialExpressionFunctionInput>(Expression))
 					{
-						if (UMaterialExpressionFunctionInput* InputExpression = Cast<UMaterialExpressionFunctionInput>(Expression))
-						{
-							++NumInputs;
-							if (NumInputs > 1 || !InputExpression->IsResultMaterialAttributes(0))
-							{
-								CompileErrors.Add(TEXT("Layer graphs only support a single material attributes input."));
-								FailingExpression.Add(nullptr);
-							}
-						}
-						else if (UMaterialExpressionFunctionOutput* OutputExpression = Cast<UMaterialExpressionFunctionOutput>(Expression))
+						++NumInputs;
+						if (NumInputs > 1 || !InputExpression->IsResultMaterialAttributes(0))
 						{
-							++NumOutputs;
-							if (NumOutputs > 1 || !OutputExpression->IsResultMaterialAttributes(0))
-							{
-								CompileErrors.Add(TEXT("Layer graphs only support a single material attributes output."));
-								FailingExpression.Add(nullptr);
-							}
+							CompileErrors.Add(TEXT("Layer graphs only support a single material attributes input."));
+							FailingExpression.Add(nullptr);
 						}
-						else if (UMaterialExpressionMaterialAttributeLayers* RecursiveLayer = Cast<UMaterialExpressionMaterialAttributeLayers>(Expression))
+					}
+					else if (UMaterialExpressionFunctionOutput* OutputExpression = Cast<UMaterialExpressionFunctionOutput>(Expression))
+					{
+						++NumOutputs;
+						if (NumOutputs > 1 || !OutputExpression->IsResultMaterialAttributes(0))
 						{
-							CompileErrors.Add(TEXT("Layer graphs do not support layers within layers."));
+							CompileErrors.Add(TEXT("Layer graphs only support a single material attributes output."));
 							FailingExpression.Add(nullptr);
 						}
 					}
-
-					if (NumInputs > 1 || NumOutputs < 1)
+					else if (UMaterialExpressionMaterialAttributeLayers* RecursiveLayer = Cast<UMaterialExpressionMaterialAttributeLayers>(Expression))
 					{
-						CompileErrors.Add(TEXT("Layer graphs require a single material attributes output and optionally, a single material attributes input."));
+						CompileErrors.Add(TEXT("Layer graphs do not support layers within layers."));
 						FailingExpression.Add(nullptr);
 					}
 				}
-				else if (MaterialFunction->GetMaterialFunctionUsage() == EMaterialFunctionUsage::MaterialLayerBlend)
+
+				if (NumInputs > 1 || NumOutputs < 1)
 				{
-					// Material layer blends can have two MA inputs and single MA output only
-					for (UMaterialExpression* Expression : MaterialFunction->GetExpressions())
+					CompileErrors.Add(TEXT("Layer graphs require a single material attributes output and optionally, a single material attributes input."));
+					FailingExpression.Add(nullptr);
+				}
+			}
+			else if (MaterialFunction->GetMaterialFunctionUsage() == EMaterialFunctionUsage::MaterialLayerBlend)
+			{
+				// Material layer blends can have two MA inputs and single MA output only
+				for (UMaterialExpression* Expression : MaterialFunction->GetExpressions())
+				{
+					if (UMaterialExpressionFunctionInput* InputExpression = Cast<UMaterialExpressionFunctionInput>(Expression))
 					{
-						if (UMaterialExpressionFunctionInput* InputExpression = Cast<UMaterialExpressionFunctionInput>(Expression))
-						{
-							++NumInputs;
-							if (NumInputs > 2 || !InputExpression->IsResultMaterialAttributes(0))
-							{
-								CompileErrors.Add(TEXT("Layer blend graphs only support two material attributes inputs."));
-								FailingExpression.Add(nullptr);
-							}
-						}
-						else if (UMaterialExpressionFunctionOutput* OutputExpression = Cast<UMaterialExpressionFunctionOutput>(Expression))
+						++NumInputs;
+						if (NumInputs > 2 || !InputExpression->IsResultMaterialAttributes(0))
 						{
-							++NumOutputs;
-							if (NumOutputs > 1 || !OutputExpression->IsResultMaterialAttributes(0))
-							{
-								CompileErrors.Add(TEXT("Layer blend graphs only support a single material attributes output."));
-								FailingExpression.Add(nullptr);
-							}
+							CompileErrors.Add(TEXT("Layer blend graphs only support two material attributes inputs."));
+							FailingExpression.Add(nullptr);
 						}
-						else if (UMaterialExpressionMaterialAttributeLayers* RecursiveLayer = Cast<UMaterialExpressionMaterialAttributeLayers>(Expression))
+					}
+					else if (UMaterialExpressionFunctionOutput* OutputExpression = Cast<UMaterialExpressionFunctionOutput>(Expression))
+					{
+						++NumOutputs;
+						if (NumOutputs > 1 || !OutputExpression->IsResultMaterialAttributes(0))
 						{
-							CompileErrors.Add(TEXT("Layer blend graphs do not support layers within layers."));
+							CompileErrors.Add(TEXT("Layer blend graphs only support a single material attributes output."));
 							FailingExpression.Add(nullptr);
 						}
 					}
-
-					if (NumOutputs < 1)
+					else if (UMaterialExpressionMaterialAttributeLayers* RecursiveLayer = Cast<UMaterialExpressionMaterialAttributeLayers>(Expression))
 					{
-						CompileErrors.Add(TEXT("Layer blend graphs can have up to two material attributes inputs and a single output."));
+						CompileErrors.Add(TEXT("Layer blend graphs do not support layers within layers."));
 						FailingExpression.Add(nullptr);
 					}
 				}
-				else
+
+				if (NumOutputs < 1)
 				{
-					// Add a compile error message for functions missing an output
-					FMaterialResource* CurrentResource = ExpressionPreviewMaterial->GetMaterialResource(FeatureLevel);
-					if (CurrentResource)
-					{
-						CompileErrors = CurrentResource->GetCompileErrors();
-						FailingExpression = CurrentResource->GetErrorExpressions();
-					}
+					CompileErrors.Add(TEXT("Layer blend graphs can have up to two material attributes inputs and a single output."));
+					FailingExpression.Add(nullptr);
+				}
+			}
+			else
+			{
+				// Add a compile error message for functions missing an output
+				FMaterialResource* CurrentResource = ExpressionPreviewMaterial->GetMaterialResource(FeatureLevel);
+				if (CurrentResource)
+				{
+					CompileErrors = CurrentResource->GetCompileErrors();
+					FailingExpression = CurrentResource->GetErrorExpressions();
+				}
 
-					bool bFoundFunctionOutput = false;
-					for (UMaterialExpression* MaterialExpression : Material->GetExpressions())
+				bool bFoundFunctionOutput = false;
+				for (UMaterialExpression* MaterialExpression : Material->GetExpressions())
+				{
+					if (MaterialExpression->IsA(UMaterialExpressionFunctionOutput::StaticClass()))
 					{
-						if (MaterialExpression->IsA(UMaterialExpressionFunctionOutput::StaticClass()))
-						{
-							bFoundFunctionOutput = true;
-							break;
-						}
+						bFoundFunctionOutput = true;
+						break;
 					}
+				}
 
-					if (!bFoundFunctionOutput)
-					{
-						CompileErrors.Add(TEXT("Missing a function output"));
-						FailingExpression.Add(nullptr);
-					}
+				if (!bFoundFunctionOutput)
+				{
+					CompileErrors.Add(TEXT("Missing a function output"));
+					FailingExpression.Add(nullptr);
 				}
 			}
-			else
-			{
-				CompileErrors = MaterialResource->GetCompileErrors();
-				FailingExpression = MaterialResource->GetErrorExpressions();
-			}
+		}
+		else
+		{
+			CompileErrors = MaterialResource->GetCompileErrors();
+			FailingExpression = MaterialResource->GetErrorExpressions();
+		}
 
-			// Only show general info if there are no errors and stats are enabled - Stats show for Materials, layers and blends
+		// Only show general info if there are no errors and stats are enabled - Stats show for Materials, layers and blends
 		if (CompileErrors.Num() == 0 && (!MaterialFunction || MaterialFunction->GetMaterialFunctionUsage() != EMaterialFunctionUsage::Default))
-			{
+		{
 			TArray<FMaterialStatsUtils::FShaderInstructionsInfo> Results;
 			TArray<FMaterialStatsUtils::FShaderInstructionsInfo> EmptyMaterialResults;
 			FMaterialStatsUtils::GetRepresentativeInstructionCounts(Results, MaterialResource);
 
-				//Built in stats is no longer exposed to the UI but may still be useful so they're still in the code.
-				bool bBuiltinStats = false;
+			//Built in stats is no longer exposed to the UI but may still be useful so they're still in the code.
+			bool bBuiltinStats = false;
 			const FMaterialResource* EmptyMaterialResource = EmptyMaterial ? EmptyMaterial->GetMaterialResource(FeatureLevel) : nullptr;
 			if (bShowBuiltinStats && bStatsFromPreviewMaterial && EmptyMaterialResource && Results.Num() > 0)
-				{
+			{
 				FMaterialStatsUtils::GetRepresentativeInstructionCounts(EmptyMaterialResults, EmptyMaterialResource);
 
 				if (EmptyMaterialResults.Num() > 0)
-					{
-						//The instruction counts should match. If not, the preview material has been changed without the EmptyMaterial being updated to match.
+				{
+					//The instruction counts should match. If not, the preview material has been changed without the EmptyMaterial being updated to match.
 					if (ensure(Results.Num() == EmptyMaterialResults.Num()))
-						{
-							bBuiltinStats = true;
-						}
+					{
+						bBuiltinStats = true;
 					}
 				}
+			}
 
 			for (int32 InstructionIndex = 0; InstructionIndex < Results.Num(); InstructionIndex++)
-				{
+			{
 				FString InstructionCountString = FString::Printf(TEXT("%s: %u instructions"), *Results[InstructionIndex].ShaderDescription, Results[InstructionIndex].InstructionCount);
-					if (bBuiltinStats)
-					{
+				if (bBuiltinStats)
+				{
 					InstructionCountString += FString::Printf(TEXT(" - Built-in instructions: %u"), EmptyMaterialResults[InstructionIndex].InstructionCount);
-					}
-					TempMaterialInfoList.Add(MakeShareable(new FMaterialInfo(InstructionCountString, FLinearColor::Yellow)));
-					TSharedRef<FTokenizedMessage> Line = FTokenizedMessage::Create(EMessageSeverity::Info);
-					Line->AddToken(FTextToken::Create(FText::FromString(InstructionCountString)));
-					Messages.Add(Line);
 				}
+				TempMaterialInfoList.Add(MakeShareable(new FMaterialInfo(InstructionCountString, FLinearColor::Yellow)));
+				TSharedRef<FTokenizedMessage> Line = FTokenizedMessage::Create(EMessageSeverity::Info);
+				Line->AddToken(FTextToken::Create(FText::FromString(InstructionCountString)));
+				Messages.Add(Line);
+			}
 
-				// Display the number of samplers used by the material.
-				const int32 SamplersUsed = MaterialResource->GetSamplerUsage();
+			// Display the number of samplers used by the material.
+			const int32 SamplersUsed = MaterialResource->GetSamplerUsage();
 
-				if (SamplersUsed >= 0)
-				{
-					int32 MaxSamplers = GetExpectedFeatureLevelMaxTextureSamplers(MaterialResource->GetFeatureLevel());
-					FString SamplersString = FString::Printf(TEXT("%s samplers: %u/%u"), FeatureLevel <= ERHIFeatureLevel::ES3_1 ? TEXT("Mobile texture") : TEXT("Texture"), SamplersUsed, MaxSamplers);
-					TempMaterialInfoList.Add(MakeShareable(new FMaterialInfo(SamplersString, FLinearColor::Yellow)));
-					TSharedRef<FTokenizedMessage> Line = FTokenizedMessage::Create( EMessageSeverity::Info );
-					Line->AddToken( FTextToken::Create( FText::FromString( SamplersString ) ) );
-					Messages.Add(Line);
-				}
+			if (SamplersUsed >= 0)
+			{
+				int32 MaxSamplers = GetExpectedFeatureLevelMaxTextureSamplers(MaterialResource->GetFeatureLevel());
+				FString SamplersString = FString::Printf(TEXT("%s samplers: %u/%u"), FeatureLevel <= ERHIFeatureLevel::ES3_1 ? TEXT("Mobile texture") : TEXT("Texture"), SamplersUsed, MaxSamplers);
+				TempMaterialInfoList.Add(MakeShareable(new FMaterialInfo(SamplersString, FLinearColor::Yellow)));
+				TSharedRef<FTokenizedMessage> Line = FTokenizedMessage::Create( EMessageSeverity::Info );
+				Line->AddToken( FTextToken::Create( FText::FromString( SamplersString ) ) );
+				Messages.Add(Line);
+			}
 				
-				// Display estimated texture look-up/sample counts
-				uint32 NumVSTextureSamples = 0, NumPSTextureSamples = 0;
-				MaterialResource->GetEstimatedNumTextureSamples(NumVSTextureSamples, NumPSTextureSamples);
+			// Display estimated texture look-up/sample counts
+			uint32 NumVSTextureSamples = 0, NumPSTextureSamples = 0;
+			MaterialResource->GetEstimatedNumTextureSamples(NumVSTextureSamples, NumPSTextureSamples);
 
-				if (NumVSTextureSamples > 0 || NumPSTextureSamples > 0)
-				{
-					FString SamplesString = FString::Printf(TEXT("Texture Lookups (Est.): VS(%u), PS(%u)"), NumVSTextureSamples, NumPSTextureSamples);
+			if (NumVSTextureSamples > 0 || NumPSTextureSamples > 0)
+			{
+				FString SamplesString = FString::Printf(TEXT("Texture Lookups (Est.): VS(%u), PS(%u)"), NumVSTextureSamples, NumPSTextureSamples);
 
-					TempMaterialInfoList.Add(MakeShareable(new FMaterialInfo(SamplesString, FLinearColor::Yellow)));
-					TSharedRef<FTokenizedMessage> Line = FTokenizedMessage::Create(EMessageSeverity::Info);
-					Line->AddToken(FTextToken::Create(FText::FromString(SamplesString)));
-					Messages.Add(Line);
-				}
+				TempMaterialInfoList.Add(MakeShareable(new FMaterialInfo(SamplesString, FLinearColor::Yellow)));
+				TSharedRef<FTokenizedMessage> Line = FTokenizedMessage::Create(EMessageSeverity::Info);
+				Line->AddToken(FTextToken::Create(FText::FromString(SamplesString)));
+				Messages.Add(Line);
+			}
 
-				// Display estimated virtual texture look-up counts
-				uint32 NumVirtualTextureLookups = MaterialResource->GetEstimatedNumVirtualTextureLookups();
-				if (NumVirtualTextureLookups > 0)
-				{
-					FString LookupsString = FString::Printf(TEXT("Virtual Texture Lookups (Est.): %u"), NumVirtualTextureLookups);
+			// Display estimated virtual texture look-up counts
+			uint32 NumVirtualTextureLookups = MaterialResource->GetEstimatedNumVirtualTextureLookups();
+			if (NumVirtualTextureLookups > 0)
+			{
+				FString LookupsString = FString::Printf(TEXT("Virtual Texture Lookups (Est.): %u"), NumVirtualTextureLookups);
 
-					TempMaterialInfoList.Add(MakeShareable(new FMaterialInfo(LookupsString, FLinearColor::Yellow)));
-					TSharedRef<FTokenizedMessage> Line = FTokenizedMessage::Create(EMessageSeverity::Info);
-					Line->AddToken(FTextToken::Create(FText::FromString(LookupsString)));
-					Messages.Add(Line);
-				}
+				TempMaterialInfoList.Add(MakeShareable(new FMaterialInfo(LookupsString, FLinearColor::Yellow)));
+				TSharedRef<FTokenizedMessage> Line = FTokenizedMessage::Create(EMessageSeverity::Info);
+				Line->AddToken(FTextToken::Create(FText::FromString(LookupsString)));
+				Messages.Add(Line);
+			}
 
-				const uint32 NumVirtualTextureStacks = MaterialResource->GetNumVirtualTextureStacks();
-				if (NumVirtualTextureStacks > 0u)
-				{
-					FString VTString = FString::Printf(TEXT("Virtual Texture Stacks: %u"), NumVirtualTextureStacks);
+			const uint32 NumVirtualTextureStacks = MaterialResource->GetNumVirtualTextureStacks();
+			if (NumVirtualTextureStacks > 0u)
+			{
+				FString VTString = FString::Printf(TEXT("Virtual Texture Stacks: %u"), NumVirtualTextureStacks);
 
-					TempMaterialInfoList.Add(MakeShareable(new FMaterialInfo(VTString, FLinearColor::Yellow)));
-					TSharedRef<FTokenizedMessage> Line = FTokenizedMessage::Create(EMessageSeverity::Info);
-					Line->AddToken(FTextToken::Create(FText::FromString(VTString)));
-					Messages.Add(Line);
-				}
+				TempMaterialInfoList.Add(MakeShareable(new FMaterialInfo(VTString, FLinearColor::Yellow)));
+				TSharedRef<FTokenizedMessage> Line = FTokenizedMessage::Create(EMessageSeverity::Info);
+				Line->AddToken(FTextToken::Create(FText::FromString(VTString)));
+				Messages.Add(Line);
+			}
 
-				// Display the number of custom/user interpolators used by the material.
-				uint32 UVScalarsUsed, CustomInterpolatorScalarsUsed;
-				MaterialResource->GetUserInterpolatorUsage(UVScalarsUsed, CustomInterpolatorScalarsUsed);
+			// Display the number of custom/user interpolators used by the material.
+			uint32 UVScalarsUsed, CustomInterpolatorScalarsUsed;
+			MaterialResource->GetUserInterpolatorUsage(UVScalarsUsed, CustomInterpolatorScalarsUsed);
 
-				if (UVScalarsUsed > 0 || CustomInterpolatorScalarsUsed > 0)
-				{
-					uint32 TotalScalars = UVScalarsUsed + CustomInterpolatorScalarsUsed;
-					uint32 MaxScalars = FMath::DivideAndRoundUp(TotalScalars, 4u) * 4;
+			if (UVScalarsUsed > 0 || CustomInterpolatorScalarsUsed > 0)
+			{
+				uint32 TotalScalars = UVScalarsUsed + CustomInterpolatorScalarsUsed;
+				uint32 MaxScalars = FMath::DivideAndRoundUp(TotalScalars, 4u) * 4;
 
-					FString InterpolatorsString = FString::Printf(TEXT("User interpolators: %u/%u Scalars (%u/4 Vectors) (TexCoords: %i, Custom: %i)"),
-						TotalScalars, MaxScalars, MaxScalars / 4, UVScalarsUsed, CustomInterpolatorScalarsUsed);
+				FString InterpolatorsString = FString::Printf(TEXT("User interpolators: %u/%u Scalars (%u/4 Vectors) (TexCoords: %i, Custom: %i)"),
+															  TotalScalars, MaxScalars, MaxScalars / 4, UVScalarsUsed, CustomInterpolatorScalarsUsed);
 
-					TempMaterialInfoList.Add(MakeShareable(new FMaterialInfo(InterpolatorsString, FLinearColor::Yellow)));
-					TSharedRef<FTokenizedMessage> Line = FTokenizedMessage::Create( EMessageSeverity::Info );
-					Line->AddToken(FTextToken::Create(FText::FromString(InterpolatorsString)));
-					Messages.Add(Line);
-				}
+				TempMaterialInfoList.Add(MakeShareable(new FMaterialInfo(InterpolatorsString, FLinearColor::Yellow)));
+				TSharedRef<FTokenizedMessage> Line = FTokenizedMessage::Create( EMessageSeverity::Info );
+				Line->AddToken(FTextToken::Create(FText::FromString(InterpolatorsString)));
+				Messages.Add(Line);
+			}
 
-			TStaticArray<uint16, (int)ELWCFunctionKind::Max> LWCFuncUsagesVS;
-			TStaticArray<uint16, (int)ELWCFunctionKind::Max> LWCFuncUsagesPS;
-			MaterialResource->GetEstimatedLWCFuncUsages(LWCFuncUsagesVS, LWCFuncUsagesPS);
+			FMaterialResource::FLWCUsagesArray LWCFuncUsagesVS;
+			FMaterialResource::FLWCUsagesArray LWCFuncUsagesPS;
+			FMaterialResource::FLWCUsagesArray LWCFuncUsagesCS;
+			MaterialResource->GetEstimatedLWCFuncUsages(LWCFuncUsagesVS, LWCFuncUsagesPS, LWCFuncUsagesCS);
 			for (int KindIndex = 0; KindIndex < (int)ELWCFunctionKind::Max; ++KindIndex)
 			{
-				int Usages = LWCFuncUsagesVS[KindIndex] + LWCFuncUsagesPS[KindIndex];
+				int Usages = LWCFuncUsagesVS[KindIndex] + LWCFuncUsagesPS[KindIndex] + LWCFuncUsagesCS[KindIndex];
 				if (Usages > 0)
 				{
-					FString Message = FString::Printf(TEXT("LWC %s usages (Est.): %u"), *UEnum::GetDisplayValueAsText((ELWCFunctionKind)KindIndex).ToString(), Usages);
+					FString Message = FString::Printf(
+						TEXT("LWC %s usages (Est.): %u (VS), %u (PS), %u (CS)"), 
+						*UEnum::GetDisplayValueAsText((ELWCFunctionKind)KindIndex).ToString(), 
+						LWCFuncUsagesVS[KindIndex],
+						LWCFuncUsagesPS[KindIndex],
+						LWCFuncUsagesCS[KindIndex]);
 						
 					TempMaterialInfoList.Add(MakeShareable(new FMaterialInfo(Message, FLinearColor::Yellow)));
 					TSharedRef<FTokenizedMessage> Line = FTokenizedMessage::Create(EMessageSeverity::Info);
@@ -3184,49 +3196,49 @@ void FMaterialEditor::UpdateMaterialinfoList_Old()
 				}
 			}
 
-				if (FMaterialShaderMap* ShaderMap = MaterialResource->GetGameThreadShaderMap())
+			if (FMaterialShaderMap* ShaderMap = MaterialResource->GetGameThreadShaderMap())
+			{
+				// Add shader count
+				FString ShaderCountString = FString::Printf(TEXT("Shader Count: %u"), ShaderMap->GetShaderNum());
+				TSharedRef<FTokenizedMessage> ShaderCountLine = FTokenizedMessage::Create(EMessageSeverity::Info);
+				ShaderCountLine->AddToken(FTextToken::Create(FText::FromString(ShaderCountString)));
+				Messages.Add(ShaderCountLine);
+
+				// Add number of preshaders and stats
+				uint32 TotalParams, TotalOps;
+				MaterialResource->GetPreshaderStats(TotalParams, TotalOps);
+				FString PreshaderCountString = FString::Printf(TEXT("Preshaders: %u  (%u param fetches, %u ops)"), ShaderMap->GetNumPreshaders(), TotalParams, TotalOps);
+				TSharedRef<FTokenizedMessage> PreshaderCountLine = FTokenizedMessage::Create(EMessageSeverity::Info);
+				PreshaderCountLine->AddToken(FTextToken::Create(FText::FromString(PreshaderCountString)));
+				Messages.Add(PreshaderCountLine);
+
+				if (MaterialResource->GetMaterialDomain() == MD_LightFunction)
 				{
-					// Add shader count
-					FString ShaderCountString = FString::Printf(TEXT("Shader Count: %u"), ShaderMap->GetShaderNum());
-					TSharedRef<FTokenizedMessage> ShaderCountLine = FTokenizedMessage::Create(EMessageSeverity::Info);
-					ShaderCountLine->AddToken(FTextToken::Create(FText::FromString(ShaderCountString)));
-					Messages.Add(ShaderCountLine);
-
-					// Add number of preshaders and stats
-					uint32 TotalParams, TotalOps;
-					MaterialResource->GetPreshaderStats(TotalParams, TotalOps);
-					FString PreshaderCountString = FString::Printf(TEXT("Preshaders: %u  (%u param fetches, %u ops)"), ShaderMap->GetNumPreshaders(), TotalParams, TotalOps);
-					TSharedRef<FTokenizedMessage> PreshaderCountLine = FTokenizedMessage::Create(EMessageSeverity::Info);
-					PreshaderCountLine->AddToken(FTextToken::Create(FText::FromString(PreshaderCountString)));
-					Messages.Add(PreshaderCountLine);
-
-					if (MaterialResource->GetMaterialDomain() == MD_LightFunction)
-					{
-						const bool bIsCompatibleWithLightFunctionAtlas = ShaderMap->IsLightFunctionAtlasCompatible();
-						FString LightFunctionAtlasStr = FString::Printf(TEXT("Light function material%s compatible with the light function atlas for fast batched deferred light shading."), bIsCompatibleWithLightFunctionAtlas ? TEXT(" IS") : TEXT(" IS NOT"));
-						TSharedRef<FTokenizedMessage> LightFunctionAtlasCountLine = FTokenizedMessage::Create(EMessageSeverity::Info);
-						LightFunctionAtlasCountLine->AddToken(FTextToken::Create(FText::FromString(LightFunctionAtlasStr)));
-						Messages.Add(LightFunctionAtlasCountLine);
-					}
+					const bool bIsCompatibleWithLightFunctionAtlas = ShaderMap->IsLightFunctionAtlasCompatible();
+					FString LightFunctionAtlasStr = FString::Printf(TEXT("Light function material%s compatible with the light function atlas for fast batched deferred light shading."), bIsCompatibleWithLightFunctionAtlas ? TEXT(" IS") : TEXT(" IS NOT"));
+					TSharedRef<FTokenizedMessage> LightFunctionAtlasCountLine = FTokenizedMessage::Create(EMessageSeverity::Info);
+					LightFunctionAtlasCountLine->AddToken(FTextToken::Create(FText::FromString(LightFunctionAtlasStr)));
+					Messages.Add(LightFunctionAtlasCountLine);
 				}
 			}
+		}
 
-			FString FeatureLevelName;
-			GetFeatureLevelName(FeatureLevel,FeatureLevelName);
-			for(int32 ErrorIndex = 0; ErrorIndex < CompileErrors.Num(); ErrorIndex++)
+		FString FeatureLevelName;
+		GetFeatureLevelName(FeatureLevel,FeatureLevelName);
+		for(int32 ErrorIndex = 0; ErrorIndex < CompileErrors.Num(); ErrorIndex++)
+		{
+			FString ErrorString = FString::Printf(TEXT("[%s] %s"), *FeatureLevelName, *CompileErrors[ErrorIndex]);
+			TempMaterialInfoList.Add(MakeShareable(new FMaterialInfo(ErrorString, FLinearColor::Red)));
+			TSharedRef<FTokenizedMessage> Line = FTokenizedMessage::Create( EMessageSeverity::Error );
+			if(FailingExpression.Num() && ensure(FailingExpression.Num() == CompileErrors.Num()) && FailingExpression[ErrorIndex])
 			{
-				FString ErrorString = FString::Printf(TEXT("[%s] %s"), *FeatureLevelName, *CompileErrors[ErrorIndex]);
-				TempMaterialInfoList.Add(MakeShareable(new FMaterialInfo(ErrorString, FLinearColor::Red)));
-				TSharedRef<FTokenizedMessage> Line = FTokenizedMessage::Create( EMessageSeverity::Error );
-				if(FailingExpression.Num() && ensure(FailingExpression.Num() == CompileErrors.Num()) && FailingExpression[ErrorIndex])
-				{
-					Line->SetMessageLink(FUObjectToken::Create(FailingExpression[ErrorIndex]));
-				}
-				Line->AddToken( FTextToken::Create( FText::FromString( ErrorString ) ) );
-				Messages.Add(Line);
-				bForceDisplay = true;
+				Line->SetMessageLink(FUObjectToken::Create(FailingExpression[ErrorIndex]));
 			}
+			Line->AddToken( FTextToken::Create( FText::FromString( ErrorString ) ) );
+			Messages.Add(Line);
+			bForceDisplay = true;
 		}
+	}
 
 	bool bNeedsRefresh = false;
 	if (TempMaterialInfoList.Num() != MaterialInfoList.Num())
diff --git a/Engine/Source/Editor/MaterialEditor/Private/MaterialStatsCommon.cpp b/Engine/Source/Editor/MaterialEditor/Private/MaterialStatsCommon.cpp
--- a/Engine/Source/Editor/MaterialEditor/Private/MaterialStatsCommon.cpp
+++ b/Engine/Source/Editor/MaterialEditor/Private/MaterialStatsCommon.cpp
@@ -730,15 +730,21 @@ void FMaterialStatsUtils::ExtractMatertialStatsInfo(EShaderPlatform ShaderPlatfo
 		OutInfo.ShaderCount.StrDescriptionLong = FString::Printf(TEXT("Total Shaders: %u"), TotalShadersForMaterial);
 
 		FString LWCMessage;
-		TStaticArray<uint16, (int)ELWCFunctionKind::Max> LWCFuncUsagesVS; 
-		TStaticArray<uint16, (int)ELWCFunctionKind::Max> LWCFuncUsagesPS; 
-		MaterialResource->GetEstimatedLWCFuncUsages(LWCFuncUsagesVS, LWCFuncUsagesPS);
+		FMaterialResource::FLWCUsagesArray LWCFuncUsagesVS;
+		FMaterialResource::FLWCUsagesArray LWCFuncUsagesPS;
+		FMaterialResource::FLWCUsagesArray LWCFuncUsagesCS;
+		MaterialResource->GetEstimatedLWCFuncUsages(LWCFuncUsagesVS, LWCFuncUsagesPS, LWCFuncUsagesCS);
 		for (int KindIndex = 0; KindIndex < (int)ELWCFunctionKind::Max; ++KindIndex)
 		{
-			int Usages = LWCFuncUsagesVS[KindIndex] + LWCFuncUsagesPS[KindIndex];
+			int Usages = LWCFuncUsagesVS[KindIndex] + LWCFuncUsagesPS[KindIndex] + LWCFuncUsagesCS[KindIndex];
 			if (Usages > 0)
 			{
-				LWCMessage += FString::Printf(TEXT("%s: %u\n"), *UEnum::GetDisplayValueAsText((ELWCFunctionKind)KindIndex).ToString(), Usages);
+				LWCMessage += FString::Printf(
+					TEXT("%s: %u (VS), %u (PS), %u (CS)\n"), 
+					*UEnum::GetDisplayValueAsText((ELWCFunctionKind)KindIndex).ToString(), 
+					LWCFuncUsagesVS[KindIndex],
+					LWCFuncUsagesPS[KindIndex],
+					LWCFuncUsagesCS[KindIndex]);
 			}
 		}
 
diff --git a/Engine/Source/Editor/UnrealEd/Classes/Commandlets/DumpMaterialInfo.h b/Engine/Source/Editor/UnrealEd/Classes/Commandlets/DumpMaterialInfo.h
new file mode 100644
--- /dev/null
+++ b/Engine/Source/Editor/UnrealEd/Classes/Commandlets/DumpMaterialInfo.h
@@ -0,0 +1,22 @@
+// Copyright Epic Games, Inc. All Rights Reserved.
+
+#pragma once
+
+#include "CoreMinimal.h"
+#include "UObject/ObjectMacros.h"
+#include "Commandlets/Commandlet.h"
+#include "DumpMaterialInfo.generated.h"
+
+/**
+  Writes a CSV that lists many properties of all material (instances).
+  Useful to inspect all materials and query them, e.g. for performance characteristics.
+  */
+UCLASS(config = Editor)
+class UDumpMaterialInfoCommandlet : public UCommandlet
+{
+	GENERATED_UCLASS_BODY()
+
+	//~ Begin UCommandlet Interface
+	virtual int32 Main(const FString& Params) override;
+	//~ End UCommandlet Interface
+};
diff --git a/Engine/Source/Editor/UnrealEd/Private/Commandlets/DumpMaterialInfo.cpp b/Engine/Source/Editor/UnrealEd/Private/Commandlets/DumpMaterialInfo.cpp
new file mode 100644
--- /dev/null
+++ b/Engine/Source/Editor/UnrealEd/Private/Commandlets/DumpMaterialInfo.cpp
@@ -0,0 +1,754 @@
+// Copyright Epic Games, Inc. All Rights Reserved.
+
+#include "Commandlets/DumpMaterialInfo.h"
+#include "Modules/ModuleManager.h"
+#include "AssetRegistry/AssetRegistryModule.h"
+#include "AssetRegistry/AssetData.h"
+#include "Interfaces/ITargetPlatform.h"
+#include "Interfaces/ITargetPlatformManagerModule.h"
+#include "Materials/Material.h"
+#include "Materials/MaterialInstance.h"
+#include "MaterialStatsCommon.h"
+#include "MaterialShared.h"
+#include "MaterialDomain.h"
+#include "ShaderCompiler.h"
+#include "Algo/Accumulate.h"
+#include "MaterialShaderType.h"
+#include "RendererUtils.h"
+#include "Internationalization/Regex.h"
+#include "String/ParseTokens.h"
+#include "ProfilingDebugging/DiagnosticTable.h"
+
+DEFINE_LOG_CATEGORY_STATIC(LogDumpMaterialInfo, Log, All);
+
+namespace MaterialInfo
+{
+	class FOutput;
+
+	// Material data from which properties are read to be dumped to csv
+	struct FPropertyDumpInput
+	{
+		UMaterial* Material;
+		FMaterialResource* MaterialResource;
+		FMaterialRelevance MaterialRelevance;
+		TArray<const FShaderType*>* ShaderTypes;
+	};
+
+	// Set of material properties that can be dumped
+	struct FPropertySet
+	{
+		using FDumpFunction = TFunction<void(const FPropertyDumpInput& Input, FOutput& Output)>;
+
+		// Names of all the properties in this set, used as column headers and for filtering
+		TArray<FString> PropertyNames;
+		// Read each property value in this set from Input, and write it to Output, in order specified by PropertyNames
+		FDumpFunction DumpFunction;
+
+		FPropertySet() { }
+		FPropertySet(TArray<FString> InNames, FDumpFunction InDumpFunction)
+			: PropertyNames(InNames)
+			, DumpFunction(InDumpFunction)
+		{ }
+		FPropertySet(FString InName, FDumpFunction InDumpFunction)
+			: PropertyNames( { InName } )
+			, DumpFunction(InDumpFunction)
+		{ }
+
+		template<typename TValue, TValue(FMaterialResource::* MemberAccessor)() const>
+		static FPropertySet FromAccessor(const FString& InName);
+
+		template<typename TValue, TValue(FMaterialResource::* MemberAccessor)() const>
+		static FPropertySet FromEnumAccessor(const FString& InName);
+	};
+
+	struct FPropertyValue
+	{
+		// Column name
+		const FString& Name;
+		// Value as string view, valid until FOutput::Reset is called.
+		FStringView Value;
+	};
+
+	struct FSlice
+	{
+		int Start;
+		int End;
+	};
+
+	// Class to efficiently write out property values to a stringbuffer and verify the output.
+	class FOutput
+	{
+		int NumExpectedValues = 0;
+		// Storage for output strings
+		TStringBuilder<2048> Buffer;
+		TArray<FSlice> BufferSlices;
+		// Populated after all values have been written.
+		TArray<FPropertyValue> Values;
+
+	public:
+		void DumpPropertySet(const FPropertySet& PropertySet, const FPropertyDumpInput& Input)
+		{
+			NumExpectedValues = PropertySet.PropertyNames.Num();
+			PropertySet.DumpFunction(Input, *this);
+			check(NumExpectedValues == 0);
+			for (int Index = 0; Index < BufferSlices.Num(); Index++)
+			{
+				const FSlice& Slice = BufferSlices[Index];
+				Values.Push( { PropertySet.PropertyNames[Index], Buffer.ToView().SubStr(Slice.Start, Slice.End - Slice.Start) });
+			}
+			BufferSlices.Reset();
+		}
+
+		const TArray<FPropertyValue>& GetValues()
+		{
+			return Values;
+		}
+
+		void Reset()
+		{
+			NumExpectedValues = 0;
+			Buffer.Reset();
+			BufferSlices.Reset();
+			Values.Reset();
+		}
+
+		template <typename FmtType, typename... Types>
+		void WriteFormat(const FmtType& Fmt, Types... Args)
+		{
+			check(NumExpectedValues-- > 0);
+			int Start = Buffer.Len();
+			Buffer.Appendf(Fmt, Forward<Types>(Args)...);
+			int End = Buffer.Len();
+			BufferSlices.Push( { Start, End } );
+		}
+
+		void Write(int32 Value)
+		{
+			WriteFormat(TEXT("%d"), Value);
+		}
+
+		void Write(uint32 Value)
+		{
+			WriteFormat(TEXT("%u"), Value);
+		}
+
+		void Write(bool Value)
+		{
+			WriteFormat(TEXT("%d"), Value);
+		}
+
+		void Write(float Value)
+		{
+			WriteFormat(TEXT("%f"), Value);
+		}
+
+		void Write(const FString& Value)
+		{
+			WriteFormat(TEXT("%s"), *Value);
+		}
+	};
+
+	template<typename TValue, TValue(FMaterialResource::*MemberAccessor)() const>
+	FPropertySet FPropertySet::FromAccessor(const FString& InName)
+	{ 
+		return FPropertySet(
+			InName, 
+			[](const FPropertyDumpInput& Input, FOutput& Output)
+			{ 
+				Output.Write((Input.MaterialResource->*MemberAccessor)()); 
+			});
+	}
+
+	template<typename TValue, TValue(FMaterialResource::*MemberAccessor)() const>
+	FPropertySet FPropertySet::FromEnumAccessor(const FString& InName)
+	{ 
+		return FPropertySet(
+			InName, 
+			[](const FPropertyDumpInput& Input, FOutput& Output)
+			{ 
+				Output.Write(UEnum::GetValueOrBitfieldAsString((Input.MaterialResource->*MemberAccessor)()));
+			});
+	}
+}
+
+
+UDumpMaterialInfoCommandlet::UDumpMaterialInfoCommandlet(const FObjectInitializer& ObjectInitializer)
+	: Super(ObjectInitializer)
+{
+}
+
+TArray<MaterialInfo::FPropertySet> GetMaterialInfoProperties()
+{
+#if WITH_EDITOR
+	using namespace MaterialInfo;
+	TArray<FPropertySet> PropertySets = 
+	{
+		FPropertySet::FromAccessor<int32, &FMaterialResource::GetSamplerUsage>(TEXT("GetSamplerUsage")),
+		FPropertySet
+		{
+			{ TEXT("NumUsedUVScalars"), TEXT("NumUsedCustomInterpolatorScalars") }, 
+			[](const FPropertyDumpInput& Input, FOutput& Output)
+			{
+				uint32 NumUsedUVScalars = 0;
+				uint32 NumUsedCustomInterpolatorScalars = 0;
+				Input.MaterialResource->GetUserInterpolatorUsage(NumUsedUVScalars, NumUsedCustomInterpolatorScalars);
+				Output.Write(NumUsedUVScalars);
+				Output.Write(NumUsedCustomInterpolatorScalars);
+			} 
+		},
+		FPropertySet::FromAccessor<uint32, &FMaterialResource::GetEstimatedNumVirtualTextureLookups>(TEXT("GetEstimatedNumVirtualTextureLookups")),
+		FPropertySet
+		{
+			{ TEXT("LWCUsagesVS"), TEXT("LWCUsagesPS"), TEXT("LWCUsagesCS") }, 
+			[](const FPropertyDumpInput& Input, FOutput& Output)
+			{
+				FMaterialResource::FLWCUsagesArray LWCUsagesVS {};
+				FMaterialResource::FLWCUsagesArray LWCUsagesPS {};
+				FMaterialResource::FLWCUsagesArray LWCUsagesCS {};
+				Input.MaterialResource->GetEstimatedLWCFuncUsages(LWCUsagesVS, LWCUsagesPS, LWCUsagesCS);
+				Output.Write(Algo::Accumulate(LWCUsagesVS, 0));
+				Output.Write(Algo::Accumulate(LWCUsagesPS, 0));
+				Output.Write(Algo::Accumulate(LWCUsagesCS, 0));
+			} 
+		},
+		FPropertySet::FromAccessor<uint32, &FMaterialResource::GetNumVirtualTextureStacks>(TEXT("GetNumVirtualTextureStacks")),
+		//FPropertySet::FromAccessor<FString, &FMaterialResource::GetMaterialUsageDescription>(TEXT("MaterialUsageDescription")),
+		//FPropertySet::FromAccessor<void, &FMaterialResource::GetShaderMapId>(TEXT("ShaderMapId")),
+		//FPropertySet::FromAccessor<void, &FMaterialResource::GetStaticParameterSet>(TEXT("StaticParameterSet")),
+		FPropertySet::FromEnumAccessor<EMaterialDomain, &FMaterialResource::GetMaterialDomain>(TEXT("GetMaterialDomain")),
+		FPropertySet::FromAccessor<bool, &FMaterialResource::IsTranslucencyWritingFrontLayerTransparency>(TEXT("IsTranslucencyWritingFrontLayerTransparency")),
+		FPropertySet::FromAccessor<bool, &FMaterialResource::IsTranslucencyWritingFrontLayerTransparency>(TEXT("IsTranslucencyWritingFrontLayerTransparency")),
+		FPropertySet::FromAccessor<bool, &FMaterialResource::IsTangentSpaceNormal>(TEXT("IsTangentSpaceNormal")),
+		FPropertySet::FromAccessor<bool, &FMaterialResource::ShouldGenerateSphericalParticleNormals>(TEXT("ShouldGenerateSphericalParticleNormals")),
+		FPropertySet::FromAccessor<bool, &FMaterialResource::ShouldDisableDepthTest>(TEXT("ShouldDisableDepthTest")),
+		FPropertySet::FromAccessor<bool, &FMaterialResource::ShouldWriteOnlyAlpha>(TEXT("ShouldWriteOnlyAlpha")),
+		FPropertySet::FromAccessor<bool, &FMaterialResource::ShouldEnableResponsiveAA>(TEXT("ShouldEnableResponsiveAA")),
+		FPropertySet::FromAccessor<bool, &FMaterialResource::ShouldDoSSR>(TEXT("ShouldDoSSR")),
+		FPropertySet::FromAccessor<bool, &FMaterialResource::ShouldDoContactShadows>(TEXT("ShouldDoContactShadows")),
+		FPropertySet::FromAccessor<bool, &FMaterialResource::HasPixelAnimation>(TEXT("HasPixelAnimation")),
+		FPropertySet::FromAccessor<bool, &FMaterialResource::IsLightFunction>(TEXT("IsLightFunction")),
+		FPropertySet::FromAccessor<bool, &FMaterialResource::IsUsedWithEditorCompositing>(TEXT("IsUsedWithEditorCompositing")),
+		FPropertySet::FromAccessor<bool, &FMaterialResource::IsDeferredDecal>(TEXT("IsDeferredDecal")),
+		FPropertySet::FromAccessor<bool, &FMaterialResource::IsVolumetricPrimitive>(TEXT("IsVolumetricPrimitive")),
+		FPropertySet::FromAccessor<bool, &FMaterialResource::IsWireframe>(TEXT("IsWireframe")),
+		FPropertySet::FromAccessor<bool, &FMaterialResource::IsVariableRateShadingAllowed>(TEXT("IsVariableRateShadingAllowed")),
+		FPropertySet::FromEnumAccessor<EMaterialShadingRate, &FMaterialResource::GetShadingRate>(TEXT("GetShadingRate")),
+		FPropertySet::FromAccessor<bool, &FMaterialResource::IsUIMaterial>(TEXT("IsUIMaterial")),
+		FPropertySet::FromAccessor<bool, &FMaterialResource::IsPostProcessMaterial>(TEXT("IsPostProcessMaterial")),
+		FPropertySet::FromAccessor<bool, &FMaterialResource::IsSpecialEngineMaterial>(TEXT("IsSpecialEngineMaterial")),
+		FPropertySet::FromAccessor<bool, &FMaterialResource::IsUsedWithSkeletalMesh>(TEXT("IsUsedWithSkeletalMesh")),
+		FPropertySet::FromAccessor<bool, &FMaterialResource::IsUsedWithLandscape>(TEXT("IsUsedWithLandscape")),
+		FPropertySet::FromAccessor<bool, &FMaterialResource::IsUsedWithParticleSystem>(TEXT("IsUsedWithParticleSystem")),
+		FPropertySet::FromAccessor<bool, &FMaterialResource::IsUsedWithParticleSprites>(TEXT("IsUsedWithParticleSprites")),
+		FPropertySet::FromAccessor<bool, &FMaterialResource::IsUsedWithBeamTrails>(TEXT("IsUsedWithBeamTrails")),
+		FPropertySet::FromAccessor<bool, &FMaterialResource::IsUsedWithMeshParticles>(TEXT("IsUsedWithMeshParticles")),
+		FPropertySet::FromAccessor<bool, &FMaterialResource::IsUsedWithNiagaraSprites>(TEXT("IsUsedWithNiagaraSprites")),
+		FPropertySet::FromAccessor<bool, &FMaterialResource::IsUsedWithNiagaraRibbons>(TEXT("IsUsedWithNiagaraRibbons")),
+		FPropertySet::FromAccessor<bool, &FMaterialResource::IsUsedWithNiagaraMeshParticles>(TEXT("IsUsedWithNiagaraMeshParticles")),
+		FPropertySet::FromAccessor<bool, &FMaterialResource::IsUsedWithStaticLighting>(TEXT("IsUsedWithStaticLighting")),
+		FPropertySet::FromAccessor<bool, &FMaterialResource::IsUsedWithMorphTargets>(TEXT("IsUsedWithMorphTargets")),
+		FPropertySet::FromAccessor<bool, &FMaterialResource::IsUsedWithSplineMeshes>(TEXT("IsUsedWithSplineMeshes")),
+		FPropertySet::FromAccessor<bool, &FMaterialResource::IsUsedWithInstancedStaticMeshes>(TEXT("IsUsedWithInstancedStaticMeshes")),
+		FPropertySet::FromAccessor<bool, &FMaterialResource::IsUsedWithGeometryCollections>(TEXT("IsUsedWithGeometryCollections")),
+		FPropertySet::FromAccessor<bool, &FMaterialResource::IsUsedWithAPEXCloth>(TEXT("IsUsedWithAPEXCloth")),
+		FPropertySet::FromAccessor<bool, &FMaterialResource::IsUsedWithGeometryCache>(TEXT("IsUsedWithGeometryCache")),
+		FPropertySet::FromAccessor<bool, &FMaterialResource::IsUsedWithWater>(TEXT("IsUsedWithWater")),
+		FPropertySet::FromAccessor<bool, &FMaterialResource::IsUsedWithHairStrands>(TEXT("IsUsedWithHairStrands")),
+		FPropertySet::FromAccessor<bool, &FMaterialResource::IsUsedWithLidarPointCloud>(TEXT("IsUsedWithLidarPointCloud")),
+		FPropertySet::FromAccessor<bool, &FMaterialResource::IsUsedWithVirtualHeightfieldMesh>(TEXT("IsUsedWithVirtualHeightfieldMesh")),
+		FPropertySet::FromAccessor<bool, &FMaterialResource::IsUsedWithNeuralNetworks>(TEXT("IsUsedWithNeuralNetworks")),
+		FPropertySet::FromAccessor<bool, &FMaterialResource::IsUsedWithNanite>(TEXT("IsUsedWithNanite")),
+		FPropertySet::FromAccessor<bool, &FMaterialResource::IsUsedWithVolumetricCloud>(TEXT("IsUsedWithVolumetricCloud")),
+		FPropertySet::FromAccessor<bool, &FMaterialResource::IsUsedWithHeterogeneousVolumes>(TEXT("IsUsedWithHeterogeneousVolumes")),
+		FPropertySet::FromAccessor<bool, &FMaterialResource::IsFullyRough>(TEXT("IsFullyRough")),
+		FPropertySet::FromAccessor<bool, &FMaterialResource::GetForceCompatibleWithLightFunctionAtlas>(TEXT("GetForceCompatibleWithLightFunctionAtlas")),
+		FPropertySet::FromAccessor<bool, &FMaterialResource::UseNormalCurvatureToRoughness>(TEXT("UseNormalCurvatureToRoughness")),
+		FPropertySet::FromEnumAccessor<EMaterialFloatPrecisionMode, &FMaterialResource::GetMaterialFloatPrecisionMode>(TEXT("GetMaterialFloatPrecisionMode")),
+		FPropertySet::FromAccessor<bool, &FMaterialResource::IsUsingAlphaToCoverage>(TEXT("IsUsingAlphaToCoverage")),
+		FPropertySet::FromAccessor<bool, &FMaterialResource::IsUsingPreintegratedGFForSimpleIBL>(TEXT("IsUsingPreintegratedGFForSimpleIBL")),
+		FPropertySet::FromAccessor<bool, &FMaterialResource::IsUsingHQForwardReflections>(TEXT("IsUsingHQForwardReflections")),
+		FPropertySet::FromAccessor<bool, &FMaterialResource::GetForwardBlendsSkyLightCubemaps>(TEXT("GetForwardBlendsSkyLightCubemaps")),
+		FPropertySet::FromAccessor<bool, &FMaterialResource::IsUsingPlanarForwardReflections>(TEXT("IsUsingPlanarForwardReflections")),
+		FPropertySet::FromAccessor<bool, &FMaterialResource::IsNonmetal>(TEXT("IsNonmetal")),
+		FPropertySet::FromAccessor<bool, &FMaterialResource::UseLmDirectionality>(TEXT("UseLmDirectionality")),
+		FPropertySet::FromAccessor<bool, &FMaterialResource::IsMobileHighQualityBRDFEnabled>(TEXT("IsMobileHighQualityBRDFEnabled")),
+		FPropertySet::FromEnumAccessor<EBlendMode, &FMaterialResource::GetBlendMode>(TEXT("GetBlendMode")),
+		FPropertySet::FromEnumAccessor<ERefractionMode, &FMaterialResource::GetRefractionMode>(TEXT("GetRefractionMode")),
+		FPropertySet::FromAccessor<bool, &FMaterialResource::GetRootNodeOverridesDefaultRefraction>(TEXT("GetRootNodeOverridesDefaultRefraction")),
+		FPropertySet::FromAccessor<uint32, &FMaterialResource::GetMaterialDecalResponse>(TEXT("GetMaterialDecalResponse")),
+		FPropertySet::FromAccessor<bool, &FMaterialResource::HasBaseColorConnected>(TEXT("HasBaseColorConnected")),
+		FPropertySet::FromAccessor<bool, &FMaterialResource::HasNormalConnected>(TEXT("HasNormalConnected")),
+		FPropertySet::FromAccessor<bool, &FMaterialResource::HasRoughnessConnected>(TEXT("HasRoughnessConnected")),
+		FPropertySet::FromAccessor<bool, &FMaterialResource::HasSpecularConnected>(TEXT("HasSpecularConnected")),
+		FPropertySet::FromAccessor<bool, &FMaterialResource::HasMetallicConnected>(TEXT("HasMetallicConnected")),
+		FPropertySet::FromAccessor<bool, &FMaterialResource::HasEmissiveColorConnected>(TEXT("HasEmissiveColorConnected")),
+		FPropertySet::FromAccessor<bool, &FMaterialResource::HasAnisotropyConnected>(TEXT("HasAnisotropyConnected")),
+		FPropertySet::FromAccessor<bool, &FMaterialResource::HasAmbientOcclusionConnected>(TEXT("HasAmbientOcclusionConnected")),
+		FPropertySet::FromAccessor<bool, &FMaterialResource::HasDisplacementConnected>(TEXT("HasDisplacementConnected")),
+		FPropertySet::FromAccessor<bool, &FMaterialResource::IsSubstrateMaterial>(TEXT("IsSubstrateMaterial")),
+		//FPropertySet::FromAccessor<bool, &FMaterialResource::HasMaterialPropertyConnected>(TEXT("HasMaterialPropertyConnected")),
+		//FPropertySet::FromAccessor<FMaterialShadingModelField, &FMaterialResource::GetShadingModels>(TEXT("GetShadingModels")),
+		FPropertySet::FromAccessor<bool, &FMaterialResource::IsShadingModelFromMaterialExpression>(TEXT("IsShadingModelFromMaterialExpression")),
+		FPropertySet::FromEnumAccessor<ETranslucencyLightingMode, &FMaterialResource::GetTranslucencyLightingMode>(TEXT("GetTranslucencyLightingMode")),
+		FPropertySet::FromAccessor<float, &FMaterialResource::GetOpacityMaskClipValue>(TEXT("GetOpacityMaskClipValue")),
+		FPropertySet::FromAccessor<bool, &FMaterialResource::GetCastDynamicShadowAsMasked>(TEXT("GetCastDynamicShadowAsMasked")),
+		FPropertySet::FromAccessor<bool, &FMaterialResource::IsDistorted>(TEXT("IsDistorted")),
+		FPropertySet::FromEnumAccessor<ERefractionCoverageMode, &FMaterialResource::GetRefractionCoverageMode>(TEXT("GetRefractionCoverageMode")),
+		FPropertySet::FromAccessor<float, &FMaterialResource::GetTranslucencyDirectionalLightingIntensity>(TEXT("GetTranslucencyDirectionalLightingIntensity")),
+		FPropertySet::FromAccessor<float, &FMaterialResource::GetTranslucentShadowDensityScale>(TEXT("GetTranslucentShadowDensityScale")),
+		FPropertySet::FromAccessor<float, &FMaterialResource::GetTranslucentSelfShadowDensityScale>(TEXT("GetTranslucentSelfShadowDensityScale")),
+		FPropertySet::FromAccessor<float, &FMaterialResource::GetTranslucentSelfShadowSecondDensityScale>(TEXT("GetTranslucentSelfShadowSecondDensityScale")),
+		FPropertySet::FromAccessor<float, &FMaterialResource::GetTranslucentSelfShadowSecondOpacity>(TEXT("GetTranslucentSelfShadowSecondOpacity")),
+		FPropertySet::FromAccessor<float, &FMaterialResource::GetTranslucentBackscatteringExponent>(TEXT("GetTranslucentBackscatteringExponent")),
+		FPropertySet::FromAccessor<bool, &FMaterialResource::IsTranslucencyAfterDOFEnabled>(TEXT("IsTranslucencyAfterDOFEnabled")),
+		FPropertySet::FromAccessor<bool, &FMaterialResource::IsTranslucencyAfterMotionBlurEnabled>(TEXT("IsTranslucencyAfterMotionBlurEnabled")),
+		//FPropertySet::FromAccessor<bool, &FMaterialResource::IsDualBlendingEnabled>(TEXT("IsDualBlendingEnabled")),
+		FPropertySet::FromAccessor<bool, &FMaterialResource::IsMobileSeparateTranslucencyEnabled>(TEXT("IsMobileSeparateTranslucencyEnabled")),
+		//FPropertySet::FromAccessor<FDisplacementScaling, &FMaterialResource::GetDisplacementScaling>(TEXT("GetDisplacementScaling")),
+		FPropertySet::FromAccessor<bool, &FMaterialResource::IsDisplacementFadeEnabled>(TEXT("IsDisplacementFadeEnabled")),
+		//FPropertySet::FromAccessor<FDisplacementFadeRange, &FMaterialResource::GetDisplacementFadeRange>(TEXT("GetDisplacementFadeRange")),
+		//FPropertySet::FromAccessor<FLinearColor, &FMaterialResource::GetTranslucentMultipleScatteringExtinction>(TEXT("GetTranslucentMultipleScatteringExtinction")),
+		FPropertySet::FromAccessor<float, &FMaterialResource::GetTranslucentShadowStartOffset>(TEXT("GetTranslucentShadowStartOffset")),
+		FPropertySet::FromAccessor<bool, &FMaterialResource::IsMasked>(TEXT("IsMasked")),
+		FPropertySet::FromAccessor<bool, &FMaterialResource::IsDitherMasked>(TEXT("IsDitherMasked")),
+		FPropertySet::FromAccessor<bool, &FMaterialResource::AllowNegativeEmissiveColor>(TEXT("AllowNegativeEmissiveColor")),
+		//FPropertySet::FromAccessor<FString, &FMaterialResource::GetFriendlyName>(TEXT("GetFriendlyName")),
+		FPropertySet::FromAccessor<FString, &FMaterialResource::GetAssetName>(TEXT("GetAssetName")),
+		FPropertySet::FromAccessor<bool, &FMaterialResource::RequiresSynchronousCompilation>(TEXT("RequiresSynchronousCompilation")),
+		FPropertySet::FromAccessor<bool, &FMaterialResource::IsDefaultMaterial>(TEXT("IsDefaultMaterial")),
+		FPropertySet::FromAccessor<int32, &FMaterialResource::GetNumCustomizedUVs>(TEXT("GetNumCustomizedUVs")),
+		FPropertySet::FromAccessor<int32, &FMaterialResource::GetBlendableLocation>(TEXT("GetBlendableLocation")),
+		FPropertySet::FromAccessor<bool, &FMaterialResource::GetBlendableOutputAlpha>(TEXT("GetBlendableOutputAlpha")),
+		FPropertySet::FromAccessor<bool, &FMaterialResource::IsStencilTestEnabled>(TEXT("IsStencilTestEnabled")),
+		FPropertySet::FromAccessor<uint32, &FMaterialResource::GetStencilRefValue>(TEXT("GetStencilRefValue")),
+		FPropertySet::FromAccessor<uint32, &FMaterialResource::GetStencilCompare>(TEXT("GetStencilCompare")),
+		FPropertySet::FromAccessor<float, &FMaterialResource::GetRefractionDepthBiasValue>(TEXT("GetRefractionDepthBiasValue")),
+		FPropertySet::FromAccessor<bool, &FMaterialResource::ShouldApplyFogging>(TEXT("ShouldApplyFogging")),
+		FPropertySet::FromAccessor<bool, &FMaterialResource::ShouldApplyCloudFogging>(TEXT("ShouldApplyCloudFogging")),
+		FPropertySet::FromAccessor<bool, &FMaterialResource::ShouldAlwaysEvaluateWorldPositionOffset>(TEXT("ShouldAlwaysEvaluateWorldPositionOffset")),
+		FPropertySet::FromAccessor<bool, &FMaterialResource::IsSky>(TEXT("IsSky")),
+		FPropertySet::FromAccessor<bool, &FMaterialResource::ComputeFogPerPixel>(TEXT("ComputeFogPerPixel")),
+		FPropertySet::FromAccessor<bool, &FMaterialResource::HasPerInstanceCustomData>(TEXT("HasPerInstanceCustomData")),
+		FPropertySet::FromAccessor<bool, &FMaterialResource::HasPerInstanceRandom>(TEXT("HasPerInstanceRandom")),
+		FPropertySet::FromAccessor<bool, &FMaterialResource::HasVertexInterpolator>(TEXT("HasVertexInterpolator")),
+		FPropertySet::FromAccessor<bool, &FMaterialResource::HasRuntimeVirtualTextureOutput>(TEXT("HasRuntimeVirtualTextureOutput")),
+		FPropertySet::FromAccessor<bool, &FMaterialResource::CastsRayTracedShadows>(TEXT("CastsRayTracedShadows")),
+		FPropertySet::FromAccessor<bool, &FMaterialResource::IsTessellationEnabled>(TEXT("IsTessellationEnabled")),
+		FPropertySet::FromAccessor<bool, &FMaterialResource::HasRenderTracePhysicalMaterialOutputs>(TEXT("HasRenderTracePhysicalMaterialOutputs")),
+		FPropertySet::FromAccessor<uint16, &FMaterialResource::GetPreshaderGap>(TEXT("GetPreshaderGap")),
+		FPropertySet::FromAccessor<int32, &FMaterialResource::GetNeuralProfileId>(TEXT("GetNeuralProfileId")),
+
+		FPropertySet
+		{
+			{
+				TEXT("ShadingModelMask"),
+				TEXT("SubstrateUintPerPixel"),
+				TEXT("SubstrateClosureCountMask"),
+				TEXT("bUsesComplexSpecialRenderPath"),
+				TEXT("bOpaque"),
+				TEXT("bMasked"),
+				TEXT("bDistortion"),
+				TEXT("bHairStrands"),
+				TEXT("bTwoSided"),
+				TEXT("bSeparateTranslucency"),
+				TEXT("bTranslucencyModulate"),
+				TEXT("bPostMotionBlurTranslucency"),
+				TEXT("bNormalTranslucency"),
+				TEXT("bUsesSceneColorCopy"),
+				TEXT("bOutputsTranslucentVelocity"),
+				TEXT("bUsesGlobalDistanceField"),
+				TEXT("bUsesWorldPositionOffset"),
+				TEXT("bUsesDisplacement"),
+				TEXT("bUsesPixelDepthOffset"),
+				TEXT("bDecal"),
+				TEXT("bTranslucentSurfaceLighting"),
+				TEXT("bUsesSceneDepth"),
+				TEXT("bUsesSkyMaterial"),
+				TEXT("bUsesSingleLayerWaterMaterial"),
+				TEXT("bHasVolumeMaterialDomain"),
+				TEXT("CustomDepthStencilUsageMask"),
+				TEXT("bUsesDistanceCullFade"),
+				TEXT("bDisableDepthTest"),
+				TEXT("bUsesAnisotropy"),
+				TEXT("bIsLightFunctionAtlasCompatible"),
+			},
+			[](const FPropertyDumpInput& Input, FOutput& Output)
+			{
+				Output.Write(Input.MaterialRelevance.ShadingModelMask);
+				Output.Write(Input.MaterialRelevance.SubstrateUintPerPixel);
+				Output.Write(Input.MaterialRelevance.SubstrateClosureCountMask);
+				Output.Write(Input.MaterialRelevance.bUsesComplexSpecialRenderPath);
+				Output.Write(Input.MaterialRelevance.bOpaque);
+				Output.Write(Input.MaterialRelevance.bMasked);
+				Output.Write(Input.MaterialRelevance.bDistortion);
+				Output.Write(Input.MaterialRelevance.bHairStrands);
+				Output.Write(Input.MaterialRelevance.bTwoSided);
+				Output.Write(Input.MaterialRelevance.bSeparateTranslucency);
+				Output.Write(Input.MaterialRelevance.bTranslucencyModulate);
+				Output.Write(Input.MaterialRelevance.bPostMotionBlurTranslucency);
+				Output.Write(Input.MaterialRelevance.bNormalTranslucency);
+				Output.Write(Input.MaterialRelevance.bUsesSceneColorCopy);
+				Output.Write(Input.MaterialRelevance.bOutputsTranslucentVelocity);
+				Output.Write(Input.MaterialRelevance.bUsesGlobalDistanceField);
+				Output.Write(Input.MaterialRelevance.bUsesWorldPositionOffset);
+				Output.Write(Input.MaterialRelevance.bUsesDisplacement);
+				Output.Write(Input.MaterialRelevance.bUsesPixelDepthOffset);
+				Output.Write(Input.MaterialRelevance.bDecal);
+				Output.Write(Input.MaterialRelevance.bTranslucentSurfaceLighting);
+				Output.Write(Input.MaterialRelevance.bUsesSceneDepth);
+				Output.Write(Input.MaterialRelevance.bUsesSkyMaterial);
+				Output.Write(Input.MaterialRelevance.bUsesSingleLayerWaterMaterial);
+				Output.Write(Input.MaterialRelevance.bHasVolumeMaterialDomain);
+				Output.Write(Input.MaterialRelevance.CustomDepthStencilUsageMask);
+				Output.Write(Input.MaterialRelevance.bUsesDistanceCullFade);
+				Output.Write(Input.MaterialRelevance.bDisableDepthTest);
+				Output.Write(Input.MaterialRelevance.bUsesAnisotropy);
+				Output.Write(Input.MaterialRelevance.bIsLightFunctionAtlasCompatible);
+			} 
+		},
+	};
+	return PropertySets;
+#else
+	return {};
+#endif
+}
+
+int32 UDumpMaterialInfoCommandlet::Main(const FString& Params)
+{
+	TArray<FString> Tokens;
+	TArray<FString> Switches;
+	TMap<FString, FString> ParamVals;
+	UCommandlet::ParseCommandLine(*Params, Tokens, Switches, ParamVals);
+
+	// Display help
+	if (Switches.Contains("help"))
+	{
+		UE_LOG(LogDumpMaterialInfo, Log, TEXT("DumpMaterialInfo"));
+		UE_LOG(LogDumpMaterialInfo, Log, TEXT("This commandlet will dump information about materials."));
+		UE_LOG(LogDumpMaterialInfo, Log, TEXT("A typical way to invoke it is: <YourProject> -run=DumpMaterialInfo -targetplatform=Windows -unattended -sm6 -allowcommandletrendering -nomaterialshaderddc."));
+		UE_LOG(LogDumpMaterialInfo, Log, TEXT(""));
+		UE_LOG(LogDumpMaterialInfo, Log, TEXT("Options:"));
+		UE_LOG(LogDumpMaterialInfo, Log, TEXT(" -help           Print this message"));
+		UE_LOG(LogDumpMaterialInfo, Log, TEXT(" -help=columns   Print the list of available columns"));
+		UE_LOG(LogDumpMaterialInfo, Log, TEXT(" -csv=filename   Writes the output to a CSV file"));
+		UE_LOG(LogDumpMaterialInfo, Log, TEXT(" -material=name  Only dump materials matching this material name or regular expression"));
+		UE_LOG(LogDumpMaterialInfo, Log, TEXT(" -columns=a,b    Comma-seperated list of the columns that should be included in the output"));
+		return 0;
+	}
+
+	FString* Help = ParamVals.Find(TEXT("help"));
+	if (Help)
+	{
+		if ((*Help) == TEXT("columns"))
+		{
+			UE_LOG(LogDumpMaterialInfo, Display, TEXT("Available columns:"), **Help);
+			for (const MaterialInfo::FPropertySet& PropertySet : GetMaterialInfoProperties())
+			{
+				for (const FString& PropertyName : PropertySet.PropertyNames)
+				{
+					UE_LOG(LogDumpMaterialInfo, Display, TEXT("   %s"), *PropertyName);
+				}
+			}
+		}
+		else
+		{
+			UE_LOG(LogDumpMaterialInfo, Error, TEXT("Unknown help option %s"), **Help);
+			return 1;
+		}
+		return 0;
+	}
+
+	// Parse params
+	FString* CsvPath = ParamVals.Find(TEXT("csv"));
+	if (!CsvPath)
+	{
+		UE_LOG(LogDumpMaterialInfo, Error, TEXT("No output CSV file path was specified"));
+		return 1;
+	}
+	FString* RequestedMaterialPatternString = ParamVals.Find(TEXT("material"));
+	bool bMatchAllMaterials = !RequestedMaterialPatternString;
+	FRegexPattern RequestedMaterialPattern(RequestedMaterialPatternString ? *RequestedMaterialPatternString : FString());
+	FString* ColumnsString = ParamVals.Find(TEXT("columns"));
+	TSet<FString> Columns {};
+	if (ColumnsString)
+	{
+		UE::String::ParseTokens(*ColumnsString, TEXT(","), [&Columns](const auto& SubString)
+		{
+			Columns.Add(FString(SubString));
+		}, UE::String::EParseTokensOptions::SkipEmpty | UE::String::EParseTokensOptions::Trim);
+	}
+
+	ERHIFeatureLevel::Type FeatureLevel = ERHIFeatureLevel::SM6;
+	EMaterialQualityLevel::Type MaterialQualityLevel = EMaterialQualityLevel::High;
+	EShaderPlatform ShaderPlatform = SP_PCD3D_SM6;
+
+	// Get available material properties and filter them
+	TArray<MaterialInfo::FPropertySet> MaterialInfoProperties = GetMaterialInfoProperties();
+	if (!Columns.IsEmpty())
+	{
+		// Retain property sets that have at least 1 requested property
+		MaterialInfoProperties = MaterialInfoProperties.FilterByPredicate([&Columns](const MaterialInfo::FPropertySet& PropertySet)
+		{
+			return PropertySet.PropertyNames.ContainsByPredicate([&Columns](const FString& PropertyName)
+			{
+				return Columns.Contains(PropertyName);
+			});
+		});
+	}
+
+	UE_LOG(LogDumpMaterialInfo, Log, TEXT("Searching for materials within the project..."));
+
+	FAssetRegistryModule& AssetRegistryModule = FModuleManager::LoadModuleChecked<FAssetRegistryModule>("AssetRegistry");
+	IAssetRegistry& AssetRegistry = AssetRegistryModule.Get();
+	AssetRegistry.SearchAllAssets(true);
+
+	TArray<FAssetData> MaterialInterfaceAssets;
+	{
+		TArray<FAssetData> MaterialAssets;
+		AssetRegistry.GetAssetsByClass(UMaterial::StaticClass()->GetClassPathName(), MaterialAssets, true);
+		UE_LOG(LogDumpMaterialInfo, Log, TEXT("Found %d materials"), MaterialAssets.Num());
+		MaterialInterfaceAssets = MaterialAssets;
+	}
+	{
+		TArray<FAssetData> MaterialInstanceAssets;
+		AssetRegistry.GetAssetsByClass(UMaterialInstance::StaticClass()->GetClassPathName(), MaterialInstanceAssets, true);
+		UE_LOG(LogDumpMaterialInfo, Log, TEXT("Found %d material instances"), MaterialInstanceAssets.Num());
+		MaterialInterfaceAssets.Append(MaterialInstanceAssets);
+	}
+
+	ITargetPlatformManagerModule* TPM = GetTargetPlatformManager();
+	const TArray<ITargetPlatform*>& Platforms = TPM->GetActiveTargetPlatforms();
+
+	FArchive* CsvFileWriter = IFileManager::Get().CreateFileWriter(**CsvPath);
+	if (!CsvFileWriter)
+	{
+		UE_LOG(LogDumpMaterialInfo, Error, TEXT("Failed to open output file %s"), **CsvPath);
+		return 1;
+	}
+	FDiagnosticTableWriterCSV CsvWriter(CsvFileWriter);
+
+	for (ITargetPlatform* Platform : Platforms)
+	{
+		UE_LOG(LogDumpMaterialInfo, Display, TEXT("Compiling shaders for %s..."), *Platform->PlatformName());
+
+		TSet<UMaterialInterface*> MaterialsToCompile;
+		for (const FAssetData& AssetData : MaterialInterfaceAssets)
+		{
+			bool bInclude = (bMatchAllMaterials || FRegexMatcher(RequestedMaterialPattern, AssetData.GetFullName()).FindNext());
+			if (!bInclude)
+			{
+				continue;
+			}
+			if (UMaterialInterface* MaterialInterface = Cast<UMaterialInterface>(AssetData.GetAsset()))
+			{
+				UMaterial* Material = MaterialInterface->GetMaterial();
+				if (Material)
+				{
+					UE_LOG(LogDumpMaterialInfo, Display, TEXT("BeginCache for %s"), *MaterialInterface->GetFullName());
+					MaterialInterface->BeginCacheForCookedPlatformData(Platform);
+					// need to call this once for all objects before any calls to ProcessAsyncResults as otherwise we'll potentially upload
+					// incremental/incomplete shadermaps to DDC (as this function actually triggers compilation, some compiles for a particular
+					// material may finish before we've even started others - if we call ProcessAsyncResults in that case the associated shader
+					// maps will think they are "finished" due to having no outstanding dependencies).
+					if (!MaterialInterface->IsCachedCookedPlatformDataLoaded(Platform))
+					{
+						MaterialsToCompile.Add(MaterialInterface);
+					}
+				}
+			}
+		}
+		TSet<UMaterialInterface*> MaterialsToAnalyse = MaterialsToCompile;
+
+		UE_LOG(LogDumpMaterialInfo, Log, TEXT("Found %d materials to compile."), MaterialsToCompile.Num());
+
+		static constexpr bool bLimitExecutationTime = false;
+		int32 PreviousOutstandingJobs = 0;
+		constexpr int32 MaxOutstandingJobs = 20000; // Having a max is a way to try to reduce memory usage.. otherwise outstanding jobs can reach 100k+ and use up 300gb committed memory
+		// Submit all the jobs.
+		{
+			TRACE_CPUPROFILER_EVENT_SCOPE(SubmitJobs);
+
+			UE_LOG(LogDumpMaterialInfo, Display, TEXT("Submit Jobs"));
+
+			while (MaterialsToCompile.Num())
+			{
+				for (auto It = MaterialsToCompile.CreateIterator(); It; ++It)
+				{
+					UMaterialInterface* MaterialInterface = *It;
+					if (MaterialInterface->IsCachedCookedPlatformDataLoaded(Platform))
+					{
+						It.RemoveCurrent();
+						UE_LOG(LogDumpMaterialInfo, Display, TEXT("Finished cache for %s."), *MaterialInterface->GetFullName());
+						UE_LOG(LogDumpMaterialInfo, Display, TEXT("Materials remaining: %d"), MaterialsToCompile.Num());
+					}
+
+					GShaderCompilingManager->ProcessAsyncResults(bLimitExecutationTime, false /* bBlockOnGlobalShaderCompilation */);
+
+					while (true)
+					{
+						const int32 CurrentOutstandingJobs = GShaderCompilingManager->GetNumOutstandingJobs();
+						if (CurrentOutstandingJobs != PreviousOutstandingJobs)
+						{
+							UE_LOG(LogDumpMaterialInfo, Display, TEXT("Outstanding Jobs: %d"), CurrentOutstandingJobs);
+							PreviousOutstandingJobs = CurrentOutstandingJobs;
+						}
+
+						// Flush rendering commands to release any RHI resources (shaders and shader maps).
+						// Delete any FPendingCleanupObjects (shader maps).
+						FlushRenderingCommands();
+
+						if (CurrentOutstandingJobs < MaxOutstandingJobs)
+						{
+							break;
+						}
+						FPlatformProcess::Sleep(1);
+					}
+				}
+			}
+		}
+
+		// Process the shader maps and save to the DDC.
+		{
+			TRACE_CPUPROFILER_EVENT_SCOPE(ProcessShaderCompileResults);
+
+			UE_LOG(LogDumpMaterialInfo, Log, TEXT("ProcessAsyncResults"));
+
+			while (GShaderCompilingManager->IsCompiling())
+			{
+				GShaderCompilingManager->ProcessAsyncResults(bLimitExecutationTime, false /* bBlockOnGlobalShaderCompilation */);
+
+				while (true)
+				{
+					const int32 CurrentOutstandingJobs = GShaderCompilingManager->GetNumOutstandingJobs();
+					if (CurrentOutstandingJobs != PreviousOutstandingJobs)
+					{
+						UE_LOG(LogDumpMaterialInfo, Display, TEXT("Outstanding Jobs: %d"), CurrentOutstandingJobs);
+						PreviousOutstandingJobs = CurrentOutstandingJobs;
+					}
+
+					// Flush rendering commands to release any RHI resources (shaders and shader maps).
+					// Delete any FPendingCleanupObjects (shader maps).
+					FlushRenderingCommands();
+
+					if (CurrentOutstandingJobs < MaxOutstandingJobs)
+					{
+						break;
+					}
+					FPlatformProcess::Sleep(1);
+				}
+			}
+		}
+
+		// CSV header
+		{
+			for (const MaterialInfo::FPropertySet& Property : MaterialInfoProperties)
+			{
+				for (const FString& PropertyName : Property.PropertyNames)
+				{
+					if (Columns.IsEmpty() || Columns.Contains(PropertyName))
+					{
+						CsvWriter.AddColumn(TEXT("%s"), *PropertyName);
+					}
+				}
+			}
+			CsvWriter.CycleRow();
+		}
+
+		// Look up compilation result for the materials
+		TArray<const FVertexFactoryType*> VFTypes;
+		TArray<const FShaderPipelineType*> PipelineTypes;
+		TArray<const FShaderType*> ShaderTypes;
+		for(UMaterialInterface* MaterialInterface : MaterialsToAnalyse)
+		{
+			VFTypes.Empty();
+			PipelineTypes.Empty();
+			ShaderTypes.Empty();
+
+			UMaterial* Material = MaterialInterface->GetMaterial();
+			if (Material)
+			{
+				TArray<FMaterialResource*> ResourcesToCache;
+
+				FMaterialResource* CurrentResource = FindOrCreateMaterialResource(ResourcesToCache, Material, nullptr, FeatureLevel, MaterialQualityLevel);
+				check(CurrentResource);
+
+				TMap<FName, TArray<FMaterialStatsUtils::FRepresentativeShaderInfo>> ShaderTypeNamesAndDescriptions;
+				FMaterialStatsUtils::GetRepresentativeShaderTypesAndDescriptions(ShaderTypeNamesAndDescriptions, CurrentResource);
+
+				for (auto& DescriptionPair : ShaderTypeNamesAndDescriptions)
+				{
+					const FVertexFactoryType* VFType = FindVertexFactoryType(DescriptionPair.Key);
+					check(VFType);
+
+					auto& DescriptionArray = DescriptionPair.Value;
+					for (const FMaterialStatsUtils::FRepresentativeShaderInfo& ShaderInfo : DescriptionArray)
+					{
+						const FShaderType* ShaderType = FindShaderTypeByName(ShaderInfo.ShaderName);
+
+						if (ShaderType && VFType)
+						{
+							VFTypes.Add(VFType);
+							ShaderTypes.Add(ShaderType);
+							PipelineTypes.Add(nullptr);
+						}
+					}
+				}
+
+				// Prepare the resource for compilation, but don't compile the completed shader map.
+				const bool bSuccess = CurrentResource->CacheShaders(ShaderPlatform, EMaterialShaderPrecompileMode::None);
+
+				if (bSuccess)
+				{
+					// Compile just the types we want.
+					CurrentResource->CacheGivenTypes(ShaderPlatform, VFTypes, PipelineTypes, ShaderTypes);
+				}
+				
+				if(!CurrentResource->IsGameThreadShaderMapComplete()) { UE_LOG(LogDumpMaterialInfo, Warning, TEXT("Missing shader map data")); }
+
+				FMaterialRelevance MaterialRelevance = CurrentResource->GetMaterialInterface()->GetRelevance(FeatureLevel);
+
+				// CSV line
+				{
+					MaterialInfo::FOutput Output;
+					for (const MaterialInfo::FPropertySet& Property : MaterialInfoProperties)
+					{
+						MaterialInfo::FPropertyDumpInput Input;
+						Input.Material = Material;
+						Input.MaterialResource = CurrentResource;
+						Input.MaterialRelevance = MaterialRelevance;
+						Input.ShaderTypes = &ShaderTypes;
+						Output.DumpPropertySet(Property, Input);
+						for (const MaterialInfo::FPropertyValue& Value : Output.GetValues())
+						{
+							if (Columns.IsEmpty() || Columns.Contains(Value.Name))
+							{
+								CsvWriter.AddColumn(TEXT("%s"), *FString(Value.Value));
+							}
+						}
+						Output.Reset();
+					}
+					CsvWriter.CycleRow();
+				}
+
+				FMaterial::DeferredDeleteArray(ResourcesToCache);
+			}
+		}
+
+		// Perform cleanup and clear cached data for cooking.
+		{
+			TRACE_CPUPROFILER_EVENT_SCOPE(ClearCachedCookedPlatformData);
+
+			UE_LOG(LogDumpMaterialInfo, Display, TEXT("Clear Cached Cooked Platform Data"));
+
+			for (const FAssetData& AssetData : MaterialInterfaceAssets)
+			{
+				if (UMaterialInterface* MaterialInterface = Cast<UMaterialInterface>(AssetData.GetAsset()))
+				{
+					MaterialInterface->ClearAllCachedCookedPlatformData();
+				}
+			}
+		}
+	} // Platforms
+
+	return 0;
+}
diff --git a/Engine/Source/Runtime/Engine/Private/Materials/MaterialShared.cpp b/Engine/Source/Runtime/Engine/Private/Materials/MaterialShared.cpp
--- a/Engine/Source/Runtime/Engine/Private/Materials/MaterialShared.cpp
+++ b/Engine/Source/Runtime/Engine/Private/Materials/MaterialShared.cpp
@@ -4146,11 +4146,11 @@ uint32 FMaterialResource::GetEstimatedNumVirtualTextureLookups() const
 	return 0;
 }
 
-void FMaterialResource::GetEstimatedLWCFuncUsages(FLWCUsagesArray& UsagesVS, FLWCUsagesArray& UsagesPS) const
+void FMaterialResource::GetEstimatedLWCFuncUsages(FLWCUsagesArray& UsagesVS, FLWCUsagesArray& UsagesPS, FLWCUsagesArray& UsagesCS) const
 {
 	if (const FMaterialShaderMap* ShaderMap = GetGameThreadShaderMap())
 	{
-		ShaderMap->GetEstimatedLWCFuncUsages(UsagesVS, UsagesPS);
+		ShaderMap->GetEstimatedLWCFuncUsages(UsagesVS, UsagesPS, UsagesCS);
 	}
 }
 
diff --git a/Engine/Source/Runtime/Engine/Public/MaterialShared.h b/Engine/Source/Runtime/Engine/Public/MaterialShared.h
--- a/Engine/Source/Runtime/Engine/Public/MaterialShared.h
+++ b/Engine/Source/Runtime/Engine/Public/MaterialShared.h
@@ -1550,10 +1550,12 @@ public:
 	uint32 GetNumUsedCustomInterpolatorScalars() const { return GetContent()->MaterialCompilationOutput.NumUsedCustomInterpolatorScalars; }
 	void GetEstimatedNumTextureSamples(uint32& VSSamples, uint32& PSSamples) const { VSSamples = GetContent()->MaterialCompilationOutput.EstimatedNumTextureSamplesVS; PSSamples = GetContent()->MaterialCompilationOutput.EstimatedNumTextureSamplesPS; }
 	uint32 GetEstimatedNumVirtualTextureLookups() const { return GetContent()->MaterialCompilationOutput.EstimatedNumVirtualTextureLookups; }
-	void GetEstimatedLWCFuncUsages(TStaticArray<uint16, (int)ELWCFunctionKind::Max>& UsagesVS, TStaticArray<uint16, (int)ELWCFunctionKind::Max>& UsagesPS) const
+	using FLWCUsagesArray = TStaticArray<uint16, (int)ELWCFunctionKind::Max>;
+	void GetEstimatedLWCFuncUsages(FLWCUsagesArray& UsagesVS, FLWCUsagesArray& UsagesPS, FLWCUsagesArray& UsagesCS) const
 	{
 		CopyAssignItems(UsagesVS.GetData(), GetContent()->MaterialCompilationOutput.EstimatedLWCFuncUsagesVS, (int)ELWCFunctionKind::Max);
 		CopyAssignItems(UsagesPS.GetData(), GetContent()->MaterialCompilationOutput.EstimatedLWCFuncUsagesPS, (int)ELWCFunctionKind::Max);
+		CopyAssignItems(UsagesCS.GetData(), GetContent()->MaterialCompilationOutput.EstimatedLWCFuncUsagesCS, (int)ELWCFunctionKind::Max);
 	}
 	ENGINE_API void GetEstimatedLWCFuncUsageComplexity(uint32& LWCComplexityVS, uint32& LWCComplexityPS, uint32& LWCComplexityCS) const;
 	uint32 GetNumPreshaders() const { return GetContent()->MaterialCompilationOutput.UniformExpressionSet.UniformPreshaders.Num(); }
@@ -2762,7 +2764,7 @@ public:
 	ENGINE_API void GetEstimatedNumTextureSamples(uint32& VSSamples, uint32& PSSamples) const;
 	ENGINE_API uint32 GetEstimatedNumVirtualTextureLookups() const;
 	using FLWCUsagesArray = TStaticArray<uint16, (int)ELWCFunctionKind::Max>;
-	ENGINE_API void GetEstimatedLWCFuncUsages(FLWCUsagesArray& UsagesVS, FLWCUsagesArray& UsagesPS) const;
+	ENGINE_API void GetEstimatedLWCFuncUsages(FLWCUsagesArray& UsagesVS, FLWCUsagesArray& UsagesPS, FLWCUsagesArray& UsagesCS) const;
 #endif
 	ENGINE_API uint32 GetNumVirtualTextureStacks() const;
 
