MPCook: CookPackageSplitter: Fixes for enabling retraction and for iterative cook of the new CookPackageSplitter code.
* Remove the garbage collection when saving the generator if the generated packages already exist. We no longer need that behavior because the API of CookPackageSplitter has changed to not guarantee empty generated packages during generator save, and the code complicates correctness proofs.
* Change the modification of the generator save state when demoted to return to StartSave instead of StartPopulate for a similar reason - calling QueueGeneratedPackages twice is not a problem, and maintaining the separate path complicates correctness proofs.
* Fix loss of bHasSaved data for generated packages when GenerationInfo downgrades to the uninitialized state on the cook director but saves are still ongoing on cookworkers; loss of that data prevented the CompletedAllSaves message from being sent.
* Fix incorrect assertions that some data was reset when leaving the save state.
* Add feedback to the assertion that we have left the save state before calling uninitilize.
#rnx

diff --git a/Engine/Source/Editor/UnrealEd/Private/CookOnTheFlyServer.cpp b/Engine/Source/Editor/UnrealEd/Private/CookOnTheFlyServer.cpp
--- a/Engine/Source/Editor/UnrealEd/Private/CookOnTheFlyServer.cpp
+++ b/Engine/Source/Editor/UnrealEd/Private/CookOnTheFlyServer.cpp
@@ -3311,32 +3311,6 @@ UE::Cook::EPollStatus UCookOnTheFlyServer::QueueGeneratedPackages(UE::Cook::FGen
 
 	FCookGenerationInfo& Info = GenerationHelper.GetOwnerInfo();
 
-	if (Info.GetSaveState() <= FCookGenerationInfo::ESaveState::ClearOldPackagesLastAttempt)
-	{
-		for (const FCookGenerationInfo& ChildInfo: GenerationHelper.GetPackagesToGenerate())
-		{
-			const FString GeneratedPackageName = ChildInfo.PackageData->GetPackageName().ToString();
-			if (FindObject<UPackage>(nullptr, *GeneratedPackageName))
-			{
-				if (Info.GetSaveState() < FCookGenerationInfo::ESaveState::ClearOldPackagesLastAttempt)
-				{
-					PackageData.SetIsPrepareSaveRequiresGC(true);
-					Info.SetSaveState(FCookGenerationInfo::ESaveState::ClearOldPackagesLastAttempt);
-					return EPollStatus::Incomplete;
-				}
-				else
-				{
-					UE_LOG(LogCook, Error,
-						TEXT("PackageSplitter was unable to construct new generated packages because an old version of the package is already in memory and GC did not remove it. Splitter=%s, Generated=%s."),
-						*GenerationHelper.GetSplitDataObjectName().ToString(), *ChildInfo.RelativePath);
-					return EPollStatus::Error;
-				}
-			}
-		}
-
-		Info.SetSaveStateComplete(FCookGenerationInfo::ESaveState::ClearOldPackagesLastAttempt);
-	}
-
 	UPackage* Owner = PackageData.GetPackage();
 	FName OwnerName = Owner->GetFName();
 	if (Info.GetSaveState() <= FCookGenerationInfo::ESaveState::QueueGeneratedPackages)
@@ -6733,14 +6707,14 @@ void UCookOnTheFlyServer::RecordExternalActorDependencies(TConstArrayView<FName>
 
 	if (IsCookWorkerMode())
 	{
-		// The dependencies will be replicated to the CookDirectory during ReportPromoteToSaveComplete
+		// The dependencies will be replicated to the CookDirector during ReportPromoteToSaveComplete
 		return;
 	}
 
 	// External actors are a special case in the cooker, and they are only referenced through the
-	// WorldPartitionCookPackageSplitter. They are marked as NeverCook, but we need to add them
-	// to the cook results so we can detect whether they change in iterative cooks. Call a function
-	// on the splitter to get a list of them and add them.
+	// WorldPartitionCookPackageSplitter. They are marked as NeverCook, but we need to add them to the cook results so
+	// we can detect whether they change in iterative cooks. The splitter has passed in its list of
+	// ExternalActorDependencies; add them to the list of cooked packages stored in the AssetRegistry.
 	for (FName DependencyName : ExternalActorDependencies)
 	{
 		FPackageData* DependencyData = PackageDatas->TryAddPackageDataByPackageName(DependencyName);
diff --git a/Engine/Source/Editor/UnrealEd/Private/Cooker/CookGenerationHelper.cpp b/Engine/Source/Editor/UnrealEd/Private/Cooker/CookGenerationHelper.cpp
--- a/Engine/Source/Editor/UnrealEd/Private/Cooker/CookGenerationHelper.cpp
+++ b/Engine/Source/Editor/UnrealEd/Private/Cooker/CookGenerationHelper.cpp
@@ -129,14 +129,14 @@ void FGenerationHelper::Uninitialize()
 	}
 	OwnerPackage.Reset();
 	// Keep PreviousGeneratedPackages; they are allowed in the uninitialized state
-	check(ExternalActorDependencies.IsEmpty()); // We can not still be in the save state, so this should be empty
+	ExternalActorDependencies.Empty();
 	check(OwnerObjectsToMove.IsEmpty()); // We can not still be in the save state, so this should be empty
 	// Do not modify the reference tracking variables
 	// ReferenceFromKeepForIterative
 	// ReferenceFromKeepForQueueResults
 	// ReferenceFromKeepForGeneratorSave
 	// ReferenceFromKeepForAllSavedOrGC
-	check(MPCookNextAssignmentIndex == 0); // We can not still be in the save state, so should have been cleared
+	MPCookNextAssignmentIndex = 0;
 	// Keep NumSaved; it is allowed in the uninitialized state
 	// Keep WorkerIdThatSavedGenerator; it is allowed in the uninitialized state
 	// InitializeStatus was modified above
@@ -887,28 +887,8 @@ void FGenerationHelper::ResetSaveState(FCookGenerationInfo& Info, UPackage* Pack
 		}
 	}
 
-	if (ReleaseSaveReason == EStateChangeReason::RecreateObjectCache ||
-		ReleaseSaveReason == EStateChangeReason::DoneForNow)
-	{
-		if (Info.IsGenerator())
-		{
-			if (Info.GetSaveState() >= FCookGenerationInfo::ESaveState::StartPopulate)
-			{
-				Info.SetSaveState(FCookGenerationInfo::ESaveState::StartPopulate);
-			}
-			else
-			{
-				// Redo all the steps since we didn't make it to the FinishCachePreObjectsToMove.
-				// Restarting in the middle of that flow is not robust
-				Info.SetSaveState(FCookGenerationInfo::ESaveState::StartSave);
-			}
-		}
-		else
-		{
-			Info.SetSaveState(FCookGenerationInfo::ESaveState::StartPopulate);
-		}
-	}
-	else
+	if (ReleaseSaveReason != EStateChangeReason::RecreateObjectCache &&
+		ReleaseSaveReason != EStateChangeReason::DoneForNow)
 	{
 		// The save is completed and we will not come back to it; set state back to initial
 		// state and drop our reference keeping this GenerationHelper in memory for the save.
@@ -923,14 +903,14 @@ void FGenerationHelper::ResetSaveState(FCookGenerationInfo& Info, UPackage* Pack
 
 			// And also teardown data needed during save
 			ClearKeepForGeneratorSave();
-			Info.SetSaveState(FCookGenerationInfo::ESaveState::StartSave);
 		}
 		else
 		{
-			Info.SetSaveState(FCookGenerationInfo::ESaveState::StartPopulate);
 			Info.PackageData->SetParentGenerationHelper(nullptr);
 		}
 	}
+	Info.SetSaveState(Info.IsGenerator() ? FCookGenerationInfo::ESaveState::StartSave
+		: FCookGenerationInfo::ESaveState::StartPopulate);
 
 	if (Info.HasTakenOverCachedCookedPlatformData())
 	{
@@ -1422,6 +1402,12 @@ FCookGenerationInfo::FCookGenerationInfo(FPackageData& InPackageData, bool bInGe
 
 void FCookGenerationInfo::Uninitialize()
 {
+	// Check that we have left the save state first, since other assertions assume we have left the save state
+	checkf(GeneratorSaveState == (bGenerator ? ESaveState::StartSave : ESaveState::StartPopulate),
+		TEXT("Cooker bug: Expected FCookGenerationInfo::Uninitialize to not be called for a package still in the save state, ")
+		TEXT("but %s package %s has SaveState %d."),
+		bGenerator ? TEXT("generator") : TEXT("generated"), *GetPackageName(), static_cast<int32>(GeneratorSaveState));
+
 	PackageHash.Reset();
 	RelativePath.Empty();
 	GeneratedRootPath.Empty();
@@ -1430,11 +1416,9 @@ void FCookGenerationInfo::Uninitialize()
 	// Keep PackageData; it is allowed in the uninitialized state
 	KeepReferencedPackages.Empty();
 	check(CachedObjectsInOuterInfo.IsEmpty()); // We can not still be in the save state, so this should be empty
-	// We can not still be in the save state, so this should have been cleared
-	check(GeneratorSaveState == (bGenerator ? ESaveState::StartSave : ESaveState::StartPopulate));
 	bCreateAsMap = false;
 	bHasCreatedPackage = false;
-	bHasSaved = false;
+	// Keep bHasSaved; it is allowed in the uninitialized state
 	bTakenOverCachedCookedPlatformData = false;
 	bIssuedUndeclaredMovedObjectsWarning = false;
 	// Keep bGenerator; it is allowed in the uninitialized state
diff --git a/Engine/Source/Editor/UnrealEd/Private/Cooker/CookGenerationHelper.h b/Engine/Source/Editor/UnrealEd/Private/Cooker/CookGenerationHelper.h
--- a/Engine/Source/Editor/UnrealEd/Private/Cooker/CookGenerationHelper.h
+++ b/Engine/Source/Editor/UnrealEd/Private/Cooker/CookGenerationHelper.h
@@ -63,9 +63,7 @@ public:
 	{
 		StartSave = 0,
 
-		ClearOldPackagesFirstAttempt = StartSave,
-		ClearOldPackagesLastAttempt,
-		QueueGeneratedPackages,
+		QueueGeneratedPackages = StartSave,
 
 		StartPopulate,
 		FinishCachePreMove = StartPopulate,
@@ -642,6 +640,7 @@ inline TConstArrayView<FName> FGenerationHelper::GetExternalActorDependencies()
 inline TArray<FName> FGenerationHelper::ReleaseExternalActorDependencies()
 {
 	TArray<FName> Result = MoveTemp(ExternalActorDependencies);
+	ExternalActorDependencies.Empty();
 	return Result;
 }
 
