Adding support for wireframe opacity adjustment in editor viewports.

diff --git a/Engine/Source/Editor/UnrealEd/Private/EditorViewportClient.cpp b/Engine/Source/Editor/UnrealEd/Private/EditorViewportClient.cpp
--- a/Engine/Source/Editor/UnrealEd/Private/EditorViewportClient.cpp
+++ b/Engine/Source/Editor/UnrealEd/Private/EditorViewportClient.cpp
@@ -21,6 +21,7 @@
 #include "Settings/LevelEditorMiscSettings.h"
 #include "Engine/DebugDisplayProperty.h"
 #include "Engine/RendererSettings.h"
+#include "MeshEdges.h"
 #include "Components/DirectionalLightComponent.h"
 #include "Components/BillboardComponent.h"
 #include "Audio/AudioDebug.h"
@@ -4075,6 +4076,9 @@ void FEditorViewportClient::SetupViewForRendering(FSceneViewFamily& ViewFamily,
 	View.CurrentRayTracingDebugVisualizationMode = CurrentRayTracingDebugVisualizationMode;
 #endif
 
+	// assign wireframe opacity to the view
+	GetMeshEdgesViewSettings(View).Opacity = WireframeOpacity;
+	
 	//Look if the pixel inspector tool is on
 	View.bUsePixelInspector = false;
 	FPixelInspectorModule& PixelInspectorModule = FModuleManager::LoadModuleChecked<FPixelInspectorModule>(TEXT("PixelInspectorModule"));
diff --git a/Engine/Source/Editor/UnrealEd/Private/SEditorViewport.cpp b/Engine/Source/Editor/UnrealEd/Private/SEditorViewport.cpp
--- a/Engine/Source/Editor/UnrealEd/Private/SEditorViewport.cpp
+++ b/Engine/Source/Editor/UnrealEd/Private/SEditorViewport.cpp
@@ -678,6 +678,33 @@ TSharedRef<SWidget> SEditorViewport::BuildFixedEV100Menu()  const
 		];
 };
 
+TSharedRef<SWidget> SEditorViewport::BuildWireframeMenu() const
+{
+	return 
+		SNew( SBox )
+		.HAlign( HAlign_Right )
+		[
+			SNew( SBox )
+			.Padding( FMargin(0.0f, 0.0f, 0.0f, 0.0f) )
+			.WidthOverride( 100.0f )
+			[
+				SNew ( SBorder )
+				.BorderImage(FAppStyle::Get().GetBrush("Menu.WidgetBorder"))
+				.Padding(FMargin(1.0f))
+				[
+					SNew(SSpinBox<float>)
+					.Style(&FAppStyle::Get(), "Menu.SpinBox")
+					.Font( FAppStyle::GetFontStyle( TEXT( "MenuItem.Font" ) ) )
+					.MinValue(0.f)
+					.MaxValue(1.f)
+					.SupportDynamicSliderMaxValue(false)
+					.Value( this, &SEditorViewport::OnGetWireframeOpacity )
+					.OnValueChanged( const_cast<SEditorViewport*>(this), &SEditorViewport::OnWireframeOpacityChanged )
+					.ToolTipText(LOCTEXT("WireframeOpacity_ToolTip", "Adjust opacity of wireframes in view."))
+				]
+			]
+		];
+};
 				
 void SEditorViewport::UpdateInViewportMenuLocation(const FVector2D InLocation)
 {
@@ -716,6 +743,24 @@ void SEditorViewport::OnFixedEV100ValueChanged(float NewValue)
 	}
 }
 
+void SEditorViewport::OnWireframeOpacityChanged(float Opacity)
+{
+	if( Client.IsValid() )
+	{
+		Client->WireframeOpacity = Opacity;
+		Client->Invalidate();
+	}
+}
+float SEditorViewport::OnGetWireframeOpacity() const
+{
+	if(Client.IsValid())
+	{
+		return Client->WireframeOpacity;
+	}
+	
+	return 0.8f;
+}
+
 bool SEditorViewport::IsWidgetModeActive( UE::Widget::EWidgetMode Mode ) const
 {
 	return Client->GetWidgetMode() == Mode;
diff --git a/Engine/Source/Editor/UnrealEd/Private/SEditorViewportViewMenu.cpp b/Engine/Source/Editor/UnrealEd/Private/SEditorViewportViewMenu.cpp
--- a/Engine/Source/Editor/UnrealEd/Private/SEditorViewportViewMenu.cpp
+++ b/Engine/Source/Editor/UnrealEd/Private/SEditorViewportViewMenu.cpp
@@ -333,6 +333,13 @@ void SEditorViewportViewMenu::FillViewMenu(UToolMenu* Menu) const
 			Section.AddMenuEntry(bIsLevelEditor ? BaseViewportCommands.ToggleInGameExposure : BaseViewportCommands.ToggleAutoExposure);
 			Section.AddEntry(FToolMenuEntry::InitWidget("FixedEV100", FixedEV100Menu, LOCTEXT("FixedEV100", "EV100")));
 		}
+
+		// Wireframe Opacity
+		{
+			TSharedRef<SWidget> WireOpacityMenu = Viewport.Pin()->BuildWireframeMenu();
+			FToolMenuSection& Section = Menu->AddSection("Wireframe", LOCTEXT("WireframeHeader", "Wireframe"));
+			Section.AddEntry(FToolMenuEntry::InitWidget("WireframeOpacity", WireOpacityMenu, LOCTEXT("WireframeOpacity", "Opacity")));
+		}
 	}
 }
 
diff --git a/Engine/Source/Editor/UnrealEd/Public/EditorViewportClient.h b/Engine/Source/Editor/UnrealEd/Public/EditorViewportClient.h
--- a/Engine/Source/Editor/UnrealEd/Public/EditorViewportClient.h
+++ b/Engine/Source/Editor/UnrealEd/Public/EditorViewportClient.h
@@ -1782,6 +1782,9 @@ public:
 	/** Editor setting to allow designers to override the automatic expose */
 	FExposureSettings		ExposureSettings;
 
+	/** Editor setting to adjust wireframe opacity */
+	float WireframeOpacity = 0.2f;
+
 	FName CurrentBufferVisualizationMode;
 	FName CurrentNaniteVisualizationMode;
 	FName CurrentLumenVisualizationMode;
diff --git a/Engine/Source/Editor/UnrealEd/Public/SEditorViewport.h b/Engine/Source/Editor/UnrealEd/Public/SEditorViewport.h
--- a/Engine/Source/Editor/UnrealEd/Public/SEditorViewport.h
+++ b/Engine/Source/Editor/UnrealEd/Public/SEditorViewport.h
@@ -96,6 +96,9 @@ public:
 
 	/** Build the exposure menu using EV100 settings */
 	UNREALED_API TSharedRef<SWidget> BuildFixedEV100Menu()  const;
+	
+	/** Build the slider to adjust wireframe opacity */
+	UNREALED_API TSharedRef<SWidget> BuildWireframeMenu()  const;
 
 	/**
  * Called when the user wants to show the in-viewport context menu
@@ -142,6 +145,9 @@ protected:
 	/** Called to know whether the fixed EV100 slider is enabled. */
 	UNREALED_API bool IsFixedEV100Enabled() const;
 
+	/** Adjustable wireframe opacity */
+	UNREALED_API void OnWireframeOpacityChanged(float Opacity);
+	UNREALED_API float OnGetWireframeOpacity() const;
 
 	virtual TSharedRef<FEditorViewportClient> MakeEditorViewportClient() = 0;
 
diff --git a/Engine/Source/Editor/UnrealEd/UnrealEd.Build.cs b/Engine/Source/Editor/UnrealEd/UnrealEd.Build.cs
--- a/Engine/Source/Editor/UnrealEd/UnrealEd.Build.cs
+++ b/Engine/Source/Editor/UnrealEd/UnrealEd.Build.cs
@@ -185,6 +185,7 @@ public class UnrealEd : ModuleRules
 				"VirtualizationEditor",
 				"GeometryCore",
 				"UniversalObjectLocatorEditor",
+				"Renderer",
 			}
 		);
 
