Added CorePreciseFP to the list of libraries necessary for Lightmass, to follow changes to UObjectCore dependencies

diff --git a/Engine/Source/Editor/UnrealEd/Private/Lightmass/Lightmass.cpp b/Engine/Source/Editor/UnrealEd/Private/Lightmass/Lightmass.cpp
--- a/Engine/Source/Editor/UnrealEd/Private/Lightmass/Lightmass.cpp
+++ b/Engine/Source/Editor/UnrealEd/Private/Lightmass/Lightmass.cpp
@@ -2933,6 +2933,7 @@ bool FLightmassProcessor::BeginRun()
 		TEXT("Binaries/Win32/UnrealLightmass-ApplicationCore.dll"),
 		TEXT("Binaries/Win32/UnrealLightmass-Core.dll"),
 		TEXT("Binaries/Win32/UnrealLightmass-CoreUObject.dll"),
+		TEXT("Binaries/Win32/UnrealLightmass-CorePreciseFP.dll"),
 		TEXT("Binaries/Win32/UnrealLightmass-DerivedDataCache.dll"),
 		TEXT("Binaries/Win32/UnrealLightmass-Sockets.dll"),
 		TEXT("Binaries/Win32/UnrealLightmass-Zen.dll"),
@@ -2953,6 +2954,7 @@ bool FLightmassProcessor::BeginRun()
 		TEXT("Binaries/Win64/UnrealLightmass-ApplicationCore.dll"),
 		TEXT("Binaries/Win64/UnrealLightmass-Core.dll"),
 		TEXT("Binaries/Win64/UnrealLightmass-CoreUObject.dll"),
+		TEXT("Binaries/Win64/UnrealLightmass-CorePreciseFP.dll"),
 		TEXT("Binaries/Win64/UnrealLightmass-DerivedDataCache.dll"),
 		TEXT("Binaries/Win64/UnrealLightmass-Sockets.dll"),
 		TEXT("Binaries/Win64/UnrealLightmass-Zen.dll"),
@@ -2972,6 +2974,7 @@ bool FLightmassProcessor::BeginRun()
 		TEXT("Binaries/Mac/UnrealLightmass-ApplicationCore.dylib"),
 		TEXT("Binaries/Mac/UnrealLightmass-Core.dylib"),
 		TEXT("Binaries/Mac/UnrealLightmass-CoreUObject.dylib"),
+		TEXT("Binaries/Mac/UnrealLightmass-CorePreciseFP.dylib"),
 		TEXT("Binaries/Mac/UnrealLightmass-DerivedDataCache.dylib"),
 		TEXT("Binaries/Mac/UnrealLightmass-Sockets.dylib"),
 		TEXT("Binaries/Mac/UnrealLightmass-Zen.dylib"),
@@ -2992,6 +2995,7 @@ bool FLightmassProcessor::BeginRun()
 		TEXT("Binaries/Linux/libUnrealLightmass-ApplicationCore.so"),
 		TEXT("Binaries/Linux/libUnrealLightmass-Core.so"),
 		TEXT("Binaries/Linux/libUnrealLightmass-CoreUObject.so"),
+		TEXT("Binaries/Linux/libUnrealLightmass-CorePrecise.so"),
 		TEXT("Binaries/Linux/libUnrealLightmass-DerivedDataCache.so"),
 		TEXT("Binaries/Linux/libUnrealLightmass-Sockets.so"),
 		TEXT("Binaries/Linux/libUnrealLightmass-Zen.so"),
