MPCook CookPackageSplitter retraction: Add diagnostics and bugfixes for GenerationHelpers still allocated at end of cook.
#rnx

diff --git a/Engine/Source/Editor/UnrealEd/Private/CookOnTheFlyServer.cpp b/Engine/Source/Editor/UnrealEd/Private/CookOnTheFlyServer.cpp
--- a/Engine/Source/Editor/UnrealEd/Private/CookOnTheFlyServer.cpp
+++ b/Engine/Source/Editor/UnrealEd/Private/CookOnTheFlyServer.cpp
@@ -10159,6 +10159,24 @@ void UCookOnTheFlyServer::CookByTheBookFinishedInternal()
 	check(PackageDatas->GetLoadReadyQueue().IsEmpty());
 	check(PackageDatas->GetSaveQueue().IsEmpty());
 
+	TArray<FPackageData*> DanglingGenerationHelpers;
+	PackageDatas->LockAndEnumeratePackageDatas([&DanglingGenerationHelpers](FPackageData* PackageData)
+		{
+			if (PackageData->GetGenerationHelper())
+			{
+				DanglingGenerationHelpers.Add(PackageData);
+			}
+		});
+	for (FPackageData* PackageData : DanglingGenerationHelpers)
+	{
+		TRefCountPtr<FGenerationHelper> GenerationHelper = PackageData->GetGenerationHelper();
+		if (GenerationHelper)
+		{
+			GenerationHelper->DiagnoseWhyNotShutdown();
+			GenerationHelper->ForceUninitialize();
+		}
+	};
+
 	UE_LOG(LogCook, Display, TEXT("Finishing up..."));
 
 	{
diff --git a/Engine/Source/Editor/UnrealEd/Private/Cooker/CookGenerationHelper.cpp b/Engine/Source/Editor/UnrealEd/Private/Cooker/CookGenerationHelper.cpp
--- a/Engine/Source/Editor/UnrealEd/Private/Cooker/CookGenerationHelper.cpp
+++ b/Engine/Source/Editor/UnrealEd/Private/Cooker/CookGenerationHelper.cpp
@@ -5,6 +5,7 @@
 #include "Algo/Unique.h"
 #include "AssetRegistry/AssetData.h"
 #include "Cooker/CookDirector.h"
+#include "Cooker/CookPlatformManager.h"
 #include "Cooker/CookWorkerServer.h"
 #include "Cooker/IWorkerRequests.h"
 #include "Cooker/PackageTracker.h"
@@ -136,9 +137,8 @@ void FGenerationHelper::Uninitialize()
 	// ReferenceFromKeepForQueueResults
 	// ReferenceFromKeepForGeneratorSave
 	// ReferenceFromKeepForAllSavedOrGC
-	MPCookNextAssignmentIndex = 0;
+	// Keep MPCookNextAssignmentIndex; it is allowed in the uninitialized state
 	// Keep NumSaved; it is allowed in the uninitialized state
-	// Keep WorkerIdThatSavedGenerator; it is allowed in the uninitialized state
 	// InitializeStatus was modified above
 	// Keep DoesGeneratedRequireGeneratorValue; it is allowed in the uninitialized state
 	// Keep bUseInternalReferenceToAvoidGarbageCollect; it is allowed in the uninitialized state
@@ -180,6 +180,161 @@ void FGenerationHelper::OnAllSavesCompleted(UCookOnTheFlyServer& COTFS)
 	ClearKeepForAllSavedOrGC();
 }
 
+void FGenerationHelper::DiagnoseWhyNotShutdown()
+{
+	TStringBuilder<256> Lines;
+	int32 ExpectedNumSaved = PackagesToGenerate.Num() + 1;
+	if (NumSaved != ExpectedNumSaved)
+	{
+		Lines.Appendf(TEXT("\tNumSaved == %d, ExpectedNumSaved == %d.\n"), NumSaved, ExpectedNumSaved);
+	}
+	UCookOnTheFlyServer& COTFS = GetOwner().GetPackageDatas().GetCookOnTheFlyServer();
+	uint32 ExpectedRefCount = 1;
+	auto TestInfo = [this, &Lines, &COTFS, &ExpectedRefCount](FCookGenerationInfo& Info)
+		{
+			if (Info.PackageData->GetState() != EPackageState::Idle)
+			{
+				Lines.Appendf(TEXT("\t%s%s is not idle; it is in state %d.\n"),
+					Info.IsGenerator() ? TEXT("OwnerInfo") : TEXT("GeneratedPackage "),
+					Info.IsGenerator() ? TEXT("") : *Info.GetPackageName(),
+					static_cast<int32>(Info.PackageData->GetState()));
+			}
+			else
+			{
+				bool bMissingPlatform = false;
+				for (const ITargetPlatform* TargetPlatform : COTFS.PlatformManager->GetSessionPlatforms())
+				{
+					const FPackagePlatformData* PlatformData = Info.PackageData->GetPlatformDatas().Find(TargetPlatform);
+					if (!PlatformData || PlatformData->GetCookResults() == ECookResult::NotAttempted)
+					{
+						bMissingPlatform = true;
+					}
+				}
+				if (bMissingPlatform)
+				{
+					Lines.Appendf(TEXT("\t%s%s was not cooked.\n"),
+						Info.IsGenerator() ? TEXT("OwnerInfo") : TEXT("GeneratedPackage "),
+						Info.IsGenerator() ? TEXT("") : *Info.GetPackageName());
+				}
+			}
+			if (!Info.HasSaved())
+			{
+				Lines.Appendf(TEXT("\t%s%s has not marked saved.\n"),
+					Info.IsGenerator() ? TEXT("OwnerInfo") : TEXT("GeneratedPackage "),
+					Info.IsGenerator() ? TEXT("") : *Info.GetPackageName());
+			}
+			if (!Info.IsGenerator() && Info.PackageData->GetParentGenerationHelper())
+			{
+				Lines.Appendf(TEXT("\tGeneratedPackage %s has ParentGenerationHelper set.\n"), *Info.GetPackageName());
+				++ExpectedRefCount;
+			}
+		};
+	TestInfo(GetOwnerInfo());
+	for (FCookGenerationInfo& Info : GetPackagesToGenerate())
+	{
+		TestInfo(Info);
+	}
+
+	if (ReferenceFromKeepForIterative)
+	{
+		Lines.Append(TEXT("\tReferenceFromKeepForIterative is set.\n"));
+		++ExpectedRefCount;
+	}
+	if (ReferenceFromKeepForQueueResults)
+	{
+		Lines.Append(TEXT("\tReferenceFromKeepForQueueResults is set.\n"));
+		++ExpectedRefCount;
+	}
+	if (ReferenceFromKeepForGeneratorSave)
+	{
+		Lines.Append(TEXT("\tReferenceFromKeepForGeneratorSave is set.\n"));
+		++ExpectedRefCount;
+	}
+	if (bKeepForAllSavedOrGC)
+	{
+		Lines.Append(TEXT("\tbKeepForAllSavedOrGC is true.\n"));
+	}
+	if (bKeepForCompletedAllSavesMessage)
+	{
+		Lines.Append(TEXT("\tbKeepForCompletedAllSavesMessage is true.\n"));
+	}
+	if (ReferenceFromKeepForAllSavedOrGC)
+	{
+		if (!bKeepForAllSavedOrGC && !bKeepForCompletedAllSavesMessage)
+		{
+			Lines.Append(TEXT("\tReferenceFromKeepForAllSavedOrGC is set, despite bKeepForAllSavedOrGC and bKeepForCompletedAllSavesMessage being false.\n"));
+		}
+		++ExpectedRefCount;
+	}
+	if (GetRefCount() > ExpectedRefCount)
+	{
+		GetOwner().GetPackageDatas().LockAndEnumeratePackageDatas(
+			[this, &ExpectedRefCount, &Lines](FPackageData* PackageData)
+			{
+				if (PackageData->GetParentGenerationHelper().GetReference() == this &&
+					FindInfo(*PackageData) == nullptr)
+				{
+					Lines.Appendf(TEXT("\tGenerated package %s has ParentGenerationHelper set, but is not listed as a PackageToGenerate from the GenerationHelper.\n"),
+						*PackageData->GetPackageName().ToString());
+					++ExpectedRefCount;
+				}
+			});
+	}
+	if (GetRefCount() > ExpectedRefCount)
+	{
+		Lines.Appendf(TEXT("\tGetRefCount() has references from unknown sources. GetRefCount() == %u, ExpectedRefCount == %u.\n"),
+			GetRefCount(), ExpectedRefCount);
+	}
+
+	if (Lines.Len() != 0)
+	{
+		FWorkerId WorkerId = GetWorkerIdThatSavedGenerator();
+		Lines.Appendf(TEXT("\tGenerator: Saved on %s.\n"), *GetOwnerInfo().SavedOnWorker.ToString());
+		for (FCookGenerationInfo& Info : PackagesToGenerate)
+		{
+			Lines.Appendf(TEXT("\tGeneratedPackage %s: Saved on %s.\n"), *Info.GetPackageName(),
+				*Info.SavedOnWorker.ToString());
+		}
+	}
+	else
+	{
+		Lines.Appendf(TEXT("\tDiagnoseWhyNotShutdown was called unexpectedly; GetRefCount() == 1 so this GenerationHelper should be shut down.\n"));
+	}
+	if (Lines.ToView().EndsWith(TEXT("\n")))
+	{
+		Lines.RemoveSuffix(1);
+	}
+
+	UE_LOG(LogCook, Error,
+		TEXT("GenerationHelper for package %s is still allocated at end of cooksession. This is unexpected and could indicate some generated packages are missing."),
+		*GetOwner().GetPackageName().ToString());
+	UE_LOG(LogCook, Display, TEXT("Diagnostics:\n%s"), *Lines);
+}
+
+void FGenerationHelper::ForceUninitialize()
+{
+	TArray<FPackageData*> PackagesToDemote;
+	auto TestInfo = [&PackagesToDemote](FCookGenerationInfo& Info)
+		{
+			if (Info.PackageData->GetState() != EPackageState::Idle)
+			{
+				PackagesToDemote.Add(Info.PackageData);
+			}
+		};
+	TestInfo(GetOwnerInfo());
+	for (FCookGenerationInfo& Info : GetPackagesToGenerate())
+	{
+		TestInfo(Info);
+	}
+
+	UCookOnTheFlyServer& COTFS = GetOwner().GetPackageDatas().GetCookOnTheFlyServer();
+	for (FPackageData* PackageData : PackagesToDemote)
+	{
+		COTFS.DemoteToIdle(*PackageData, ESendFlags::QueueAddAndRemove, ESuppressCookReason::CookCanceled);
+	}
+	Uninitialize();
+}
+
 UPackage* FGenerationHelper::FindOrLoadPackage(UCookOnTheFlyServer& COTFS, FPackageData& OwnerPackageData)
 {
 	// This is the static helper function on FGenerationHelper that loads the package for any FPackageData; for the
@@ -738,11 +893,11 @@ void FGenerationHelper::NotifyStartQueueGeneratedPackages(UCookOnTheFlyServer& C
 	// Note this function can be called on an uninitialized Generator; the generator is only needed
 	// on the director so it can serve as the passer of messages. We have to keep ourselves referenced after
 	// this call, until after we send EGeneratorEvent::QueuedGeneratedPackagesFencePassed, so that we don't destruct
-	// and lose the WorkerIdThatSavedGenerator information.
-	if (COTFS.CookDirector)
+	// and lose the information from SavedOnWorker or TryGenerateList.
+	if (!COTFS.CookWorkerClient)
 	{
 		COTFS.PackageDatas->GetRequestQueue().AddRequestFenceListener(GetOwner().GetPackageName());
-		WorkerIdThatSavedGenerator = SourceWorkerId;
+		GetOwnerInfo().SavedOnWorker = SourceWorkerId;
 		SetKeepForCompletedAllSavesMessage();
 	}
 	SetKeepForQueueResults();
@@ -761,10 +916,18 @@ void FGenerationHelper::EndQueueGeneratedPackagesOnDirector(UCookOnTheFlyServer&
 	COTFS.PackageDatas->GetRequestQueue().AddRequestFenceListener(GetOwner().GetPackageName());
 	SetKeepForQueueResults();
 
-	// Setting the WorkerIdThatSavedGenerator in response to this event is usually not needed because it is also set
-	// from the reported GeneratedPackages, but we set it anyway in case there is an edge condition that skips those
-	// notifications.
-	WorkerIdThatSavedGenerator = SourceWorkerId;
+	// Setting OwnerInfo.SavedOnWorker and KeepForCompletedAllSavesMessage in response to this event is usually not
+	// needed because they are set from NotifyStartQueueGeneratedPackages, but we set them anyway in case there is an
+	// edge condition that skips those notifications.
+	SetKeepForCompletedAllSavesMessage();
+	GetOwnerInfo().SavedOnWorker = SourceWorkerId;
+
+	// The save message for the owner may have come in before this GenerationHelper was created and thus
+	// MarkPackageSavedRemotely was not called. Check for that case now and marked saved if so.
+	if (GetOwner().HasAnyCookedPlatform())
+	{
+		GetOwnerInfo().SetHasSaved(*this, true, SourceWorkerId);
+	}
 }
 
 void FGenerationHelper::OnRequestFencePassedBroadcast(UCookOnTheFlyServer& COTFS)
@@ -892,7 +1055,7 @@ void FGenerationHelper::ResetSaveState(FCookGenerationInfo& Info, UPackage* Pack
 	{
 		// The save is completed and we will not come back to it; set state back to initial
 		// state and drop our reference keeping this GenerationHelper in memory for the save.
-		Info.SetHasSaved(*this, true);
+		Info.SetHasSaved(*this, true, FWorkerId::Local());
 
 		if (Info.IsGenerator())
 		{
@@ -1175,12 +1338,13 @@ void FGenerationHelper::TrackGeneratedPackageListedRemotely(UCookOnTheFlyServer&
 	}
 }
 
-void FGenerationHelper::MarkPackageSavedRemotely(UCookOnTheFlyServer& COTFS, FPackageData& PackageData)
+void FGenerationHelper::MarkPackageSavedRemotely(UCookOnTheFlyServer& COTFS, FPackageData& PackageData,
+	FWorkerId SourceWorkerId)
 {
 	FCookGenerationInfo* Info = FindInfo(PackageData);
 	if (Info)
 	{
-		Info->SetHasSaved(*this, true);
+		Info->SetHasSaved(*this, true, SourceWorkerId);
 	}
 }
 
@@ -1416,6 +1580,7 @@ void FCookGenerationInfo::Uninitialize()
 	// Keep PackageData; it is allowed in the uninitialized state
 	KeepReferencedPackages.Empty();
 	check(CachedObjectsInOuterInfo.IsEmpty()); // We can not still be in the save state, so this should be empty
+	// Keep SavedOnWorker; it is allowed in the uninitialized state
 	bCreateAsMap = false;
 	bHasCreatedPackage = false;
 	// Keep bHasSaved; it is allowed in the uninitialized state
diff --git a/Engine/Source/Editor/UnrealEd/Private/Cooker/CookGenerationHelper.h b/Engine/Source/Editor/UnrealEd/Private/Cooker/CookGenerationHelper.h
--- a/Engine/Source/Editor/UnrealEd/Private/Cooker/CookGenerationHelper.h
+++ b/Engine/Source/Editor/UnrealEd/Private/Cooker/CookGenerationHelper.h
@@ -90,7 +90,7 @@ public:
 	bool HasCreatedPackage() const;
 	void SetHasCreatedPackage(bool bValue);
 	bool HasSaved() const;
-	void SetHasSaved(FGenerationHelper& GenerationHelper, bool bValue);
+	void SetHasSaved(FGenerationHelper& GenerationHelper, bool bValue, FWorkerId SourceWorkerId);
 	bool HasTakenOverCachedCookedPlatformData() const;
 	void SetHasTakenOverCachedCookedPlatformData(bool bValue);
 	bool HasIssuedUndeclaredMovedObjectsWarning() const;
@@ -152,6 +152,7 @@ public:
 	FPackageData* PackageData;
 	TArray<TWeakObjectPtr<UPackage>> KeepReferencedPackages;
 	TMap<UObject*, FCachedObjectInOuterGeneratorInfo> CachedObjectsInOuterInfo;
+	FWorkerId SavedOnWorker = FWorkerId::Invalid();
 private:
 	ESaveState GeneratorSaveState = ESaveState::StartSave;
 	bool bCreateAsMap : 1;
@@ -373,7 +374,7 @@ public:
 	 * Called on the director when the generator or one of the generated packages was saved on a remote worker.
 	 * Used to manage KeepForGCOrAllSaved lifetime. Does not call Initialize.
 	 */
-	void MarkPackageSavedRemotely(UCookOnTheFlyServer& COTFS, FPackageData& PackageData);
+	void MarkPackageSavedRemotely(UCookOnTheFlyServer& COTFS, FPackageData& PackageData, FWorkerId SourceWorkerId);
 
 	/**
 	 * Called on the director and every CookWorker when all saves have been completed; this indicates that
@@ -381,6 +382,14 @@ public:
 	 */
 	void OnAllSavesCompleted(UCookOnTheFlyServer& COTFS);
 
+	/** Diagnostics for shutdown errors. Report why the GenerationHelper is still referenced. */
+	void DiagnoseWhyNotShutdown();
+	/**
+	 * Helper for shutdown errors. Force the GenerationHelper to uninitialize so that the CookPackageSplitter
+	 * is shutdown correctly before the cooksession ends.
+	 */
+	void ForceUninitialize();
+
 	/** Helper function for Initialize and for TryCreateValidGenerationHelper. */
 	static void SearchForRegisteredSplitDataObject(UCookOnTheFlyServer& COTFS, FName PackageName, UPackage* Package,
 		TOptional<TConstArrayView<FCachedObjectInOuter>> CachedObjectsInOuter, UObject*& OutSplitDataObject,
@@ -435,7 +444,6 @@ private:
 	TRefCountPtr<FGenerationHelper> ReferenceFromKeepForAllSavedOrGC;
 	int32 MPCookNextAssignmentIndex = 0;
 	int32 NumSaved = 0;
-	FWorkerId WorkerIdThatSavedGenerator = FWorkerId::Invalid();
 	EInitializeStatus InitializeStatus = EInitializeStatus::Uninitialized;
 	ICookPackageSplitter::EGeneratedRequiresGenerator DoesGeneratedRequireGeneratorValue =
 		ICookPackageSplitter::EGeneratedRequiresGenerator::None;
@@ -491,12 +499,17 @@ inline bool FCookGenerationInfo::HasSaved() const
 	return bHasSaved;
 }
 
-inline void FCookGenerationInfo::SetHasSaved(FGenerationHelper& GenerationHelper, bool bValue)
+inline void FCookGenerationInfo::SetHasSaved(FGenerationHelper& GenerationHelper, bool bValue,
+	FWorkerId SourceWorkerId)
 {
 	if (bValue != bHasSaved)
 	{
 		bHasSaved = bValue;
 		GenerationHelper.ModifyNumSaved(bValue ? 1 : -1);
+		if (bHasSaved)
+		{
+			SavedOnWorker = SourceWorkerId;
+		}
 	}
 }
 
@@ -716,7 +729,7 @@ inline void FGenerationHelper::ClearKeepForCompletedAllSavesMessage()
 
 inline FWorkerId FGenerationHelper::GetWorkerIdThatSavedGenerator() const
 {
-	return WorkerIdThatSavedGenerator;
+	return OwnerInfo.SavedOnWorker;
 }
 
 inline int32& FGenerationHelper::GetMPCookNextAssignmentIndex()
diff --git a/Engine/Source/Editor/UnrealEd/Private/Cooker/CookWorkerServer.cpp b/Engine/Source/Editor/UnrealEd/Private/Cooker/CookWorkerServer.cpp
--- a/Engine/Source/Editor/UnrealEd/Private/Cooker/CookWorkerServer.cpp
+++ b/Engine/Source/Editor/UnrealEd/Private/Cooker/CookWorkerServer.cpp
@@ -932,14 +932,14 @@ void FCookWorkerServer::RecordResults(FPackageResultsMessage& Message)
 			}
 			else
 			{
-				ParentGenerationHelper->MarkPackageSavedRemotely(COTFS, *PackageData);
+				ParentGenerationHelper->MarkPackageSavedRemotely(COTFS, *PackageData, GetWorkerId());
 				PackageData->SetParentGenerationHelper(nullptr);
 			}
 		}
 		TRefCountPtr<FGenerationHelper> GenerationHelper = PackageData->GetGenerationHelper();
 		if (GenerationHelper)
 		{
-			GenerationHelper->MarkPackageSavedRemotely(COTFS, *PackageData);
+			GenerationHelper->MarkPackageSavedRemotely(COTFS, *PackageData, GetWorkerId());
 			GenerationHelper.SafeRelease();
 		}
 
diff --git a/Engine/Source/Runtime/CoreUObject/Public/Cooker/MPCollector.h b/Engine/Source/Runtime/CoreUObject/Public/Cooker/MPCollector.h
--- a/Engine/Source/Runtime/CoreUObject/Public/Cooker/MPCollector.h
+++ b/Engine/Source/Runtime/CoreUObject/Public/Cooker/MPCollector.h
@@ -6,6 +6,7 @@
 #include "Async/Future.h"
 #include "Containers/Array.h"
 #include "Containers/ArrayView.h"
+#include "Containers/UnrealString.h"
 #include "HAL/Platform.h"
 #include "Misc/Guid.h"
 #include "Serialization/CompactBinary.h"
@@ -52,6 +53,8 @@ public:
 	bool operator<(const FWorkerId& Other) const { return Id < Other.Id; }
 	inline friend int32 GetTypeHash(const FWorkerId& WorkerId) { return WorkerId.Id; }
 
+	FString ToString();
+
 private:
 	constexpr explicit FWorkerId(uint8 InId) : Id(InId) {}
 
@@ -273,6 +276,12 @@ private:
 	TUniqueFunction<void(FMPCollectorServerMessageContext& Context, bool bReadSuccessful, MessageType&& Message)> Callback;
 };
 
+inline FString FWorkerId::ToString()
+{
+	return IsInvalid() ? TEXT("<Invalid>") : (IsLocal() ? TEXT("Local") :
+		FString::Printf(TEXT("CookWorker %u"), static_cast<uint32>(GetRemoteIndex())));
+}
+
 } // namespace UE::Cook
 
 #endif // WITH_EDITOR
\ No newline at end of file
