Cooker: Add parsing of -cookfirst=<PackageName> and -cooklast=<PackageName> to allow specifying which commandline packages should be cookfirst/cooklast.
Change RandomizePackageOrder to completely randomize the clusters of requested packages rather than preserving the leaf to root order. RandomizePackageOrder needs more freedom to be random so we can cook referencers before dependencies or vice versa to track down determinism issues.
#rnx

diff --git a/Engine/Source/Editor/UnrealEd/Classes/CookOnTheSide/CookOnTheFlyServer.h b/Engine/Source/Editor/UnrealEd/Classes/CookOnTheSide/CookOnTheFlyServer.h
--- a/Engine/Source/Editor/UnrealEd/Classes/CookOnTheSide/CookOnTheFlyServer.h
+++ b/Engine/Source/Editor/UnrealEd/Classes/CookOnTheSide/CookOnTheFlyServer.h
@@ -1520,10 +1520,6 @@ private:
 	bool bRandomizeCookOrder = false;
 	/** True if commandline arguments specified that we suppress the cook of packages based on filter criteria. */
 	bool bCookFilter = false;
-	/** True if commandline arguments specify that packages on commandline should be cooked first. */
-	bool bCookFirst = false;
-	/** True if commandline arguments specify that packages on commandline should be cooked last. */
-	bool bCookLast = false;
 	/** True if experimental optimizations for fast startup should be used. */
 	bool bCookFastStartup = false;
 	/**
diff --git a/Engine/Source/Editor/UnrealEd/Private/CookOnTheFlyServer.cpp b/Engine/Source/Editor/UnrealEd/Private/CookOnTheFlyServer.cpp
--- a/Engine/Source/Editor/UnrealEd/Private/CookOnTheFlyServer.cpp
+++ b/Engine/Source/Editor/UnrealEd/Private/CookOnTheFlyServer.cpp
@@ -9,6 +9,7 @@
 #include "Algo/AllOf.h"
 #include "Algo/AnyOf.h"
 #include "Algo/Find.h"
+#include "Algo/RandomShuffle.h"
 #include "Algo/Unique.h"
 #include "AssetCompilingManager.h"
 #include "AssetRegistry/AssetRegistryModule.h"
@@ -2725,6 +2726,14 @@ void UCookOnTheFlyServer::AssignRequests(TArrayView<UE::Cook::FPackageData*> Req
 	}
 	else
 	{
+		TArray<FPackageData*> Shuffled;
+		if (bRandomizeCookOrder)
+		{
+			Shuffled = Requests;
+			Algo::RandomShuffle(Shuffled);
+			Requests = Shuffled;
+		}
+
 		for (FPackageData* PackageData : Requests)
 		{
 			RequestQueue.AddReadyRequest(PackageData);
@@ -6980,6 +6989,9 @@ void FInitializeConfigSettings::LoadLocal(const FString& InOutputDirectoryOverri
 			MaxAsyncCacheForType.Add(CacheSetting.Key, Count);
 		}
 	}
+
+	bRandomizeCookOrder = FParse::Param(FCommandLine::Get(), TEXT("RANDOMPACKAGEORDER")) ||
+		(FParse::Param(FCommandLine::Get(), TEXT("DIFFONLY")) && !FParse::Param(FCommandLine::Get(), TEXT("DIFFNORANDCOOK")));
 }
 
 static EMPCookGeneratorSplit ParseMPCookGeneratorSplitFromString(const FString& Text)
@@ -7120,7 +7132,7 @@ void UCookOnTheFlyServer::SetInitializeConfigSettings(UE::Cook::FInitializeConfi
 		}
 		if (!ClassPathListStr.IsEmpty())
 		{
-			UE::String::ParseTokensMultiple(ClassPathListStr, { '+', ',', ';' },
+			UE::String::ParseTokensMultiple(ClassPathListStr, UE::Cook::GetCommandLineDelimiterChars(),
 				[this](FStringView Token)
 				{
 					FTopLevelAssetPath Path(Token);
@@ -7135,16 +7147,6 @@ void UCookOnTheFlyServer::SetInitializeConfigSettings(UE::Cook::FInitializeConfi
 		}
 	}
 
-	bCookFirst = FParse::Param(FCommandLine::Get(), TEXT("CookFirst"));
-	bCookLast = FParse::Param(FCommandLine::Get(), TEXT("CookLast"));
-	if (bCookFirst && bCookLast)
-	{
-		UE_LOG(LogCook, Error, TEXT("-CookFirst and -CookLast are mutually exclusive. Ignoring -CookLast"));
-		bCookLast = false;
-	}
-	bRandomizeCookOrder = !bCookFirst && !bCookLast && (FParse::Param(FCommandLine::Get(), TEXT("RANDOMPACKAGEORDER")) ||
-		(FParse::Param(FCommandLine::Get(), TEXT("DIFFONLY")) && !FParse::Param(FCommandLine::Get(), TEXT("DIFFNORANDCOOK"))));
-
 	ParseCookFilters();
 
 	bCallIsCachedOnSaveCreatedObjects = FParse::Param(FCommandLine::Get(), TEXT("CallIsCachedOnSaveCreatedObjects"));
@@ -7219,9 +7221,9 @@ void UCookOnTheFlyServer::ParseCookFilters(const TCHAR* Parameter, const TCHAR*
 	if (FParse::Value(FCommandLine::Get(), *FullParameter, IncludeClassesString))
 	{
 		TArray<FString> IncludeClasses;
-		const TCHAR* Delimiters[] = { TEXT(","), TEXT("+"), TEXT(";")};
-		IncludeClassesString.ParseIntoArray(IncludeClasses, Delimiters,
-			UE_ARRAY_COUNT(Delimiters), true /* bCullEmpty */);
+		TConstArrayView<const TCHAR*> Delimiters = UE::Cook::GetCommandLineDelimiterStrs();
+		IncludeClassesString.ParseIntoArray(IncludeClasses, Delimiters.GetData(), Delimiters.Num(),
+			true /* bCullEmpty */);
 		TArray<FTopLevelAssetPath> RootNames;
 		for (FString& IncludeClassString : IncludeClasses)
 		{
@@ -11813,9 +11815,36 @@ void UCookOnTheFlyServer::GenerateInitialRequests(FBeginCookContext& BeginContex
 		}
 	}
 
+	TArray<FString> CookMaps = BeginContext.StartupOptions->CookMaps;
+	TArray<FString> CookFirstPackages;
+	TArray<FString> CookLastPackages;
+	FString Text;
+	TConstArrayView<const TCHAR*> CommandLineDelimiters = UE::Cook::GetCommandLineDelimiterStrs();
+	if (FParse::Param(FCommandLine::Get(), TEXT("CookFirst")))
+	{
+		CookFirstPackages.Append(CookMaps);
+	}
+	if (FParse::Value(FCommandLine::Get(), TEXT("-CookFirst="), Text))
+	{
+		TArray<FString> Array;
+		Text.ParseIntoArray(Array, CommandLineDelimiters.GetData(), CommandLineDelimiters.Num(), true);
+		CookFirstPackages.Append(Array);
+		CookMaps.Append(Array);
+	}
+	if (FParse::Param(FCommandLine::Get(), TEXT("CookLast")))
+	{
+		CookLastPackages.Append(CookMaps);
+	}
+	if (FParse::Value(FCommandLine::Get(), TEXT("-CookLast="), Text))
+	{
+		TArray<FString> Array;
+		Text.ParseIntoArray(Array, CommandLineDelimiters.GetData(), CommandLineDelimiters.Num(), true);
+		CookLastPackages.Append(Array);
+		CookMaps.Append(Array);
+	}
+
 	TArray<FName> FilesInPath;
 	TMap<FName, UE::Cook::FInstigator> FilesInPathInstigators;
-	const TArray<FString>& CookMaps = BeginContext.StartupOptions->CookMaps;
 	const TArray<FString>& CookDirectories = BeginContext.StartupOptions->CookDirectories;
 	const TArray<FString>& IniMapSections = BeginContext.StartupOptions->IniMapSections;
 	ECookByTheBookOptions CookOptions = CookByTheBookOptions->StartupOptions;
@@ -11856,13 +11885,31 @@ void UCookOnTheFlyServer::GenerateInitialRequests(FBeginCookContext& BeginContex
 		GenerateLongPackageNames(FilesInPath, FilesInPathInstigators);
 	}
 	TSet<FName> CookFirstOrLastPackages;
-	bool bCookFirstOrLast = bCookFirst || bCookLast;
-	if (bCookFirstOrLast)
+	TMap<FString, TOptional<bool>> CookFirstOrLastPackagesInputs;
+	for (const FString& PackageName : CookFirstPackages)
+	{
+		CookFirstOrLastPackagesInputs.Add(PackageName, TOptional<bool>(true));
+	}
+	for (const FString& PackageName : CookLastPackages)
+	{
+		TOptional<bool>& IsCookFirst = CookFirstOrLastPackagesInputs.FindOrAdd(PackageName);
+		if (IsCookFirst.IsSet() && IsCookFirst.GetValue())
+		{
+			UE_LOG(LogCook, Error, TEXT("-CookFirst and -CookLast are mutually exclusive. Ignoring -CookLast for %s."),
+				*PackageName);
+		}
+		else
+		{
+			IsCookFirst.Emplace(false);
+		}
+	}
+	if (!CookFirstOrLastPackagesInputs.IsEmpty())
 	{
-		for (const FString& CookMap : CookMaps)
+		for (const TPair<FString, TOptional<bool>>& Pair : CookFirstOrLastPackagesInputs)
 		{
+			bool bCookLast = !Pair.Value.GetValue();
 			FString LongPackageName;
-			if (FPackageName::TryConvertFilenameToLongPackageName(CookMap, LongPackageName))
+			if (FPackageName::TryConvertFilenameToLongPackageName(Pair.Key, LongPackageName))
 			{
 				FName LongPackageFName(*LongPackageName);
 				CookFirstOrLastPackages.Add(LongPackageFName);
@@ -11877,6 +11924,7 @@ void UCookOnTheFlyServer::GenerateInitialRequests(FBeginCookContext& BeginContex
 			}
 		}
 	}
+	bool bCookFirstOrLast = !CookFirstOrLastPackages.IsEmpty();
 
 	// add all the files to the cook list for the requested platforms
 	for (FName PackageName : FilesInPath)
diff --git a/Engine/Source/Editor/UnrealEd/Private/Cooker/CookRequestCluster.cpp b/Engine/Source/Editor/UnrealEd/Private/Cooker/CookRequestCluster.cpp
--- a/Engine/Source/Editor/UnrealEd/Private/Cooker/CookRequestCluster.cpp
+++ b/Engine/Source/Editor/UnrealEd/Private/Cooker/CookRequestCluster.cpp
@@ -600,14 +600,6 @@ void FRequestCluster::PumpExploration(const FCookerTimer& CookerTimer, bool& bOu
 	};
 
 	Algo::TopologicalSort(SortedPackages, GetElementDependencies, Algo::ETopologicalSort::AllowCycles);
-	if (COTFS.bRandomizeCookOrder)
-	{
-		RandomizeCookOrder(SortedPackages, GraphSearch->GetGraphEdges());
-		// Only shuffle the first cluster. We don't need to spend time shuffling later clusters because they
-		// will already occur at random times, and we want to suppress the log information about them since
-		// there can be a lot of small discovered clusters.
-		COTFS.bRandomizeCookOrder = false;
-	}
 	TMap<FPackageData*, int32> SortOrder;
 	int32 Counter = 0;
 	SortOrder.Reserve(SortedPackages.Num());
@@ -2300,143 +2292,4 @@ static TArray<T> FindRootsFromLeafToRootOrderList(TConstArrayView<T> LeafToRootO
 	return Roots;
 }
 
-void FRequestCluster::RandomizeCookOrder(TArray<FPackageData*>& InOutLeafToRootOrder,
-	const TMap<FPackageData*, TArray<FPackageData*>>& Edges)
-{
-	// Notes on the ideal solution:
-	// In a graph without cycles, the visitation order of a DepthFirstSearch starting from the graph roots is a
-	// RootToLeaf ordering. We can randomize by randomly iterating the graph roots and by randomly iterating the
-	// edges when moving from each graph.
-	// In a graph with cycles, a RootToLeaf order is only defined in the condensed graph, and for
-	// each chain within a node of the condensed graph, we can randomize the vertices in the chain.
-	// But this requires creating the condensed graph, which is expensive.
-	// 
-	// Notes on the practical solution:
-	// Cycles are not supposed to be a large part of our graph, so we will not seek to do a good job of randomizing
-	// them. We do a DFS as we would in an acyclic graph, and when we encounter an already-visited vertex due to a
-	// cycle we just pretend that edge does not exist.
-	//
-	// Two DFS passes
-	// Pass 1, find all roots by iterating from root to leaf and DFSing each remaining head element.
-	// Pass 2: Iterate all roots in random order and DFS each one; append the leaf-to-root order of their search
-	// to the final leaf-to-root order.
-	if (InOutLeafToRootOrder.IsEmpty())
-	{
-		return;
-	}
-
-	struct FVisitEntry
-	{
-		FPackageData* Vertex;
-		TArray<FPackageData*> Edges;
-		int32 NextEdge;
-		void Set(FPackageData* V, const TMap<FPackageData*,TArray<FPackageData*>>& AllEdges)
-		{
-			Vertex = V;
-			Edges.Reset();
-			const TArray<FPackageData*>* EdgesFromV = AllEdges.Find(V);
-			if (EdgesFromV)
-			{
-				Edges.Append(*EdgesFromV);
-				ArrayShuffle(Edges);
-			}
-			NextEdge = 0;
-		}
-	};
-	
-	TArray<FVisitEntry> DFSStack;
-	int32 StackNum = 0;
-	auto Push = [&DFSStack, &Edges, &StackNum](FPackageData* Vertex)
-	{
-		while (DFSStack.Num() <= StackNum)
-		{
-			DFSStack.Emplace();
-		}
-		DFSStack[StackNum++].Set(Vertex, Edges);
-	};
-	auto Pop = [&StackNum]()
-	{
-		--StackNum;
-	};
-
-	TSet<FPackageData*> ValidVertices;
-	ValidVertices.Append(InOutLeafToRootOrder);
-
-	TArray<FPackageData*> Roots = FindRootsFromLeafToRootOrderList<FPackageData*>(InOutLeafToRootOrder, Edges, 
-		ValidVertices);
-	ArrayShuffle(Roots);
-
-	TSet<FPackageData*> Visited;
-	Visited.Reserve(InOutLeafToRootOrder.Num());
-	TArray<FPackageData*> OutputVisitOrder;
-
-	for (FPackageData* Root : Roots)
-	{
-		bool bAlreadyExists;
-		Visited.Add(Root, &bAlreadyExists);
-		if (bAlreadyExists)
-		{
-			continue;
-		}
-
-		Push(Root);
-		check(StackNum == 1);
-		while (StackNum > 0)
-		{
-			FVisitEntry& Entry = DFSStack[StackNum - 1];
-			bool bPushed = false;
-			while (Entry.NextEdge < Entry.Edges.Num())
-			{
-				FPackageData* Target = Entry.Edges[Entry.NextEdge++];
-				Visited.Add(Target, &bAlreadyExists);
-				if (!bAlreadyExists && ValidVertices.Contains(Target))
-				{
-					Push(Target);
-					bPushed = true;
-					break;
-				}
-			}
-			if (!bPushed)
-			{
-				// Add this vertex now as the most rootwards encountered; it is farther toward the root than any of the
-				// vertices it depends on
-				OutputVisitOrder.Add(Entry.Vertex);
-				Pop();
-			}
-		}
-	}
-	Visited.Empty();
-	Roots.Empty();
-
-	TMap<FPackageData*, int32> OriginalIndices;
-	OriginalIndices.Reserve(InOutLeafToRootOrder.Num());
-	for (int32 Index = 0; Index < InOutLeafToRootOrder.Num(); ++Index)
-	{
-		OriginalIndices.Add(InOutLeafToRootOrder[Index], Index);
-	}
-	check(OriginalIndices.Num() == InOutLeafToRootOrder.Num());
-	
-	// Copy the OutputVisitOrder over top of our InOut ordered parameter and report the diagnostic
-	// for the average shuffle distance.
-
-	double SumSquaredDistances = 0.0;
-
-	int32 WriteIndex = 0;
-	for (FPackageData* Vertex : OutputVisitOrder)
-	{
-		int32 OriginalIndex;
-		verify(OriginalIndices.RemoveAndCopyValue(Vertex, OriginalIndex));
-		double Distance = static_cast<double>(OriginalIndex - WriteIndex);
-		SumSquaredDistances += FMath::Square(Distance);
-
-		check(WriteIndex < InOutLeafToRootOrder.Num());
-		InOutLeafToRootOrder[WriteIndex++] = Vertex;
-	}
-	check(WriteIndex == InOutLeafToRootOrder.Num());
-
-	UE_LOG(LogCook, Display,
-		TEXT("RandomPackageOrder used, packages in the cluster were shuffled. %d elements in cluster, average shuffle distance == %0.1f."),
-		InOutLeafToRootOrder.Num(), (float) FMath::Sqrt(SumSquaredDistances/InOutLeafToRootOrder.Num()));
-}
-
 } // namespace UE::Cook
diff --git a/Engine/Source/Editor/UnrealEd/Private/Cooker/CookRequestCluster.h b/Engine/Source/Editor/UnrealEd/Private/Cooker/CookRequestCluster.h
--- a/Engine/Source/Editor/UnrealEd/Private/Cooker/CookRequestCluster.h
+++ b/Engine/Source/Editor/UnrealEd/Private/Cooker/CookRequestCluster.h
@@ -416,8 +416,6 @@ private:
 	static void IsRequestCookable(const ITargetPlatform* TargetPlatform, FName PackageName, FPackageData& PackageData,
 		UCookOnTheFlyServer& InCOTFS, FStringView InDLCPath, ESuppressCookReason& OutReason, bool& bOutCookable,
 		bool& bOutExplorable);
-	static void RandomizeCookOrder(TArray<FPackageData*>& InOutLeafToRootOrder,
-		const TMap<FPackageData*, TArray<FPackageData*>>& Edges);
 
 	TArray<FFilePlatformRequest> FilePlatformRequests;
 	/**
diff --git a/Engine/Source/Editor/UnrealEd/Private/Cooker/CookTypes.cpp b/Engine/Source/Editor/UnrealEd/Private/Cooker/CookTypes.cpp
--- a/Engine/Source/Editor/UnrealEd/Private/Cooker/CookTypes.cpp
+++ b/Engine/Source/Editor/UnrealEd/Private/Cooker/CookTypes.cpp
@@ -281,6 +281,18 @@ bool IsCookIgnoreTimeouts()
 	return bIsIgnoreCookTimeouts;
 }
 
+TConstArrayView<const TCHAR*> GetCommandLineDelimiterStrs()
+{
+	static const TCHAR* Delimiters[] = { TEXT(","), TEXT("+"), TEXT(";") };
+	return TConstArrayView<const TCHAR*>(Delimiters, UE_ARRAY_COUNT(Delimiters));
+}
+
+TConstArrayView<const TCHAR> GetCommandLineDelimiterChars()
+{
+	static const TCHAR Delimiters[] = { ',', '+', ';' };
+	return TConstArrayView<TCHAR>(Delimiters, UE_ARRAY_COUNT(Delimiters));
+}
+
 FDiscoveredPlatformSet::FDiscoveredPlatformSet(EDiscoveredPlatformSet InSource)
 	: Source(InSource)
 {
@@ -647,6 +659,7 @@ FCbWriter& operator<<(FCbWriter& Writer, const UE::Cook::FInitializeConfigSettin
 	Writer << "MaxNumPackagesBeforePartialGC" << Value.MaxNumPackagesBeforePartialGC;
 	Writer << "ConfigSettingDenyList" << Value.ConfigSettingDenyList;
 	Writer << "MaxAsyncCacheForType" << Value.MaxAsyncCacheForType;
+	Writer << "bRandomizeCookOrder" << Value.bRandomizeCookOrder;
 	// Make sure new values are added to LoadFromCompactBinary and MoveOrCopy
 	Writer.EndObject();
 	return Writer;
@@ -682,6 +695,7 @@ bool LoadFromCompactBinary(FCbFieldView Field, UE::Cook::FInitializeConfigSettin
 	bOk = LoadFromCompactBinary(Field["MaxNumPackagesBeforePartialGC"], OutValue.MaxNumPackagesBeforePartialGC) & bOk;
 	bOk = LoadFromCompactBinary(Field["ConfigSettingDenyList"], OutValue.ConfigSettingDenyList) & bOk;
 	bOk = LoadFromCompactBinary(Field["MaxAsyncCacheForType"], OutValue.MaxAsyncCacheForType) & bOk;
+	bOk = LoadFromCompactBinary(Field["bRandomizeCookOrder"], OutValue.bRandomizeCookOrder) & bOk;
 	// Make sure new values are added to MoveOrCopy and operator<<
 	return bOk;
 }
@@ -710,6 +724,7 @@ void FInitializeConfigSettings::MoveOrCopy(SourceType&& Source, TargetType&& Tar
 	Target.MaxNumPackagesBeforePartialGC = Source.MaxNumPackagesBeforePartialGC;
 	Target.ConfigSettingDenyList = MoveTempIfPossible(Source.ConfigSettingDenyList);
 	Target.MaxAsyncCacheForType = MoveTempIfPossible(Source.MaxAsyncCacheForType);
+	Target.bRandomizeCookOrder = Source.bRandomizeCookOrder;
 	// Make sure new values are added to operator<< and LoadFromCompactBinary
 }
 
diff --git a/Engine/Source/Editor/UnrealEd/Private/Cooker/CookTypes.h b/Engine/Source/Editor/UnrealEd/Private/Cooker/CookTypes.h
--- a/Engine/Source/Editor/UnrealEd/Private/Cooker/CookTypes.h
+++ b/Engine/Source/Editor/UnrealEd/Private/Cooker/CookTypes.h
@@ -363,6 +363,7 @@ namespace UE::Cook
 		/** max number of objects of a specific type which are allowed to async cache at once */
 		TMap<FName, int32> MaxAsyncCacheForType;
 		bool bUseSoftGC = false;
+		bool bRandomizeCookOrder = false;
 
 		friend FCbWriter& ::operator<<(FCbWriter& Writer, const UE::Cook::FInitializeConfigSettings& Value);
 		friend bool ::LoadFromCompactBinary(FCbFieldView Field, UE::Cook::FInitializeConfigSettings& Value);
@@ -387,6 +388,9 @@ namespace UE::Cook
 
 	/** Report whether commandline/config has disabled use of timeouts throughout the cooker, useful for debugging. */
 	bool IsCookIgnoreTimeouts();
+
+	TConstArrayView<const TCHAR*> GetCommandLineDelimiterStrs();
+	TConstArrayView<TCHAR> GetCommandLineDelimiterChars();
 }
 
 bool LexTryParseString(FPlatformMemoryStats::EMemoryPressureStatus& OutValue, FStringView Text);
