Cooker: Add -cookskiphardrefs, -cookskipsoftrefs, -cookskiprequests, and -cookreferencersof.
This enables testing for indeterminism via minimal cook that cook a package either before or after its referencers:
-run=cook -cookskiphardrefs -cookskiprequests -cookreferencersof=/Game/Package -cookfirst=/Game/Package
-run=cook -cookskiphardrefs -cookskiprequests -cookreferencersof=/Game/Package -cooklast=/Game/Package
#rnx

diff --git a/Engine/Source/Editor/UnrealEd/Classes/CookOnTheSide/CookOnTheFlyServer.h b/Engine/Source/Editor/UnrealEd/Classes/CookOnTheSide/CookOnTheFlyServer.h
--- a/Engine/Source/Editor/UnrealEd/Classes/CookOnTheSide/CookOnTheFlyServer.h
+++ b/Engine/Source/Editor/UnrealEd/Classes/CookOnTheSide/CookOnTheFlyServer.h
@@ -84,7 +84,7 @@ enum class ECookByTheBookOptions
 	NoGameAlwaysCookPackages =			0x00000080, // don't include the packages specified by the game in the cook (this cook will probably be missing content unless you know what you are doing)
 	NoAlwaysCookMaps =					0x00000100, // don't include always cook maps (this cook will probably be missing content unless you know what you are doing)
 	NoDefaultMaps =						0x00000200, // don't include default cook maps (this cook will probably be missing content unless you know what you are doing)
-	// Unused =							0x00000400,
+	NoStartupPackages =					0x00000400, // Don't include packages that are loaded by engine startup (this cook will probably be missing content unless you know what you are doing)
 	NoInputPackages =					0x00000800, // don't include slate content (this cook will probably be missing content unless you know what you are doing)
 	SkipSoftReferences =				0x00001000, // Don't follow soft references when cooking. Usually not viable for a real cook and the results probably wont load properly, but can be useful for debugging
 	SkipHardReferences =				0x00002000, // Don't follow hard references when cooking. Not viable for a real cook, only useful for debugging
diff --git a/Engine/Source/Editor/UnrealEd/Private/Commandlets/CookCommandlet.cpp b/Engine/Source/Editor/UnrealEd/Private/Commandlets/CookCommandlet.cpp
--- a/Engine/Source/Editor/UnrealEd/Private/Commandlets/CookCommandlet.cpp
+++ b/Engine/Source/Editor/UnrealEd/Private/Commandlets/CookCommandlet.cpp
@@ -436,13 +436,28 @@ bool UCookCommandlet::CookByTheBook(const TArray<ITargetPlatform*>& Platforms)
 		CookOptions |= Switches.Contains(TEXT("ValidationErrorsAreFatal")) ? ECookByTheBookOptions::ValidationErrorsAreFatal : ECookByTheBookOptions::None;
 	}
 
+	const ECookByTheBookOptions SkipRequestFlags = ECookByTheBookOptions::NoAlwaysCookMaps |
+		ECookByTheBookOptions::NoDefaultMaps | ECookByTheBookOptions::NoGameAlwaysCookPackages |
+		ECookByTheBookOptions::NoInputPackages | ECookByTheBookOptions::ForceDisableSaveGlobalShaders;
 	if (bCookSinglePackage)
 	{
-		const ECookByTheBookOptions SinglePackageFlags = ECookByTheBookOptions::NoAlwaysCookMaps | ECookByTheBookOptions::NoDefaultMaps | ECookByTheBookOptions::NoGameAlwaysCookPackages |
-			ECookByTheBookOptions::NoInputPackages | ECookByTheBookOptions::SkipSoftReferences | ECookByTheBookOptions::ForceDisableSaveGlobalShaders;
-		CookOptions |= SinglePackageFlags;
+		CookOptions |= SkipRequestFlags;
+		CookOptions |= ECookByTheBookOptions::SkipSoftReferences;
 		CookOptions |= bKeepSinglePackageRefs ? ECookByTheBookOptions::None : ECookByTheBookOptions::SkipHardReferences;
 	}
+	if (Switches.Contains(TEXT("CookSkipRequests")))
+	{
+		CookOptions |= SkipRequestFlags;
+		CookOptions |= ECookByTheBookOptions::NoStartupPackages;
+	}
+	if (Switches.Contains(TEXT("CookSkipSoftRefs")))
+	{
+		CookOptions |= ECookByTheBookOptions::SkipSoftReferences;
+	}
+	if (Switches.Contains(TEXT("CookSkipHardRefs")))
+	{
+		CookOptions |= ECookByTheBookOptions::SkipHardReferences;
+	}
 
 	// Also append any cookdirs from the project ini files; these dirs are relative to the game content directory or start with a / root
 	if (!(CookOptions & ECookByTheBookOptions::NoGameAlwaysCookPackages))
diff --git a/Engine/Source/Editor/UnrealEd/Private/CookOnTheFlyServer.cpp b/Engine/Source/Editor/UnrealEd/Private/CookOnTheFlyServer.cpp
--- a/Engine/Source/Editor/UnrealEd/Private/CookOnTheFlyServer.cpp
+++ b/Engine/Source/Editor/UnrealEd/Private/CookOnTheFlyServer.cpp
@@ -11477,7 +11477,7 @@ FBeginCookContext UCookOnTheFlyServer::CreateBeginCookByTheBookContext(const FCo
 	CookByTheBookOptions->DlcName = StartupOptions.DLCName;
 	if (CookByTheBookOptions->bSkipHardReferences && !CookByTheBookOptions->bSkipSoftReferences)
 	{
-		UE_LOG(LogCook, Warning, TEXT("Setting bSkipSoftReferences to true since bSkipHardReferences is true and skipping hard references requires skipping soft references."));
+		UE_LOG(LogCook, Display, TEXT("Setting bSkipSoftReferences to true since bSkipHardReferences is true and skipping hard references requires skipping soft references."));
 		CookByTheBookOptions->bSkipSoftReferences = true;
 	}
 
@@ -11842,6 +11842,22 @@ void UCookOnTheFlyServer::GenerateInitialRequests(FBeginCookContext& BeginContex
 		CookLastPackages.Append(Array);
 		CookMaps.Append(Array);
 	}
+	if (FParse::Value(FCommandLine::Get(), TEXT("-CookReferencersOf="), Text))
+	{
+		TArray<FString> Array;
+		Text.ParseIntoArray(Array, CommandLineDelimiters.GetData(), CommandLineDelimiters.Num(), true);
+		for (FString& PackageName : Array)
+		{
+			CookMaps.Add(PackageName);
+			TArray<FName> Referencers;
+			AssetRegistry->GetReferencers(FName(FStringView(PackageName)), Referencers);
+			for (FName Referencer : Referencers)
+			{
+				CookMaps.Add(Referencer.ToString());
+			}
+		}
+	}
+
 
 	TArray<FName> FilesInPath;
 	TMap<FName, UE::Cook::FInstigator> FilesInPathInstigators;
@@ -11851,7 +11867,7 @@ void UCookOnTheFlyServer::GenerateInitialRequests(FBeginCookContext& BeginContex
 	CollectFilesToCook(FilesInPath, FilesInPathInstigators, CookMaps, CookDirectories, IniMapSections, CookOptions, TargetPlatforms, GameDefaultObjects);
 
 	// Add soft/hard startup references after collecting requested files and handling empty requests
-	if (!CookByTheBookOptions->bSkipHardReferences)
+	if (!CookByTheBookOptions->bSkipHardReferences && !EnumHasAnyFlags(CookOptions, ECookByTheBookOptions::NoStartupPackages))
 	{
 		ProcessUnsolicitedPackages(&FilesInPath, &FilesInPathInstigators);
 	}
