MPCook CookPackageSplitter retraction: Fix dangling GenerationHelper when a CookWorker sends a discovery message for a generated package after the generator and all of its generated packages have already saved.
#rnx

diff --git a/Engine/Source/Editor/UnrealEd/Private/Cooker/CookWorkerServer.cpp b/Engine/Source/Editor/UnrealEd/Private/Cooker/CookWorkerServer.cpp
--- a/Engine/Source/Editor/UnrealEd/Private/Cooker/CookWorkerServer.cpp
+++ b/Engine/Source/Editor/UnrealEd/Private/Cooker/CookWorkerServer.cpp
@@ -1017,20 +1017,6 @@ void FCookWorkerServer::QueueDiscoveredPackage(FDiscoveredPackageReplication&& D
 	FDiscoveredPlatformSet& Platforms = DiscoveredPackage.Platforms;
 	FPackageData& PackageData = PackageDatas.FindOrAddPackageData(DiscoveredPackage.PackageName,
 		DiscoveredPackage.NormalizedFileName);
-	if (!DiscoveredPackage.ParentGenerator.IsNone())
-	{
-		PackageData.SetGenerated(DiscoveredPackage.ParentGenerator);
-		PackageData.SetDoesGeneratedRequireGenerator(DiscoveredPackage.DoesGeneratedRequireGenerator);
-		FPackageData* GeneratorPackageData = PackageDatas.FindPackageDataByPackageName(
-			DiscoveredPackage.ParentGenerator);
-		if (GeneratorPackageData)
-		{
-			TRefCountPtr<FGenerationHelper> GenerationHelper =
-				GeneratorPackageData->CreateUninitializedGenerationHelper();
-			GenerationHelper->NotifyStartQueueGeneratedPackages(COTFS, WorkerId);
-			GenerationHelper->TrackGeneratedPackageListedRemotely(COTFS, PackageData);
-		}
-	}
 
 	TArray<const ITargetPlatform*, TInlineAllocator<ExpectedMaxNumPlatforms>> BufferPlatforms;
 	TConstArrayView<const ITargetPlatform*> DiscoveredPlatforms;
@@ -1067,6 +1053,26 @@ void FCookWorkerServer::QueueDiscoveredPackage(FDiscoveredPackageReplication&& D
 		return;
 	}
 
+	if (!DiscoveredPackage.ParentGenerator.IsNone())
+	{
+		// Registration of the discovered Generated package with its generator needs to come after we early-exit
+		// for already discovered packages, because when one generated package can refer to another from the same
+		// generator, the message that a CookWorker has discovered the referred-to generated package can show up
+		// on the director AFTER all save messages have already been processed and the GenerationHelper has shut
+		// down and destroyed its information about the list of generated packages.
+		PackageData.SetGenerated(DiscoveredPackage.ParentGenerator);
+		PackageData.SetDoesGeneratedRequireGenerator(DiscoveredPackage.DoesGeneratedRequireGenerator);
+		FPackageData* GeneratorPackageData = PackageDatas.FindPackageDataByPackageName(
+			DiscoveredPackage.ParentGenerator);
+		if (GeneratorPackageData)
+		{
+			TRefCountPtr<FGenerationHelper> GenerationHelper =
+				GeneratorPackageData->CreateUninitializedGenerationHelper();
+			GenerationHelper->NotifyStartQueueGeneratedPackages(COTFS, WorkerId);
+			GenerationHelper->TrackGeneratedPackageListedRemotely(COTFS, PackageData);
+		}
+	}
+
 	if (PackageData.IsGenerated()
 		&& (PackageData.DoesGeneratedRequireGenerator() >= ICookPackageSplitter::EGeneratedRequiresGenerator::Save
 				|| COTFS.MPCookGeneratorSplit == EMPCookGeneratorSplit::AllOnSameWorker))
