HWRT: Reducing RHI_RAYTRACING throughout the engine - Editor

diff --git a/Engine/Plugins/Runtime/nDisplay/Source/DisplayCluster/Private/Render/Viewport/DisplayClusterViewport.cpp b/Engine/Plugins/Runtime/nDisplay/Source/DisplayCluster/Private/Render/Viewport/DisplayClusterViewport.cpp
--- a/Engine/Plugins/Runtime/nDisplay/Source/DisplayCluster/Private/Render/Viewport/DisplayClusterViewport.cpp
+++ b/Engine/Plugins/Runtime/nDisplay/Source/DisplayCluster/Private/Render/Viewport/DisplayClusterViewport.cpp
@@ -339,7 +339,6 @@ void FDisplayClusterViewport::SetupSceneView(uint32 ContextNum, class UWorld* Wo
 	}
 
 	// Disable raytracing for lightcard and chromakey
-#if RHI_RAYTRACING
 	switch (RenderSettings.CaptureMode)
 	{
 	case EDisplayClusterViewportCaptureMode::Chromakey:
@@ -350,7 +349,6 @@ void FDisplayClusterViewport::SetupSceneView(uint32 ContextNum, class UWorld* Wo
 	default:
 		break;
 	}
-#endif // RHI_RAYTRACING
 
 	// Apply visibility settigns to view
 	VisibilitySettings.SetupSceneView(World, InOutView);
diff --git a/Engine/Source/Developer/ShaderCompilerCommon/Private/ShaderCompilerCommon.cpp b/Engine/Source/Developer/ShaderCompilerCommon/Private/ShaderCompilerCommon.cpp
--- a/Engine/Source/Developer/ShaderCompilerCommon/Private/ShaderCompilerCommon.cpp
+++ b/Engine/Source/Developer/ShaderCompilerCommon/Private/ShaderCompilerCommon.cpp
@@ -1435,12 +1435,10 @@ FString CreateShaderCompilerWorkerDirectCommandLine(const FShaderCompilerInput&
 	case SF_Geometry:		Text += TEXT(" -gs"); break;
 	case SF_Pixel:			Text += TEXT(" -ps"); break;
 	case SF_Compute:		Text += TEXT(" -cs"); break;
-#if RHI_RAYTRACING
 	case SF_RayGen:			Text += TEXT(" -rgs"); break;
 	case SF_RayMiss:		Text += TEXT(" -rms"); break;
 	case SF_RayHitGroup:	Text += TEXT(" -rhs"); break;
 	case SF_RayCallable:	Text += TEXT(" -rcs"); break;
-#endif // RHI_RAYTRACING
 	case SF_WorkGraphRoot:			Text += TEXT(" -wrs"); break;
 	case SF_WorkGraphComputeNode:	Text += TEXT(" -wcs"); break;
 	default: break;
diff --git a/Engine/Source/Editor/DetailCustomizations/Private/DebugCameraControllerSettingsCustomization.cpp b/Engine/Source/Editor/DetailCustomizations/Private/DebugCameraControllerSettingsCustomization.cpp
--- a/Engine/Source/Editor/DetailCustomizations/Private/DebugCameraControllerSettingsCustomization.cpp
+++ b/Engine/Source/Editor/DetailCustomizations/Private/DebugCameraControllerSettingsCustomization.cpp
@@ -51,13 +51,13 @@ void FDebugCameraControllerSettingsViewModeIndexCustomization::CustomizeHeader(T
 	EnumRestriction->AddHiddenValue(ViewModeIndexEnum->GetNameStringByValue((uint8)EViewModeIndex::VMI_VisualizeGPUSkinCache));
 	EnumRestriction->AddHiddenValue(ViewModeIndexEnum->GetNameStringByValue((uint8)EViewModeIndex::VMI_LWCComplexity));
 	EnumRestriction->AddHiddenValue(ViewModeIndexEnum->GetNameStringByValue((uint8)EViewModeIndex::VMI_StationaryLightOverlap));
-#if RHI_RAYTRACING
+
 	if (!GRHISupportsRayTracing || !GRHISupportsRayTracingShaders)
 	{
 		EnumRestriction->AddHiddenValue(ViewModeIndexEnum->GetNameStringByValue((uint8)EViewModeIndex::VMI_PathTracing));
 		EnumRestriction->AddHiddenValue(ViewModeIndexEnum->GetNameStringByValue((uint8)EViewModeIndex::VMI_RayTracingDebug));
 	}
-#endif
+
 	ViewModeIndexHandle->AddRestriction(EnumRestriction.ToSharedRef());
 
 	HeaderRow
diff --git a/Engine/Source/Editor/UnrealEd/Private/EditorViewportClient.cpp b/Engine/Source/Editor/UnrealEd/Private/EditorViewportClient.cpp
--- a/Engine/Source/Editor/UnrealEd/Private/EditorViewportClient.cpp
+++ b/Engine/Source/Editor/UnrealEd/Private/EditorViewportClient.cpp
@@ -4072,9 +4072,7 @@ void FEditorViewportClient::SetupViewForRendering(FSceneViewFamily& ViewFamily,
 	View.CurrentGroomVisualizationMode = CurrentGroomVisualizationMode;
 	View.CurrentVirtualShadowMapVisualizationMode = CurrentVirtualShadowMapVisualizationMode;
 	View.CurrentGPUSkinCacheVisualizationMode = CurrentGPUSkinCacheVisualizationMode;
-#if RHI_RAYTRACING
 	View.CurrentRayTracingDebugVisualizationMode = CurrentRayTracingDebugVisualizationMode;
-#endif
 
 	// assign wireframe opacity to the view
 	GetMeshEdgesViewSettings(View).Opacity = WireframeOpacity;
diff --git a/Engine/Source/Editor/UnrealEd/Private/EditorViewportCommands.cpp b/Engine/Source/Editor/UnrealEd/Private/EditorViewportCommands.cpp
--- a/Engine/Source/Editor/UnrealEd/Private/EditorViewportCommands.cpp
+++ b/Engine/Source/Editor/UnrealEd/Private/EditorViewportCommands.cpp
@@ -65,11 +65,11 @@ void FEditorViewportCommands::RegisterCommands()
 	UI_COMMAND( UnlitMode, "Unlit View Mode", "Renders the scene with no lights", EUserInterfaceActionType::RadioButton, FInputChord( EModifierKey::Alt, EKeys::Three ) );
 	UI_COMMAND( LitMode, "Lit View Mode", "Renders the scene with normal lighting", EUserInterfaceActionType::RadioButton, FInputChord( EModifierKey::Alt, EKeys::Four ) );
 	UI_COMMAND( LitWireframeMode, "Lit Wireframe View Mode", "Renders the scene with normal lighting and a wireframe overlay", EUserInterfaceActionType::RadioButton, FInputChord(EModifierKey::Alt, EKeys::Nine));
-#if RHI_RAYTRACING
+
 	UI_COMMAND( PathTracingMode, "PathTracing View Mode", "Renders the scene using a path tracer", EUserInterfaceActionType::RadioButton, FInputChord());
 	UI_COMMAND( RayTracingDebugMode, "Ray tracing Debug View Mode", "Renders the scene in ray tracing debug mode", EUserInterfaceActionType::RadioButton, FInputChord());
 	FRayTracingDebugVisualizationMenuCommands::Register();
-#endif
+
 	UI_COMMAND( DetailLightingMode, "Detail Lighting View Mode", "Renders the scene with detailed lighting only", EUserInterfaceActionType::RadioButton, FInputChord( EModifierKey::Alt, EKeys::Five ) );
 	UI_COMMAND( LightingOnlyMode, "Lighting Only View Mode", "Renders the scene with lights only, no textures", EUserInterfaceActionType::RadioButton, FInputChord( EModifierKey::Alt, EKeys::Six ) );
 	UI_COMMAND( LightComplexityMode, "Light Complexity View Mode", "Renders the scene with light complexity visualization", EUserInterfaceActionType::RadioButton, FInputChord( EModifierKey::Alt, EKeys::Seven ) );
diff --git a/Engine/Source/Editor/UnrealEd/Private/SEditorViewport.cpp b/Engine/Source/Editor/UnrealEd/Private/SEditorViewport.cpp
--- a/Engine/Source/Editor/UnrealEd/Private/SEditorViewport.cpp
+++ b/Engine/Source/Editor/UnrealEd/Private/SEditorViewport.cpp
@@ -417,7 +417,7 @@ void SEditorViewport::BindCommands()
 	MAP_VIEWMODE_ACTION( Commands.UnlitMode, VMI_Unlit );
 	MAP_VIEWMODE_ACTION( Commands.LitMode, VMI_Lit );
 	MAP_VIEWMODE_ACTION( Commands.LitWireframeMode, VMI_Lit_Wireframe);
-#if RHI_RAYTRACING
+
 	if (IsRayTracingAllowed())
 	{
 		MAP_VIEWMODE_ACTION(Commands.PathTracingMode, VMI_PathTracing);
@@ -426,7 +426,7 @@ void SEditorViewport::BindCommands()
 		const FRayTracingDebugVisualizationMenuCommands& RtDebugCommands = FRayTracingDebugVisualizationMenuCommands::Get();
 		RtDebugCommands.BindCommands(CommandListRef, Client);
 	}
-#endif
+
 	MAP_VIEWMODE_ACTION( Commands.DetailLightingMode, VMI_Lit_DetailLighting );
 	MAP_VIEWMODE_ACTION( Commands.LightingOnlyMode, VMI_LightingOnly );
 	MAP_VIEWMODE_ACTION( Commands.LightComplexityMode, VMI_LightComplexity );
diff --git a/Engine/Source/Editor/UnrealEd/Private/SEditorViewportViewMenu.cpp b/Engine/Source/Editor/UnrealEd/Private/SEditorViewportViewMenu.cpp
--- a/Engine/Source/Editor/UnrealEd/Private/SEditorViewportViewMenu.cpp
+++ b/Engine/Source/Editor/UnrealEd/Private/SEditorViewportViewMenu.cpp
@@ -144,7 +144,6 @@ void SEditorViewportViewMenu::FillViewMenu(UToolMenu* Menu) const
 				Section.AddMenuEntry(BaseViewportActions.CollisionVisibility, UViewModeUtils::GetViewModeDisplayName(VMI_CollisionVisibility));
 			}
 
-#if RHI_RAYTRACING
 			if (IsRayTracingEnabled())
 			{
 				static auto PathTracingCvar = IConsoleManager::Get().FindTConsoleVariableDataInt(TEXT("r.PathTracing"));
@@ -155,7 +154,6 @@ void SEditorViewportViewMenu::FillViewMenu(UToolMenu* Menu) const
 					Section.AddMenuEntry(BaseViewportActions.PathTracingMode, UViewModeUtils::GetViewModeDisplayName(VMI_PathTracing));
 				}
 			}
-#endif
 
 			// Optimization
 			{
@@ -250,7 +248,6 @@ void SEditorViewportViewMenu::FillViewMenu(UToolMenu* Menu) const
 					/* bInOpenSubMenuOnClick = */ false, FSlateIcon(FAppStyle::GetAppStyleSetName(), "EditorViewport.QuadOverdrawMode"));
 			}
 
-#if RHI_RAYTRACING
 			if (IsRayTracingEnabled())
 			{
 				struct Local
@@ -264,7 +261,6 @@ void SEditorViewportViewMenu::FillViewMenu(UToolMenu* Menu) const
 
 				Section.AddSubMenu("RayTracingDebugSubMenu", LOCTEXT("RayTracingDebugSubMenu", "Ray Tracing Debug"), LOCTEXT("RayTracing_ToolTip", "Select ray tracing buffer visualization view modes"), FNewMenuDelegate::CreateStatic(&Local::BuildRayTracingDebugMenu, ParentToolBar), false, FSlateIcon(FAppStyle::GetAppStyleSetName(), "EditorViewport.RayTracingDebugMode"));
 			}
-#endif
 
 			{
 				struct Local
diff --git a/Engine/Source/Editor/UnrealEd/Public/EditorViewportCommands.h b/Engine/Source/Editor/UnrealEd/Public/EditorViewportCommands.h
--- a/Engine/Source/Editor/UnrealEd/Public/EditorViewportCommands.h
+++ b/Engine/Source/Editor/UnrealEd/Public/EditorViewportCommands.h
@@ -71,10 +71,8 @@ public:
 	/** Changes the viewport to lit wireframe mode */
 	TSharedPtr< FUICommandInfo > LitWireframeMode;
 
-#if RHI_RAYTRACING
 	TSharedPtr< FUICommandInfo > PathTracingMode;
 	TSharedPtr< FUICommandInfo > RayTracingDebugMode;
-#endif
 
 	/** Changes the viewport to detail lighting mode */
 	TSharedPtr< FUICommandInfo > DetailLightingMode;
diff --git a/Engine/Source/Programs/ShaderCompileWorker/Private/ShaderCompileWorker.cpp b/Engine/Source/Programs/ShaderCompileWorker/Private/ShaderCompileWorker.cpp
--- a/Engine/Source/Programs/ShaderCompileWorker/Private/ShaderCompileWorker.cpp
+++ b/Engine/Source/Programs/ShaderCompileWorker/Private/ShaderCompileWorker.cpp
@@ -942,7 +942,6 @@ static void DirectCompile(const TArray<const class IShaderFormat*>& ShaderFormat
 			{
 				Frequency = SF_Compute;
 			}
-#if RHI_RAYTRACING
 			else if (!FCString::Strcmp(*Token, TEXT("rgs")))
 			{
 				Frequency = SF_RayGen;
@@ -959,7 +958,6 @@ static void DirectCompile(const TArray<const class IShaderFormat*>& ShaderFormat
 			{
 				Frequency = SF_RayCallable;
 			}
-#endif // RHI_RAYTRACING
 			else if (!FCString::Strcmp(*Token, TEXT("wrs")))
 			{
 				Frequency = SF_WorkGraphRoot;
diff --git a/Engine/Source/Runtime/Engine/Private/GPUSkinCacheVisualizationData.cpp b/Engine/Source/Runtime/Engine/Private/GPUSkinCacheVisualizationData.cpp
--- a/Engine/Source/Runtime/Engine/Private/GPUSkinCacheVisualizationData.cpp
+++ b/Engine/Source/Runtime/Engine/Private/GPUSkinCacheVisualizationData.cpp
@@ -30,7 +30,6 @@ void FGPUSkinCacheVisualizationData::Initialize()
 			LOCTEXT("MemoryDesc", "Memory usage"),
 			FModeType::Memory);
 
-#if RHI_RAYTRACING
 		if (IsRayTracingEnabled())
 		{
 			AddVisualizationMode(
@@ -39,7 +38,6 @@ void FGPUSkinCacheVisualizationData::Initialize()
 				LOCTEXT("RayTracingLODOffsetDesc", "Ray Tracing LOD index offset from rasterization"),
 				FModeType::RayTracingLODOffset);
 		}
-#endif
 
 		ConfigureConsoleCommand();
 
diff --git a/Engine/Source/Runtime/Engine/Private/GameViewportClient.cpp b/Engine/Source/Runtime/Engine/Private/GameViewportClient.cpp
--- a/Engine/Source/Runtime/Engine/Private/GameViewportClient.cpp
+++ b/Engine/Source/Runtime/Engine/Private/GameViewportClient.cpp
@@ -3673,7 +3673,6 @@ bool UGameViewportClient::HandleViewModeCommand( const TCHAR* Cmd, FOutputDevice
 		ViewModeIndex = VMI_Lit;
 	}
 
-#if RHI_RAYTRACING
 	if (!GRHISupportsRayTracing || !GRHISupportsRayTracingShaders)
 	{
 		if (ViewModeIndex == VMI_PathTracing)
@@ -3688,7 +3687,6 @@ bool UGameViewportClient::HandleViewModeCommand( const TCHAR* Cmd, FOutputDevice
 			ViewModeIndex = VMI_Lit;
 		}
 	}
-#endif
 #endif
 
 	ApplyViewMode((EViewModeIndex)ViewModeIndex, true, EngineShowFlags);
diff --git a/Engine/Source/Runtime/Engine/Public/SceneView.h b/Engine/Source/Runtime/Engine/Public/SceneView.h
--- a/Engine/Source/Runtime/Engine/Public/SceneView.h
+++ b/Engine/Source/Runtime/Engine/Public/SceneView.h
@@ -1614,10 +1614,8 @@ public:
 	/** Feature level for this scene */
 	const ERHIFeatureLevel::Type FeatureLevel;
 
-#if RHI_RAYTRACING
 	/** Use to allow ray tracing on this view. */
 	bool bAllowRayTracing = true;
-#endif
 
 protected:
 	friend class FSceneRenderer;
@@ -1815,10 +1813,8 @@ public:
 		const FViewMatrices& InViewMatrices,
 		const FViewMatrices& InPrevViewMatrices) const;
 
-#if RHI_RAYTRACING
 	/** Current ray tracing debug visualization mode */
 	FName CurrentRayTracingDebugVisualizationMode;
-#endif
 
 	UE_DEPRECATED(5.2, "Use HasValidEyeAdaptationBuffer() instead.")
 	ENGINE_API bool HasValidEyeAdaptationTexture() const;
diff --git a/Engine/Source/Runtime/Landscape/Private/LandscapeGrass.cpp b/Engine/Source/Runtime/Landscape/Private/LandscapeGrass.cpp
--- a/Engine/Source/Runtime/Landscape/Private/LandscapeGrass.cpp
+++ b/Engine/Source/Runtime/Landscape/Private/LandscapeGrass.cpp
@@ -203,12 +203,10 @@ static FAutoConsoleVariableRef CVarUpdateAllOnRebuild(
 	GGrassUpdateAllOnRebuild,
 	TEXT(""));
 
-#if RHI_RAYTRACING
 static TAutoConsoleVariable<int32> CVarRayTracingLandscapeGrass(
 	TEXT("r.RayTracing.Geometry.LandscapeGrass"),
 	0,
 	TEXT("Include landscapes grass in ray tracing effects (default = 1)"));
-#endif
 
 const TCHAR* GGrassQualityLevelCVarName = TEXT("r.grass.DensityQualityLevel");
 const TCHAR* GGrassQualityLevelScalabilitySection = TEXT("ViewDistanceQuality");
diff --git a/Engine/Source/Runtime/Renderer/Private/CompositionLighting/CompositionLighting.cpp b/Engine/Source/Runtime/Renderer/Private/CompositionLighting/CompositionLighting.cpp
--- a/Engine/Source/Runtime/Renderer/Private/CompositionLighting/CompositionLighting.cpp
+++ b/Engine/Source/Runtime/Renderer/Private/CompositionLighting/CompositionLighting.cpp
@@ -93,9 +93,9 @@ bool ShouldRenderScreenSpaceAmbientOcclusion(const FViewInfo& View, bool bLumenW
 		&& View.FinalPostProcessSettings.AmbientOcclusionRadius >= 0.1f
 		&& !View.Family->UseDebugViewPS()
 		&& (FSSAOHelper::IsBasePassAmbientOcclusionRequired(View) || IsAmbientCubemapPassRequired(View) || IsReflectionEnvironmentActive(View) || IsSkylightActive(View) || IsForwardShadingEnabled(View.GetShaderPlatform()) || View.Family->EngineShowFlags.VisualizeBuffer || bLumenWantsSSAO);
-#if RHI_RAYTRACING
+
 	bEnabled &= !ShouldRenderRayTracingAmbientOcclusion(View);
-#endif
+
 	return bEnabled;
 }
 
diff --git a/Engine/Source/Runtime/Renderer/Private/ShadingEnergyConservation.cpp b/Engine/Source/Runtime/Renderer/Private/ShadingEnergyConservation.cpp
--- a/Engine/Source/Runtime/Renderer/Private/ShadingEnergyConservation.cpp
+++ b/Engine/Source/Runtime/Renderer/Private/ShadingEnergyConservation.cpp
@@ -227,7 +227,7 @@ void Init(FRDGBuilder& GraphBuilder, FViewInfo& View)
 	// Build/bind table if energy conservation is enabled or if Substrate is enabled in order to have 
 	// the correct tables built & bound. Even if we are not using energy conservation, we want to 
 	// have access to directional albedo information for env. lighting for instance)
-	const bool bBindEnergyData = (View.ViewState != nullptr) && (bMaterialEnergyConservationEnabled || Substrate::IsSubstrateEnabled() || (View.Family->EngineShowFlags.PathTracing && RHI_RAYTRACING)) && (bIsEnergyPreservationEnabled || bIsEnergyConservationEnabled);
+	const bool bBindEnergyData = (View.ViewState != nullptr) && (bMaterialEnergyConservationEnabled || Substrate::IsSubstrateEnabled() || (View.Family->EngineShowFlags.PathTracing)) && (bIsEnergyPreservationEnabled || bIsEnergyConservationEnabled);
 	if (bBindEnergyData)
 	{
 		// Change this to true in order to regenerate the energy tables, and manually copy the coefficients into ShadingEnergyConservationData.h
diff --git a/Engine/Source/Runtime/Renderer/Private/TranslucentRendering.cpp b/Engine/Source/Runtime/Renderer/Private/TranslucentRendering.cpp
--- a/Engine/Source/Runtime/Renderer/Private/TranslucentRendering.cpp
+++ b/Engine/Source/Runtime/Renderer/Private/TranslucentRendering.cpp
@@ -178,12 +178,11 @@ EMeshPass::Type TranslucencyPassToMeshPass(ETranslucencyPass::Type TranslucencyP
 
 ETranslucencyView GetTranslucencyView(const FViewInfo& View)
 {
-#if RHI_RAYTRACING
 	if (ShouldRenderRayTracingTranslucency(View))
 	{
 		return ETranslucencyView::RayTracing;
 	}
-#endif
+
 	return View.IsUnderwater() ? ETranslucencyView::UnderWater : ETranslucencyView::AboveWater;
 }
 
diff --git a/Engine/Source/Runtime/Renderer/Public/RayTracingInstance.h b/Engine/Source/Runtime/Renderer/Public/RayTracingInstance.h
--- a/Engine/Source/Runtime/Renderer/Public/RayTracingInstance.h
+++ b/Engine/Source/Runtime/Renderer/Public/RayTracingInstance.h
@@ -5,8 +5,6 @@
 #include "RenderResource.h"
 #include "MeshBatch.h"
 
-#if RHI_RAYTRACING
-
 class FRayTracingGeometry;
 
 struct FRayTracingMaskAndFlags
@@ -168,5 +166,3 @@ struct FRayTracingInstance
 	~FRayTracingInstance() = default;
 	PRAGMA_ENABLE_DEPRECATION_WARNINGS
 };
-
-#endif
