MPCook: Fix GenerationHelper DiagnoseWhyNotShutdown diagnostic to only fire when initialized; it's not a problem when the helper is not initialized. Firing when uninitialized can occur when using iterative cook (which creates uninitialized generation helpers at startup) and not cooking all packages due to -cooksinglepackagenorefs.
#rnx

diff --git a/Engine/Source/Editor/UnrealEd/Private/CookOnTheFlyServer.cpp b/Engine/Source/Editor/UnrealEd/Private/CookOnTheFlyServer.cpp
--- a/Engine/Source/Editor/UnrealEd/Private/CookOnTheFlyServer.cpp
+++ b/Engine/Source/Editor/UnrealEd/Private/CookOnTheFlyServer.cpp
@@ -10164,7 +10164,8 @@ void UCookOnTheFlyServer::CookByTheBookFinishedInternal()
 	TArray<FPackageData*> DanglingGenerationHelpers;
 	PackageDatas->LockAndEnumeratePackageDatas([&DanglingGenerationHelpers](FPackageData* PackageData)
 		{
-			if (PackageData->GetGenerationHelper())
+			TRefCountPtr<FGenerationHelper> GenerationHelper = PackageData->GetGenerationHelper();
+			if (GenerationHelper && GenerationHelper->IsInitialized())
 			{
 				DanglingGenerationHelpers.Add(PackageData);
 			}
