Support reordering of variables

diff --git a/Engine/Source/Editor/UnrealEd/Private/Kismet2/BlueprintEditorUtils.cpp b/Engine/Source/Editor/UnrealEd/Private/Kismet2/BlueprintEditorUtils.cpp
--- a/Engine/Source/Editor/UnrealEd/Private/Kismet2/BlueprintEditorUtils.cpp
+++ b/Engine/Source/Editor/UnrealEd/Private/Kismet2/BlueprintEditorUtils.cpp
@@ -3527,7 +3527,8 @@ int32 FBlueprintEditorUtils::FindLocalVariableIndex(const UBlueprint* Blueprint,
 	return INDEX_NONE;
 }
 
-bool FBlueprintEditorUtils::MoveVariableBeforeVariable(UBlueprint* Blueprint, UStruct* VariableScope, FName VarNameToMove, FName TargetVarName, bool bDontRecompile)
+/** Helper function for moving variables around in the list relative to a target variable */
+static bool MoveVariableRelativeToVariable(UBlueprint* Blueprint, UStruct* VariableScope, FName VarNameToMove, FName TargetVarName, bool bDontRecompile, bool bInsertBeforeTargetVariable)
 {
 	check(Blueprint && VariableScope);
 
@@ -3538,13 +3539,13 @@ bool FBlueprintEditorUtils::MoveVariableBeforeVariable(UBlueprint* Blueprint, US
 	//Get the indices of the variables to be re-ordered
 	if (VariableScope->IsA(UFunction::StaticClass()))
 	{
-		VarIndexToMove = FindLocalVariableIndex(Blueprint, VariableScope, VarNameToMove);
-		TargetVarIndex = FindLocalVariableIndex(Blueprint, VariableScope, TargetVarName);
+		VarIndexToMove = FBlueprintEditorUtils::FindLocalVariableIndex(Blueprint, VariableScope, VarNameToMove);
+		TargetVarIndex = FBlueprintEditorUtils::FindLocalVariableIndex(Blueprint, VariableScope, TargetVarName);
 	}
 	else
 	{
-		VarIndexToMove = FindNewVariableIndex(Blueprint, VarNameToMove);
-		TargetVarIndex = FindNewVariableIndex(Blueprint, TargetVarName);
+		VarIndexToMove = FBlueprintEditorUtils::FindNewVariableIndex(Blueprint, VarNameToMove);
+		TargetVarIndex = FBlueprintEditorUtils::FindNewVariableIndex(Blueprint, TargetVarName);
 	}
 
 	if (VarIndexToMove != INDEX_NONE && TargetVarIndex != INDEX_NONE)
@@ -3559,7 +3560,7 @@ bool FBlueprintEditorUtils::MoveVariableBeforeVariable(UBlueprint* Blueprint, US
 		if (VariableScope->IsA(UFunction::StaticClass()))
 		{
 			UK2Node_FunctionEntry* FunctionEntryNode = nullptr;
-			FindLocalVariable(Blueprint, VariableScope, VarNameToMove, &FunctionEntryNode);
+			FBlueprintEditorUtils::FindLocalVariable(Blueprint, VariableScope, VarNameToMove, &FunctionEntryNode);
 
 			if (FunctionEntryNode != nullptr)
 			{
@@ -3569,7 +3570,7 @@ bool FBlueprintEditorUtils::MoveVariableBeforeVariable(UBlueprint* Blueprint, US
 				// Remove var we are moving
 				FunctionEntryNode->LocalVariables.RemoveAt(VarIndexToMove);
 				// Add in before target variable
-				FunctionEntryNode->LocalVariables.Insert(MoveVar, TargetVarIndex);				
+				FunctionEntryNode->LocalVariables.Insert(MoveVar, bInsertBeforeTargetVariable ? TargetVarIndex : TargetVarIndex + 1);
 			}
 		}
 		else
@@ -3580,7 +3581,7 @@ bool FBlueprintEditorUtils::MoveVariableBeforeVariable(UBlueprint* Blueprint, US
 			// Remove var we are moving
 			Blueprint->NewVariables.RemoveAt(VarIndexToMove);
 			// Add in before target variable
-			Blueprint->NewVariables.Insert(MoveVar, TargetVarIndex);
+			Blueprint->NewVariables.Insert(MoveVar, bInsertBeforeTargetVariable ? TargetVarIndex : TargetVarIndex + 1);
 		}
 
 		if (!bDontRecompile)
@@ -3592,6 +3593,16 @@ bool FBlueprintEditorUtils::MoveVariableBeforeVariable(UBlueprint* Blueprint, US
 	return bMoved;
 }
 
+bool FBlueprintEditorUtils::MoveVariableBeforeVariable(UBlueprint* Blueprint, UStruct* VariableScope, FName VarNameToMove, FName TargetVarName, bool bDontRecompile)
+{
+	return MoveVariableRelativeToVariable(Blueprint, VariableScope, VarNameToMove, TargetVarName, bDontRecompile, true);
+}
+
+bool FBlueprintEditorUtils::MoveVariableAfterVariable(UBlueprint* Blueprint, UStruct* VariableScope, FName VarNameToMove, FName TargetVarName, bool bDontRecompile)
+{
+	return MoveVariableRelativeToVariable(Blueprint, VariableScope, VarNameToMove, TargetVarName, bDontRecompile, false);
+}
+
 int32 FBlueprintEditorUtils::FindTimelineIndex(const UBlueprint* Blueprint, const FName& InName) 
 {
 	const FName TimelineTemplateName = *UTimelineTemplate::TimelineVariableNameToTemplateName(InName);
diff --git a/Engine/Source/Editor/UnrealEd/Public/Kismet2/BlueprintEditorUtils.h b/Engine/Source/Editor/UnrealEd/Public/Kismet2/BlueprintEditorUtils.h
--- a/Engine/Source/Editor/UnrealEd/Public/Kismet2/BlueprintEditorUtils.h
+++ b/Engine/Source/Editor/UnrealEd/Public/Kismet2/BlueprintEditorUtils.h
@@ -1272,9 +1272,12 @@ public:
 	 */
 	static UNREALED_API int32 FindLocalVariableIndex(const UBlueprint* Blueprint, UStruct* VariableScope, const FName& InVariableName);
 
-	/** Change the order of variables in the Blueprint */
+	/** Change the order of variables in the Blueprint by placing a variable before the target variable. */
 	static UNREALED_API bool MoveVariableBeforeVariable(UBlueprint* Blueprint, UStruct* VariableScope, FName VarNameToMove, FName TargetVarName, bool bDontRecompile);
 
+	/** Change the order of variables in the Blueprint by placing a variable after the target variable. */
+	static UNREALED_API bool MoveVariableAfterVariable(UBlueprint* Blueprint, UStruct* VariableScope, FName VarNameToMove, FName TargetVarName, bool bDontRecompile);
+
 	/**
 	 * Find the index of a timeline first declared in this blueprint. Returns INDEX_NONE if not found.
 	 *
