Hide Preview Platform indicator when we hide viewport UI

diff --git a/Engine/Source/Editor/UnrealEd/Private/SEditorViewport.cpp b/Engine/Source/Editor/UnrealEd/Private/SEditorViewport.cpp
--- a/Engine/Source/Editor/UnrealEd/Private/SEditorViewport.cpp
+++ b/Engine/Source/Editor/UnrealEd/Private/SEditorViewport.cpp
@@ -933,7 +933,7 @@ TSharedRef<SWidget> SEditorViewport::BuildFeatureLevelWidget() const
 
 EVisibility SEditorViewport::GetCurrentFeatureLevelPreviewTextVisibility() const
 {
-	if (Client->GetWorld())
+	if (Client->GetWorld() && !GLevelEditorModeTools().IsViewportUIHidden())
 	{
 		return (GEditor && GEditor->IsFeatureLevelPreviewActive()) ? EVisibility::SelfHitTestInvisible : EVisibility::Collapsed;
 	}
