we now have other chunktype hashes available with non zero indices so for now limit to bulk data.

diff --git a/Engine/Source/Editor/UnrealEd/Private/Commandlets/DiffAssetBulkDataCommandlet.cpp b/Engine/Source/Editor/UnrealEd/Private/Commandlets/DiffAssetBulkDataCommandlet.cpp
--- a/Engine/Source/Editor/UnrealEd/Private/Commandlets/DiffAssetBulkDataCommandlet.cpp
+++ b/Engine/Source/Editor/UnrealEd/Private/Commandlets/DiffAssetBulkDataCommandlet.cpp
@@ -185,6 +185,11 @@ int32 UDiffAssetBulkDataCommandlet::Main(const FString& FullCommandLine)
 
 		for (const TPair<FIoChunkId, FIoHash>& ChunkHashPair : BasePackage->ChunkHashes)
 		{
+			if (ChunkHashPair.Key.GetChunkType() != EIoChunkType::BulkData &&
+				ChunkHashPair.Key.GetChunkType() != EIoChunkType::OptionalBulkData &&
+				ChunkHashPair.Key.GetChunkType() != EIoChunkType::MemoryMappedBulkData)
+				continue;
+
 			const FIoHash* CurrentHash = CurrentPackage->ChunkHashes.Find(ChunkHashPair.Key);
 			if (CurrentHash == nullptr)
 			{
@@ -204,6 +209,11 @@ int32 UDiffAssetBulkDataCommandlet::Main(const FString& FullCommandLine)
 
 		for (const TPair<FIoChunkId, FIoHash>& ChunkHashPair : CurrentPackage->ChunkHashes)
 		{
+			if (ChunkHashPair.Key.GetChunkType() != EIoChunkType::BulkData &&
+				ChunkHashPair.Key.GetChunkType() != EIoChunkType::OptionalBulkData &&
+				ChunkHashPair.Key.GetChunkType() != EIoChunkType::MemoryMappedBulkData)
+				continue;
+
 			if (BasePackage->ChunkHashes.Contains(ChunkHashPair.Key) == false)
 			{
 				TArray<FName>& New = NewChunksByType[(uint32)ChunkHashPair.Key.GetChunkType()];
