Maintain looping state in anim editor when switching between assets

diff --git a/Engine/Source/Editor/UnrealEd/Private/Animation/DebugSkelMeshComponent.cpp b/Engine/Source/Editor/UnrealEd/Private/Animation/DebugSkelMeshComponent.cpp
--- a/Engine/Source/Editor/UnrealEd/Private/Animation/DebugSkelMeshComponent.cpp
+++ b/Engine/Source/Editor/UnrealEd/Private/Animation/DebugSkelMeshComponent.cpp
@@ -681,7 +681,8 @@ void UDebugSkelMeshComponent::EnablePreview(bool bEnable, UAnimationAsset* Previ
 		    // restore previous state
 		    bDisableClothSimulation = bPrevDisableClothSimulation;
     
-			PreviewInstance->SetAnimationAsset(PreviewAsset); 
+			const bool bWasLooping = PreviewInstance->IsLooping();
+			PreviewInstance->SetAnimationAsset(PreviewAsset, bWasLooping);
 			
 			// Reset to previous animation asset's root motion playback time to prevent this from influencing the new animation asset previewing during root motion consumption.
 			ConsumeRootMotionPreviousPlaybackTime = 0.0f;
