Remove duplicate menu entries in asset editor viewport toolbars.

diff --git a/Engine/Source/Editor/UnrealEd/Private/SCommonEditorViewportToolbarBase.cpp b/Engine/Source/Editor/UnrealEd/Private/SCommonEditorViewportToolbarBase.cpp
--- a/Engine/Source/Editor/UnrealEd/Private/SCommonEditorViewportToolbarBase.cpp
+++ b/Engine/Source/Editor/UnrealEd/Private/SCommonEditorViewportToolbarBase.cpp
@@ -801,13 +801,6 @@ void SCommonEditorViewportToolbarBase::CreateViewMenuExtensions(FMenuBuilder& Me
 	MenuBuilder.BeginSection("LevelViewportDeferredRendering", LOCTEXT("DeferredRenderingHeader", "Deferred Rendering") );
 	MenuBuilder.EndSection();
 
-	MenuBuilder.BeginSection("LevelViewportCollision", LOCTEXT("CollisionViewModeHeader", "Collision") );
-	{
-		MenuBuilder.AddMenuEntry( FEditorViewportCommands::Get().CollisionPawn, NAME_None, LOCTEXT("CollisionPawnViewModeDisplayName", "Player Collision") );
-		MenuBuilder.AddMenuEntry( FEditorViewportCommands::Get().CollisionVisibility, NAME_None, LOCTEXT("CollisionVisibilityViewModeDisplayName", "Visibility Collision") );
-	}
-	MenuBuilder.EndSection();
-
 //FINDME
 // 	MenuBuilder.BeginSection("LevelViewportLandscape", LOCTEXT("LandscapeHeader", "Landscape") );
 // 	{
