MPCook: Log replication: Call FMsg::LogF instead of GLog->CategorizedLogf so that messages at display verbosity and above from CookWorkers are routed into GWarn instead of GLog. GWarn handles writing to stdout and collection of errors for display in the Warning/Error Summary at the end of the commandlet.
#rnx

diff --git a/Engine/Source/Editor/UnrealEd/Private/Cooker/CookWorkerServer.cpp b/Engine/Source/Editor/UnrealEd/Private/Cooker/CookWorkerServer.cpp
--- a/Engine/Source/Editor/UnrealEd/Private/Cooker/CookWorkerServer.cpp
+++ b/Engine/Source/Editor/UnrealEd/Private/Cooker/CookWorkerServer.cpp
@@ -1604,7 +1604,8 @@ void FLogMessagesMessageHandler::ServerReceiveMessage(FMPCollectorServerMessageC
 			// Do not spam heartbeat messages into the CookDirector log
 			continue;
 		}
-		GLog->CategorizedLogf(LogData.Category, LogData.Verbosity, TEXT("[CookWorker %d]: %s"),
+
+		FMsg::Logf(__FILE__, __LINE__, LogData.Category, LogData.Verbosity, TEXT("[CookWorker %d]: %s"),
 			Context.GetProfileId(), *LogData.Message);
 	}
 }
