Cooker:
Bugfix for incremental ZenStoreWriter: it was broadcasting iterativelyunmodified for every package whether it was modified or not.
MPCook Performance bugfix for GenerationHelper on cookdirector: it was initializing them (and loading the generator package) when it did not need to in response to discovery and savepackage messages from cookworkers.
Fix some coding standard issues.
Readability improvements for some save code that is about to be extended.
#rnx

diff --git a/Engine/Source/Developer/IoStoreUtilities/Private/ZenStoreWriter.cpp b/Engine/Source/Developer/IoStoreUtilities/Private/ZenStoreWriter.cpp
--- a/Engine/Source/Developer/IoStoreUtilities/Private/ZenStoreWriter.cpp
+++ b/Engine/Source/Developer/IoStoreUtilities/Private/ZenStoreWriter.cpp
@@ -1226,6 +1226,11 @@ void FZenStoreWriter::UpdatePackageModificationStatus(FName PackageName, bool bI
 {
 	TRACE_CPUPROFILER_EVENT_SCOPE(FZenStoreWriter::UpdatePackageModificationStatus);
 
+	if (!bIterativelyUnmodified)
+	{
+		return;
+	}
+
 	IPackageStoreWriter::FMarkUpToDateEventArgs MarkUpToDateEventArgs;
 
 	{
diff --git a/Engine/Source/Editor/UnrealEd/Classes/CookOnTheSide/CookOnTheFlyServer.h b/Engine/Source/Editor/UnrealEd/Classes/CookOnTheSide/CookOnTheFlyServer.h
--- a/Engine/Source/Editor/UnrealEd/Classes/CookOnTheSide/CookOnTheFlyServer.h
+++ b/Engine/Source/Editor/UnrealEd/Classes/CookOnTheSide/CookOnTheFlyServer.h
@@ -306,7 +306,6 @@ private:
 	//////////////////////////////////////////////////////////////////////////
 	// Cook on the fly server interface adapter
 	class FCookOnTheFlyServerInterface;
-	friend class FCookOnTheFlyServerInterface;
 	TUniquePtr<FCookOnTheFlyServerInterface> CookOnTheFlyServerInterface;
 
 	/** Current cook mode the cook on the fly server is running in */
@@ -1329,7 +1328,6 @@ private:
 	void SaveCookedPackage(UE::Cook::FSaveCookedPackageContext& Context);
 	/** Helper for package saves using ExternalActors: record ExternalActors for iterative builds. */
 	void RecordExternalActorDependencies(TConstArrayView<FName> ExternalActorDependencies);
-	friend class UE::Cook::FSaveCookedPackageContext;
 
 	/**
 	 * Save the global shader map
@@ -1623,6 +1621,7 @@ private:
 	friend UE::Cook::FPendingCookedPlatformData;
 	friend UE::Cook::FPlatformManager;
 	friend UE::Cook::FRequestCluster;
+	friend UE::Cook::FSaveCookedPackageContext;
 	friend UE::Cook::FWorkerRequestsLocal;
 	friend UE::Cook::FWorkerRequestsRemote;
 };
diff --git a/Engine/Source/Editor/UnrealEd/Private/CookOnTheFlyServer.cpp b/Engine/Source/Editor/UnrealEd/Private/CookOnTheFlyServer.cpp
--- a/Engine/Source/Editor/UnrealEd/Private/CookOnTheFlyServer.cpp
+++ b/Engine/Source/Editor/UnrealEd/Private/CookOnTheFlyServer.cpp
@@ -4427,7 +4427,11 @@ private: // Used only by UCookOnTheFlyServer, which has private access
 	// private helper functions
 	void CalculatePlatformAgnosticRuntimeDependencies();
 	void CalculatePlatformRuntimeDependencies();
+	TArray<FName> GetPlatformRuntimeDependencies() const;
+	TArray<IPackageWriter::FCommitAttachmentInfo> GetCommitAttachments();
+	IPackageWriter::EWriteOptions GetCommitWriteOptions() const;
 	static void AddDependency(TMap<FPackageData*, EInstigator>& InDependencies, FPackageData* PackageData, bool bHard);
+	static IPackageWriter::ECommitStatus PackageResultToCommitStatus(FSavePackageResultStruct& Result);
 
 	// General Package Data
 	UCookOnTheFlyServer& COTFS;
@@ -6342,50 +6346,44 @@ void FSaveCookedPackageContext::FinishPlatform()
 	bool bSuccessful = SavePackageResult.IsSuccessful();
 	ECookResult CookResult = bSuccessful ? ECookResult::Succeeded : ECookResult::Failed;
 
+	// Calculate up-to-date assetregistry data for Generator and Generated packages
+	TOptional<TArray<FAssetDependency>> OverridePackageDependencies;
+	TOptional<FAssetPackageData> AssetPackageDataBuffer;
+	TOptional<FAssetPackageData> OverrideAssetPackageData;
+	const FAssetPackageData* AssetPackageData = nullptr;
+	if (TRefCountPtr<FGenerationHelper> GenerationHelper = PackageData.GetGenerationHelper();
+		GenerationHelper)
+	{
+		OverridePackageDependencies.Emplace();
+		GenerationHelper->FinishGeneratorPlatformSave(PackageData, PlatformIndex == 0, *OverridePackageDependencies);
+	}
+	else if (TRefCountPtr<FGenerationHelper> ParentGenerationHelper = PackageData.GetParentGenerationHelper();
+		ParentGenerationHelper)
+	{
+		OverridePackageDependencies.Emplace();
+		OverrideAssetPackageData.Emplace();
+		ParentGenerationHelper->FinishGeneratedPlatformSave(PackageData, *OverridePackageDependencies, *OverrideAssetPackageData);
+		AssetPackageData = OverrideAssetPackageData.GetPtrOrNull();
+	}
+	if (!AssetPackageData)
+	{
+		AssetPackageDataBuffer = COTFS.AssetRegistry->GetAssetPackageDataCopy(Package->GetFName());
+		AssetPackageData = AssetPackageDataBuffer.GetPtrOrNull();
+	}
+
+	// Commit the saved bytes and the iterative cook data to the PackageWriter
 	if (bPlatformSetupSuccessful)
 	{
 		CalculatePlatformAgnosticRuntimeDependencies();
 		CalculatePlatformRuntimeDependencies();
 
-		TOptional<FAssetPackageData> AssetPackageData = COTFS.AssetRegistry->GetAssetPackageDataCopy(Package->GetFName());
-
 		ICookedPackageWriter::FCommitPackageInfo Info;
-		if (COTFS.bHybridIterativeEnabled)
-		{
-			TArray<FName> PlatformDependencyNames;
-			PlatformDependencyNames.Reserve(PlatformDependencies[PlatformIndex].Num());
-			for (const TPair<FPackageData*, EInstigator>& DependencyPair : PlatformDependencies[PlatformIndex])
-			{
-				PlatformDependencyNames.Add(DependencyPair.Key->GetPackageName());
-			}
-			UE_SCOPED_HIERARCHICAL_COOKTIMER(TargetDomainDependencies);
-			UE::TargetDomain::CollectAndStoreCookAttachments(Package, TargetPlatform, &SavePackageResult,
-				MoveTemp(PlatformDependencyNames), Info.Attachments);
-		}
-		if (bSuccessful)
-		{
-			Info.Status = IPackageWriter::ECommitStatus::Success;
-		}
-		else if (SavePackageResult.Result == ESavePackageResult::Timeout)
-		{
-			Info.Status = IPackageWriter::ECommitStatus::Canceled;
-		}
-		else
-		{
-			Info.Status = IPackageWriter::ECommitStatus::Error;
-		}
+		// Note GetCommitAttachments mutates the SaveResult; it moves CookDependencies out of it
+		Info.Attachments = GetCommitAttachments();
+		Info.Status = PackageResultToCommitStatus(SavePackageResult);
 		Info.PackageName = Package->GetFName();
 		Info.PackageHash = AssetPackageData ? AssetPackageData->GetPackageSavedHash() : FIoHash();
-		Info.WriteOptions = IPackageWriter::EWriteOptions::None;
-		if (!COTFS.bSkipSave)
-		{
-			Info.WriteOptions |= IPackageWriter::EWriteOptions::Write;
-
-			if (COTFS.IsDirectorCookByTheBook())
-			{
-				Info.WriteOptions |= IPackageWriter::EWriteOptions::ComputeHash;
-			}
-		}
+		Info.WriteOptions = GetCommitWriteOptions();
 
 		PackageWriter->CommitPackage(MoveTemp(Info));
 	}
@@ -6394,77 +6392,8 @@ void FSaveCookedPackageContext::FinishPlatform()
 	if (COTFS.IsDirectorCookByTheBook())
 	{
 		IAssetRegistryReporter& Reporter = *(COTFS.PlatformManager->GetPlatformData(TargetPlatform)->RegistryReporter);
-
-		// Calculate up-to-date dependencies for Generator and Generated packages
-		// For generated packages, additionally calculate FAssetPackageData; that struct is calculated
-		// during user saves for non-generated packages.
-		TOptional<TArray<FAssetDependency>> OverridePackageDependencies;
-		TOptional<FAssetPackageData> OverrideAssetPackageData;
-		FGenerationHelper* GenerationHelper;
-		if (GenerationHelper = PackageData.GetGenerationHelper(); GenerationHelper)
-		{
-			OverridePackageDependencies.Emplace();
-
-			// Set override dependencies equal to the global AssetRegistry dependencies plus a dependency on
-			// each generated package.
-			COTFS.AssetRegistry->GetDependencies(Package->GetFName(), *OverridePackageDependencies,
-				UE::AssetRegistry::EDependencyCategory::Package);
-			OverridePackageDependencies->Reserve(GenerationHelper->GetPackagesToGenerate().Num());
-			for (FCookGenerationInfo& GeneratedInfo : GenerationHelper->GetPackagesToGenerate())
-			{
-				FAssetDependency& Dependency = OverridePackageDependencies->Emplace_GetRef();
-				Dependency.AssetId = FAssetIdentifier(GeneratedInfo.PackageData->GetPackageName());
-				Dependency.Category = UE::AssetRegistry::EDependencyCategory::Package;
-				Dependency.Properties = UE::AssetRegistry::EDependencyProperty::Game;
-			}
-
-			if (PlatformIndex == 0)
-			{
-				GenerationHelper->FetchExternalActorDependencies();
-				COTFS.RecordExternalActorDependencies(GenerationHelper->GetExternalActorDependencies());
-			}
-		}
-		else if (GenerationHelper = PackageData.GetParentGenerationHelper(); GenerationHelper)
-		{
-			FCookGenerationInfo* GeneratedInfo = GenerationHelper->FindInfo(PackageData);
-			if (!GeneratedInfo)
-			{
-				UE_LOG(LogCook, Error, TEXT("GeneratedInfo missing for package %s."), *PackageData.GetPackageName().ToString());
-			}
-			else
-			{
-				// There should be no package dependencies present for the package from the global assetregistry
-				// because it is newly created. Add on the dependencies declared for it from the CookPackageSplitter.
-				OverridePackageDependencies.Emplace(GeneratedInfo->PackageDependencies);
-
-				// Update the AssetPackageData for each requested platform with Guid and ImportedClasses
-				TSet<UClass*> PackageClasses;
-				ForEachObjectWithPackage(Package, [&PackageClasses, this](UObject* Object)
-					{
-						UClass* Class = Object->GetClass();
-						if (!Class->IsInPackage(Package)) // Imported classes list does not include classes in the package
-						{
-							PackageClasses.Add(Object->GetClass());
-						}
-						return true;
-					});
-				TArray<FName> ImportedClasses;
-				ImportedClasses.Reserve(PackageClasses.Num());
-				for (UClass* Class : PackageClasses)
-				{
-					TStringBuilder<256> ClassPath;
-					Class->GetPathName(nullptr, ClassPath);
-					ImportedClasses.Add(FName(ClassPath));
-				}
-				ImportedClasses.Sort(FNameLexicalLess());
-
-				OverrideAssetPackageData.Emplace();
-				OverrideAssetPackageData->SetPackageSavedHash(GeneratedInfo->PackageHash);
-				OverrideAssetPackageData->ImportedClasses = ImportedClasses;
-			}
-		}
 		TOptional<TArray<FAssetData>> AssetDatasFromSave;
-		if (SavePackageResult.IsSuccessful())
+		if (bSuccessful)
 		{
 			AssetDatasFromSave.Emplace(MoveTemp(SavePackageResult.SavedAssets));
 		}
@@ -6694,6 +6623,44 @@ void FSaveCookedPackageContext::CalculatePlatformRuntimeDependencies()
 	}
 }
 
+TArray<FName> FSaveCookedPackageContext::GetPlatformRuntimeDependencies() const
+{
+	TArray<FName> PlatformDependencyNames;
+	PlatformDependencyNames.Reserve(PlatformDependencies[PlatformIndex].Num());
+	for (const TPair<FPackageData*, EInstigator>& DependencyPair : PlatformDependencies[PlatformIndex])
+	{
+		PlatformDependencyNames.Add(DependencyPair.Key->GetPackageName());
+	}
+	return PlatformDependencyNames;
+}
+
+TArray<IPackageWriter::FCommitAttachmentInfo> FSaveCookedPackageContext::GetCommitAttachments()
+{
+	TArray<IPackageWriter::FCommitAttachmentInfo> Result;
+	if (COTFS.bHybridIterativeEnabled)
+	{
+		UE_SCOPED_HIERARCHICAL_COOKTIMER(TargetDomainDependencies);
+		UE::TargetDomain::CollectAndStoreCookAttachments(Package, TargetPlatform, &SavePackageResult,
+			GetPlatformRuntimeDependencies(), Result);
+	}
+	return Result;
+}
+
+IPackageWriter::EWriteOptions FSaveCookedPackageContext::GetCommitWriteOptions() const
+{
+	IPackageWriter::EWriteOptions Result = IPackageWriter::EWriteOptions::None;
+	if (!COTFS.bSkipSave)
+	{
+		Result |= IPackageWriter::EWriteOptions::Write;
+
+		if (COTFS.IsDirectorCookByTheBook())
+		{
+			Result |= IPackageWriter::EWriteOptions::ComputeHash;
+		}
+	}
+	return Result;
+}
+
 void FSaveCookedPackageContext::AddDependency(TMap<FPackageData*, EInstigator>& InDependencies,
 	FPackageData* PackageData, bool bHard)
 {
@@ -6708,6 +6675,23 @@ void FSaveCookedPackageContext::AddDependency(TMap<FPackageData*, EInstigator>&
 	}
 }
 
+IPackageWriter::ECommitStatus FSaveCookedPackageContext::PackageResultToCommitStatus(
+	FSavePackageResultStruct& Result)
+{
+	if (Result.IsSuccessful())
+	{
+		return IPackageWriter::ECommitStatus::Success;
+	}
+	else if (Result.Result == ESavePackageResult::Timeout)
+	{
+		return IPackageWriter::ECommitStatus::Canceled;
+	}
+	else
+	{
+		return IPackageWriter::ECommitStatus::Error;
+	}
+}
+
 } // namespace UE::Cook
 
 void UCookOnTheFlyServer::RecordExternalActorDependencies(TConstArrayView<FName> ExternalActorDependencies)
diff --git a/Engine/Source/Editor/UnrealEd/Private/Cooker/CookGenerationHelper.cpp b/Engine/Source/Editor/UnrealEd/Private/Cooker/CookGenerationHelper.cpp
--- a/Engine/Source/Editor/UnrealEd/Private/Cooker/CookGenerationHelper.cpp
+++ b/Engine/Source/Editor/UnrealEd/Private/Cooker/CookGenerationHelper.cpp
@@ -464,7 +464,11 @@ void FGenerationHelper::ClearSelfReferences()
 FCookGenerationInfo* FGenerationHelper::FindInfo(const FPackageData& PackageData)
 {
 	ConditionalInitialize();
+	return FindInfoNoInitialize(PackageData);
+}
 
+FCookGenerationInfo* FGenerationHelper::FindInfoNoInitialize(const FPackageData& PackageData)
+{
 	if (&PackageData == &GetOwner())
 	{
 		return &OwnerInfo;
@@ -878,7 +882,7 @@ void FGenerationHelper::StartQueueGeneratedPackages(UCookOnTheFlyServer& COTFS)
 			}
 		}
 
-		if (NumIterativePrevious > 0 && !bHybridIterativeEnabled)
+		if (!bHybridIterativeEnabled)
 		{
 			UE_LOG(LogCook, Display, TEXT("Found %d cooked package(s) in package store for generator package %s."),
 				NumIterativePrevious, *WriteToString<256>(GetOwner().GetPackageName()));
@@ -1007,6 +1011,80 @@ UPackage* FGenerationHelper::TryCreateGeneratedPackage(FCookGenerationInfo& Gene
 	return GeneratedPackage;
 }
 
+void FGenerationHelper::FinishGeneratorPlatformSave(FPackageData& PackageData, bool bFirstPlatform,
+	TArray<FAssetDependency>& OutPackageDependencies)
+{
+	ConditionalInitialize();
+
+	FCookGenerationInfo* Info = &GetOwnerInfo();
+	UCookOnTheFlyServer& COTFS = Info->PackageData->GetPackageDatas().GetCookOnTheFlyServer();
+
+	// Set dependencies equal to the global AssetRegistry dependencies plus a dependency on
+	// each generated package.
+	COTFS.AssetRegistry->GetDependencies(PackageData.GetPackageName(), OutPackageDependencies,
+		UE::AssetRegistry::EDependencyCategory::Package);
+	OutPackageDependencies.Reserve(PackagesToGenerate.Num());
+	for (FCookGenerationInfo& GeneratedInfo : GetPackagesToGenerate())
+	{
+		FAssetDependency& Dependency = OutPackageDependencies.Emplace_GetRef();
+		Dependency.AssetId = FAssetIdentifier(PackageData.GetPackageName());
+		Dependency.Category = UE::AssetRegistry::EDependencyCategory::Package;
+		Dependency.Properties = UE::AssetRegistry::EDependencyProperty::Game;
+	}
+
+	if (bFirstPlatform)
+	{
+		FetchExternalActorDependencies();
+		COTFS.RecordExternalActorDependencies(GetExternalActorDependencies());
+	}
+}
+
+void FGenerationHelper::FinishGeneratedPlatformSave(FPackageData& PackageData,
+	TArray<FAssetDependency>& OutPackageDependencies,
+	FAssetPackageData& OutAssetPackageData)
+{
+	ConditionalInitialize();
+
+	FCookGenerationInfo* Info = FindInfo(PackageData);
+	if (!Info)
+	{
+		UE_LOG(LogCook, Error, TEXT("GeneratedInfo missing for package %s."),
+			*PackageData.GetPackageName().ToString());
+		return;
+	}
+	UCookOnTheFlyServer& COTFS = Info->PackageData->GetPackageDatas().GetCookOnTheFlyServer();
+
+	// There should be no package dependencies present for the package from the global assetregistry
+	// because it is newly created. Add on the dependencies declared for it from the CookPackageSplitter.
+	OutPackageDependencies = Info->PackageDependencies;
+
+	// Update the AssetPackageData for each requested platform with Guid and ImportedClasses
+	TSet<UClass*> PackageClasses;
+	UPackage* Package = PackageData.GetPackage();
+	check(Package);
+	ForEachObjectWithPackage(Package, [&PackageClasses, Package](UObject* Object)
+		{
+			UClass* Class = Object->GetClass();
+			if (!Class->IsInPackage(Package)) // Imported classes list does not include classes in the package
+			{
+				PackageClasses.Add(Object->GetClass());
+			}
+			return true;
+		});
+	TArray<FName> ImportedClasses;
+	ImportedClasses.Reserve(PackageClasses.Num());
+	for (UClass* Class : PackageClasses)
+	{
+		TStringBuilder<256> ClassPath;
+		Class->GetPathName(nullptr, ClassPath);
+		ImportedClasses.Add(FName(ClassPath));
+	}
+	ImportedClasses.Sort(FNameLexicalLess());
+
+	OutAssetPackageData.SetPackageSavedHash(Info->PackageHash);
+	OutAssetPackageData.ImportedClasses = MoveTemp(ImportedClasses);
+}
+
 void FGenerationHelper::ResetSaveState(FCookGenerationInfo& Info, UPackage* Package,
 	EStateChangeReason ReleaseSaveReason, EPackageState NewState)
 {
@@ -1330,7 +1408,7 @@ void FGenerationHelper::TrackGeneratedPackageListedRemotely(UCookOnTheFlyServer&
 	}
 	else
 	{
-		if (!FindInfo(PackageData))
+		if (!FindInfoNoInitialize(PackageData))
 		{
 			bool bGenerator = false; // Cannot be the generator, if it were we would have found it
 			PackagesToGenerate.Emplace(PackageData, bGenerator);
@@ -1341,7 +1419,7 @@ void FGenerationHelper::TrackGeneratedPackageListedRemotely(UCookOnTheFlyServer&
 void FGenerationHelper::MarkPackageSavedRemotely(UCookOnTheFlyServer& COTFS, FPackageData& PackageData,
 	FWorkerId SourceWorkerId)
 {
-	FCookGenerationInfo* Info = FindInfo(PackageData);
+	FCookGenerationInfo* Info = FindInfoNoInitialize(PackageData);
 	if (Info)
 	{
 		Info->SetHasSaved(*this, true, SourceWorkerId);
diff --git a/Engine/Source/Editor/UnrealEd/Private/Cooker/CookGenerationHelper.h b/Engine/Source/Editor/UnrealEd/Private/Cooker/CookGenerationHelper.h
--- a/Engine/Source/Editor/UnrealEd/Private/Cooker/CookGenerationHelper.h
+++ b/Engine/Source/Editor/UnrealEd/Private/Cooker/CookGenerationHelper.h
@@ -2,6 +2,7 @@
 
 #pragma once
 
+#include "AssetRegistry/AssetData.h"
 #include "AssetRegistry/IAssetRegistry.h"
 #include "Containers/Array.h"
 #include "Containers/ArrayView.h"
@@ -305,6 +306,20 @@ public:
 	void OnRequestFencePassed(UCookOnTheFlyServer& COTFS);
 	/** Call CreatePackage and set package header data; or empty it and normalize it if it already exists. */
 	UPackage* TryCreateGeneratedPackage(FCookGenerationInfo& GenerationInfo, bool bResetToEmpty);
+	/**
+	 * Called when a generator package is in FSaveCookedPackageContext::FinishPlatformSave.
+	 * Calculates AssetRegistryData.
+	 */
+	void FinishGeneratorPlatformSave(FPackageData& PackageData, bool bFirstPlatform,
+		TArray<FAssetDependency>& OutPackageDependencies);
+	/**
+	 * Called when a generated package is in FSaveCookedPackageContext::FinishPlatformSave.
+	 * Calculates AssetRegistryData.
+	 */
+	void FinishGeneratedPlatformSave(FPackageData& PackageData,
+		TArray<FAssetDependency>& OutPackageDependencies,
+		FAssetPackageData& OutAssetPackageData);
+
 	/**
 	 * Clear any data that should only be held when an FPackageData is in the save state, for the given Info. The given
 	 * Info might specify the generator package or one of the generated packages.
@@ -415,6 +430,7 @@ private:
 
 private:
 	void ConditionalInitialize() const;
+	FCookGenerationInfo* FindInfoNoInitialize(const FPackageData& PackageData);
 	void NotifyCompletion(ICookPackageSplitter::ETeardown Status);
 	void PreGarbageCollectGCLifetimeData();
 	void PostGarbageCollectGCLifetimeData();
diff --git a/Engine/Source/Editor/UnrealEd/Private/Cooker/PackageResultsMessage.h b/Engine/Source/Editor/UnrealEd/Private/Cooker/PackageResultsMessage.h
--- a/Engine/Source/Editor/UnrealEd/Private/Cooker/PackageResultsMessage.h
+++ b/Engine/Source/Editor/UnrealEd/Private/Cooker/PackageResultsMessage.h
@@ -5,6 +5,7 @@
 #include "Async/Future.h"
 #include "CompactBinaryTCP.h"
 #include "Containers/Array.h"
+#include "Containers/ArrayView.h"
 #include "Cooker/CookTypes.h"
 #include "Cooker/MPCollector.h"
 #include "HAL/CriticalSection.h"
@@ -30,14 +31,14 @@ public:
 	struct FPlatformResult
 	{
 	public:
-		const ITargetPlatform* GetPlatform() const { return Platform; }
-		void SetPlatform(const ITargetPlatform* InPlatform) { Platform = InPlatform; }
+		const ITargetPlatform* GetPlatform() const;
+		void SetPlatform(const ITargetPlatform* InPlatform);
 
-		TConstArrayView<UE::CompactBinaryTCP::FMarshalledMessage> GetMessages() const { return Messages; }
+		TConstArrayView<UE::CompactBinaryTCP::FMarshalledMessage> GetMessages() const;
 		TArray<UE::CompactBinaryTCP::FMarshalledMessage> ReleaseMessages();
 
-		ECookResult GetCookResults() const { return CookResults; }
-		void SetCookResults(ECookResult Value) { CookResults = Value; }
+		ECookResult GetCookResults() const;
+		void SetCookResults(ECookResult Value);
 
 	private:
 		TArray<UE::CompactBinaryTCP::FMarshalledMessage> Messages;
@@ -54,14 +55,14 @@ public:
 	FPackageRemoteResult& operator=(FPackageRemoteResult&&) = default;
 	FPackageRemoteResult& operator=(const FPackageRemoteResult&) = delete;
 
-	FName GetPackageName() const { return PackageName; }
-	void SetPackageName(FName InPackageName) { PackageName = InPackageName; }
+	FName GetPackageName() const;
+	void SetPackageName(FName InPackageName);
 
-	ESuppressCookReason GetSuppressCookReason() const { return SuppressCookReason; }
-	void SetSuppressCookReason(ESuppressCookReason InSuppressCookReason) { SuppressCookReason = InSuppressCookReason; }
+	ESuppressCookReason GetSuppressCookReason() const;
+	void SetSuppressCookReason(ESuppressCookReason InSuppressCookReason);
 
-	bool IsReferencedOnlyByEditorOnlyData() const { return bReferencedOnlyByEditorOnlyData; }
-	void SetReferencedOnlyByEditorOnlyData(bool bInReferencedOnlyByEditorOnlyData) { bReferencedOnlyByEditorOnlyData = bInReferencedOnlyByEditorOnlyData; }
+	bool IsReferencedOnlyByEditorOnlyData() const;
+	void SetReferencedOnlyByEditorOnlyData(bool bInReferencedOnlyByEditorOnlyData);
 
 	void AddPackageMessage(const FGuid& MessageType, FCbObject&& Object);
 	void AddAsyncPackageMessage(const FGuid& MessageType, TFuture<FCbObject>&& ObjectFuture);
@@ -69,23 +70,23 @@ public:
 	void AddAsyncPlatformMessage(const ITargetPlatform* TargetPlatform, const FGuid& MessageType, TFuture<FCbObject>&& ObjectFuture);
 
 	// GetMessages and ReleaseMessages are not thread-safe until IsComplete returns true or GetCompletionFuture().Get()/.Next().
-	TConstArrayView<UE::CompactBinaryTCP::FMarshalledMessage> GetMessages() const { return Messages; }
+	TConstArrayView<UE::CompactBinaryTCP::FMarshalledMessage> GetMessages() const;
 	TArray<UE::CompactBinaryTCP::FMarshalledMessage> ReleaseMessages();
 
 	bool IsComplete();
 	TFuture<int> GetCompletionFuture();
 
-	TArray<FPlatformResult, TInlineAllocator<1>>& GetPlatforms() { return Platforms; }
+	TArray<FPlatformResult, TInlineAllocator<1>>& GetPlatforms();
 	void SetPlatforms(TConstArrayView<ITargetPlatform*> OrderedSessionPlatforms);
 
-	void SetExternalActorDependencies(TArray<FName>&& InExternalActorDependencies) { ExternalActorDependencies = MoveTemp(InExternalActorDependencies); }
-	TConstArrayView<FName> GetExternalActorDependencies() const { return ExternalActorDependencies; }
+	void SetExternalActorDependencies(TArray<FName>&& InExternalActorDependencies);
+	TConstArrayView<FName> GetExternalActorDependencies() const;
 
 	/**
 	 * A non-atomic RefCount that can be used for storage of a refcount by the user (e.g. CookWorkerClient)
 	 * If used from multiple threads, the user must access it only within the user's external critical section.
 	 */
-	int32& GetUserRefCount() { return UserRefCount; }
+	int32& GetUserRefCount();
 
 private:
 	/** A TFuture and status data that was received from an asynchronous IMPCollector. */
@@ -147,8 +148,11 @@ struct FPackageResultsMessage : public IMPCollectorMessage
 public:
 	virtual void Write(FCbWriter& Writer) const override;
 	virtual bool TryRead(FCbObjectView Object) override;
-	virtual FGuid GetMessageType() const override { return MessageType; }
-	virtual const TCHAR* GetDebugName() const override { return TEXT("PackageResultsMessage"); }
+	virtual FGuid GetMessageType() const override;
+	virtual const TCHAR* GetDebugName() const override
+	{
+		return TEXT("PackageResultsMessage");
+	}
 
 public:
 	TArray<FPackageRemoteResult> Results;
@@ -162,4 +166,95 @@ private:
 		TArray<UE::CompactBinaryTCP::FMarshalledMessage>& InMessages);
 };
 
+
+///////////////////////////////////////////////////////
+// Inline implementations
+///////////////////////////////////////////////////////
+
+inline const ITargetPlatform* FPackageRemoteResult::FPlatformResult::GetPlatform() const
+{
+	return Platform;
+}
+
+inline void FPackageRemoteResult::FPlatformResult::SetPlatform(const ITargetPlatform* InPlatform)
+{
+	Platform = InPlatform;
+}
+
+inline TConstArrayView<UE::CompactBinaryTCP::FMarshalledMessage>
+FPackageRemoteResult::FPlatformResult::GetMessages() const
+{
+	return Messages;
+}
+
+inline ECookResult FPackageRemoteResult::FPlatformResult::GetCookResults() const
+{
+	return CookResults;
+}
+
+inline void FPackageRemoteResult::FPlatformResult::SetCookResults(ECookResult Value)
+{
+	CookResults = Value;
+}
+
+inline FName FPackageRemoteResult::GetPackageName() const
+{
+	return PackageName;
+}
+
+inline void FPackageRemoteResult::SetPackageName(FName InPackageName)
+{
+	PackageName = InPackageName;
+}
+
+inline ESuppressCookReason FPackageRemoteResult::GetSuppressCookReason() const
+{
+	return SuppressCookReason;
+}
+
+inline void FPackageRemoteResult::SetSuppressCookReason(ESuppressCookReason InSuppressCookReason)
+{
+	SuppressCookReason = InSuppressCookReason;
+}
+
+inline bool FPackageRemoteResult::IsReferencedOnlyByEditorOnlyData() const
+{
+	return bReferencedOnlyByEditorOnlyData;
+}
+
+inline void FPackageRemoteResult::SetReferencedOnlyByEditorOnlyData(bool bInReferencedOnlyByEditorOnlyData)
+{
+	bReferencedOnlyByEditorOnlyData = bInReferencedOnlyByEditorOnlyData;
+}
+
+inline TConstArrayView<UE::CompactBinaryTCP::FMarshalledMessage> FPackageRemoteResult::GetMessages() const
+{
+	return Messages;
+}
+
+inline TArray<FPackageRemoteResult::FPlatformResult, TInlineAllocator<1>>& FPackageRemoteResult::GetPlatforms()
+{
+	return Platforms;
+}
+
+inline void FPackageRemoteResult::SetExternalActorDependencies(TArray<FName>&& InExternalActorDependencies)
+{
+	ExternalActorDependencies = MoveTemp(InExternalActorDependencies);
+}
+
+inline TConstArrayView<FName> FPackageRemoteResult::GetExternalActorDependencies() const
+{
+	return ExternalActorDependencies;
+}
+
+inline int32& FPackageRemoteResult::GetUserRefCount()
+{
+	return UserRefCount;
+}
+
+inline FGuid FPackageResultsMessage::GetMessageType() const
+{
+	return MessageType;
+}
+
 }
diff --git a/Engine/Source/Editor/UnrealEd/Private/EditorDomain/EditorDomain.h b/Engine/Source/Editor/UnrealEd/Private/EditorDomain/EditorDomain.h
--- a/Engine/Source/Editor/UnrealEd/Private/EditorDomain/EditorDomain.h
+++ b/Engine/Source/Editor/UnrealEd/Private/EditorDomain/EditorDomain.h
@@ -169,8 +169,14 @@ public:
 
 	// FTickableEditorObject interface
 	virtual void Tick(float DeltaTime) override;
-	virtual ETickableTickType GetTickableTickType() const override { return ETickableTickType::Always; }
-	virtual TStatId GetStatId() const override { return TStatId(); }
+	virtual ETickableTickType GetTickableTickType() const override
+	{
+		return ETickableTickType::Always;
+	}
+	virtual TStatId GetStatId() const override
+	{
+		return TStatId();
+	}
 
 	// IPackageDigestCache interface
 	virtual UE::EditorDomain::FPackageDigest GetPackageDigest(FName PackageDigest) override;
@@ -178,8 +184,8 @@ public:
 	// EditorDomain interface
 	/** Fetch data from game-thread sources that is required to calculate the PackageDigest of the given PackageName. */
 	void PrecachePackageDigest(FName PackageName);
-	bool IsReadingPackages() const { return bEditorDomainReadEnabled; }
-	bool IsWritingPackages() const { return bEditorDomainWriteEnabled; }
+	bool IsReadingPackages() const;
+	bool IsWritingPackages() const;
 
 	/** Request the download of the given packages from the upstream DDC server. */
 	void BatchDownload(TArrayView<FName> PackageNames);
@@ -287,4 +293,19 @@ private:
 	friend class FEditorDomainPackageSegments;
 	friend class FEditorDomainReadArchive;
 	friend class FEditorDomainAsyncReadFileHandle;
-};
\ No newline at end of file
+};
+
+
+///////////////////////////////////////////////////////
+// Inline implementations
+///////////////////////////////////////////////////////
+
+inline bool FEditorDomain::IsReadingPackages() const
+{
+	return bEditorDomainReadEnabled;
+}
+
+inline bool FEditorDomain::IsWritingPackages() const
+{
+	return bEditorDomainWriteEnabled;
+}
