[DiffUtils] skip non visible properties from node Graph diffs.
Added details on the subproperty causing the difference.

diff --git a/Engine/Source/Editor/BehaviorTreeEditor/Private/SBehaviorTreeDiff.cpp b/Engine/Source/Editor/BehaviorTreeEditor/Private/SBehaviorTreeDiff.cpp
--- a/Engine/Source/Editor/BehaviorTreeEditor/Private/SBehaviorTreeDiff.cpp
+++ b/Engine/Source/Editor/BehaviorTreeEditor/Private/SBehaviorTreeDiff.cpp
@@ -255,7 +255,6 @@ TSharedRef<SWidget> SBehaviorTreeDiff::GenerateDiffListWidget()
 		TSharedRef<SHorizontalBox> Result =	SNew(SHorizontalBox)
 		+SHorizontalBox::Slot()
 		.FillWidth(1.f)
-		.MaxWidth(350.f)
 		[
 			SNew(SVerticalBox)
 			+SVerticalBox::Slot()
diff --git a/Engine/Source/Editor/Kismet/Private/DetailsDiff.cpp b/Engine/Source/Editor/Kismet/Private/DetailsDiff.cpp
--- a/Engine/Source/Editor/Kismet/Private/DetailsDiff.cpp
+++ b/Engine/Source/Editor/Kismet/Private/DetailsDiff.cpp
@@ -270,7 +270,7 @@ void FDetailsDiff::DiffAgainst(const FDetailsDiff& Newer, TArray< FSingleObjectD
 
 		TArray<FPropertySoftPath> DifferingSubProperties;
 
-		if (!DiffUtils::Identical(OldProperty, NewProperty, OldPackage, NewPackage, CommonProperty, DifferingSubProperties))
+		if (!DiffUtils::Identical(OldProperty, NewProperty, OldPackage, NewPackage, DiffUtils::FDiffParameters(CommonProperty), DifferingSubProperties))
 		{
 			for (const FPropertySoftPath& DifferingSubProperty : DifferingSubProperties)
 			{
diff --git a/Engine/Source/Editor/UnrealEd/Private/DiffUtils.cpp b/Engine/Source/Editor/UnrealEd/Private/DiffUtils.cpp
--- a/Engine/Source/Editor/UnrealEd/Private/DiffUtils.cpp
+++ b/Engine/Source/Editor/UnrealEd/Private/DiffUtils.cpp
@@ -273,8 +273,14 @@ FPropertySoftPath::FPropertySoftPath()
 {
 }
 
+FPropertySoftPath::FPropertySoftPath(const FProperty* Property)
+	: FPropertySoftPath()
+{
+	PropertyChain.Push(FChainElement(Property));
+}
+
 FPropertySoftPath::FPropertySoftPath(TArray<FName> InPropertyChain)
-	: RootTypeHint(nullptr)
+	: FPropertySoftPath()
 {
 	for (FName PropertyName : InPropertyChain)
 	{
@@ -283,7 +289,7 @@ FPropertySoftPath::FPropertySoftPath(TArray<FName> InPropertyChain)
 }
 
 FPropertySoftPath::FPropertySoftPath(FPropertyPath InPropertyPath)
-	: RootTypeHint(nullptr)
+	: FPropertySoftPath()
 {
 	for (int32 PropertyIndex = 0, end = InPropertyPath.GetNumProperties(); PropertyIndex != end; ++PropertyIndex)
 	{
@@ -308,6 +314,13 @@ FPropertySoftPath::FPropertySoftPath(FPropertyPath InPropertyPath)
 	}
 }
 
+FPropertySoftPath::FPropertySoftPath(const FPropertySoftPath& MainPropertyPath, const FPropertySoftPath& SubPropertyPath)
+	: PropertyChain(MainPropertyPath.PropertyChain)
+	, RootTypeHint(MainPropertyPath.RootTypeHint)
+{
+	PropertyChain.Append(SubPropertyPath.PropertyChain);
+}
+
 FPropertySoftPath::FPropertySoftPath(const FPropertySoftPath& SubPropertyPath, const FProperty* LeafProperty)
 	: PropertyChain(SubPropertyPath.PropertyChain)
 	, RootTypeHint(SubPropertyPath.RootTypeHint)
@@ -654,7 +667,7 @@ void DiffUtils::CompareUnrelatedStructs(const UStruct* StructA, const void* A, c
 
 			check(AProp != FResolvedProperty() && BProp != FResolvedProperty());
 			TArray<FPropertySoftPath> DifferingSubProperties;
-			if (!DiffUtils::Identical(AProp, BProp, OwningOuterA, OwningOuterB, PropertyName, DifferingSubProperties))
+			if (!DiffUtils::Identical(AProp, BProp, OwningOuterA, OwningOuterB, FDiffParameters(PropertyName), DifferingSubProperties))
 			{
 				for (int DifferingIndex = 0; DifferingIndex < DifferingSubProperties.Num(); DifferingIndex++)
 				{
@@ -778,27 +791,40 @@ void DiffUtils::CompareUnrelatedSCS(const UBlueprint* Old, const TArray< FSCSRes
 }
 
 static void IdenticalHelper(const FProperty* AProperty, const FProperty* BProperty, const void* AValue, const void* BValue,
-	const UObject* OwningOuterA, const UObject* OwningOuterB, const FPropertySoftPath& RootPath,
-	TArray<FPropertySoftPath>& DifferingSubProperties, bool bStaticArrayHandled = false)
+	const UObject* OwningOuterA, const UObject* OwningOuterB, DiffUtils::FDiffParameters DiffParameters,
+	TArray<FPropertySoftPath>& DifferingSubProperties)
 {
-	if(AProperty == nullptr || BProperty == nullptr || AProperty->ArrayDim != BProperty->ArrayDim || AProperty->GetClass() != BProperty->GetClass())
+	if (AProperty == nullptr || BProperty == nullptr || AProperty->ArrayDim != BProperty->ArrayDim || AProperty->GetClass() != BProperty->GetClass())
 	{
-		DifferingSubProperties.Push(RootPath);
+		DifferingSubProperties.Push(DiffParameters.RootPath);
 		return;
 	}
 
-	if(!bStaticArrayHandled && AProperty->ArrayDim != 1)
+	if (DiffParameters.ShouldIgnorePropertyPredicate
+		&& DiffParameters.ShouldIgnorePropertyPredicate(*AProperty))
+	{
+		return;
+	}
+
+	// Keep copy of the initial RootPath since we reuse DiffParameters with
+	// an updated path to nested calls to 'IdenticalHelper'
+	const FPropertySoftPath RootPath = DiffParameters.RootPath;
+
+	if (DiffParameters.bShouldDiffArrayElements && AProperty->ArrayDim != 1)
 	{
 		// Identical does not handle static array case automatically and we have to do the offset calculation ourself because 
 		// our container (e.g. the struct or class or dynamic array) has already done the initial offset calculation:
-		for( int32 I = 0; I < AProperty->ArrayDim; ++I )
+		for (int32 I = 0; I < AProperty->ArrayDim; ++I)
 		{
 			int32 Offset = AProperty->ElementSize * I;
+
+			DiffParameters.RootPath = FPropertySoftPath(DiffParameters.RootPath, I);
+			DiffParameters.bShouldDiffArrayElements = false;
+
 			const void* CurAValue = reinterpret_cast<const void*>(reinterpret_cast<const uint8*>(AValue) + Offset);
 			const void* CurBValue = reinterpret_cast<const void*>(reinterpret_cast<const uint8*>(BValue) + Offset);
 
-			IdenticalHelper(AProperty, BProperty, CurAValue, CurBValue, OwningOuterA, OwningOuterB,
-				FPropertySoftPath(RootPath, I), DifferingSubProperties, true);
+			IdenticalHelper(AProperty, BProperty, CurAValue, CurBValue, OwningOuterA, OwningOuterB, DiffParameters, DifferingSubProperties);
 		}
 
 		return;
@@ -828,10 +854,13 @@ static void IdenticalHelper(const FProperty* AProperty, const FProperty* BProper
 				for (TFieldIterator<FProperty> PropertyIt(APropAsStruct->Struct); PropertyIt; ++PropertyIt)
 				{
 					const FProperty* StructProp = *PropertyIt;
+
+					DiffParameters.RootPath = FPropertySoftPath(DiffParameters.RootPath, StructProp);
+
 					const void* SubValueA = StructProp->ContainerPtrToValuePtr<void>(AValue, 0);
 					const void* SubValueB = StructProp->ContainerPtrToValuePtr<void>(BValue, 0);
 					IdenticalHelper(StructProp, StructProp, SubValueA, SubValueB,
-						OwningOuterA, OwningOuterB, FPropertySoftPath(RootPath, StructProp), DifferingSubProperties);
+						OwningOuterA, OwningOuterB, DiffParameters, DifferingSubProperties);
 				}
 			}
 		}
@@ -851,10 +880,12 @@ static void IdenticalHelper(const FProperty* AProperty, const FProperty* BProper
 			// note any differences in contained types:
 			for (int32 ArrayIndex = 0; ArrayIndex < ArrayHelperA.Num() && ArrayIndex < ArrayHelperB.Num(); ArrayIndex++)
 			{
+				DiffParameters.RootPath = FPropertySoftPath(DiffParameters.RootPath, ArrayIndex);
+
 				const void* SubValueA = ArrayHelperA.GetRawPtr(ArrayIndex);
 				const void* SubValueB = ArrayHelperB.GetRawPtr(ArrayIndex);
 				IdenticalHelper(APropAsArray->Inner, BPropAsArray->Inner, SubValueA, SubValueB,
-					OwningOuterA, OwningOuterB, FPropertySoftPath(RootPath, ArrayIndex), DifferingSubProperties);
+					OwningOuterA, OwningOuterB, DiffParameters, DifferingSubProperties);
 			}
 
 			// note any size difference:
@@ -888,10 +919,12 @@ static void IdenticalHelper(const FProperty* AProperty, const FProperty* BProper
 			FScriptSetHelper::FIterator IteratorB(SetHelperB);
 			for (; IteratorA && IteratorB; ++IteratorA, ++IteratorB)
 			{
+				DiffParameters.RootPath = FPropertySoftPath(DiffParameters.RootPath, IteratorA.GetLogicalIndex());
+
 				const void* SubValueA = SetHelperA.GetElementPtr(IteratorA);
 				const void* SubValueB = SetHelperB.GetElementPtr(IteratorB);
 				IdenticalHelper(APropAsSet->ElementProp, BPropAsSet->ElementProp, SubValueA, SubValueB,
-					OwningOuterA, OwningOuterB, FPropertySoftPath(RootPath, IteratorA.GetLogicalIndex()), DifferingSubProperties);
+					OwningOuterA, OwningOuterB, DiffParameters, DifferingSubProperties);
 			}
 		}
 		else
@@ -917,10 +950,12 @@ static void IdenticalHelper(const FProperty* AProperty, const FProperty* BProper
 			FScriptMapHelper::FIterator IteratorB(MapHelperB);
 			for (; IteratorA && IteratorB; ++IteratorA, ++IteratorB)
 			{
+				DiffParameters.RootPath = FPropertySoftPath(DiffParameters.RootPath, IteratorA.GetLogicalIndex());
+
 				IdenticalHelper(APropAsMap->KeyProp, BPropAsMap->KeyProp, MapHelperA.GetKeyPtr(IteratorA), MapHelperB.GetKeyPtr(IteratorB),
-					OwningOuterA, OwningOuterB, FPropertySoftPath(RootPath, IteratorA.GetLogicalIndex()), DifferingSubProperties);
+					OwningOuterA, OwningOuterB, DiffParameters, DifferingSubProperties);
 				IdenticalHelper(APropAsMap->ValueProp, BPropAsMap->ValueProp, MapHelperA.GetValuePtr(IteratorA), MapHelperB.GetValuePtr(IteratorB),
-					OwningOuterA, OwningOuterB, FPropertySoftPath(RootPath, IteratorA.GetLogicalIndex()), DifferingSubProperties);
+					OwningOuterA, OwningOuterB, DiffParameters, DifferingSubProperties);
 
 			}
 		}
@@ -954,10 +989,12 @@ static void IdenticalHelper(const FProperty* AProperty, const FProperty* BProper
 				for (TFieldIterator<FProperty> PropertyIt(AClass); PropertyIt; ++PropertyIt)
 				{
 					const FProperty* ClassProp = *PropertyIt;
+					DiffParameters.RootPath = FPropertySoftPath(DiffParameters.RootPath, ClassProp);
+
 					const void* SubValueA = ClassProp->ContainerPtrToValuePtr<void>(A, 0);
 					const void* SubValueB = ClassProp->ContainerPtrToValuePtr<void>(B, 0);
 					IdenticalHelper(ClassProp, ClassProp, SubValueA, SubValueB,
-						OwningOuterA, OwningOuterB, FPropertySoftPath(RootPath, ClassProp), DifferingSubProperties);
+						OwningOuterA, OwningOuterB, DiffParameters, DifferingSubProperties);
 				}
 			}
 			else
@@ -1005,7 +1042,7 @@ bool DiffUtils::Identical(const FResolvedProperty& AProp, const FResolvedPropert
 	// Instead let's write a new function, specific to DiffUtils, that handles the sub properties
 	// NOTE: For Static Arrays, AValue and BValue were, and are, only references to the value at index 0.  So changes to values past index 0 didn't show up before and
 	// won't show up now.  Changes to index 0 will show up as a change to the entire array.
-	IdenticalHelper(AProp.Property, BProp.Property, AValue, BValue, nullptr, nullptr, RootPath, DifferingProperties);
+	IdenticalHelper(AProp.Property, BProp.Property, AValue, BValue, nullptr, nullptr, FDiffParameters(RootPath), DifferingProperties);
 	
 	return DifferingProperties.Num() == 0;
 }
@@ -1014,11 +1051,17 @@ bool DiffUtils::Identical(const FResolvedProperty& AProp, const FResolvedPropert
 	const UObject* OwningOuterB)
 {
 	TArray<FPropertySoftPath> DifferingProperties;
-	return Identical(AProp, BProp, OwningOuterA, OwningOuterB, {}, DifferingProperties);
+	return Identical(AProp, BProp, OwningOuterA, OwningOuterB, FDiffParameters{}, DifferingProperties);
 }
 
 bool DiffUtils::Identical(const FResolvedProperty& AProp, const FResolvedProperty& BProp, const UObject* OwningOuterA,
 	const UObject* OwningOuterB, const FPropertySoftPath& RootPath, TArray<FPropertySoftPath>& DifferingProperties)
+{
+	return Identical(AProp, BProp, OwningOuterA, OwningOuterB, FDiffParameters(RootPath), DifferingProperties);
+}
+
+bool DiffUtils::Identical(const FResolvedProperty& AProp, const FResolvedProperty& BProp, const UObject* OwningOuterA,
+	const UObject* OwningOuterB, FDiffParameters DiffParameters, TArray<FPropertySoftPath>& DifferingProperties)
 {
 	if( AProp.Property == nullptr && BProp.Property == nullptr )
 	{
@@ -1034,7 +1077,7 @@ bool DiffUtils::Identical(const FResolvedProperty& AProp, const FResolvedPropert
 
 	// note that we're not directly calling FProperty::Identical because sub-object properties should be weakly compared based on
 	// their paths instead of their pointers or data
-	IdenticalHelper(AProp.Property, BProp.Property, AValue, BValue, OwningOuterA, OwningOuterB, RootPath, DifferingProperties);
+	IdenticalHelper(AProp.Property, BProp.Property, AValue, BValue, OwningOuterA, OwningOuterB, DiffParameters, DifferingProperties);
 	
 	return DifferingProperties.Num() == 0;
 }
@@ -1087,8 +1130,11 @@ bool DiffUtils::Identical(const TSharedPtr<IPropertyHandle>& PropertyHandleA, co
 		// note that we're not directly calling FProperty::Identical because sub-object properties should be weakly compared based on
 		// their paths instead of their pointers or data
 		TArray<FPropertySoftPath> DifferingProperties;
+		FDiffParameters DiffParameters;
+		DiffParameters.bShouldDiffArrayElements = false;
+
 		IdenticalHelper(PropertyHandleA->GetProperty(), PropertyHandleB->GetProperty(), ValueA, ValueB,
-			OwningOuterA, OwningOuterB, {}, DifferingProperties, true);
+			OwningOuterA, OwningOuterB, DiffParameters, DifferingProperties);
 
 		return DifferingProperties.IsEmpty();
 	};
diff --git a/Engine/Source/Editor/UnrealEd/Public/DiffUtils.h b/Engine/Source/Editor/UnrealEd/Public/DiffUtils.h
--- a/Engine/Source/Editor/UnrealEd/Public/DiffUtils.h
+++ b/Engine/Source/Editor/UnrealEd/Public/DiffUtils.h
@@ -64,8 +64,10 @@ struct FResolvedProperty
 struct FPropertySoftPath
 {
 	UNREALED_API FPropertySoftPath();
+	UNREALED_API FPropertySoftPath(const FProperty* Property);
 	UNREALED_API FPropertySoftPath(TArray<FName> InPropertyChain);
 	UNREALED_API FPropertySoftPath(FPropertyPath InPropertyPath);
+	UNREALED_API FPropertySoftPath(const FPropertySoftPath& MainPropertyPath, const FPropertySoftPath& SubPropertyPath);
 	UNREALED_API FPropertySoftPath(const FPropertySoftPath& SubPropertyPath, const FProperty* LeafProperty);
 	UNREALED_API FPropertySoftPath(const FPropertySoftPath& SubPropertyPath, int32 ContainerIndex);
 
@@ -97,10 +99,17 @@ struct FPropertySoftPath
 		return PropertyChain == RHS.PropertyChain;
 	}
 
-	inline bool operator!=(FPropertySoftPath const& RHS ) const
+	inline bool operator!=(FPropertySoftPath const& RHS) const
 	{
 		return !(*this == RHS);
 	}
+
+	inline FPropertySoftPath& operator+=(FPropertySoftPath const& RHS)
+	{
+		PropertyChain.Append(RHS.PropertyChain);
+		return *this;
+	}
+
 private:
 	struct FChainElement
 	{
@@ -248,6 +257,23 @@ struct FSCSDiffRoot
 
 namespace DiffUtils
 {
+	struct FDiffParameters
+	{
+		FDiffParameters() = default;
+		explicit FDiffParameters(const FPropertySoftPath& RootPath)
+			: RootPath(RootPath)
+		{
+		}
+
+		FPropertySoftPath RootPath;
+
+		/** Predicate evaluated on visited properties to determine if they should be part of the comparison. */
+		TFunction<bool(const FProperty&)> ShouldIgnorePropertyPredicate;
+
+		/** Indicates if the elements of a static array should be compared. Otherwise, only the size is compared. */
+		bool bShouldDiffArrayElements = true;
+	};
+
 	UNREALED_API const UObject* GetCDO(const UBlueprint* ForBlueprint);
 	UE_DEPRECATED(5.3, "DiffUtils now requires root objects so that object topology can be meaningfully compared.")
 	UNREALED_API void CompareUnrelatedStructs(const UStruct* StructA, const void* A, const UStruct* StructB, const void* B, TArray<FSingleObjectDiffEntry>& OutDifferingProperties);
@@ -257,14 +283,18 @@ namespace DiffUtils
 	UNREALED_API void CompareUnrelatedSCS(const UBlueprint* Old, const TArray< FSCSResolvedIdentifier >& OldHierarchy, const UBlueprint* New, const TArray< FSCSResolvedIdentifier >& NewHierarchy, FSCSDiffRoot& OutDifferingEntries );
 
 	UE_DEPRECATED(5.3, "DiffUtils now requires root objects so that object topology can be meaningfully compared.")
-	UNREALED_API bool Identical(const FResolvedProperty& AProp, const FResolvedProperty& BProp, const FPropertySoftPath& RootPath, TArray<FPropertySoftPath>& DifferingProperties); 
+	UNREALED_API bool Identical(const FResolvedProperty& AProp, const FResolvedProperty& BProp, const FPropertySoftPath& RootPath, TArray<FPropertySoftPath>& DifferingProperties);
+
+	UE_DEPRECATED(5.5, "Use version with FDiffParameters instead.")
+	UNREALED_API bool Identical(const FResolvedProperty& AProp, const FResolvedProperty& BProp, const UObject* OwningOuterA, const UObject* OwningOuterB, const FPropertySoftPath& RootPath, TArray<FPropertySoftPath>& DifferingProperties);
+
 	/**
 	 * DiffUtils now requires root objects so that object topology can be meaningfully compared.
 	 * DiffUtils::Identical works similar to FProperty::Identical except when a UObject is found, that is in OwningOuter*,
 	 * it's compared by topology instead. This allows sub-objects to diff correctly.
 	 */
 	UNREALED_API bool Identical(const FResolvedProperty& AProp, const FResolvedProperty& BProp, const UObject* OwningOuterA, const UObject* OwningOuterB);
-	UNREALED_API bool Identical(const FResolvedProperty& AProp, const FResolvedProperty& BProp, const UObject* OwningOuterA, const UObject* OwningOuterB, const FPropertySoftPath& RootPath, TArray<FPropertySoftPath>& DifferingProperties);
+	UNREALED_API bool Identical(const FResolvedProperty& AProp, const FResolvedProperty& BProp, const UObject* OwningOuterA, const UObject* OwningOuterB, FDiffParameters DiffParameters, TArray<FPropertySoftPath>& DifferingProperties);
 	UNREALED_API bool Identical(const TSharedPtr<IPropertyHandle>& PropertyHandleA, const TSharedPtr<IPropertyHandle>& PropertyHandleB, const TArray<TWeakObjectPtr<UObject>>& OwningOutersA = {}, const TArray<TWeakObjectPtr<UObject>>& OwningOutersB = {});
 	UNREALED_API TArray<FPropertySoftPath> GetVisiblePropertiesInOrderDeclared(const UStruct* ForStruct, const FPropertySoftPath& Scope = FPropertySoftPath());
 
diff --git a/Engine/Source/Runtime/Engine/Private/EdGraph/EdGraphNode.cpp b/Engine/Source/Runtime/Engine/Private/EdGraph/EdGraphNode.cpp
--- a/Engine/Source/Runtime/Engine/Private/EdGraph/EdGraphNode.cpp
+++ b/Engine/Source/Runtime/Engine/Private/EdGraph/EdGraphNode.cpp
@@ -306,22 +306,46 @@ void UEdGraphNode::DiffProperties(UStruct* StructA, UStruct* StructB, uint8* Dat
 			continue;
 		}
 
+		auto ShouldSkipProperty = [](const FProperty& Prop)
+			{
+				return !Prop.HasAnyPropertyFlags(CPF_Edit | CPF_BlueprintVisible) ||
+					Prop.HasAnyPropertyFlags(CPF_Transient) ||
+					Prop.HasAnyPropertyFlags(CPF_DisableEditOnInstance) ||
+					Prop.IsA(FDelegateProperty::StaticClass()) ||
+					Prop.IsA(FMulticastDelegateProperty::StaticClass());
+			};
+
 		// skip properties we cant see
-		if (!Prop->HasAnyPropertyFlags(CPF_Edit | CPF_BlueprintVisible) ||
-			Prop->HasAnyPropertyFlags(CPF_Transient) ||
-			Prop->HasAnyPropertyFlags(CPF_DisableEditOnInstance) ||
-			Prop->IsA(FDelegateProperty::StaticClass()) ||
-			Prop->IsA(FMulticastDelegateProperty::StaticClass()))
+		if (ShouldSkipProperty(*Prop))
 		{
 			continue;
 		}
 
-		if (!DiffUtils::Identical(FResolvedProperty(DataA, Prop), FResolvedProperty(DataB, PropB), StructA, StructB))
+		DiffUtils::FDiffParameters ComparisonParameters;
+		ComparisonParameters.ShouldIgnorePropertyPredicate = ShouldSkipProperty;
+		TArray<FPropertySoftPath> DifferingSubProperties;
+
+		if (!DiffUtils::Identical(FResolvedProperty(DataA, Prop), FResolvedProperty(DataB, PropB), StructA, StructB, ComparisonParameters, DifferingSubProperties))
 		{
 			// Only bother setting up the display data if we're storing the result
 			if (Results.CanStoreResults())
 			{
-				Diff.DisplayString = FText::Format(LOCTEXT("DIF_NodePropertyFmt", "Property Changed: {0} "), FText::FromString(Prop->GetName()));
+				FStringBuilderBase StringBuilder;
+				for (int DifferingIndex = 0; DifferingIndex < DifferingSubProperties.Num(); DifferingIndex++)
+				{
+					StringBuilder += DifferingSubProperties[DifferingIndex].ToDisplayName();
+				}
+
+				if (StringBuilder.Len() > 0)
+				{
+					Diff.DisplayString = FText::Format(LOCTEXT("DIF_NodePropertyFmt", "Property Changed: {0} - {1}"),
+						FText::FromString(Prop->GetName()),
+						FText::FromString(StringBuilder.ToString()));
+				}
+				else
+				{
+					Diff.DisplayString = FText::Format(LOCTEXT("DIF_NodePropertyFmt", "Property Changed: {0}"), FText::FromString(Prop->GetName()));
+				}
 			}
 			Results.Add(Diff);
 		}
