Cooker: Typo fix for previous changelist; moved code in FinishGeneratorPlatformSave had incorrect variable PackageData substituted for GeneratedInfo.PackageData.
#rnx

diff --git a/Engine/Source/Editor/UnrealEd/Private/Cooker/CookGenerationHelper.cpp b/Engine/Source/Editor/UnrealEd/Private/Cooker/CookGenerationHelper.cpp
--- a/Engine/Source/Editor/UnrealEd/Private/Cooker/CookGenerationHelper.cpp
+++ b/Engine/Source/Editor/UnrealEd/Private/Cooker/CookGenerationHelper.cpp
@@ -1027,7 +1027,7 @@ void FGenerationHelper::FinishGeneratorPlatformSave(FPackageData& PackageData, b
 	for (FCookGenerationInfo& GeneratedInfo : GetPackagesToGenerate())
 	{
 		FAssetDependency& Dependency = OutPackageDependencies.Emplace_GetRef();
-		Dependency.AssetId = FAssetIdentifier(PackageData.GetPackageName());
+		Dependency.AssetId = FAssetIdentifier(GeneratedInfo.PackageData->GetPackageName());
 		Dependency.Category = UE::AssetRegistry::EDependencyCategory::Package;
 		Dependency.Properties = UE::AssetRegistry::EDependencyProperty::Game;
 	}
