Add ability to replicate structured logging for multi-process cook.

diff --git a/Engine/Source/Editor/UnrealEd/Private/Cooker/CookWorkerServer.cpp b/Engine/Source/Editor/UnrealEd/Private/Cooker/CookWorkerServer.cpp
--- a/Engine/Source/Editor/UnrealEd/Private/Cooker/CookWorkerServer.cpp
+++ b/Engine/Source/Editor/UnrealEd/Private/Cooker/CookWorkerServer.cpp
@@ -4,17 +4,22 @@
 
 #include "Algo/Find.h"
 #include "Commandlets/AssetRegistryGenerator.h"
+#include "Containers/AnsiString.h"
 #include "Cooker/CompactBinaryTCP.h"
 #include "Cooker/CookDirector.h"
 #include "Cooker/CookGenerationHelper.h"
 #include "Cooker/CookPackageData.h"
 #include "Cooker/CookPlatformManager.h"
+#include "HAL/Platform.h"
 #include "HAL/PlatformProcess.h"
 #include "Interfaces/ITargetPlatform.h"
 #include "Interfaces/ITargetPlatformManagerModule.h"
+#include "Logging/StructuredLog.h"
+#include "Logging/StructuredLogFormat.h"
 #include "Math/NumericLimits.h"
 #include "Misc/AssertionMacros.h"
 #include "Misc/Char.h"
+#include "Misc/FeedbackContext.h"
 #include "Misc/FileHelper.h"
 #include "Misc/Parse.h"
 #include "Misc/ScopeLock.h"
@@ -1524,36 +1529,231 @@ bool FGeneratorEventMessage::TryRead(FCbObjectView Object)
 
 FGuid FGeneratorEventMessage::MessageType(TEXT("B6EE94CA70EC4F40B0D2214EDC11ED03"));
 
-FCbWriter& operator<<(FCbWriter& Writer, const FReplicatedLogData& Package)
+FCbWriter& operator<<(FCbWriter& Writer, const FReplicatedLogData& LogData)
 {
+	// Serializing as an array of unnamed fields and using the quantity of fields
+	// as the discriminator between structured and unstructured log data.
 	Writer.BeginArray();
-	Writer << Package.Category;
-	uint8 Verbosity = static_cast<uint8>(Package.Verbosity);
-	Writer << Verbosity;
-	Writer << Package.Message;
+	if (LogData.LogDataVariant.IsType<FReplicatedLogData::FUnstructuredLogData>())
+	{
+		const FReplicatedLogData::FUnstructuredLogData& UnstructuredLogData = LogData.LogDataVariant.Get<FReplicatedLogData::FUnstructuredLogData>();
+		Writer << UnstructuredLogData.Category;
+		uint8 Verbosity = static_cast<uint8>(UnstructuredLogData.Verbosity);
+		Writer << Verbosity;
+		Writer << UnstructuredLogData.Message;
+	}
+	else if (LogData.LogDataVariant.IsType<FCbObject>())
+	{
+		Writer << LogData.LogDataVariant.Get<FCbObject>();
+	}
+	else
+	{
+		checkNoEntry();
+	}
 	Writer.EndArray();
 	return Writer;
 }
 
-bool LoadFromCompactBinary(FCbFieldView Field, FReplicatedLogData& OutPackage)
+bool LoadFromCompactBinary(FCbFieldView Field, FReplicatedLogData& OutLogData)
+{
+	bool bOk = true;
+	FCbArrayView ArrayView = Field.AsArrayView();
+	switch (ArrayView.Num())
+	{
+	case 3:
+	{
+		OutLogData.LogDataVariant.Emplace<FReplicatedLogData::FUnstructuredLogData>();
+		FReplicatedLogData::FUnstructuredLogData& UnstructuredLogData = OutLogData.LogDataVariant.Get<FReplicatedLogData::FUnstructuredLogData>();
+		FCbFieldViewIterator It = ArrayView.CreateViewIterator();
+		bOk = LoadFromCompactBinary(*It++, UnstructuredLogData.Category) & bOk;
+		uint8 Verbosity;
+		if (LoadFromCompactBinary(*It++, Verbosity))
+		{
+			UnstructuredLogData.Verbosity = static_cast<ELogVerbosity::Type>(Verbosity);
+		}
+		else
+		{
+			bOk = false;
+			UnstructuredLogData.Verbosity = static_cast<ELogVerbosity::Type>(0);
+		}
+		bOk = LoadFromCompactBinary(*It++, UnstructuredLogData.Message) & bOk;
+		break;
+	}
+	case 1:
+	{
+		OutLogData.LogDataVariant.Emplace<FCbObject>();
+		FCbObject& StructuredLogData = OutLogData.LogDataVariant.Get<FCbObject>();
+		FCbFieldViewIterator It = ArrayView.CreateViewIterator();
+		if (It->IsObject())
+		{
+			StructuredLogData = FCbObject::Clone(It->AsObjectView());
+		}
+		else
+		{
+			bOk = false;
+		}
+		break;
+	}
+	default:
+		bOk = false;
+	}
+	return bOk;
+}
+
+FCbWriter& FLogMessagesMessageHandler::FLogRecordSerializationContext::Serialize(FCbWriter& Writer, const FLogRecord& LogRecord)
+{
+	Writer.BeginArray();
+	Writer << LogRecord.GetCategory();
+	Writer << static_cast<uint8>(LogRecord.GetVerbosity());
+	Writer << LogRecord.GetTime().GetUtcTime();
+	Writer << LogRecord.GetFormat();
+	Writer << LogRecord.GetFields();
+	Writer << LogRecord.GetFile();
+	Writer << LogRecord.GetLine();
+	Writer << LogRecord.GetTextNamespace();
+	Writer << LogRecord.GetTextKey();
+	Writer.EndArray();
+	return Writer;
+}
+
+bool FLogMessagesMessageHandler::FLogRecordSerializationContext::Deserialize(FCbFieldView Field, FLogRecord& OutLogRecord, int32 ProfileId)
 {
 	bool bOk = true;
 	FCbFieldViewIterator It = Field.CreateViewIterator();
-	bOk = LoadFromCompactBinary(*It++, OutPackage.Category) & bOk;
-	uint8 Verbosity;
-	if (LoadFromCompactBinary(*It++, Verbosity))
+	if (FName Category; LoadFromCompactBinary(*It++, Category))
+	{
+		OutLogRecord.SetCategory(Category);
+	}
+	else
+	{
+		bOk = false;
+	}
+	if (uint8 Verbosity; LoadFromCompactBinary(*It++, Verbosity) && Verbosity < ELogVerbosity::NumVerbosity)
+	{
+		OutLogRecord.SetVerbosity(static_cast<ELogVerbosity::Type>(Verbosity));
+	}
+	else
+	{
+		bOk = false;
+	}
+	if (FDateTime Time; LoadFromCompactBinary(*It++, Time))
+	{
+		OutLogRecord.SetTime(FLogTime::FromUtcTime(Time));
+	}
+	else
 	{
-		OutPackage.Verbosity = static_cast<ELogVerbosity::Type>(Verbosity);
+		bOk = false;
+	}
+	if (FString SerializedString; LoadFromCompactBinary(*It++, SerializedString))
+	{
+		FString& FormatString = StringTable.AddDefaulted_GetRef();
+		FormatString = FString::Printf(TEXT("[CookWorker %d]: %s"), ProfileId, *SerializedString);
+		OutLogRecord.SetFormat(*FormatString);
+	}
+	else
+	{
+		bOk = false;
+	}
+
+	FCbObject Object(FCbObject::Clone(It->AsObjectView()));
+	OutLogRecord.SetFields(MoveTemp(Object));
+	bOk = !It->HasError() && bOk;
+	It++;
+
+	if (TUtf8StringBuilder<64> FileStringBuilder; LoadFromCompactBinary(*It++, FileStringBuilder))
+	{
+		FAnsiString& FileString = AnsiStringTable.AddDefaulted_GetRef();
+		FileString = FileStringBuilder.ToString();
+		OutLogRecord.SetFile(*FileString);
 	}
 	else
 	{
 		bOk = false;
-		OutPackage.Verbosity = static_cast<ELogVerbosity::Type>(0);
 	}
-	bOk = LoadFromCompactBinary(*It++, OutPackage.Message) & bOk;
+	if (int32 Line; LoadFromCompactBinary(*It++, Line))
+	{
+		OutLogRecord.SetLine(Line);
+	}
+	else
+	{
+		bOk = false;
+	}
+	if (FString TextNamespaceString; LoadFromCompactBinary(*It++, TextNamespaceString))
+	{
+		if (!TextNamespaceString.IsEmpty())
+		{
+			OutLogRecord.SetTextNamespace(*StringTable.Emplace_GetRef(MoveTemp(TextNamespaceString)));
+		}
+		else
+		{
+			OutLogRecord.SetTextNamespace(nullptr);
+		}
+	}
+	else
+	{
+		bOk = false;
+	}
+	bool bHasTextKey = false;
+	if (FString TextKeyString; LoadFromCompactBinary(*It++, TextKeyString))
+	{
+		if (!TextKeyString.IsEmpty())
+		{
+			bHasTextKey = true;
+			OutLogRecord.SetTextKey(*StringTable.Emplace_GetRef(MoveTemp(TextKeyString)));
+		}
+		else
+		{
+			OutLogRecord.SetTextKey(nullptr);
+		}
+	}
+	else
+	{
+		bOk = false;
+	}
+
+	if (bHasTextKey)
+	{
+		FLogTemplate* LogTemplate = CreateLogTemplate(OutLogRecord.GetTextNamespace(), OutLogRecord.GetTextKey(), OutLogRecord.GetFormat());
+		TemplateTable.Add(LogTemplate);
+		OutLogRecord.SetTemplate(LogTemplate);
+	}
+	else
+	{
+		FLogTemplate* LogTemplate = CreateLogTemplate(OutLogRecord.GetFormat());
+		TemplateTable.Add(LogTemplate);
+		OutLogRecord.SetTemplate(LogTemplate);
+	}
+
 	return bOk;
 }
 
+void FLogMessagesMessageHandler::FLogRecordSerializationContext::ConditionalFlush(int32 TableSize)
+{
+	if ((StringTable.Num() > TableSize) || (AnsiStringTable.Num() > TableSize) || (TemplateTable.Num() > TableSize))
+	{
+		Flush();
+	}
+}
+
+void FLogMessagesMessageHandler::FLogRecordSerializationContext::Flush()
+{
+	if (!StringTable.IsEmpty() || !AnsiStringTable.IsEmpty() || !TemplateTable.IsEmpty())
+	{
+		// NOTE: We only call FlushThreadedLogs on GLog even though we might serialize structured logs via GLog or GWarn.
+		// GWarn is an output device, but GLog is a an output redirector, and only the redirector has/needs FlushThreadedLogs.
+		// Output devices are expected to not use any pointer on a structured log record after completion of the SerializeRecord call.
+		GLog->FlushThreadedLogs();
+	}
+	for (FLogTemplate* LogTemplate : TemplateTable)
+	{
+		DestroyLogTemplate(LogTemplate);
+	}
+
+	StringTable.Empty();
+	AnsiStringTable.Empty();
+	TemplateTable.Empty();
+}
+
 FGuid FLogMessagesMessageHandler::MessageType(TEXT("DB024D28203D4FBAAAF6AAD7080CF277"));
 
 FLogMessagesMessageHandler::~FLogMessagesMessageHandler()
@@ -1591,6 +1791,7 @@ void FLogMessagesMessageHandler::ServerReceiveMessage(FMPCollectorServerMessageC
 	FCbObjectView InMessage)
 {
 	TArray<FReplicatedLogData> Messages;
+
 	if (!LoadFromCompactBinary(InMessage["Messages"], Messages))
 	{
 		UE_LOG(LogCook, Error, TEXT("FLogMessagesMessageHandler received corrupted message from CookWorker"));
@@ -1599,22 +1800,61 @@ void FLogMessagesMessageHandler::ServerReceiveMessage(FMPCollectorServerMessageC
 
 	for (FReplicatedLogData& LogData : Messages)
 	{
-		if (LogData.Category == LogCookName && LogData.Message.Contains(HeartbeatCategoryText))
+		if (const FReplicatedLogData::FUnstructuredLogData* UnStructuredLogData = LogData.LogDataVariant.TryGet<FReplicatedLogData::FUnstructuredLogData>())
 		{
-			// Do not spam heartbeat messages into the CookDirector log
-			continue;
-		}
+			if (UnStructuredLogData->Category == LogCookName && UnStructuredLogData->Message.Contains(HeartbeatCategoryText))
+			{
+				// Do not spam heartbeat messages into the CookDirector log
+				continue;
+			}
 
-		FMsg::Logf(__FILE__, __LINE__, LogData.Category, LogData.Verbosity, TEXT("[CookWorker %d]: %s"),
+			FMsg::Logf(__FILE__, __LINE__, LogData.Category, LogData.Verbosity, TEXT("[CookWorker %d]: %s"),
 			Context.GetProfileId(), *LogData.Message);
+		}
+		else if (const FCbObject* StructuredLogObject = LogData.LogDataVariant.TryGet<FCbObject>())
+		{
+			FLogRecord LogRecord;
+			if (LogRecordSerializationContext.Deserialize((*StructuredLogObject)["S"], LogRecord, Context.GetProfileId()))
+			{
+				FOutputDevice* LogOverride = nullptr;
+				switch (LogRecord.GetVerbosity())
+				{
+				case ELogVerbosity::Error:
+				case ELogVerbosity::Warning:
+				case ELogVerbosity::Display:
+				case ELogVerbosity::SetColor:
+					LogOverride = GWarn;
+					break;
+				default:
+					break;
+				}
+				if (LogOverride)
+				{
+					LogOverride->SerializeRecord(LogRecord);
+				}
+				else
+				{
+					GLog->SerializeRecord(LogRecord);
+				}
+			}
+		}
+		else
+		{
+			checkNoEntry();
+		}
 	}
+
+	// Flush if the tables in the serialization context have exceeded 100 entries
+	const int32 TableSizeToFlushAt = 100;
+	LogRecordSerializationContext.ConditionalFlush(TableSizeToFlushAt);
 }
 
 void FLogMessagesMessageHandler::Serialize(const TCHAR* V, ELogVerbosity::Type Verbosity,
 	const FName& Category)
 {
 	FScopeLock QueueScopeLock(&QueueLock);
-	QueuedLogs.Add(FReplicatedLogData{ FString(V), Category, Verbosity });
+	FReplicatedLogData& LogData = QueuedLogs.Emplace_GetRef();
+	LogData.LogDataVariant.Emplace<FReplicatedLogData::FUnstructuredLogData>(FString(V), Category, Verbosity);
 }
 
 void FLogMessagesMessageHandler::Serialize(const TCHAR* V, ELogVerbosity::Type Verbosity,
@@ -1623,6 +1863,20 @@ void FLogMessagesMessageHandler::Serialize(const TCHAR* V, ELogVerbosity::Type V
 	Serialize(V, Verbosity, Category);
 }
 
+void FLogMessagesMessageHandler::SerializeRecord(const UE::FLogRecord& Record)
+{
+	FCbWriter Writer;
+	Writer.BeginObject();
+	Writer << "S";
+	FLogRecordSerializationContext::Serialize(Writer, Record);
+	Writer.EndObject();
+	FCbObject Object = Writer.Save().AsObject();
+
+	FScopeLock QueueScopeLock(&QueueLock);
+	FReplicatedLogData& LogData = QueuedLogs.Emplace_GetRef();
+	LogData.LogDataVariant.Emplace<FCbObject>(MoveTemp(Object));
+}
+
 FGuid FHeartbeatMessage::MessageType(TEXT("C08FFAF07BF34DD3A2FFB8A287CDDE83"));
 
 FHeartbeatMessage::FHeartbeatMessage(int32 InHeartbeatNumber)
diff --git a/Engine/Source/Editor/UnrealEd/Private/Cooker/CookWorkerServer.h b/Engine/Source/Editor/UnrealEd/Private/Cooker/CookWorkerServer.h
--- a/Engine/Source/Editor/UnrealEd/Private/Cooker/CookWorkerServer.h
+++ b/Engine/Source/Editor/UnrealEd/Private/Cooker/CookWorkerServer.h
@@ -24,6 +24,7 @@
 class FSocket;
 class ITargetPlatform;
 struct FProcHandle;
+namespace UE { class FLogTemplate; }
 namespace UE::CompactBinaryTCP { struct FMarshalledMessage; }
 namespace UE::Cook { class FCookDirector; }
 namespace UE::Cook { struct FDiscoveredPackageReplication; }
@@ -473,12 +474,16 @@ public:
 /** Stores the data passed into FOutputDevice::Serialize, for replication to the CookDirector. */
 struct FReplicatedLogData
 {
-	FString Message;
-	FName Category;
-	ELogVerbosity::Type Verbosity;
+	struct FUnstructuredLogData
+	{
+		FString Message;
+		FName Category;
+		ELogVerbosity::Type Verbosity;
+	};
+	TVariant<FUnstructuredLogData, FCbObject> LogDataVariant;
 };
-FCbWriter& operator<<(FCbWriter& Writer, const FReplicatedLogData& Package);
-bool LoadFromCompactBinary(FCbFieldView Field, FReplicatedLogData& OutPackage);
+FCbWriter& operator<<(FCbWriter& Writer, const FReplicatedLogData& LogData);
+bool LoadFromCompactBinary(FCbFieldView Field, FReplicatedLogData& OutLogData);
 
 /**
  * Send log messages from CookWorkers to the CookDirector, which marks them up with the CookWorkerId and
@@ -500,10 +505,32 @@ public:
 	virtual void Serialize(const TCHAR* V, ELogVerbosity::Type Verbosity, const FName& Category) override;
 	virtual void Serialize(const TCHAR* V, ELogVerbosity::Type Verbosity, const FName& Category,
 		const double Time) override;
+	virtual void SerializeRecord(const FLogRecord& Record) override;
 	virtual bool CanBeUsedOnAnyThread() const override { return true; }
 	virtual bool CanBeUsedOnMultipleThreads() const override { return true; }
 
 private:
+	class FLogRecordSerializationContext
+	{
+	public:
+		~FLogRecordSerializationContext()
+		{
+			Flush();
+		}
+
+		static FCbWriter& Serialize(FCbWriter& Writer, const FLogRecord& LogRecord);
+		bool Deserialize(FCbFieldView Field, FLogRecord& OutLogRecord, int32 ProfileId);
+
+		void ConditionalFlush(int32 TableSize);
+		void Flush();
+	private:
+		TArray<FString> StringTable;
+		TArray<FAnsiString> AnsiStringTable;
+		TArray<FLogTemplate*> TemplateTable;
+	};
+
+	FLogRecordSerializationContext LogRecordSerializationContext;
+
 	FCriticalSection QueueLock;
 	TArray<FReplicatedLogData> QueuedLogs;
 	TArray<FReplicatedLogData> QueuedLogsBackBuffer;
diff --git a/Engine/Source/Runtime/Core/Private/Logging/StructuredLog.cpp b/Engine/Source/Runtime/Core/Private/Logging/StructuredLog.cpp
--- a/Engine/Source/Runtime/Core/Private/Logging/StructuredLog.cpp
+++ b/Engine/Source/Runtime/Core/Private/Logging/StructuredLog.cpp
@@ -588,6 +588,13 @@ FLogTime FLogTime::Now()
 	return Time;
 }
 
+FLogTime FLogTime::FromUtcTime(const FDateTime& UtcTime)
+{
+	FLogTime Time;
+	Time.UtcTicks = UtcTime.GetTicks();
+	return Time;
+}
+
 FDateTime FLogTime::GetUtcTime() const
 {
 	return FDateTime(UtcTicks);
diff --git a/Engine/Source/Runtime/Core/Private/Serialization/CompactBinarySerialization.cpp b/Engine/Source/Runtime/Core/Private/Serialization/CompactBinarySerialization.cpp
--- a/Engine/Source/Runtime/Core/Private/Serialization/CompactBinarySerialization.cpp
+++ b/Engine/Source/Runtime/Core/Private/Serialization/CompactBinarySerialization.cpp
@@ -325,6 +325,18 @@ bool LoadFromCompactBinary(FCbFieldView Field, FGuid& OutValue, const FGuid& Def
 	return !Field.HasError();
 }
 
+bool LoadFromCompactBinary(FCbFieldView Field, FDateTime& OutValue)
+{
+	OutValue = Field.AsDateTime();
+	return !Field.HasError();
+}
+
+bool LoadFromCompactBinary(FCbFieldView Field, FDateTime& OutValue, FDateTime Default)
+{
+	OutValue = Field.AsDateTime(Default);
+	return !Field.HasError();
+}
+
 ///////////////////////////////////////////////////////////////////////////////////////////////////////////////////////
 
 template <typename CharType>
diff --git a/Engine/Source/Runtime/Core/Public/Logging/StructuredLog.h b/Engine/Source/Runtime/Core/Public/Logging/StructuredLog.h
--- a/Engine/Source/Runtime/Core/Public/Logging/StructuredLog.h
+++ b/Engine/Source/Runtime/Core/Public/Logging/StructuredLog.h
@@ -135,6 +135,7 @@ class FLogTime
 {
 public:
 	UE_API static FLogTime Now();
+	UE_API static FLogTime FromUtcTime(const FDateTime& UtcTime);
 
 	constexpr FLogTime() = default;
 
diff --git a/Engine/Source/Runtime/Core/Public/Serialization/CompactBinarySerialization.h b/Engine/Source/Runtime/Core/Public/Serialization/CompactBinarySerialization.h
--- a/Engine/Source/Runtime/Core/Public/Serialization/CompactBinarySerialization.h
+++ b/Engine/Source/Runtime/Core/Public/Serialization/CompactBinarySerialization.h
@@ -14,6 +14,7 @@
 
 class FArchive;
 class FName;
+struct FDateTime;
 struct FGuid;
 
 ///////////////////////////////////////////////////////////////////////////////////////////////////////////////////////
@@ -182,6 +183,9 @@ inline bool LoadFromCompactBinary(FCbFieldView Field, FIoHash& OutValue, const F
 CORE_API bool LoadFromCompactBinary(FCbFieldView Field, FGuid& OutValue);
 CORE_API bool LoadFromCompactBinary(FCbFieldView Field, FGuid& OutValue, const FGuid& Default);
 
+CORE_API bool LoadFromCompactBinary(FCbFieldView Field, FDateTime& OutValue);
+CORE_API bool LoadFromCompactBinary(FCbFieldView Field, FDateTime& OutValue, FDateTime Default);
+
 inline bool LoadFromCompactBinary(FCbFieldView Field, FCbObjectId& OutValue, const FCbObjectId& Default = FCbObjectId())
 {
 	OutValue = Field.AsObjectId(Default);
