Speculative workaround for Mac compile issue.

diff --git a/Engine/Source/Editor/UnrealEd/Private/Cooker/CookWorkerServer.cpp b/Engine/Source/Editor/UnrealEd/Private/Cooker/CookWorkerServer.cpp
--- a/Engine/Source/Editor/UnrealEd/Private/Cooker/CookWorkerServer.cpp
+++ b/Engine/Source/Editor/UnrealEd/Private/Cooker/CookWorkerServer.cpp
@@ -1854,7 +1854,11 @@ void FLogMessagesMessageHandler::Serialize(const TCHAR* V, ELogVerbosity::Type V
 {
 	FScopeLock QueueScopeLock(&QueueLock);
 	FReplicatedLogData& LogData = QueuedLogs.Emplace_GetRef();
-	LogData.LogDataVariant.Emplace<FReplicatedLogData::FUnstructuredLogData>(FString(V), Category, Verbosity);
+	LogData.LogDataVariant.Emplace<FReplicatedLogData::FUnstructuredLogData>();
+	FReplicatedLogData::FUnstructuredLogData& NewVal = LogData.LogDataVariant.Get<FReplicatedLogData::FUnstructuredLogData>();
+	NewVal.Message = V;
+	NewVal.Category = Category;
+	NewVal.Verbosity = Verbosity;
 }
 
 void FLogMessagesMessageHandler::Serialize(const TCHAR* V, ELogVerbosity::Type Verbosity,
