Fix compile error caused by manual merge from robomerge.

diff --git a/Engine/Source/Editor/UnrealEd/Private/Cooker/CookWorkerServer.cpp b/Engine/Source/Editor/UnrealEd/Private/Cooker/CookWorkerServer.cpp
--- a/Engine/Source/Editor/UnrealEd/Private/Cooker/CookWorkerServer.cpp
+++ b/Engine/Source/Editor/UnrealEd/Private/Cooker/CookWorkerServer.cpp
@@ -1808,8 +1808,8 @@ void FLogMessagesMessageHandler::ServerReceiveMessage(FMPCollectorServerMessageC
 				continue;
 			}
 
-			FMsg::Logf(__FILE__, __LINE__, LogData.Category, LogData.Verbosity, TEXT("[CookWorker %d]: %s"),
-			Context.GetProfileId(), *LogData.Message);
+			FMsg::Logf(__FILE__, __LINE__, UnStructuredLogData->Category, UnStructuredLogData->Verbosity, TEXT("[CookWorker %d]: %s"),
+			Context.GetProfileId(), *UnStructuredLogData->Message);
 		}
 		else if (const FCbObject* StructuredLogObject = LogData.LogDataVariant.TryGet<FCbObject>())
 		{
