Cooker: Remove call to UpdateManagementDatabase from FAssetRegistryGenerator::SaveAssetRegistry. The cooker does it earlier, and FAssetRegistryGenerator is tightly coupled to the cooker and not really functional for other uses.
#rnx

diff --git a/Engine/Source/Editor/UnrealEd/Private/Commandlets/AssetRegistryGenerator.cpp b/Engine/Source/Editor/UnrealEd/Private/Commandlets/AssetRegistryGenerator.cpp
--- a/Engine/Source/Editor/UnrealEd/Private/Commandlets/AssetRegistryGenerator.cpp
+++ b/Engine/Source/Editor/UnrealEd/Private/Commandlets/AssetRegistryGenerator.cpp
@@ -1751,8 +1751,6 @@ bool FAssetRegistryGenerator::SaveAssetRegistry(const FString& SandboxPath, bool
 		// an asset registry for each chunk.
 		if (FinalChunkManifests.Num() > 0)
 		{
-			// before final assignment ensure the database has been made.
-			UAssetManager::Get().UpdateManagementDatabase();
 			// Pass over all chunks and build a mapping of chunk index to asset registry name. All chunks that don't have a unique registry are assigned to the "generic bucket"
 			// which will be written to the master asset registry in chunk 0
 			for (int32 PakchunkIndex = 0; PakchunkIndex < FinalChunkManifests.Num(); ++PakchunkIndex)
