MPCook: Fix race condition between package assignment on scheduler thread and a worker shutting down on the communication thread that causes packages assigned to that worker to fail to be reassigned and cause an assertion on the director.
#issue https://horde.devtools.epicgames.com/log/663d9eb10e4261369615c431?lineindex=2812
#rnx

diff --git a/Engine/Source/Editor/UnrealEd/Private/Cooker/CookDirector.cpp b/Engine/Source/Editor/UnrealEd/Private/Cooker/CookDirector.cpp
--- a/Engine/Source/Editor/UnrealEd/Private/Cooker/CookDirector.cpp
+++ b/Engine/Source/Editor/UnrealEd/Private/Cooker/CookDirector.cpp
@@ -343,6 +343,23 @@ void FCookDirector::AssignRequests(TArrayView<FPackageData*> Requests, TArray<FW
 
 	AssignRequests(MoveTemp(WorkerIds), LocalRemoteWorkers, Requests, OutAssignments, MoveTemp(RequestGraph),
 		true /* bInitialAssignment */);
+
+	// Check for a race condition with the communication thread; if a Server was aborted after we CopyRemoteWorkers
+	// above but before we assigned packages to it, we need to abort those assigments now. 
+	TSet<FPackageData*> PackagesToReassignSet;
+	for (TRefCountPtr<FCookWorkerServer>& RemoteWorker : LocalRemoteWorkers)
+	{
+		if (RemoteWorker->IsShuttingDown())
+		{
+			RemoteWorker->AbortAllAssignments(PackagesToReassignSet, ECookDirectorThread::SchedulerThread);
+		}
+	}
+	if (!PackagesToReassignSet.IsEmpty())
+	{
+		// Defer the reassign because our caller has not yet put the packages into the assigned state.
+		FScopeLock CommunicationScopeLock(&CommunicationLock);
+		DeferredPackagesToReassign.Append(PackagesToReassignSet.Array());
+	}
 }
 
 TMap<FPackageData*, FAssignPackageExtraData> FCookDirector::GetAssignPackageExtraDatas(
