Landscape Spline Meshes Builder ignored the HLOD relevancy settings and always created HLOD relevant splines meshes components.

diff --git a/Engine/Source/Editor/UnrealEd/Private/WorldPartition/WorldPartitionLandscapeSplineMeshesBuilder.cpp b/Engine/Source/Editor/UnrealEd/Private/WorldPartition/WorldPartitionLandscapeSplineMeshesBuilder.cpp
--- a/Engine/Source/Editor/UnrealEd/Private/WorldPartition/WorldPartitionLandscapeSplineMeshesBuilder.cpp
+++ b/Engine/Source/Editor/UnrealEd/Private/WorldPartition/WorldPartitionLandscapeSplineMeshesBuilder.cpp
@@ -77,6 +77,8 @@ FName UWorldPartitionLandscapeSplineMeshesBuilder::GetSplineCollisionProfileName
 
 int32 UWorldPartitionLandscapeSplineMeshesBuilder::HashStaticMeshComponent(const UStaticMeshComponent* InComponent)
 {
+	const AActor* Owner = InComponent->GetOwner();
+
 	// Here we use ToCompactString as it internally clamps to 2 decimal digits, or else we get error precision and hash will differ)
 	uint32 HashValue = GetTypeHash(InComponent->GetComponentLocation().ToCompactString());
 	HashValue = HashCombine(HashValue, GetTypeHash(InComponent->GetComponentScale().ToCompactString()));
@@ -151,7 +153,7 @@ int32 UWorldPartitionLandscapeSplineMeshesBuilder::HashStaticMeshComponent(const
 	}
 
 	// Hash the owning actor's tags
-	if (const AActor* Owner = InComponent->GetOwner())
+	if (Owner)
 	{
 		TArray<FName> SortedTags = Owner->Tags;
 		SortedTags.Sort(FNameLexicalLess());
@@ -164,6 +166,10 @@ int32 UWorldPartitionLandscapeSplineMeshesBuilder::HashStaticMeshComponent(const
 		}
 	}
 
+	// Hash HLOD relevancy
+	const bool bOwnerHLODRelevant = !Owner || Owner->bEnableAutoLODGeneration;
+	HashValue = HashCombine(HashValue, !!InComponent->bEnableAutoLODGeneration && bOwnerHLODRelevant);
+
 	return HashValue;
 }
 
@@ -172,6 +178,8 @@ void UWorldPartitionLandscapeSplineMeshesBuilder::CloneStaticMeshComponent(const
 	check(InSrcMeshComponent);
 	check(DstMeshComponent);
 
+	const AActor* Owner = InSrcMeshComponent->GetOwner();
+
 	// Clone a USplineMeshComponent
 	if (const USplineMeshComponent* SplineMeshComponent = Cast<USplineMeshComponent>(InSrcMeshComponent))
 	{
@@ -234,10 +242,14 @@ void UWorldPartitionLandscapeSplineMeshesBuilder::CloneStaticMeshComponent(const
 	}
 
 	// Copy the owning actor's tags
-	if (const AActor* Owner = InSrcMeshComponent->GetOwner())
+	if (Owner)
 	{
 		DstMeshComponent->ComponentTags.Append(Owner->Tags);
 	}
+
+	// HLOD relevant
+	const bool bOwnerHLODRelevant = !Owner || Owner->bEnableAutoLODGeneration;
+	DstMeshComponent->bEnableAutoLODGeneration = !!InSrcMeshComponent->bEnableAutoLODGeneration && bOwnerHLODRelevant;
 }
 
 void UWorldPartitionLandscapeSplineMeshesBuilder::CloneStaticMeshComponentInActor(ALandscapeSplineMeshesActor* InActor, const UStaticMeshComponent* InMeshComponent)
@@ -399,7 +411,7 @@ bool UWorldPartitionLandscapeSplineMeshesBuilder::RunInternal(UWorld* InWorld, c
 			{
 				continue;
 			}
-
+						
 			TArray<UStaticMeshComponent*> Components;
 			SplineActor->GetComponents(Components);
 
