Increased default size of Undo buffer from 32MB to 256MB

diff --git a/Engine/Config/BaseEditorPerProjectUserSettings.ini b/Engine/Config/BaseEditorPerProjectUserSettings.ini
--- a/Engine/Config/BaseEditorPerProjectUserSettings.ini
+++ b/Engine/Config/BaseEditorPerProjectUserSettings.ini
@@ -468,7 +468,7 @@ bWheelMode=True
 
 [Undo]
 ; Size of Undo buffer in MB. Bigger number allows more Undo history especially when working with Landscape
-UndoBufferSize=32
+UndoBufferSize=256
 
 
 [PropertySettings]
diff --git a/Engine/Source/Editor/UnrealEd/Private/EditorServer.cpp b/Engine/Source/Editor/UnrealEd/Private/EditorServer.cpp
--- a/Engine/Source/Editor/UnrealEd/Private/EditorServer.cpp
+++ b/Engine/Source/Editor/UnrealEd/Private/EditorServer.cpp
@@ -1244,7 +1244,7 @@ UTransactor* UEditorEngine::CreateTrans()
 
 	if (!GConfig->GetInt(TEXT("Undo"), TEXT("UndoBufferSize"), UndoBufferSize, GEditorPerProjectIni))
 	{
-		UndoBufferSize = 16;
+		UndoBufferSize = 256;
 	}
 
 	UE_LOG(LogInit, Log, TEXT("Undo buffer set to %d MB"), UndoBufferSize);
