Changing LoadAndSaveOnePackage so that the UWorld::InitializationValues can be configured in derived classes

diff --git a/Engine/Source/Editor/UnrealEd/Classes/Commandlets/BaseIteratePackagesCommandlet.h b/Engine/Source/Editor/UnrealEd/Classes/Commandlets/BaseIteratePackagesCommandlet.h
--- a/Engine/Source/Editor/UnrealEd/Classes/Commandlets/BaseIteratePackagesCommandlet.h
+++ b/Engine/Source/Editor/UnrealEd/Classes/Commandlets/BaseIteratePackagesCommandlet.h
@@ -97,6 +97,9 @@ protected:
 
 	bool bUseWorldPartitionBuilder;
 
+	/** Initialisation for the FScopedEditorWorld if bUseWorldPartitionBuilder is enabled */
+	UWorld::InitializationValues WorldPartitionInitialisationValues;
+
 	/** Should we build lighting for the packages we are saving? **/
 	//bool bShouldBuildLighting;
 
diff --git a/Engine/Source/Editor/UnrealEd/Private/Commandlets/BaseIteratePackagesCommandlet.cpp b/Engine/Source/Editor/UnrealEd/Private/Commandlets/BaseIteratePackagesCommandlet.cpp
--- a/Engine/Source/Editor/UnrealEd/Private/Commandlets/BaseIteratePackagesCommandlet.cpp
+++ b/Engine/Source/Editor/UnrealEd/Private/Commandlets/BaseIteratePackagesCommandlet.cpp
@@ -86,7 +86,6 @@ UBaseIteratePackagesCommandlet::UBaseIteratePackagesCommandlet(const FObjectInit
 	: Super(ObjectInitializer)
 	, bForceUATEnvironmentVariableSet(false)
 {
-
 }
 
 int32 UBaseIteratePackagesCommandlet::InitializeParameters( const TArray<FString>& Tokens, TArray<FString>& PackageNames )
@@ -665,16 +664,7 @@ void UBaseIteratePackagesCommandlet::LoadAndSaveOnePackage(const FString& Filena
 						PerformAdditionalOperations(Object, bSavePackage);
 					});
 
-				UWorld::InitializationValues IVS;
-				IVS.RequiresHitProxies(false);
-				IVS.ShouldSimulatePhysics(false);
-				IVS.EnableTraceCollision(false);
-				IVS.CreateNavigation(false);
-				IVS.CreateAISystem(false);
-				IVS.AllowAudioPlayback(false);
-				IVS.CreatePhysicsScene(true);
-
-				FScopedEditorWorld ScopeEditorWorld(World, IVS);
+				FScopedEditorWorld ScopeEditorWorld(World, WorldPartitionInitialisationValues);
 
 				/*TUniquePtr<FLoaderAdapterShape> LoaderAdapterShape;
 				FBox Bounds = FBox(FVector(-HALF_WORLD_MAX, -HALF_WORLD_MAX, -HALF_WORLD_MAX), FVector(HALF_WORLD_MAX, HALF_WORLD_MAX, HALF_WORLD_MAX));
@@ -1070,6 +1060,13 @@ int32 UBaseIteratePackagesCommandlet::Main( const FString& Params )
 	bForceFinishAllCompilationBeforeGC = false;
 
 	bUseWorldPartitionBuilder = false;
+	WorldPartitionInitialisationValues.RequiresHitProxies(false);
+	WorldPartitionInitialisationValues.ShouldSimulatePhysics(false);
+	WorldPartitionInitialisationValues.EnableTraceCollision(false);
+	WorldPartitionInitialisationValues.CreateNavigation(false);
+	WorldPartitionInitialisationValues.CreateAISystem(false);
+	WorldPartitionInitialisationValues.AllowAudioPlayback(false);
+	WorldPartitionInitialisationValues.CreatePhysicsScene(true);
 
 	bKeepPackageGUIDOnSave = Switches.Contains(TEXT("KeepPackageGUIDOnSave"));
 	
