Remove superfluous search field from static mesh viewport hamburger > Layouts menu by making it unsearchable

diff --git a/Engine/Source/Editor/UnrealEd/Private/SAssetEditorViewport.cpp b/Engine/Source/Editor/UnrealEd/Private/SAssetEditorViewport.cpp
--- a/Engine/Source/Editor/UnrealEd/Private/SAssetEditorViewport.cpp
+++ b/Engine/Source/Editor/UnrealEd/Private/SAssetEditorViewport.cpp
@@ -138,6 +138,10 @@ bool SAssetEditorViewport::IsViewportConfigurationSet(FName ConfigurationName) c
 
 void SAssetEditorViewport::GenerateLayoutMenu(FMenuBuilder& MenuBuilder) const
 {
+	// Disable searching in this menu because it only contains visual representations of
+	// viewport layouts without any searchable text.
+	MenuBuilder.SetSearchable(false);
+
 	MenuBuilder.BeginSection("EditorViewportOnePaneConfigs", LOCTEXT("OnePaneConfigHeader", "One Pane"));
 	{
 		FSlimHorizontalToolBarBuilder OnePaneButton(CommandList, FMultiBoxCustomization::None);
