IterativeValidate: Give an error if the entire cook was invalidated due to a global settings change.
#rnx

diff --git a/Engine/Source/Editor/UnrealEd/Private/CookOnTheFlyServer.cpp b/Engine/Source/Editor/UnrealEd/Private/CookOnTheFlyServer.cpp
--- a/Engine/Source/Editor/UnrealEd/Private/CookOnTheFlyServer.cpp
+++ b/Engine/Source/Editor/UnrealEd/Private/CookOnTheFlyServer.cpp
@@ -8089,8 +8089,8 @@ bool UCookOnTheFlyServer::ArePreviousCookSettingsCompatible(const TMap<FName, FS
 	const FConfigSection* CookSettings = ConfigFile.FindSection(TEXT_CookSettings);
 	if (CookSettings == nullptr)
 	{
-		UE_LOG(LogCook, Display, TEXT("Cook invalidated for CookSettings file %s is invalid. Clearing previously cooked packages."),
-			*GetCookSettingsFileName(TargetPlatform));
+		UE_LOG(LogCook, Display, TEXT("Cook invalidated for platform %s because CookSettings file %s is invalid. Clearing previously cooked packages."),
+			*TargetPlatform->PlatformName(), *GetCookSettingsFileName(TargetPlatform));
 		return false;
 	}
 
diff --git a/Engine/Source/Editor/UnrealEd/Private/Cooker/IterativeValidatePackageWriter.cpp b/Engine/Source/Editor/UnrealEd/Private/Cooker/IterativeValidatePackageWriter.cpp
--- a/Engine/Source/Editor/UnrealEd/Private/Cooker/IterativeValidatePackageWriter.cpp
+++ b/Engine/Source/Editor/UnrealEd/Private/Cooker/IterativeValidatePackageWriter.cpp
@@ -349,10 +349,20 @@ void FIterativeValidatePackageWriter::BeginCook(const FCookInfo& Info)
 			UE_LOG(LogIterativeValidate, Display,
 				TEXT("-IterativeValidateAllowWrite is present, writable mode. Resaving packages as in a normal cook, but also running -diffonly on all packages that were found to be iteratively unmodified."));
 		}
+		if (Info.bFullBuild)
+		{
+			UE_LOG(LogIterativeValidate, Error,
+				TEXT("IterativeValidate was bypassed on this run; it is a full cook and all packages are marked iteratively modified."));
+		}
 		break;
 	case EPhase::Phase1:
 		UE_LOG(LogIterativeValidate, Display,
 			TEXT("Phase1: running -diffonly and a resave on all packages discovered to be iteratively unmodified."));
+		if (Info.bFullBuild)
+		{
+			UE_LOG(LogIterativeValidate, Error,
+				TEXT("IterativeValidate was bypassed on this run; it is a full cook and all packages are marked iteratively modified."));
+		}
 		break;
 	case EPhase::Phase2:
 		Load();
