Extending the DumpMaterialExpressions commandlet to contain class flags and some metadata for UMaterialExpressions.

- This is so I can audit the class flags and metadata for all the expressions.
- Some UMaterialExpressions classes are missing MinimalAPI and since this class flag isn't inherited from the base class I need to audit all derived material expressions.

[MATERIAL EXPRESSIONS]
NAME                                          TYPE                        FLAGS
LandscapeGrassOutput                          N/A                     MinimalAPI|CollapseCategories|Object
LandscapeLayerBlend                           N/A                      MinimalAPI|CollapseCategories|Object
LandscapeLayerCoords                          N/A                      MinimalAPI|CollapseCategories|Object
LandscapeLayerSample                      N/A                         MinimalAPI|CollapseCategories|Object
LandscapeLayerSwitch                        N/A                         MinimalAPI|CollapseCategories|Object
LandscapeLayerWeight                       N/A                         MinimalAPI|CollapseCategories|Object
LandscapePhysicalMaterialOutput        N/A                         MinimalAPI|CollapseCategories|Object
LandscapeVisibilityMask                     N/A                         MinimalAPI|CollapseCategories|Object
Abs                                                 N/A                         MinimalAPI|CollapseCategories|Object

diff --git a/Engine/Source/Editor/UnrealEd/Private/Commandlets/DumpMaterialExpressionsCommandlet.cpp b/Engine/Source/Editor/UnrealEd/Private/Commandlets/DumpMaterialExpressionsCommandlet.cpp
--- a/Engine/Source/Editor/UnrealEd/Private/Commandlets/DumpMaterialExpressionsCommandlet.cpp
+++ b/Engine/Source/Editor/UnrealEd/Private/Commandlets/DumpMaterialExpressionsCommandlet.cpp
@@ -90,12 +90,14 @@ void WriteOutMaterialExpressions(FArchive* FileWriter, const FString& OutputFile
 		FString Description;
 		FString Tooltip;
 		FString	Type;
+		FString ClassFlags;
 		FString ShowInCreateMenu;
 	};
 	TArray<FMaterialExpressionInfo> MaterialExpressionInfos;
 
 	const FString NameField = TEXT("NAME");
 	const FString TypeField = TEXT("TYPE");
+	const FString ClassField = TEXT("CLASS_FLAGS");
 	const FString ShowInCreateMenuField = TEXT("SHOW_IN_CREATE_MENU");
 	const FString KeywordsField = TEXT("KEYWORDS");
 	const FString CreationNameField = TEXT("CREATION_NAME");
@@ -106,6 +108,7 @@ void WriteOutMaterialExpressions(FArchive* FileWriter, const FString& OutputFile
 
 	int32 MaxNameLength = NameField.Len();
 	int32 MaxTypeLength = TypeField.Len();
+	int32 MaxClassFlagsLength = ClassField.Len();
 	int32 MaxShowInCreateMenuLength = ShowInCreateMenuField.Len();
 	int32 MaxKeywordsLength = KeywordsField.Len();
 	int32 MaxCreationNameLength = CreationNameField.Len();
@@ -144,6 +147,15 @@ void WriteOutMaterialExpressions(FArchive* FileWriter, const FString& OutputFile
 					|| Class == UMaterialExpressionFunctionOutput::StaticClass()
 					|| Class == UMaterialExpressionComposite::StaticClass());
 
+				const bool bCollapseCategories = Class->HasAnyClassFlags(CLASS_CollapseCategories);
+				const bool bHideCategories = Class->HasMetaData(TEXT("HideCategories"));
+				FString ClassFlags;
+				if (Class->HasAnyClassFlags(CLASS_MinimalAPI)) { ClassFlags = TEXT("MinimalAPI"); }
+				if (!ClassFlags.IsEmpty() && bCollapseCategories) { ClassFlags += TEXT("|"); }
+				if (bCollapseCategories) { ClassFlags += TEXT("CollapseCategories"); }
+				if (!ClassFlags.IsEmpty() && bHideCategories) { ClassFlags += TEXT("|"); }
+				if (bHideCategories) { ClassFlags += Class->GetMetaData(TEXT("HideCategories")); }
+
 				FString ExpressionType;
 				if (bControlFlow) { ExpressionType = TEXT("ControlFlow"); }
 				if (!ExpressionType.IsEmpty() && bNewHLSLGenerator) { ExpressionType += TEXT("|"); }
@@ -180,11 +192,13 @@ void WriteOutMaterialExpressions(FArchive* FileWriter, const FString& OutputFile
 				ExpressionInfo.Description = DefaultExpression->GetDescription();
 				ExpressionInfo.Tooltip = Tooltip;
 				ExpressionInfo.Type = ExpressionType;
+				ExpressionInfo.ClassFlags = ClassFlags;
 				ExpressionInfo.ShowInCreateMenu = bShowInCreateMenu ? TEXT("Yes") : TEXT("No");
 				MaterialExpressionInfos.Add(ExpressionInfo);
 
 				MaxNameLength = FMath::Max(MaxNameLength, ExpressionInfo.Name.Len());
 				MaxTypeLength = FMath::Max(MaxTypeLength, ExpressionInfo.Type.Len());
+				MaxClassFlagsLength = FMath::Max(MaxClassFlagsLength, ExpressionInfo.ClassFlags.Len());
 				MaxShowInCreateMenuLength = FMath::Max(MaxShowInCreateMenuLength, ExpressionInfo.ShowInCreateMenu.Len());
 				MaxKeywordsLength = FMath::Max(MaxKeywordsLength, ExpressionInfo.Keywords.Len());
 				MaxCreationNameLength = FMath::Max(MaxCreationNameLength, ExpressionInfo.CreationName.Len());
@@ -200,6 +214,7 @@ void WriteOutMaterialExpressions(FArchive* FileWriter, const FString& OutputFile
 	const int32 AdditionalPadding = 3;
 	MaxNameLength += AdditionalPadding;
 	MaxTypeLength += AdditionalPadding;
+	MaxClassFlagsLength += AdditionalPadding;
 	MaxShowInCreateMenuLength += AdditionalPadding;
 	MaxKeywordsLength += AdditionalPadding;
 	MaxCreationNameLength += AdditionalPadding;
@@ -211,6 +226,7 @@ void WriteOutMaterialExpressions(FArchive* FileWriter, const FString& OutputFile
 	TArray<uint32> MaxFieldLength;
 	MaxFieldLength.Add(MaxNameLength);
 	MaxFieldLength.Add(MaxTypeLength);
+	MaxFieldLength.Add(MaxClassFlagsLength);
 	MaxFieldLength.Add(MaxShowInCreateMenuLength);
 	MaxFieldLength.Add(MaxKeywordsLength);
 	MaxFieldLength.Add(MaxCreationNameLength);
@@ -223,6 +239,7 @@ void WriteOutMaterialExpressions(FArchive* FileWriter, const FString& OutputFile
 	TArray<FString> FieldNames;
 	FieldNames.Add(NameField);
 	FieldNames.Add(TypeField);
+	FieldNames.Add(ClassField);
 	FieldNames.Add(ShowInCreateMenuField);
 	FieldNames.Add(KeywordsField);
 	FieldNames.Add(CreationNameField);
@@ -237,6 +254,7 @@ void WriteOutMaterialExpressions(FArchive* FileWriter, const FString& OutputFile
 		FieldNames.Reset();
 		FieldNames.Add(GetFormattedText(ExpressionInfo.Name));
 		FieldNames.Add(GetFormattedText(ExpressionInfo.Type));
+		FieldNames.Add(GetFormattedText(ExpressionInfo.ClassFlags));
 		FieldNames.Add(GetFormattedText(ExpressionInfo.ShowInCreateMenu));
 		FieldNames.Add(GetFormattedText(ExpressionInfo.Keywords));
 		FieldNames.Add(GetFormattedText(ExpressionInfo.CreationName));
