Fix local variable scope after duplication - run post rename logic after the duplicated graph's name has been selected
[RN] Fixed compile error in a duplicated blueprint function after renaming a local variable in the originating function - run post rename logic after the duplicated graph's name has been selected

diff --git a/Engine/Source/Editor/BlueprintGraph/Private/EdGraphSchema_K2.cpp b/Engine/Source/Editor/BlueprintGraph/Private/EdGraphSchema_K2.cpp
--- a/Engine/Source/Editor/BlueprintGraph/Private/EdGraphSchema_K2.cpp
+++ b/Engine/Source/Editor/BlueprintGraph/Private/EdGraphSchema_K2.cpp
@@ -6803,8 +6803,6 @@ UEdGraph* UEdGraphSchema_K2::DuplicateGraph(UEdGraph* GraphToDuplicate) const
 				}
 			}
 
-			FName NewGraphName = FBlueprintEditorUtils::FindUniqueKismetName(Blueprint, GraphToDuplicate->GetFName().GetPlainNameString());
-			FEdGraphUtilities::RenameGraphCloseToName(NewGraph,NewGraphName.ToString());
 			// can't have two graphs with the same guid... that'd be silly!
 			NewGraph->GraphGuid = FGuid::NewGuid();
 
diff --git a/Engine/Source/Editor/Kismet/Private/SMyBlueprint.cpp b/Engine/Source/Editor/Kismet/Private/SMyBlueprint.cpp
--- a/Engine/Source/Editor/Kismet/Private/SMyBlueprint.cpp
+++ b/Engine/Source/Editor/Kismet/Private/SMyBlueprint.cpp
@@ -3480,9 +3480,12 @@ void SMyBlueprint::OnDuplicateAction()
 		else
 		{
 			const FScopedTransaction Transaction(LOCTEXT("DuplicateGraph", "Duplicate Graph"));
-			GetBlueprintObj()->Modify();
+			UBlueprint* BlueprintObj = GetBlueprintObj();
+			UEdGraph* GraphToDuplicate = GraphAction->EdGraph;
+
+			BlueprintObj->Modify();
 
-			UEdGraph* DuplicatedGraph = GraphAction->EdGraph->GetSchema()->DuplicateGraph(GraphAction->EdGraph);
+			UEdGraph* DuplicatedGraph = GraphToDuplicate->GetSchema()->DuplicateGraph(GraphToDuplicate);
 			check(DuplicatedGraph);
 
 			DuplicatedGraph->Modify();
@@ -3504,13 +3507,20 @@ void SMyBlueprint::OnDuplicateAction()
 
 			if (GraphType == GT_Function || GraphType == GT_Animation)
 			{
-				GetBlueprintObj()->FunctionGraphs.Add(DuplicatedGraph);
+				BlueprintObj->FunctionGraphs.Add(DuplicatedGraph);
 			}
 			else if (GraphType == GT_Macro)
 			{
-				GetBlueprintObj()->MacroGraphs.Add(DuplicatedGraph);
+				BlueprintObj->MacroGraphs.Add(DuplicatedGraph);
 			}
-			FBlueprintEditorUtils::MarkBlueprintAsStructurallyModified(GetBlueprintObj());
+
+			// Now that we've taken ownership of the graph (assigned it into macrographs
+			// or functiongraphs, as appropriate) it will be properly categorized
+			// and we can run post rename logic reliably:
+			FName NewGraphName = FBlueprintEditorUtils::FindUniqueKismetName(BlueprintObj, GraphToDuplicate->GetFName().GetPlainNameString());
+			FBlueprintEditorUtils::RenameGraph(DuplicatedGraph, NewGraphName.ToString());
+
+			FBlueprintEditorUtils::MarkBlueprintAsStructurallyModified(BlueprintObj);
 
 			BlueprintEditorPtr.Pin()->OpenDocument(DuplicatedGraph, FDocumentTracker::ForceOpenNewDocument);
 			DuplicateActionName = DuplicatedGraph->GetFName();
diff --git a/Engine/Source/Editor/UnrealEd/Private/Kismet2/BlueprintEditorUtils.cpp b/Engine/Source/Editor/UnrealEd/Private/Kismet2/BlueprintEditorUtils.cpp
--- a/Engine/Source/Editor/UnrealEd/Private/Kismet2/BlueprintEditorUtils.cpp
+++ b/Engine/Source/Editor/UnrealEd/Private/Kismet2/BlueprintEditorUtils.cpp
@@ -2677,21 +2677,21 @@ void FBlueprintEditorUtils::RenameGraph(UEdGraph* Graph, const FString& NewNameS
 		// Note: This will find ALL children (including nested children) so there's no need to do this recursively.
 		ValidateBlueprintChildVariables(Blueprint, Graph->GetFName(), PostValidateChildBlueprintLambda);
 
-		// Find all variable nodes in this graph.
-		TArray<UK2Node_Variable*> VariableNodes;
-		Graph->GetNodesOfClass<UK2Node_Variable>(VariableNodes);
-		GetAllChildGraphVariables(Graph, VariableNodes);
-
 		// if it's index is >= 0 we know it was found in the array of functiongraphs
-		bool bGraphIsFunction = (Blueprint->FunctionGraphs.IndexOfByKey(Graph) > -1);
+		const bool bGraphIsFunction = (Blueprint->FunctionGraphs.IndexOfByKey(Graph) > -1);
 		// For any nodes that reference a local variable, update the variable's scope to be the graph's new name (which will mirror the UFunction).
-		for (UK2Node_Variable* const VariableNode : VariableNodes)
+		if (bGraphIsFunction)
 		{
-			if (VariableNode->VariableReference.IsLocalScope())
+			// Find all variable nodes in this graph.
+			TArray<UK2Node_Variable*> VariableNodes;
+			Graph->GetNodesOfClass<UK2Node_Variable>(VariableNodes);
+			GetAllChildGraphVariables(Graph, VariableNodes);
+
+			for (UK2Node_Variable* const VariableNode : VariableNodes)
 			{
-				// if the rename is the function set the local variable scope to the new name otherwise we leave it with the same scope (Ex: subgraphs in a function)
-				if (bGraphIsFunction)
+				if (VariableNode->VariableReference.IsLocalScope())
 				{
+					// if the rename is the function set the local variable scope to the new name otherwise we leave it with the same scope (Ex: subgraphs in a function)
 					VariableNode->VariableReference.SetLocalMember(VariableNode->VariableReference.GetMemberName(), NewNameStr, VariableNode->VariableReference.GetMemberGuid());
 				}
 			}
