Add status description when the cook director fails to read from a socket connection with the client.

diff --git a/Engine/Source/Editor/UnrealEd/Private/Cooker/CookWorkerServer.cpp b/Engine/Source/Editor/UnrealEd/Private/Cooker/CookWorkerServer.cpp
--- a/Engine/Source/Editor/UnrealEd/Private/Cooker/CookWorkerServer.cpp
+++ b/Engine/Source/Editor/UnrealEd/Private/Cooker/CookWorkerServer.cpp
@@ -730,8 +730,9 @@ void FCookWorkerServer::PumpReceiveMessages()
 	if (SocketStatus != EConnectionStatus::Okay && SocketStatus != EConnectionStatus::Incomplete)
 	{
 		CrashDiagnosticsError = FString::Printf(
-			TEXT("CookWorkerCrash: CookWorker %d failed to read from socket, we will shutdown the remote process. Assigned packages will be returned to the director."),
-			ProfileId);
+			TEXT("CookWorkerCrash: CookWorker %d failed to read from socket with description: %s. we will shutdown the remote process. Assigned packages will be returned to the director."),
+			ProfileId,
+			DescribeStatus(SocketStatus));
 		bNeedCrashDiagnostics = true;
 		SendToState(EConnectStatus::WaitForDisconnect);
 		bTerminateImmediately = true;
