Cooker: TransitiveBuildDependencies: fix edge condition in graph search where a vertex previously fetched for transitive build dependencies would not be processed for whether it should be cooked when later requested.
#rnx

diff --git a/Engine/Source/Editor/UnrealEd/Private/Cooker/CookRequestCluster.cpp b/Engine/Source/Editor/UnrealEd/Private/Cooker/CookRequestCluster.cpp
--- a/Engine/Source/Editor/UnrealEd/Private/Cooker/CookRequestCluster.cpp
+++ b/Engine/Source/Editor/UnrealEd/Private/Cooker/CookRequestCluster.cpp
@@ -747,22 +747,22 @@ void FRequestCluster::FGraphSearch::QueueEdgesFetch(FVertexData& Vertex, TConstA
 
 	bool bAnyRequestedNeedsPlatformAgnostic = false;
 	bool bAnyRequested = false;
-	bool bAllHaveAlreadyCompletedFetch = false;
+	bool bAllHaveAlreadyCompletedFetch = true;
 
 	for (int32 PlatformIndex : PlatformIndexes)
 	{
 		// The platform data may have already been requested; request it only if current status is NotRequested
 		FQueryPlatformData& QueryData = Vertex.PlatformData[PlatformIndex];
 		if (!QueryData.bSchedulerThreadFetchCompleted)
-			{
+		{
 			bAllHaveAlreadyCompletedFetch = false;
 			EAsyncQueryStatus ExpectedStatus = EAsyncQueryStatus::NotRequested;
 			if (QueryData.CompareExchangeAsyncQueryStatus(ExpectedStatus, EAsyncQueryStatus::SchedulerRequested))
-		{
+			{
 				bAnyRequested = true;
+			}
 		}
 	}
-	}
 
 	if (bAnyRequested)
 	{
