Fix camera issue when stopping PIE - camera will teleport to the player's camera transform even if you had ejected from player control and flown the camera somewhere else.

diff --git a/Engine/Source/Editor/UnrealEd/Private/PlayLevel.cpp b/Engine/Source/Editor/UnrealEd/Private/PlayLevel.cpp
--- a/Engine/Source/Editor/UnrealEd/Private/PlayLevel.cpp
+++ b/Engine/Source/Editor/UnrealEd/Private/PlayLevel.cpp
@@ -705,7 +705,7 @@ void UEditorEngine::TeardownPlaySession(FWorldContext& PieWorldContext)
 			{
 				TSharedPtr<IAssetViewport> Viewport = SlatePlayInEditorSession->DestinationSlateViewport.Pin();
 
-				if(PlayInEditorSessionInfo.IsSet() && PlayInEditorSessionInfo->OriginalRequestParams.WorldType == EPlaySessionWorldType::PlayInEditor)
+				if(!bIsSimulatingInEditor)
 				{
 					// Set the editor viewport location to match that of Play in Viewport if we aren't simulating in the editor, we have a valid player to get the location from (unless we're going back to VR Editor, in which case we won't teleport the user.)
 					if (bLastViewAndLocationValid == true && !GEngine->IsStereoscopic3D( Viewport->GetActiveViewport() ) )
