Add HasExternalPackages virtual to AssetDefinition, and use it to determine save-able packages

diff --git a/Engine/Plugins/Experimental/AnimNext/Source/AnimNextEditor/Private/Graph/AnimNextGraphAssetDefinition.h b/Engine/Plugins/Experimental/AnimNext/Source/AnimNextEditor/Private/Graph/AnimNextGraphAssetDefinition.h
--- a/Engine/Plugins/Experimental/AnimNext/Source/AnimNextEditor/Private/Graph/AnimNextGraphAssetDefinition.h
+++ b/Engine/Plugins/Experimental/AnimNext/Source/AnimNextEditor/Private/Graph/AnimNextGraphAssetDefinition.h
@@ -27,6 +27,7 @@ public:
 		static const auto Categories = { FAssetCategoryPath(EAssetCategoryPaths::Animation, LOCTEXT("AnimNextSubMenu", "AnimNext")) };
 		return Categories;
 	}
+	virtual bool ShouldSaveExternalPackages() const override { return true; }
 };
 
 UCLASS()
diff --git a/Engine/Plugins/Experimental/Workspace/Source/WorkspaceEditor/Private/WorkspaceAssetDefinition.h b/Engine/Plugins/Experimental/Workspace/Source/WorkspaceEditor/Private/WorkspaceAssetDefinition.h
--- a/Engine/Plugins/Experimental/Workspace/Source/WorkspaceEditor/Private/WorkspaceAssetDefinition.h
+++ b/Engine/Plugins/Experimental/Workspace/Source/WorkspaceEditor/Private/WorkspaceAssetDefinition.h
@@ -25,6 +25,7 @@ public:
 		static const auto Categories = { FAssetCategoryPath(EAssetCategoryPaths::Misc) };
 		return Categories;
 	}
+	virtual bool ShouldSaveExternalPackages() const override { return true; }
 };
 
 #undef LOCTEXT_NAMESPACE
\ No newline at end of file
diff --git a/Engine/Source/Editor/AssetDefinition/Public/AssetDefinition.h b/Engine/Source/Editor/AssetDefinition/Public/AssetDefinition.h
--- a/Engine/Source/Editor/AssetDefinition/Public/AssetDefinition.h
+++ b/Engine/Source/Editor/AssetDefinition/Public/AssetDefinition.h
@@ -566,6 +566,12 @@ public:
 		return TSharedPtr<SWidget>();
 	}
 
+	/** Whether or not external packages associated with this asset should be saved in addition to the asset itself */
+	virtual bool ShouldSaveExternalPackages() const
+	{
+		return false;
+	}	
+
 	// DEVELOPER NOTE:
 	// Originally this class was based on the IAssetTypeActions implementation.  Several of the functions on there
 	// were created organically and added without a larger discussion about if such a thing belonged on those classes.
diff --git a/Engine/Source/Editor/UnrealEd/Private/FileHelpers.cpp b/Engine/Source/Editor/UnrealEd/Private/FileHelpers.cpp
--- a/Engine/Source/Editor/UnrealEd/Private/FileHelpers.cpp
+++ b/Engine/Source/Editor/UnrealEd/Private/FileHelpers.cpp
@@ -73,6 +73,7 @@
 #include "AssetRegistry/AssetRegistryModule.h"
 #include "Misc/NamePermissionList.h"
 #include "AnalyticsEventAttribute.h"
+#include "AssetDefinitionRegistry.h"
 #include "HierarchicalLOD.h"
 #include "WorldPartition/IWorldPartitionEditorModule.h"
 #include "WorldPartition/ActorDescContainer.h"
@@ -4314,10 +4315,27 @@ FEditorFileUtils::EPromptReturnCode FEditorFileUtils::PromptForCheckoutAndSave(c
 	TArray<UPackage*> PackagesToSave(InPackages);
 
 	// When saving a package which owns other packages, add those to the prompt as well,
-	// if we do not check dirty, we aren't already checked out and we prompt
-	if (!InOutParams.bAlreadyCheckedOut && !InOutParams.bCheckDirty && InOutParams.bPromptToSave)
+	for (UPackage* Package : InPackages)
 	{
-		for (UPackage* Package : InPackages)
+		const bool bShouldAddExternalPackages = [Package, InOutParams]() -> bool
+		{
+			if (const UObject* MainAsset = Package->FindAssetInPackage())
+			{
+				if (const UAssetDefinition* AssetDefinition = UAssetDefinitionRegistry::Get()->GetAssetDefinitionForClass(MainAsset->GetClass()))
+				{
+					if(AssetDefinition->ShouldSaveExternalPackages())
+					{
+						// if forced according to the top-level asset definition
+						return true;
+					}
+				}
+			}
+
+			// or if we do not check dirty, we aren't already checked out and we prompt
+			return !InOutParams.bAlreadyCheckedOut && !InOutParams.bCheckDirty && InOutParams.bPromptToSave;
+		}();
+		
+		if (bShouldAddExternalPackages)
 		{
 			for (UPackage* ExternalPackage : Package->GetExternalPackages())
 			{
