Cooker: Add log for rejection of a package due to EngineEditorResources. Change other rejection logs to match the standard format Cooking %s -> Rejected.
#rnx

diff --git a/Engine/Source/Editor/UnrealEd/Private/CookOnTheFlyServer.cpp b/Engine/Source/Editor/UnrealEd/Private/CookOnTheFlyServer.cpp
--- a/Engine/Source/Editor/UnrealEd/Private/CookOnTheFlyServer.cpp
+++ b/Engine/Source/Editor/UnrealEd/Private/CookOnTheFlyServer.cpp
@@ -6245,13 +6245,32 @@ void FSaveCookedPackageContext::SetupPlatform(const ITargetPlatform* InTargetPla
 		!TargetPlatform->HasEditorOnlyData())
 	{
 		SavePackageResult = ESavePackageResult::ContainsEditorOnlyData;
+		const TCHAR* RejectedReason = TEXT("EngineEditorContent");
+		if (GCookProgressDisplay & (int32)ECookProgressDisplayMode::Instigators)
+		{
+			UE_LOG(LogCook, Display, TEXT("Cooking %s, Instigator: { %s } -> Rejected %s"), *PackageName,
+				*(PackageData.GetInstigator().ToString()), RejectedReason);
+		}
+		else
+		{
+			UE_LOG(LogCook, Display, TEXT("Cooking %s -> Rejected %s"), *PackageName, RejectedReason);
+		}
 		return;
 	}
 	// Check whether or not game-specific behaviour should prevent this package from being cooked for the target platform
 	else if (!UAssetManager::Get().ShouldCookForPlatform(Package, TargetPlatform))
 	{
 		SavePackageResult = ESavePackageResult::ContainsEditorOnlyData;
-		UE_LOG(LogCook, Display, TEXT("Excluding %s"), *PackageName);
+		const TCHAR* RejectedReason = TEXT("NotAssetManagerShouldCookForPlatform");
+		if (GCookProgressDisplay & (int32)ECookProgressDisplayMode::Instigators)
+		{
+			UE_LOG(LogCook, Display, TEXT("Cooking %s, Instigator: { %s } -> Rejected %s"), *PackageName,
+				*(PackageData.GetInstigator().ToString()), RejectedReason);
+		}
+		else
+		{
+			UE_LOG(LogCook, Display, TEXT("Cooking %s -> Rejected %s"), *PackageName, RejectedReason);
+		}
 		return;
 	}
 	// check if this package is unsupported for the target platform (typically plugin content)
@@ -6266,7 +6285,16 @@ void FSaveCookedPackageContext::SetupPlatform(const ITargetPlatform* InTargetPla
 				(GenerationHelper && NeverCookPackages->Find(GenerationHelper->GetOwner().GetPackageName())))
 			{
 				SavePackageResult = ESavePackageResult::ContainsEditorOnlyData;
-				UE_LOG(LogCook, Display, TEXT("Excluding %s"), *PackageName);
+				const TCHAR* RejectedReason = TEXT("PlatformSpecificNeverCook");
+				if (GCookProgressDisplay & (int32)ECookProgressDisplayMode::Instigators)
+				{
+					UE_LOG(LogCook, Display, TEXT("Cooking %s, Instigator: { %s } -> Rejected %s"), *PackageName,
+						*(PackageData.GetInstigator().ToString()), RejectedReason);
+				}
+				else
+				{
+					UE_LOG(LogCook, Display, TEXT("Cooking %s -> Rejected %s"), *PackageName, RejectedReason);
+				}
 				return;
 			}
 		}
