[Shaders] Improve output of DumpMaterialShaderTypes commandlet.

This adds the static switch parameters to the output HTML file visualization what static switches can be considered dead-code.
It also reduces the framing as well as the color highlighting in the CSS layout, to make it easier to read.

#rnx

diff --git a/Engine/Source/Editor/UnrealEd/Private/Commandlets/DumpMaterialShaderTypes.cpp b/Engine/Source/Editor/UnrealEd/Private/Commandlets/DumpMaterialShaderTypes.cpp
--- a/Engine/Source/Editor/UnrealEd/Private/Commandlets/DumpMaterialShaderTypes.cpp
+++ b/Engine/Source/Editor/UnrealEd/Private/Commandlets/DumpMaterialShaderTypes.cpp
@@ -409,7 +409,8 @@ static void ProcessSwitchOptimizer(const ITargetPlatform* TargetPlatform, const
 	Output.Log(TEXT("<head>"));
 	Output.Log(TEXT("\t<title>StaticSwitchOptimizer</title>"));
 	Output.Log(TEXT("\t<style>"));
-	Output.Log(TEXT("\t\ttable, th, td {border: 1px solid black;}"));
+	Output.Log(TEXT("\t\ttable {border: 1px solid black; font-family: monospace;}"));
+	Output.Log(TEXT("\t\tth {padding-right: 5px; padding-left: 5px; padding-top: 2px; padding-bottom: 2px;}"));
 	Output.Log(TEXT("\t</style>"));
 	Output.Log(TEXT("</head>"));
 	Output.Log(TEXT("<body>"));
@@ -512,20 +513,30 @@ static void ProcessSwitchOptimizer(const ITargetPlatform* TargetPlatform, const
 	using ShaderIdAndSwitchGroup = TPair<const FMaterialShaderMapId, StaticSwitchGroupSet>;
 	using StaticSwitchGroupType = TPair<const StaticSwitchArrayType, MaterialAndSizeSet>;
 
+	struct FStaticSwitchMetaData
+	{
+		bool bIsVarying;
+		TArray<bool> PermutationVector; // Row-vector in the permutation matrix
+	};
+
 	using ShaderIdAndStaticSwitchGroupType = TPair<const FMaterialShaderMapId, StaticSwitchGroupType>;
 	using ShaderIdGroupArray = TArray<ShaderIdAndStaticSwitchGroupType>;
-	using VaryingSwitchesType = Experimental::TRobinHoodHashMap<FName, TArray<bool>>;
-	using InnerFilteredType = TPair<const ShaderIdGroupArray, VaryingSwitchesType>;
+	using SwitchPermutationVectorType = Experimental::TRobinHoodHashMap<const FName, FStaticSwitchMetaData>;
+
+	struct FInnerFilteredType
+	{
+		ShaderIdGroupArray ShaderIdGroups;
+		SwitchPermutationVectorType StaticSwitches;
+	};
 	
-	using FlattenedArrayType = TArray<InnerFilteredType>;
-	FlattenedArrayType FilteredHashMap;
+	TArray<FInnerFilteredType> FilteredHashMap;
 
 	for(const ShaderIdAndSwitchGroup& OuterElement : ShaderMapHashMap)
 	{
 		if(OuterElement.Value.Num() > 1)
 		{
 			ShaderIdGroupArray InnerMap;
-			VaryingSwitchesType VaryingSwitches;
+			SwitchPermutationVectorType VaryingSwitches;
 			StaticSwitchArrayType First;
 			for(const StaticSwitchGroupType& Value : OuterElement.Value)
 			{
@@ -541,9 +552,11 @@ static void ProcessSwitchOptimizer(const ITargetPlatform* TargetPlatform, const
 				for(int32 i = 0; i < Inner.Key.Num(); i++)
 				{
 					const FStaticSwitchParameter& Key = Inner.Key[i];
-					if(Key.Value != First[i].Value)
+					const bool bIsVarying = (Key.Value != First[i].Value);
+					FStaticSwitchMetaData* MetaData = VaryingSwitches.FindOrAdd(Key.ParameterInfo.Name, FStaticSwitchMetaData{});
+					if (bIsVarying)
 					{
-						VaryingSwitches.FindOrAdd(Key.ParameterInfo.Name, TArray<bool>());
+						MetaData->bIsVarying = true;
 					}
 				}
 				InnerMap.Emplace(OuterElement.Key, Inner);
@@ -567,21 +580,21 @@ static void ProcessSwitchOptimizer(const ITargetPlatform* TargetPlatform, const
 
 			if(VaryingSwitches.Num())
 			{
-				FilteredHashMap.Emplace(MoveTemp(InnerMap), MoveTemp(VaryingSwitches));
+				FilteredHashMap.Add(FInnerFilteredType{ MoveTemp(InnerMap), MoveTemp(VaryingSwitches) });
 			}
 		}
 	}
 
-	FilteredHashMap.Sort([](const InnerFilteredType& A, const InnerFilteredType& B)
+	FilteredHashMap.Sort([](const FInnerFilteredType& A, const FInnerFilteredType& B)
 	{
 		int32 NumA = 0;
-		for (const ShaderIdAndStaticSwitchGroupType& InnerA : A.Key)
+		for (const ShaderIdAndStaticSwitchGroupType& InnerA : A.ShaderIdGroups)
 		{
 			NumA += (*InnerA.Value.Value.begin()).Value;
 		}
 
 		int32 NumB = 0;
-		for (const ShaderIdAndStaticSwitchGroupType& InnerB : B.Key)
+		for (const ShaderIdAndStaticSwitchGroupType& InnerB : B.ShaderIdGroups)
 		{
 			NumB += (*InnerB.Value.Value.begin()).Value;
 		}
@@ -590,45 +603,75 @@ static void ProcessSwitchOptimizer(const ITargetPlatform* TargetPlatform, const
 
 	for(int32 i = 0; i < FilteredHashMap.Num(); i++)
 	{
-		const ShaderIdGroupArray& InnerMap = FilteredHashMap[i].Key;
-		VaryingSwitchesType& Varying = FilteredHashMap[i].Value;
+		const ShaderIdGroupArray& InnerMap = FilteredHashMap[i].ShaderIdGroups;
+		SwitchPermutationVectorType& StaticSwitches = FilteredHashMap[i].StaticSwitches;
 		const UMaterialInterface* Parent = (*(*InnerMap.begin()).Value.Value.begin()).Key->GetMaterial();
 		const FMaterialShaderMapId& ShaderId = (*InnerMap.begin()).Key;
 
-		int32 Num = 0;
+		int32 NumShaders = 0;
+		int32 NumStaticSwitchPermutations = 0;
+		int32 NumStaticSwitchesTotal = 0;
+
 		for (const ShaderIdAndStaticSwitchGroupType& Inner : InnerMap)
 		{
-			Num += (*Inner.Value.Value.begin()).Value;
+			NumShaders += (*Inner.Value.Value.begin()).Value;
 		}
 
-		Output.Log(TEXT("<table>"));
-		Output.Log(FString::Printf(TEXT("<tr><th><h3>Candidate %s with NumShaders: %d</h3></th></tr>"), *Parent->GetOuter()->GetFName().ToString(), Num));
-		Output.Log(TEXT("<tr><td><table>"));
-
 		for (const ShaderIdAndStaticSwitchGroupType& Inner : InnerMap)
 		{
 			const StaticSwitchGroupType& Groups = Inner.Value;
 
-			if(Groups.Key.Num())
+			if (Groups.Key.Num())
 			{
-				for(const FStaticSwitchParameter& Param : Groups.Key)
+				NumStaticSwitchesTotal = FMath::Max(NumStaticSwitchesTotal, Groups.Key.Num());
+				for (const FStaticSwitchParameter& Param : Groups.Key)
 				{
-					if(TArray<bool>* ValueArray = Varying.Find(Param.ParameterInfo.Name))
+					if (FStaticSwitchMetaData* MetaData = StaticSwitches.Find(Param.ParameterInfo.Name))
 					{
-						ValueArray->Add(Param.Value);
+						MetaData->PermutationVector.Add(Param.Value);
+						NumStaticSwitchPermutations = FMath::Max(NumStaticSwitchPermutations, MetaData->PermutationVector.Num());
 					}
 				}
 			}
 		}
 
-		for (const TPair<const FName, TArray<bool>>& Param : Varying)
+		Output.Log(TEXT("<table>"));
+		Output.Log(FString::Printf(TEXT("<tr><th><h3>Candidate %s (Shaders: %d)</h3></th></tr>"), *Parent->GetOuter()->GetFName().ToString(), NumShaders));
+		Output.Log(TEXT("<tr><td><table>"));
+
+		// Row for captions
+		Output.Log(TEXT("\t<tr>"));
+		Output.Log(FString::Printf(TEXT("\t\t<th style=\"background-color:gray;\">%d Static Switch(es)</th>"), NumStaticSwitchesTotal));
+		Output.Log(FString::Printf(TEXT("\t\t<th style=\"background-color:gray;\" colspan=\"%d\">%d Permutation(s)</th>"), NumStaticSwitchPermutations, NumStaticSwitchPermutations));
+		Output.Log(TEXT("\t</tr>"));
+
+		// Row for each static switch parameter that is included in at least one permutation
+		for (const TPair<const FName, FStaticSwitchMetaData>& Param : StaticSwitches)
 		{
 			Output.Log(TEXT("\t<tr>"));
 			Output.Log(FString::Printf(TEXT("\t\t<th>%s</th>"), *Param.Key.ToString()));
 
-			for(bool Switch : Param.Value)
+			if (Param.Value.bIsVarying)
+			{
+				for (bool bSwitchEnabled : Param.Value.PermutationVector)
+				{
+					Output.Log(
+						bSwitchEnabled
+							? TEXT("\t\t<th style=\"background-color:red;\">1</th>")
+							: TEXT("\t\t<th>0</th>")
+					);
+				}
+			}
+			else
 			{
-				Output.Log(Switch ? TEXT("\t\t<td style=\"background-color:red;\">1</td>") : TEXT("\t\t<td style=\"background-color:green;\">0</td>"));
+				for (bool bSwitchEnabled : Param.Value.PermutationVector)
+				{
+					Output.Log(
+						bSwitchEnabled
+							? TEXT("\t\t<th style=\"background-color:silver;\">1</th>")
+							: TEXT("\t\t<th style=\"background-color:silver;\">0</th>")
+					);
+				}
 			}
 			Output.Log(TEXT("\t</tr>"));
 		}
