AssetEditorToolkit : fix CIS error on Mac Arm, maybe?

diff --git a/Engine/Source/Editor/UnrealEd/Private/Toolkits/AssetEditorToolkit.cpp b/Engine/Source/Editor/UnrealEd/Private/Toolkits/AssetEditorToolkit.cpp
--- a/Engine/Source/Editor/UnrealEd/Private/Toolkits/AssetEditorToolkit.cpp
+++ b/Engine/Source/Editor/UnrealEd/Private/Toolkits/AssetEditorToolkit.cpp
@@ -963,6 +963,15 @@ void FAssetEditorToolkit::SwitchToWorldCentricEditor_Execute( TWeakPtr< FAssetEd
 }
 
 
+
+EVisibility FAssetEditorToolkit::GetVisibilityWhileAssetCompiling() const
+{
+	// don't tick GUI of asset editor when asset is compiling
+	//	this is to prevent deadlocks in broken asset editors
+	// @todo : change this default to Visible and instead return Collapsed only in the editors that need this bodge
+	return EVisibility::Collapsed;
+}
+
 void FAssetEditorToolkit::FindInContentBrowser_Execute()
 {
 	TArray< UObject* > ObjectsToSyncTo;
diff --git a/Engine/Source/Editor/UnrealEd/Public/Toolkits/AssetEditorToolkit.h b/Engine/Source/Editor/UnrealEd/Public/Toolkits/AssetEditorToolkit.h
--- a/Engine/Source/Editor/UnrealEd/Public/Toolkits/AssetEditorToolkit.h
+++ b/Engine/Source/Editor/UnrealEd/Public/Toolkits/AssetEditorToolkit.h
@@ -372,8 +372,8 @@ protected:
 	virtual bool CanFindInContentBrowser() const { return true; }
 
 	/** Called to get the desired visibility when the asset is compiling */
-	UNREALED_API virtual EVisibility GetVisibilityWhileAssetCompiling() const { return EVisibility::Collapsed; }
-
+	UNREALED_API virtual EVisibility GetVisibilityWhileAssetCompiling() const;
+	
 	/** Called to test if "Find in Content Browser" should be visible for this asset */
 	virtual bool IsFindInContentBrowserButtonVisible() const { return true; }
 
