IncrementalCook: Add support for iterative and incremental cooks of CookPackageSplitters.
#rnx

diff --git a/Engine/Source/Developer/IoStoreUtilities/Private/ZenStoreWriter.cpp b/Engine/Source/Developer/IoStoreUtilities/Private/ZenStoreWriter.cpp
--- a/Engine/Source/Developer/IoStoreUtilities/Private/ZenStoreWriter.cpp
+++ b/Engine/Source/Developer/IoStoreUtilities/Private/ZenStoreWriter.cpp
@@ -1102,7 +1102,16 @@ TUniquePtr<FAssetRegistryState> FZenStoreWriter::LoadPreviousAssetRegistry()
 	for (const TPair<FName, const FAssetPackageData*>& Pair : PreviousStatePackages)
 	{
 		FName PackageName = Pair.Key;
-		if (!PackageNameToIndex.Find(PackageName))
+		if (Pair.Value->DiskSize < 0)
+		{
+			// Keep the FailedSave previous cook packages; some of them (NeverCookPlaceholders) are not expected to exist
+			// in the package store
+			continue;
+		}
+		if (PackageNameToIndex.Find(PackageName))
+		{
+			continue;
+		}
 		{
 			RemoveSet.Add(PackageName);
 		}
diff --git a/Engine/Source/Editor/UnrealEd/Classes/CookOnTheSide/CookOnTheFlyServer.h b/Engine/Source/Editor/UnrealEd/Classes/CookOnTheSide/CookOnTheFlyServer.h
--- a/Engine/Source/Editor/UnrealEd/Classes/CookOnTheSide/CookOnTheFlyServer.h
+++ b/Engine/Source/Editor/UnrealEd/Classes/CookOnTheSide/CookOnTheFlyServer.h
@@ -29,6 +29,7 @@
 class FAssetRegistryGenerator;
 class FAsyncIODelete;
 class FDiffModeCookServerUtils;
+class FIterativeValidatePackageWriter;
 class FReferenceCollector;
 class FSavePackageContext;
 class IAssetRegistry;
@@ -537,7 +538,8 @@ private:
 	 * @param PackageData			The PackageData to be considered for saving.
 	 */
 	void QueueDiscoveredPackage(UE::Cook::FPackageData& PackageData, UE::Cook::FInstigator&& Instigator, 
-		UE::Cook::FDiscoveredPlatformSet&& ReachablePlatforms, bool bUrgent=false);
+		UE::Cook::FDiscoveredPlatformSet&& ReachablePlatforms, bool bUrgent=false,
+		UE::Cook::FGenerationHelper* ParentGenerationHelper = nullptr);
 	void QueueDiscoveredPackageOnDirector(UE::Cook::FPackageData& PackageData, UE::Cook::FInstigator&& Instigator,
 		UE::Cook::FDiscoveredPlatformSet&& ReachablePlatforms, bool bUrgent);
 
@@ -1605,6 +1607,7 @@ private:
 	TUniquePtr<UE::Cook::FODSCClientData> ODSCClientData;
 
 	friend FAssetRegistryGenerator;
+	friend FIterativeValidatePackageWriter;
 	friend UE::Cook::FAssetRegistryMPCollector;
 	friend UE::Cook::FBeginCookConfigSettings;
 	friend UE::Cook::FCookDirector;
diff --git a/Engine/Source/Editor/UnrealEd/Private/Commandlets/AssetRegistryGenerator.cpp b/Engine/Source/Editor/UnrealEd/Private/Commandlets/AssetRegistryGenerator.cpp
--- a/Engine/Source/Editor/UnrealEd/Private/Commandlets/AssetRegistryGenerator.cpp
+++ b/Engine/Source/Editor/UnrealEd/Private/Commandlets/AssetRegistryGenerator.cpp
@@ -1287,7 +1287,7 @@ void FAssetRegistryGenerator::ComputePackageDifferences(const FComputeDifference
 				else
 				{
 					FGeneratorPackageInfo& Info = OutDifference.GeneratorPackages.FindOrAdd(GeneratorName);
-					Info.Generated.Add(PackageName, PreviousPackageData->GetPackageSavedHash());
+					Info.Generated.Add(PackageName, CopyAssetPackageDataForIncrementalCook(*PreviousPackageData));
 				}
 			}
 			else
@@ -1323,7 +1323,7 @@ void FAssetRegistryGenerator::ComputePackageDifferences(const FComputeDifference
 		EDifference* GeneratorDifference = OutDifference.Packages.Find(GeneratorName);
 		if (GeneratorDifference && *GeneratorDifference == EDifference::RemovedCooked)
 		{
-			for (const TPair<FName, FIoHash>& Generated : Iter->Value.Generated)
+			for (const TPair<FName, FAssetPackageData>& Generated : Iter->Value.Generated)
 			{
 				OutDifference.Packages.Add(Generated.Key, EDifference::RemovedCooked);
 			}
@@ -1441,6 +1441,38 @@ bool FAssetRegistryGenerator::ComputePackageDifferences_IsPackageFileUnchanged(
 	return true;
 }
 
+FAssetPackageData FAssetRegistryGenerator::CopyAssetPackageDataForIncrementalCook(const FAssetPackageData& Source)
+{
+	FAssetPackageData Result;
+
+	// Copy small scalars since they don't cost memory or much network bandwidth
+	// Skip large scalars and containers that we don't need to save network bandwidth
+
+	// CookedHash is not read by the cook
+	// Result.CookedHash = Source.CookedHash;
+	
+	// PackageSavedHash is used during legacy iterative cooks to compare whether the package is modified
+	Result.SetPackageSavedHash(Source.GetPackageSavedHash());
+
+	// ChunkHashes is not read by the cook
+	// Result.ChunkHashes = Source.ChunkHashes;
+
+	// ImportedClasses is a large container, but incremental cook needs it to calculate the current packagedigest
+	Result.ImportedClasses = Source.ImportedClasses;
+	Result.DiskSize = Source.DiskSize;
+	Result.FileVersionUE = Source.FileVersionUE;
+	Result.FileVersionLicenseeUE = Source.FileVersionLicenseeUE;
+	Result.SetIsLicenseeVersion(Source.IsLicenseeVersion());
+	Result.SetHasVirtualizedPayloads(Source.HasVirtualizedPayloads());
+
+	// CustomVersions are not read by the cook
+	//Result.SetCustomVersions(Source.GetCustomVersions());
+
+	Result.Extension = Source.Extension;
+
+	return Result;
+}
+
 FName FAssetRegistryGenerator::GetGeneratorPackage(FName PackageName, const FAssetRegistryState& InState)
 {
 	TConstArrayView<const FAssetData*> Assets = InState.GetAssetsByPackageName(PackageName);
@@ -1487,7 +1519,8 @@ void FAssetRegistryGenerator::ComputePackageRemovals(const FAssetRegistryState&
 				}
 				else
 				{
-					OutGeneratorPackages.FindOrAdd(GeneratorName).Generated.Add(PackageName, PreviousPackageData->GetPackageSavedHash());
+					OutGeneratorPackages.FindOrAdd(GeneratorName).Generated.Add(PackageName,
+						CopyAssetPackageDataForIncrementalCook(*PreviousPackageData));
 				}
 			}
 			else
@@ -1511,7 +1544,7 @@ void FAssetRegistryGenerator::ComputePackageRemovals(const FAssetRegistryState&
 		FName GeneratorName = Iter->Key;
 		if (RemovedPackageSet.Contains(GeneratorName))
 		{
-			for (const TPair<FName, FIoHash>& Generated : Iter->Value.Generated)
+			for (const TPair<FName, FAssetPackageData>& Generated : Iter->Value.Generated)
 			{
 				RemovedPackageSet.Add(Generated.Key);
 			}
@@ -2777,13 +2810,8 @@ void FAssetRegistryPackageMessage::Write(FCbWriter& Writer) const
 	Writer.EndArray();
 	if (OverrideAssetPackageData)
 	{
-		Writer.BeginObject("P");
-		{
-			// Currently we only replicate Guid and ImportedClasses, since these are the only fields set by generated pacakges
-			Writer << "H" << OverrideAssetPackageData->GetPackageSavedHash();
-			Writer << "C" << OverrideAssetPackageData->ImportedClasses;
-		}
-		Writer.EndObject();
+		Writer.SetName("P");
+		OverrideAssetPackageData->NetworkWrite(Writer);
 	}
 	if (OverridePackageDependencies)
 	{
@@ -2818,14 +2846,7 @@ bool FAssetRegistryPackageMessage::TryRead(FCbObjectView Object)
 	if (OverrideAssetPackageDataField.HasValue())
 	{
 		OverrideAssetPackageData.Emplace();
-		// Currently we only replicate Guid and ImportedClasses, since these are the only fields set by generated packages
-		FIoHash PackageSavedHash;
-		if (!LoadFromCompactBinary(OverrideAssetPackageDataField["H"], PackageSavedHash))
-		{
-			return false;
-		}
-		OverrideAssetPackageData->SetPackageSavedHash(PackageSavedHash);
-		if (!LoadFromCompactBinary(OverrideAssetPackageDataField["C"], OverrideAssetPackageData->ImportedClasses))
+		if (!OverrideAssetPackageData->TryNetworkRead(OverrideAssetPackageDataField))
 		{
 			return false;
 		}
diff --git a/Engine/Source/Editor/UnrealEd/Private/CookOnTheFlyServer.cpp b/Engine/Source/Editor/UnrealEd/Private/CookOnTheFlyServer.cpp
--- a/Engine/Source/Editor/UnrealEd/Private/CookOnTheFlyServer.cpp
+++ b/Engine/Source/Editor/UnrealEd/Private/CookOnTheFlyServer.cpp
@@ -3202,7 +3202,8 @@ EDataValidationResult UCookOnTheFlyServer::ValidateSourcePackage(UE::Cook::FPack
 }
 
 void UCookOnTheFlyServer::QueueDiscoveredPackage(UE::Cook::FPackageData& PackageData,
-	UE::Cook::FInstigator&& Instigator, UE::Cook::FDiscoveredPlatformSet&& ReachablePlatforms, bool bUrgent)
+	UE::Cook::FInstigator&& Instigator, UE::Cook::FDiscoveredPlatformSet&& ReachablePlatforms, bool bUrgent,
+	UE::Cook::FGenerationHelper* ParentGenerationHelper)
 {
 	using namespace UE::Cook;
 
@@ -3230,7 +3231,8 @@ void UCookOnTheFlyServer::QueueDiscoveredPackage(UE::Cook::FPackageData& Package
 	{
 		OnDiscoveredPackageDebug(PackageData.GetPackageName(), Instigator);
 	}
-	WorkerRequests->QueueDiscoveredPackage(*this, PackageData, MoveTemp(Instigator), MoveTemp(ReachablePlatforms), bUrgent);
+	WorkerRequests->QueueDiscoveredPackage(*this, PackageData, MoveTemp(Instigator), MoveTemp(ReachablePlatforms),
+		bUrgent, ParentGenerationHelper);
 }
 
 void UCookOnTheFlyServer::QueueDiscoveredPackageOnDirector(UE::Cook::FPackageData& PackageData,
@@ -3338,7 +3340,7 @@ UE::Cook::EPollStatus UCookOnTheFlyServer::QueueGeneratedPackages(UE::Cook::FGen
 
 			// Queue the package for cooking
 			QueueDiscoveredPackage(*ChildPackageData, FInstigator(ChildPackageData->GetInstigator()),
-				EDiscoveredPlatformSet::CopyFromInstigator, bUrgent);
+				EDiscoveredPlatformSet::CopyFromInstigator, bUrgent, &GenerationHelper);
 		}
 		GenerationHelper.EndQueueGeneratedPackages(*this);
 
@@ -4428,7 +4430,9 @@ private: // Used only by UCookOnTheFlyServer, which has private access
 	void CalculatePlatformAgnosticRuntimeDependencies();
 	void CalculatePlatformRuntimeDependencies();
 	TArray<FName> GetPlatformRuntimeDependencies() const;
-	TArray<IPackageWriter::FCommitAttachmentInfo> GetCommitAttachments();
+	TArray<IPackageWriter::FCommitAttachmentInfo> GetCommitAttachments(
+		UE::TargetDomain::FGeneratedPackageResultStruct* GeneratedResult,
+		TArray<FAssetDependency>* OverrideDependencies);
 	IPackageWriter::EWriteOptions GetCommitWriteOptions() const;
 	static void AddDependency(TMap<FPackageData*, EInstigator>& InDependencies, FPackageData* PackageData, bool bHard);
 	static IPackageWriter::ECommitStatus PackageResultToCommitStatus(FSavePackageResultStruct& Result);
@@ -5967,7 +5971,8 @@ public:
 		IterativeValidatePhase2,
 	};
 
-	void InitializePackageWriter(ICookedPackageWriter*& CookedPackageWriter, const FString& ResolvedMetadataPath)
+	void InitializePackageWriter(UCookOnTheFlyServer& COTFS, ICookedPackageWriter*& CookedPackageWriter,
+		const FString& ResolvedMetadataPath)
 	{
 		Initialize();
 		if (DiffMode == EDiffMode::None)
@@ -5993,17 +5998,17 @@ public:
 			CookedPackageWriter = new FLinkerDiffPackageWriter(TUniquePtr<ICookedPackageWriter>(CookedPackageWriter));
 			break;
 		case EDiffMode::IterativeValidate:
-			CookedPackageWriter = new FIterativeValidatePackageWriter(
+			CookedPackageWriter = new FIterativeValidatePackageWriter(COTFS,
 				TUniquePtr<ICookedPackageWriter>(CookedPackageWriter), FIterativeValidatePackageWriter::EPhase::AllInOnePhase,
 				ResolvedMetadataPath);
 			break;
 		case EDiffMode::IterativeValidatePhase1:
-			CookedPackageWriter = new FIterativeValidatePackageWriter(
+			CookedPackageWriter = new FIterativeValidatePackageWriter(COTFS,
 				TUniquePtr<ICookedPackageWriter>(CookedPackageWriter), FIterativeValidatePackageWriter::EPhase::Phase1,
 				ResolvedMetadataPath);
 			break;
 		case EDiffMode::IterativeValidatePhase2:
-			CookedPackageWriter = new FIterativeValidatePackageWriter(
+			CookedPackageWriter = new FIterativeValidatePackageWriter(COTFS,
 				TUniquePtr<ICookedPackageWriter>(CookedPackageWriter), FIterativeValidatePackageWriter::EPhase::Phase2,
 				ResolvedMetadataPath);
 			break;
@@ -6376,6 +6381,7 @@ void FSaveCookedPackageContext::FinishPlatform()
 
 	// Calculate up-to-date assetregistry data for Generator and Generated packages
 	TOptional<TArray<FAssetDependency>> OverridePackageDependencies;
+	TOptional<UE::TargetDomain::FGeneratedPackageResultStruct> GeneratedResult;
 	TOptional<FAssetPackageData> AssetPackageDataBuffer;
 	TOptional<FAssetPackageData> OverrideAssetPackageData;
 	const FAssetPackageData* AssetPackageData = nullptr;
@@ -6388,10 +6394,11 @@ void FSaveCookedPackageContext::FinishPlatform()
 	else if (TRefCountPtr<FGenerationHelper> ParentGenerationHelper = PackageData.GetParentGenerationHelper();
 		ParentGenerationHelper)
 	{
-		OverridePackageDependencies.Emplace();
-		OverrideAssetPackageData.Emplace();
-		ParentGenerationHelper->FinishGeneratedPlatformSave(PackageData, *OverridePackageDependencies, *OverrideAssetPackageData);
-		AssetPackageData = OverrideAssetPackageData.GetPtrOrNull();
+		GeneratedResult.Emplace();
+		ParentGenerationHelper->FinishGeneratedPlatformSave(PackageData, *GeneratedResult);
+		OverridePackageDependencies.Emplace(GeneratedResult->PackageDependencies);
+		OverrideAssetPackageData.Emplace(GeneratedResult->AssetPackageData);
+		AssetPackageData = &GeneratedResult->AssetPackageData;
 	}
 	if (!AssetPackageData)
 	{
@@ -6407,7 +6414,7 @@ void FSaveCookedPackageContext::FinishPlatform()
 
 		ICookedPackageWriter::FCommitPackageInfo Info;
 		// Note GetCommitAttachments mutates the SaveResult; it moves CookDependencies out of it
-		Info.Attachments = GetCommitAttachments();
+		Info.Attachments = GetCommitAttachments(GeneratedResult.GetPtrOrNull(), OverridePackageDependencies.GetPtrOrNull());
 		Info.Status = PackageResultToCommitStatus(SavePackageResult);
 		Info.PackageName = Package->GetFName();
 		Info.PackageHash = AssetPackageData ? AssetPackageData->GetPackageSavedHash() : FIoHash();
@@ -6662,14 +6669,25 @@ TArray<FName> FSaveCookedPackageContext::GetPlatformRuntimeDependencies() const
 	return PlatformDependencyNames;
 }
 
-TArray<IPackageWriter::FCommitAttachmentInfo> FSaveCookedPackageContext::GetCommitAttachments()
+TArray<IPackageWriter::FCommitAttachmentInfo> FSaveCookedPackageContext::GetCommitAttachments(
+	UE::TargetDomain::FGeneratedPackageResultStruct* GeneratedResult,
+	TArray<FAssetDependency>* OverrideDependencies)
 {
 	TArray<IPackageWriter::FCommitAttachmentInfo> Result;
 	if (COTFS.bHybridIterativeEnabled)
 	{
 		UE_SCOPED_HIERARCHICAL_COOKTIMER(TargetDomainDependencies);
+		TArray<FName> PlatformRuntimeDependencies = GetPlatformRuntimeDependencies();
+		if (OverrideDependencies)
+		{
+			PlatformRuntimeDependencies.Reserve(PlatformRuntimeDependencies.Num() + OverrideDependencies->Num());
+			for (const FAssetDependency& Dependency : *OverrideDependencies)
+			{
+				PlatformRuntimeDependencies.Add(Dependency.AssetId.PackageName);
+			}
+		}
 		UE::TargetDomain::CollectAndStoreCookAttachments(Package, TargetPlatform, &SavePackageResult,
-			GetPlatformRuntimeDependencies(), Result);
+			GeneratedResult, MoveTemp(PlatformRuntimeDependencies), Result);
 	}
 	return Result;
 }
@@ -8483,11 +8501,15 @@ void UCookOnTheFlyServer::PopulateCookedPackages(TArrayView<const ITargetPlatfor
 			int32 RemovedCookedNum = 0;
 
 			auto AddPlaceholderPackage =
-				[this, TargetPlatform](const FName PackageName, ECookResult CookResult)
+				[this, TargetPlatform](const FName PackageName, ECookResult CookResult, bool bIterativelyUnmodified)
 				{
 					FPackageData* PackageData = PackageDatas->TryAddPackageDataByPackageName(PackageName, true /* bRequireExists */);
 					if (PackageData)
 					{
+						if (bIterativelyUnmodified)
+						{
+							PackageData->FindOrAddPlatformData(TargetPlatform).SetIterativelyUnmodified(true);
+						}
 						PackageData->SetPlatformCooked(TargetPlatform, CookResult);
 					}
 				};
@@ -8559,7 +8581,7 @@ void UCookOnTheFlyServer::PopulateCookedPackages(TArrayView<const ITargetPlatfor
 					++RemovedCookedNum;
 					break;
 				case EDifference::IdenticalUncooked:
-					AddPlaceholderPackage(PackageName, ECookResult::Failed);
+					AddPlaceholderPackage(PackageName, ECookResult::Failed, true /* bIterativelyUnmodified */);
 					PackagesToRemove.Add(PackageName);
 					break;
 				case EDifference::ModifiedUncooked:
@@ -8569,7 +8591,7 @@ void UCookOnTheFlyServer::PopulateCookedPackages(TArrayView<const ITargetPlatfor
 					PackagesToRemove.Add(PackageName);
 					break;
 				case EDifference::IdenticalNeverCookPlaceholder:
-					AddPlaceholderPackage(PackageName, ECookResult::NeverCookPlaceholder);
+					AddPlaceholderPackage(PackageName, ECookResult::NeverCookPlaceholder, true /* bIterativelyUnmodified */);
 					PackagesToRemove.Add(PackageName);
 					break;
 				case EDifference::ModifiedNeverCookPlaceholder:
@@ -8593,7 +8615,7 @@ void UCookOnTheFlyServer::PopulateCookedPackages(TArrayView<const ITargetPlatfor
 				FPackageData* PackageData = PackageDatas->TryAddPackageDataByPackageName(Generator, false /* bRequireExists */);
 				if (PackageData && PackageData->FindOrAddPlatformData(TargetPlatform).IsCookAttempted())
 				{
-					for (const TPair<FName, FIoHash>& GeneratedPair : Iter->Value.Generated)
+					for (const TPair<FName, FAssetPackageData>& GeneratedPair : Iter->Value.Generated)
 					{
 						UpdateCookedPackage(GeneratedPair.Key, false /* bRequireExists */, true /* bIterativelyUnmodified */);
 					}
@@ -11201,7 +11223,7 @@ UE::Cook::FCookSavePackageContext* UCookOnTheFlyServer::CreateSaveContext(const
 		WriterDebugName = TEXT("LooseCookedPackageWriter");
 	}
 
-	DiffModeHelper->InitializePackageWriter(PackageWriter, ResolvedMetadataPath);
+	DiffModeHelper->InitializePackageWriter(*this, PackageWriter, ResolvedMetadataPath);
 
 	// Setup save package settings (i.e. validation)
 	FSavePackageSettings SavePackageSettings = FSavePackageSettings::GetDefaultSettings();
diff --git a/Engine/Source/Editor/UnrealEd/Private/Cooker/CookGenerationHelper.cpp b/Engine/Source/Editor/UnrealEd/Private/Cooker/CookGenerationHelper.cpp
--- a/Engine/Source/Editor/UnrealEd/Private/Cooker/CookGenerationHelper.cpp
+++ b/Engine/Source/Editor/UnrealEd/Private/Cooker/CookGenerationHelper.cpp
@@ -13,6 +13,7 @@
 #include "Misc/ConfigCacheIni.h"
 #include "Misc/PackageAccessTrackingOps.h"
 #include "Misc/Parse.h"
+#include "TargetDomain/TargetDomainUtils.h"
 #include "UObject/ReferenceChainSearch.h"
 
 namespace UE::Cook
@@ -159,7 +160,14 @@ void FGenerationHelper::ModifyNumSaved(int32 Delta)
 	if (NumSaved == NumAllSaved)
 	{
 		UCookOnTheFlyServer& COTFS = GetOwner().GetPackageDatas().GetCookOnTheFlyServer();
-		if (!COTFS.CookWorkerClient)
+		// Only send OnAllSavesCompleted from director; clients have incomplete information and
+		// could send it spuriously.
+		// Additionally, only send it if we have completed queueing, to avoid sending it prematurely.
+		// ModifyNumSaved(1 == 1) will occur when the generator package is iteratively skipped,
+		// and ModifyNumSaved(k == k, k < expectednumber) can occur if we save some generated packages
+		// (or mark them iteratively skippable) before getting the full list of packages from the worker
+		// that called QueueGeneratedPackages.
+		if (!COTFS.CookWorkerClient && bHasFinishedQueueGeneratedPackages)
 		{
 			if (COTFS.CookDirector)
 			{
@@ -660,6 +668,14 @@ bool FGenerationHelper::TryGenerateList()
 					(int32)Iter->Category, *this->GetSplitDataObjectName().ToString(), *PackageName);
 				Iter.RemoveCurrent();
 			}
+			TStringBuilder<256> DependencyPackageName(InPlace, Iter->AssetId.PackageName);
+			if (ICookPackageSplitter::IsUnderGeneratedPackageSubPath(DependencyPackageName))
+			{
+				UE_LOG(LogCook, Error,
+					TEXT("PackageSplitter specified a dependency for one generated package on another generated package. Only dependencies on non-generated packages are allowed. Dependency will be ignored. Splitter=%s, Generated=%s, Dependency=%s."),
+					*this->GetSplitDataObjectName().ToString(), *PackageName, *DependencyPackageName);
+				Iter.RemoveCurrent();
+			}
 		}
 		Algo::Sort(GeneratedInfo->PackageDependencies,
 			[](const FAssetDependency& A, const FAssetDependency& B) { return A.LexicalLess(B); });
@@ -848,7 +864,12 @@ void FGenerationHelper::StartQueueGeneratedPackages(UCookOnTheFlyServer& COTFS)
 	bool bHybridIterativeEnabled = COTFS.bHybridIterativeEnabled;
 	if (!PreviousGeneratedPackages.IsEmpty())
 	{
-		TMap<FName, FIoHash> RemainingPreviousPackages = PreviousGeneratedPackages;
+		TSet<FName> RemainingPreviousPackages;
+		RemainingPreviousPackages.Reserve(PreviousGeneratedPackages.Num());
+		for (const TPair<FName, FAssetPackageData>& Pair : PreviousGeneratedPackages)
+		{
+			RemainingPreviousPackages.Add(Pair.Key);
+		}
 
 		FPackageData& OwnerPackageData = GetOwner();
 		TArray<const ITargetPlatform*, TInlineAllocator<1>> PlatformsToCook;
@@ -861,23 +882,34 @@ void FGenerationHelper::StartQueueGeneratedPackages(UCookOnTheFlyServer& COTFS)
 
 		for (FCookGenerationInfo& GeneratedInfo : PackagesToGenerate)
 		{
-			FIoHash PreviousHash;
-			if (RemainingPreviousPackages.RemoveAndCopyValue(GeneratedInfo.PackageData->GetPackageName(), PreviousHash)
-				&& !bHybridIterativeEnabled)
+			FName GeneratedPackageName = GeneratedInfo.PackageData->GetPackageName();
+			FAssetPackageData* PreviousAssetData = PreviousGeneratedPackages.Find(GeneratedPackageName);
+			RemainingPreviousPackages.Remove(GeneratedPackageName);
+			if (PreviousAssetData)
 			{
-				bool bIterativelyUnmodified;
-				GeneratedInfo.IterativeCookValidateOrClear(*this, PlatformsToCook, PreviousHash, bIterativelyUnmodified);
-				++(bIterativelyUnmodified ? NumIterativeUnmodified : NumIterativeModified);
+				if (!bHybridIterativeEnabled)
+				{
+					bool bIterativelyUnmodified;
+					GeneratedInfo.IterativeCookValidateOrClear(*this, PlatformsToCook, PreviousAssetData->GetPackageSavedHash(),
+						bIterativelyUnmodified);
+					++(bIterativelyUnmodified ? NumIterativeUnmodified : NumIterativeModified);
+				}
+				else
+				{
+					// Copy the current value for the package's hash into the PreviousPackageData, for use by incremental cook's
+					// calculation in FRequestCluster::TryCalculateIterativelyUnmodified
+					PreviousAssetData->SetPackageSavedHash(GeneratedInfo.PackageHash);
+				}
 			}
 		}
 		if (!RemainingPreviousPackages.IsEmpty())
 		{
 			NumIterativeRemoved = RemainingPreviousPackages.Num();
-			for (TPair<FName, FIoHash>& Pair : RemainingPreviousPackages)
+			for (FName PreviousPackageName : RemainingPreviousPackages)
 			{
 				for (const ITargetPlatform* TargetPlatform : PlatformsToCook)
 				{
-					COTFS.DeleteOutputForPackage(Pair.Key, TargetPlatform);
+					COTFS.DeleteOutputForPackage(PreviousPackageName, TargetPlatform);
 				}
 			}
 		}
@@ -900,7 +932,6 @@ void FGenerationHelper::NotifyStartQueueGeneratedPackages(UCookOnTheFlyServer& C
 	// and lose the information from SavedOnWorker or TryGenerateList.
 	if (!COTFS.CookWorkerClient)
 	{
-		COTFS.PackageDatas->GetRequestQueue().AddRequestFenceListener(GetOwner().GetPackageName());
 		GetOwnerInfo().SavedOnWorker = SourceWorkerId;
 		SetKeepForCompletedAllSavesMessage();
 	}
@@ -909,6 +940,7 @@ void FGenerationHelper::NotifyStartQueueGeneratedPackages(UCookOnTheFlyServer& C
 
 void FGenerationHelper::EndQueueGeneratedPackages(UCookOnTheFlyServer& COTFS)
 {
+	bHasFinishedQueueGeneratedPackages = true;
 	SetKeepForQueueResults();
 	COTFS.WorkerRequests->EndQueueGeneratedPackages(COTFS, *this);
 }
@@ -917,6 +949,14 @@ void FGenerationHelper::EndQueueGeneratedPackagesOnDirector(UCookOnTheFlyServer&
 {
 	// Note this function can be called on an uninitialized Generator; the generator is only needed
 	// on the director so it can serve as the passer of messages.
+	bHasFinishedQueueGeneratedPackages = true;
+	// When we queued locally, this function is called after QueueDiscoveredPackage was called for each package.
+	// When we queued on a remote CookWorker, the replication system from cookworker guarantees that all discovered
+	// packages have been reported via TrackGeneratedPackageListedRemotely before we receive this function call
+	// via the EGeneratorEvent::QueuedGeneratedPackages message (the package discovery messages are replicated
+	// before the EGeneratorEvent). We therefore know that all generated packages have already been requested
+	// or are in the discovery queue, so we can add a request fence listener now and know that when it is called
+	// all generated packages have been queued and assigned.
 	COTFS.PackageDatas->GetRequestQueue().AddRequestFenceListener(GetOwner().GetPackageName());
 	SetKeepForQueueResults();
 
@@ -934,23 +974,46 @@ void FGenerationHelper::EndQueueGeneratedPackagesOnDirector(UCookOnTheFlyServer&
 	}
 }
 
-void FGenerationHelper::OnRequestFencePassedBroadcast(UCookOnTheFlyServer& COTFS)
+void FGenerationHelper::OnRequestFencePassed(UCookOnTheFlyServer& COTFS)
 {
-	if (COTFS.CookDirector)
+	// This function should only be called in response to a subscription that is sent from the cook director
+	check(!COTFS.CookWorkerClient);
+
+	if (OwnerInfo.IsIterativelySkipped())
+	{
+		// PumpRequests has completed and we marked ourselves and all generated packages as iteratively skipped,
+		// so we no longer need the PreviouslyCookedData or this entire GenerationHelper
+		ClearKeepForIterative();
+		PreviousGeneratedPackages.Empty();
+	}
+
+	if (bHasFinishedQueueGeneratedPackages)
 	{
-		FName PackageName = GetOwner().GetPackageName();
-		FGeneratorEventMessage Message(EGeneratorEvent::QueuedGeneratedPackagesFencePassed, PackageName);
-		COTFS.CookDirector->BroadcastGeneratorMessage(MoveTemp(Message));
+		// We have finished EndQueueGeneratedPackagesOnDirector, so all generated packages have been requested
+		// and assigned to local ReadyRequests or to a CookWorker. Send OnQueuedGeneratedPackagesFencePassed
+		// to ourselves and all cookworkers.
+
+		// Call ModifyNumSaved to check for whether all packages have already been saved by the time we reach
+		// the request fence. This can happen in iterative cooks, or in race conditions if we sent all packages
+		// out for saving before receiving the EndQueueGeneratedPackagesOnDirector message.
+		ModifyNumSaved(0);
+
+		if (COTFS.CookDirector)
+		{
+			FName PackageName = GetOwner().GetPackageName();
+			FGeneratorEventMessage Message(EGeneratorEvent::QueuedGeneratedPackagesFencePassed, PackageName);
+			COTFS.CookDirector->BroadcastGeneratorMessage(MoveTemp(Message));
+		}
+		OnQueuedGeneratedPackagesFencePassed(COTFS);
 	}
-	OnRequestFencePassed(COTFS);
 }
 
-void FGenerationHelper::OnRequestFencePassed(UCookOnTheFlyServer& COTFS)
+void FGenerationHelper::OnQueuedGeneratedPackagesFencePassed(UCookOnTheFlyServer& COTFS)
 {
 	ClearKeepForQueueResults();
-	// We no longer need PreviousGeneratedPackages or KeepForIterative, because they are used (and cleared) in
-	// StartQueueGeneratedPackages. Clear them as well on the director and any CookWorkers that received it to
-	// free memory.
+	// We no longer need PreviousGeneratedPackages or KeepForIterative, because they are used only in 
+	// StartQueueGeneratedPackages or the request cluster that they end up in in PumpRequests, both of which
+	// are now finished. Clear them on the director and any CookWorkers that received them to free memory.
 	ClearKeepForIterative();
 	PreviousGeneratedPackages.Empty();
 }
@@ -1040,8 +1103,7 @@ void FGenerationHelper::FinishGeneratorPlatformSave(FPackageData& PackageData, b
 }
 
 void FGenerationHelper::FinishGeneratedPlatformSave(FPackageData& PackageData,
-	TArray<FAssetDependency>& OutPackageDependencies,
-	FAssetPackageData& OutAssetPackageData)
+	UE::TargetDomain::FGeneratedPackageResultStruct& OutGeneratedResult)
 {
 	ConditionalInitialize();
 
@@ -1056,7 +1118,7 @@ void FGenerationHelper::FinishGeneratedPlatformSave(FPackageData& PackageData,
 
 	// There should be no package dependencies present for the package from the global assetregistry
 	// because it is newly created. Add on the dependencies declared for it from the CookPackageSplitter.
-	OutPackageDependencies = Info->PackageDependencies;
+	OutGeneratedResult.PackageDependencies = Info->PackageDependencies;
 
 	// Update the AssetPackageData for each requested platform with Guid and ImportedClasses
 	TSet<UClass*> PackageClasses;
@@ -1081,8 +1143,23 @@ void FGenerationHelper::FinishGeneratedPlatformSave(FPackageData& PackageData,
 	}
 	ImportedClasses.Sort(FNameLexicalLess());
 
-	OutAssetPackageData.SetPackageSavedHash(Info->PackageHash);
-	OutAssetPackageData.ImportedClasses = MoveTemp(ImportedClasses);
+	OutGeneratedResult.AssetPackageData.FileVersionUE = GPackageFileUEVersion;
+	OutGeneratedResult.AssetPackageData.FileVersionLicenseeUE = GPackageFileLicenseeUEVersion;
+	OutGeneratedResult.AssetPackageData.SetIsLicenseeVersion(FEngineVersion::Current().IsLicenseeVersion());
+	OutGeneratedResult.AssetPackageData.Extension = FPackagePath::ParseExtension(
+		WriteToString<256>(PackageData.GetFileName()));
+	OutGeneratedResult.AssetPackageData.SetPackageSavedHash(Info->PackageHash);
+	OutGeneratedResult.AssetPackageData.ImportedClasses = MoveTemp(ImportedClasses);
+}
+
+const FAssetPackageData* FGenerationHelper::GetIncrementalCookAssetPackageData(FPackageData& PackageData)
+{
+	return PreviousGeneratedPackages.Find(PackageData.GetPackageName());
+}
+
+const FAssetPackageData* FGenerationHelper::GetIncrementalCookAssetPackageData(FName PackageName)
+{
+	return PreviousGeneratedPackages.Find(PackageName);
 }
 
 void FGenerationHelper::ResetSaveState(FCookGenerationInfo& Info, UPackage* Package,
@@ -1233,7 +1310,7 @@ void FGenerationHelper::FetchExternalActorDependencies()
 	ExternalActorDependencies.Shrink();
 }
 
-void FGenerationHelper::SetPreviousGeneratedPackages(TMap<FName, FIoHash>&& Packages)
+void FGenerationHelper::SetPreviousGeneratedPackages(TMap<FName, FAssetPackageData>&& Packages)
 {
 	SetKeepForIterative();
 	PreviousGeneratedPackages = MoveTemp(Packages);
@@ -1394,7 +1471,8 @@ void FGenerationHelper::PostGarbageCollectGCLifetimeData()
 	}
 }
 
-void FGenerationHelper::TrackGeneratedPackageListedRemotely(UCookOnTheFlyServer& COTFS, FPackageData& PackageData)
+void FGenerationHelper::TrackGeneratedPackageListedRemotely(UCookOnTheFlyServer& COTFS, FPackageData& PackageData,
+	const FIoHash& CurrentPackageHash)
 {
 	if (bGeneratedList)
 	{
@@ -1414,6 +1492,13 @@ void FGenerationHelper::TrackGeneratedPackageListedRemotely(UCookOnTheFlyServer&
 			PackagesToGenerate.Emplace(PackageData, bGenerator);
 		}
 	}
+	FAssetPackageData* PreviousAssetData = PreviousGeneratedPackages.Find(PackageData.GetPackageName());
+	if (PreviousAssetData)
+	{
+		// Copy the current value for the package's hash into the PreviousPackageData, for use by incremental cook's
+		// calculation in FRequestCluster::TryCalculateIterativelyUnmodified
+		PreviousAssetData->SetPackageSavedHash(CurrentPackageHash);
+	}
 }
 
 void FGenerationHelper::MarkPackageSavedRemotely(UCookOnTheFlyServer& COTFS, FPackageData& PackageData,
@@ -1426,6 +1511,24 @@ void FGenerationHelper::MarkPackageSavedRemotely(UCookOnTheFlyServer& COTFS, FPa
 	}
 }
 
+void FGenerationHelper::MarkPackageIterativelySkipped(FPackageData& PackageData)
+{
+	FCookGenerationInfo* Info = FindInfoNoInitialize(PackageData);
+	if (Info)
+	{
+		Info->SetHasSaved(*this, true, FWorkerId::Local());
+		Info->SetIterativelySkipped(true);
+	}
+
+	if (&PackageData == &GetOwner())
+	{
+		// The entire generator package has been skipped. Wait for the current cluster to complete
+		// so we can mark all of our generated packages as skipped, but then clear the iterative data;
+		// it will no longer be needed.
+		GetOwner().GetPackageDatas().GetRequestQueue().AddRequestFenceListener(GetOwner().GetPackageName());
+	}
+}
+
 void FGenerationHelper::PostGarbageCollect(const TRefCountPtr<FGenerationHelper>& RefcountHeldByCaller)
 {
 	PostGarbageCollectGCLifetimeData();
@@ -1639,6 +1742,7 @@ FCookGenerationInfo::FCookGenerationInfo(FPackageData& InPackageData, bool bInGe
 	, GeneratorSaveState(bInGenerator ? ESaveState::StartSave : ESaveState::StartPopulate)
 	, bCreateAsMap(false), bHasCreatedPackage(false), bHasSaved(false), bTakenOverCachedCookedPlatformData(false)
 	, bIssuedUndeclaredMovedObjectsWarning(false), bGenerator(bInGenerator), bHasCalledPopulate(false)
+	, bIterativelySkipped(false)
 {
 }
 
@@ -1666,6 +1770,7 @@ void FCookGenerationInfo::Uninitialize()
 	bIssuedUndeclaredMovedObjectsWarning = false;
 	// Keep bGenerator; it is allowed in the uninitialized state
 	bHasCalledPopulate = false;
+	// Keep bIterativelySkipped; it is allowed in the uninitialized state
 }
 
 void FCookGenerationInfo::SetSaveStateComplete(ESaveState CompletedState)
diff --git a/Engine/Source/Editor/UnrealEd/Private/Cooker/CookGenerationHelper.h b/Engine/Source/Editor/UnrealEd/Private/Cooker/CookGenerationHelper.h
--- a/Engine/Source/Editor/UnrealEd/Private/Cooker/CookGenerationHelper.h
+++ b/Engine/Source/Editor/UnrealEd/Private/Cooker/CookGenerationHelper.h
@@ -25,6 +25,7 @@
 
 class ITargetPlatform;
 class UCookOnTheFlyServer;
+namespace UE::TargetDomain { struct FGeneratedPackageResultStruct; }
 
 namespace UE::Cook
 {
@@ -100,6 +101,8 @@ public:
 	void SetIsGenerator(bool bValue);
 	bool HasCalledPopulate() const;
 	void SetHasCalledPopulate(bool bValue);
+	bool IsIterativelySkipped() const;
+	void SetIterativelySkipped(bool bValue);
 
 	/**
 	 * Steal the list of cached objects to call BeginCacheForCookedPlatformData on from the PackageData,
@@ -163,6 +166,7 @@ private:
 	bool bIssuedUndeclaredMovedObjectsWarning : 1;
 	bool bGenerator : 1;
 	bool bHasCalledPopulate : 1;
+	bool bIterativelySkipped : 1;
 };
 
 /**
@@ -298,12 +302,15 @@ public:
 	/** Does not call Initialize. */
 	void EndQueueGeneratedPackagesOnDirector(UCookOnTheFlyServer& COTFS, FWorkerId SourceWorkerId);
 	/**
-	 * Called from Director when the RequestFence added from EndQueueGeneratedPackages has passed. Calls
-	 * OnRequestFencePassed locally and on all Workers.
+	 * Called on Director when the RequestFence added from MarkForIterative or EndQueueGeneratedPackages has passed.
+	 * Calls the proper followup events locally and on all Workers. Does not call Initialize.
 	 */
-	void OnRequestFencePassedBroadcast(UCookOnTheFlyServer& COTFS);
-	/** Called from when the RequestFence added from EndQueueGeneratedPackages has passed. */
 	void OnRequestFencePassed(UCookOnTheFlyServer& COTFS);
+	/**
+	 * Called on Directors and Workers when the RequestFence added from EndQueueGeneratedPackages has passed.
+	 * Does not call Initialize.
+	 */
+	void OnQueuedGeneratedPackagesFencePassed(UCookOnTheFlyServer& COTFS);
 	/** Call CreatePackage and set package header data; or empty it and normalize it if it already exists. */
 	UPackage* TryCreateGeneratedPackage(FCookGenerationInfo& GenerationInfo, bool bResetToEmpty);
 	/**
@@ -317,8 +324,14 @@ public:
 	 * Calculates AssetRegistryData.
 	 */
 	void FinishGeneratedPlatformSave(FPackageData& PackageData,
-		TArray<FAssetDependency>& OutPackageDependencies,
-		FAssetPackageData& OutAssetPackageData);
+		UE::TargetDomain::FGeneratedPackageResultStruct& GeneratedResult);
+	/**
+	 * Return the AssetPackageData that was loaded for the given PackageData from the previous cook's
+	 * AssetRegistry, or nullptr if not an incremental cook or it was not previously cooked.
+	 * Does not call initialize.
+	 */
+	const FAssetPackageData* GetIncrementalCookAssetPackageData(FPackageData& PackageData);
+	const FAssetPackageData* GetIncrementalCookAssetPackageData(FName PackageName);
 
 	/**
 	 * Clear any data that should only be held when an FPackageData is in the save state, for the given Info. The given
@@ -333,9 +346,9 @@ public:
 	 */
 	void FetchExternalActorDependencies();
 	/** Iterative cook: store the list of generated packages from the last cook. Does not call Initialize. */
-	void SetPreviousGeneratedPackages(TMap<FName, FIoHash>&& Packages);
+	void SetPreviousGeneratedPackages(TMap<FName, FAssetPackageData>&& Packages);
 	/** Return the information set by SetPreviousGeneratedPackages if not yet cleared. Does not call Initialize. */
-	const TMap<FName, FIoHash>& GetPreviousGeneratedPackages() const;
+	const TMap<FName, FAssetPackageData>& GetPreviousGeneratedPackages() const;
 
 	/**
 	 * Callback during garbage collection. Does not call initialize. Caller must pass in a refcount to show
@@ -384,13 +397,18 @@ public:
 	 * Called for each of the generated packages that were discovered when TryGenerateList was called on a
 	 * remote CookWorker. Does not call Initialize.
 	 */
-	void TrackGeneratedPackageListedRemotely(UCookOnTheFlyServer& COTFS, FPackageData& PackageData);
+	void TrackGeneratedPackageListedRemotely(UCookOnTheFlyServer& COTFS, FPackageData& PackageData,
+		const FIoHash& CurrentPackageHash);
 	/**
 	 * Called on the director when the generator or one of the generated packages was saved on a remote worker.
 	 * Used to manage KeepForGCOrAllSaved lifetime. Does not call Initialize.
 	 */
 	void MarkPackageSavedRemotely(UCookOnTheFlyServer& COTFS, FPackageData& PackageData, FWorkerId SourceWorkerId);
-
+	/**
+	 * Called on the director when the generator or one of the generated packages was found to be iteratively skipped
+	 * and marked already cooked in an incremental cook. Does not call Initialize.
+	 */
+	void MarkPackageIterativelySkipped(FPackageData& PackageData);
 	/**
 	 * Called on the director and every CookWorker when all saves have been completed; this indicates that
 	 * some of our contract points are complete and the splitter can be destroyed. Does not call Initialize.
@@ -451,7 +469,7 @@ private:
 	/** Recorded list of packages to generate from the splitter, and data we need about them */
 	TArray<FCookGenerationInfo> PackagesToGenerate;
 	TWeakObjectPtr<UPackage> OwnerPackage;
-	TMap<FName, FIoHash> PreviousGeneratedPackages;
+	TMap<FName, FAssetPackageData> PreviousGeneratedPackages;
 	TArray<FName> ExternalActorDependencies;
 	TArray<FWeakObjectPtr> OwnerObjectsToMove;
 	TRefCountPtr<FGenerationHelper> ReferenceFromKeepForIterative;
@@ -471,7 +489,8 @@ private:
 	bool bKeepForAllSavedOrGC = false;
 	bool bKeepForCompletedAllSavesMessage = false;
 	bool bNeedConfirmGeneratorPackageDestroyed = false;
-	
+	bool bHasFinishedQueueGeneratedPackages = false;
+
 	friend FCookGenerationInfo;
 };
 
@@ -569,6 +588,16 @@ inline void FCookGenerationInfo::SetHasCalledPopulate(bool bValue)
 	bHasCalledPopulate = bValue;
 }
 
+inline bool FCookGenerationInfo::IsIterativelySkipped() const
+{
+	return bIterativelySkipped;
+}
+
+inline void FCookGenerationInfo::SetIterativelySkipped(bool bValue)
+{
+	bIterativelySkipped = bValue;
+}
+
 inline TConstArrayView<FAssetDependency> FCookGenerationInfo::GetDependencies() const
 {
 	return PackageDependencies;
@@ -673,7 +702,7 @@ inline TArray<FName> FGenerationHelper::ReleaseExternalActorDependencies()
 	return Result;
 }
 
-inline const TMap<FName, FIoHash>& FGenerationHelper::GetPreviousGeneratedPackages() const
+inline const TMap<FName, FAssetPackageData>& FGenerationHelper::GetPreviousGeneratedPackages() const
 {
 	return PreviousGeneratedPackages;
 }
diff --git a/Engine/Source/Editor/UnrealEd/Private/Cooker/CookPackageData.cpp b/Engine/Source/Editor/UnrealEd/Private/Cooker/CookPackageData.cpp
--- a/Engine/Source/Editor/UnrealEd/Private/Cooker/CookPackageData.cpp
+++ b/Engine/Source/Editor/UnrealEd/Private/Cooker/CookPackageData.cpp
@@ -58,7 +58,8 @@ static FAutoConsoleVariableRef CVarPollAsyncPeriod(
 // FPackageData
 FPackagePlatformData::FPackagePlatformData()
 	: bReachable(0), bVisitedByCluster(0), bSaveTimedOut(0), bCookable(1), bExplorable(1), bExplorableOverride(0)
-	, bIterativelyUnmodified(0), bRegisteredForCachedObjectsInOuter(0), CookResults((uint8)ECookResult::NotAttempted)
+	, bIterativelyUnmodified(0), bIterativelySkipped(0), bRegisteredForCachedObjectsInOuter(0),
+	CookResults((uint8)ECookResult::NotAttempted)
 {
 }
 
@@ -2923,7 +2924,7 @@ void FRequestQueue::NotifyRequestFencePassed(FPackageDatas& PackageDatas)
 			TRefCountPtr<FGenerationHelper> GenerationHelper = PackageData->GetGenerationHelper();
 			if (GenerationHelper)
 			{
-				GenerationHelper->OnRequestFencePassedBroadcast(PackageDatas.GetCookOnTheFlyServer());
+				GenerationHelper->OnRequestFencePassed(PackageDatas.GetCookOnTheFlyServer());
 			}
 		}
 	}
diff --git a/Engine/Source/Editor/UnrealEd/Private/Cooker/CookPackageData.h b/Engine/Source/Editor/UnrealEd/Private/Cooker/CookPackageData.h
--- a/Engine/Source/Editor/UnrealEd/Private/Cooker/CookPackageData.h
+++ b/Engine/Source/Editor/UnrealEd/Private/Cooker/CookPackageData.h
@@ -207,6 +207,13 @@ struct FPackagePlatformData
 	bool IsIterativelyUnmodified() const { return bIterativelyUnmodified != 0; }
 	void SetIterativelyUnmodified(bool bValue) { bIterativelyUnmodified = (uint32)bValue; }
 
+	/**
+	 * The package was skipped in the current iterative cook. This might not be equal to IterativelyUnmodified,
+	 * depending on packagewriter).
+	 */
+	bool IsIterativelySkipped() const { return bIterativelySkipped != 0; }
+	void SetIterativelySkipped(bool bValue) { bIterativelySkipped = (uint32)bValue; }
+
 	ECookResult GetCookResults() const { return (ECookResult)CookResults; }
 	bool IsCookAttempted() const { return CookResults != (uint32)ECookResult::NotAttempted; }
 	bool IsCookSucceeded() const { return CookResults == (uint32)ECookResult::Succeeded; }
@@ -236,6 +243,7 @@ private:
 	uint32 bExplorable : 1;
 	uint32 bExplorableOverride : 1;
 	uint32 bIterativelyUnmodified : 1;
+	uint32 bIterativelySkipped : 1;
 	uint32 bRegisteredForCachedObjectsInOuter : 1;
 	uint32 bReportedToDirector : 1;
 	uint32 CookResults : (int)ECookResult::NumBits;
diff --git a/Engine/Source/Editor/UnrealEd/Private/Cooker/CookPackageSplitter.cpp b/Engine/Source/Editor/UnrealEd/Private/Cooker/CookPackageSplitter.cpp
--- a/Engine/Source/Editor/UnrealEd/Private/Cooker/CookPackageSplitter.cpp
+++ b/Engine/Source/Editor/UnrealEd/Private/Cooker/CookPackageSplitter.cpp
@@ -2,15 +2,37 @@
 
 #include "CookPackageSplitter.h"
 
-#include "Misc/Paths.h"
-
 #if WITH_EDITOR
+#include "Misc/Paths.h"
+#include "Misc/PathViews.h"
 
 const TCHAR* ICookPackageSplitter::GetGeneratedPackageSubPath()
 {
 	return TEXT("_Generated_");
 }
 
+bool ICookPackageSplitter::IsUnderGeneratedPackageSubPath(FStringView FileOrLongPackagePath)
+{
+	FStringView GeneratedSubDir(GetGeneratedPackageSubPath());
+	int32 Index = FileOrLongPackagePath.Find(GeneratedSubDir, ESearchCase::IgnoreCase);
+	if (Index < 0)
+	{
+		return false;
+	}
+	if (Index == 0 || !FPathViews::IsSeparator(FileOrLongPackagePath[Index - 1]))
+	{
+		// "_Generated_/..." or ".../Prefix_Generated_/...", not the /_Generated_/ subfolder we're looking for.
+		return false;
+	}
+	if (Index + GeneratedSubDir.Len() < FileOrLongPackagePath.Len()
+		&& !FPathViews::IsSeparator(FileOrLongPackagePath[Index + GeneratedSubDir.Len()]))
+	{
+		// ".../_Generated_Suffix/...", not the /_Generated_/ subfolder we're looking for.
+		return false;
+	}
+	return true;
+}
+
 FString ICookPackageSplitter::ConstructGeneratedPackageName(FName OwnerPackageName, FStringView RelPath,
 	FStringView GeneratedRootOverride)
 {
diff --git a/Engine/Source/Editor/UnrealEd/Private/Cooker/CookRequestCluster.cpp b/Engine/Source/Editor/UnrealEd/Private/Cooker/CookRequestCluster.cpp
--- a/Engine/Source/Editor/UnrealEd/Private/Cooker/CookRequestCluster.cpp
+++ b/Engine/Source/Editor/UnrealEd/Private/Cooker/CookRequestCluster.cpp
@@ -10,6 +10,7 @@
 #include "AssetRegistry/IAssetRegistry.h"
 #include "Async/Async.h"
 #include "Cooker/CookDependency.h"
+#include "Cooker/CookGenerationHelper.h"
 #include "Cooker/CookPackageData.h"
 #include "Cooker/CookPlatformManager.h"
 #include "Cooker/CookProfiling.h"
@@ -953,7 +954,7 @@ void FRequestCluster::FGraphSearch::ResolveTransitiveBuildDependencyCycle()
 			}
 		}
 		// We can also empty the IterativelyModifiedListeners since any remaining listeners must be in
-		// PendingTransitiveBuildDependencyVertices. Empting the list here avoids the expense of kicking
+		// PendingTransitiveBuildDependencyVertices. Emptying the list here avoids the expense of kicking
 		// for a second time each of the listeners.
 		CycleVert->IterativelyModifiedListeners.Empty();
 		KickVertex(CycleVert);
@@ -1265,6 +1266,22 @@ void FRequestCluster::FGraphSearch::FExploreEdgesContext::CalculatePlatformsToPr
 	}
 }
 
+const FAssetPackageData* FRequestCluster::FVertexData::GetGeneratedAssetPackageData()
+{
+	check(PackageData); // Caller should not call if no PackageData
+	FPackageDatas& LocalPackageDatas = PackageData->GetPackageDatas();
+	FPackageData* ParentPackageData = LocalPackageDatas.FindPackageDataByPackageName(PackageData->GetParentGenerator());
+	if (ParentPackageData)
+	{
+		TRefCountPtr<FGenerationHelper> ParentGenerationHelper = ParentPackageData->GetGenerationHelper();
+		if (ParentGenerationHelper)
+		{
+			return ParentGenerationHelper->GetIncrementalCookAssetPackageData(*PackageData);
+		}
+	}
+	return nullptr;
+}
+
 bool FRequestCluster::FGraphSearch::FExploreEdgesContext::TryCalculateIterativelyUnmodified()
 {
 	using namespace UE::TargetDomain;
@@ -1294,20 +1311,50 @@ bool FRequestCluster::FGraphSearch::FExploreEdgesContext::TryCalculateIterativel
 
 		if (!PackagePlatformData.IsCookable())
 		{
-			SetIsIterativelyUnmodified(PlatformIndex, false);
+			SetIsIterativelyUnmodified(PlatformIndex, false, PackagePlatformData);
 			continue;
 		}
 
 		UE::TargetDomain::FCookDependencies& CookDependencies = QueryPlatformData.CookAttachments.Dependencies;
-		if (!CookDependencies.HasKeyMatch())
+		const FAssetPackageData* OverrideAssetPackageData = nullptr;
+		if (PackageData->IsGenerated())
+		{
+			// If a generator is marked iteratively unmodified, then by contract we are not required to test its
+			// generated packages; they are all marked iteratively unmodified as well
+			FPackageData* ParentPackageData = Cluster.PackageDatas.FindPackageDataByPackageName(
+				PackageData->GetParentGenerator());
+			if (ParentPackageData)
+			{
+				const FPackagePlatformData* ParentPlatformData =
+					ParentPackageData->GetPlatformDatas().Find(TargetPlatform);
+				if (ParentPlatformData)
+				{
+					if (ParentPlatformData->IsIterativelyUnmodified())
+					{
+						SetIsIterativelyUnmodified(PlatformIndex, true, PackagePlatformData);
+						continue;
+					}
+				}
+			}
+
+			// If the generator was not marked iteratively unmodified, then we use the data provided by the generator
+			// to decide whether the generated package is iteratively unmodified.
+			OverrideAssetPackageData = Vertex->GetGeneratedAssetPackageData();
+			if (!OverrideAssetPackageData)
+			{
+				SetIsIterativelyUnmodified(PlatformIndex, false, PackagePlatformData);
+				continue;
+			}
+		}
+		if (!CookDependencies.HasKeyMatch(OverrideAssetPackageData))
 		{
-			SetIsIterativelyUnmodified(PlatformIndex, false);
+			SetIsIterativelyUnmodified(PlatformIndex, false, PackagePlatformData);
 			continue;
 		}
 
-		if (!IsIterativeEnabled(PackageName, Cluster.COTFS.bHybridIterativeAllowAllClasses))
+		if (!IsIterativeEnabled(PackageName, Cluster.COTFS.bHybridIterativeAllowAllClasses, OverrideAssetPackageData))
 		{
-			SetIsIterativelyUnmodified(PlatformIndex, false);
+			SetIsIterativelyUnmodified(PlatformIndex, false, PackagePlatformData);
 			continue;
 		}
 
@@ -1352,7 +1399,7 @@ bool FRequestCluster::FGraphSearch::FExploreEdgesContext::TryCalculateIterativel
 
 			if (bAnyTransitiveBuildDependencyIsModified)
 			{
-				SetIsIterativelyUnmodified(PlatformIndex, false);
+				SetIsIterativelyUnmodified(PlatformIndex, false, PackagePlatformData);
 				continue;
 			}
 			if (!UnreadyTransitiveBuildVertices.IsEmpty())
@@ -1378,8 +1425,7 @@ bool FRequestCluster::FGraphSearch::FExploreEdgesContext::TryCalculateIterativel
 			}
 		}
 
-		SetIsIterativelyUnmodified(PlatformIndex, true);
-		PackagePlatformData.SetIterativelyUnmodified(true);
+		SetIsIterativelyUnmodified(PlatformIndex, true, PackagePlatformData);
 	}
 
 	if (!bAllPlatformsAreReady)
@@ -1528,6 +1574,7 @@ void FRequestCluster::FGraphSearch::FExploreEdgesContext::QueueVisitsOfDependenc
 	}
 
 	TArray<FPackageData*>* Edges = nullptr;
+	TRefCountPtr<FGenerationHelper> GenerationHelper = PackageData->GetGenerationHelper();
 	for (TPair<FName, FScratchPlatformDependencyBits>& PlatformDependencyPair : PlatformDependencyMap)
 	{
 		FName DependencyName = PlatformDependencyPair.Key;
@@ -1539,7 +1586,8 @@ void FRequestCluster::FGraphSearch::FExploreEdgesContext::QueueVisitsOfDependenc
 			FCoreRedirectObjectName(NAME_None, NAME_None, DependencyName)).PackageName;
 		DependencyName = Redirected;
 
-		FVertexData& DependencyVertex = GraphSearch.FindOrAddVertex(DependencyName);
+		FVertexData& DependencyVertex = GraphSearch.FindOrAddVertex(DependencyName,
+			GenerationHelper.GetReference());
 		if (!DependencyVertex.PackageData)
 		{
 			continue;
@@ -1683,11 +1731,36 @@ void FRequestCluster::FGraphSearch::FExploreEdgesContext::ProcessPlatformAttachm
 		bool bShouldIterativelySkip = bIterativelyUnmodified;
 		PackageWriter->UpdatePackageModificationStatus(PackageName, bIterativelyUnmodified,
 			bShouldIterativelySkip);
+
+		TRefCountPtr<FGenerationHelper> ParentGenerationHelper;
+		if (PackageData->IsGenerated())
+		{
+			// If a GeneratorPackage is iteratively skipped, its generated packages must be iteratively skipped as well
+			FPackageData* ParentPackage = Cluster.PackageDatas.FindPackageDataByPackageName(PackageData->GetParentGenerator());
+			if (ParentPackage)
+			{
+				ParentGenerationHelper = ParentPackage->GetGenerationHelper();
+				const FPackagePlatformData* ParentPlatformData = ParentPackage->GetPlatformDatas().Find(TargetPlatform);
+				if (ParentPlatformData && ParentPlatformData->IsIterativelySkipped())
+				{
+					bShouldIterativelySkip = true;
+				}
+			}
+		}
 		if (bShouldIterativelySkip)
 		{
 			// Call SetPlatformCooked instead of just PackagePlatformData.SetCookResults because we might also need
 			// to set OnFirstCookedPlatformAdded
 			PackageData->SetPlatformCooked(TargetPlatform, ECookResult::Succeeded);
+			PackagePlatformData.SetIterativelySkipped(true);
+			if (TRefCountPtr<FGenerationHelper> GenerationHelper = PackageData->GetGenerationHelper())
+			{
+				GenerationHelper->MarkPackageIterativelySkipped(*PackageData);
+			}
+			if (ParentGenerationHelper)
+			{
+				ParentGenerationHelper->MarkPackageIterativelySkipped(*PackageData);
+			}
 			Cluster.SetPackageDataWasMarkedCooked(*PackageData, true);
 			if (PlatformIndex == FirstSessionPlatformIndex)
 			{
@@ -1702,18 +1775,36 @@ void FRequestCluster::FGraphSearch::FExploreEdgesContext::ProcessPlatformAttachm
 	{
 		FQueryPlatformData& PlatformAgnosticQueryData = Vertex->PlatformData[PlatformAgnosticPlatformIndex];
 
-		if (PlatformAgnosticQueryData.bSchedulerThreadFetchCompleted &&
-			PlatformAgnosticQueryData.CookAttachments.Dependencies.HasKeyMatch())
+		if (PlatformAgnosticQueryData.bSchedulerThreadFetchCompleted)
 		{
-			Cluster.BuildDefinitions.AddBuildDefinitionList(PackageName, TargetPlatform,
-				PlatformAgnosticQueryData.CookAttachments.BuildDefinitions.Definitions);
+			bool bCanCheckHasKeyMatch = true;
+			const FAssetPackageData* OverrideAssetPackageData = nullptr;
+			if (PackageData->IsGenerated())
+			{
+				OverrideAssetPackageData = Vertex->GetGeneratedAssetPackageData();
+				if (!OverrideAssetPackageData)
+				{
+					bCanCheckHasKeyMatch = false;
+				}
+			}
+			if (bCanCheckHasKeyMatch &&
+				PlatformAgnosticQueryData.CookAttachments.Dependencies.HasKeyMatch(OverrideAssetPackageData))
+			{
+				Cluster.BuildDefinitions.AddBuildDefinitionList(PackageName, TargetPlatform,
+					PlatformAgnosticQueryData.CookAttachments.BuildDefinitions.Definitions);
+			}
 		}
 	}
 }
 
-void FRequestCluster::FGraphSearch::FExploreEdgesContext::SetIsIterativelyUnmodified(int32 PlatformIndex, bool bIterativelyUnmodified)
+void FRequestCluster::FGraphSearch::FExploreEdgesContext::SetIsIterativelyUnmodified(int32 PlatformIndex,
+	bool bIterativelyUnmodified, FPackagePlatformData& PackagePlatformData)
 {
 	Vertex->PlatformData[PlatformIndex].bIterativelyUnmodified.Emplace(bIterativelyUnmodified);
+	if (bIterativelyUnmodified)
+	{
+		PackagePlatformData.SetIterativelyUnmodified(true);
+	}
 }
 
 FRequestCluster::FVertexData* FRequestCluster::FGraphSearch::AllocateVertex(FName PackageName, FPackageData* PackageData)
@@ -1732,7 +1823,7 @@ FRequestCluster::FVertexData::FVertexData(FName InPackageName, UE::Cook::FPackag
 }
 
 FRequestCluster::FVertexData&
-FRequestCluster::FGraphSearch::FindOrAddVertex(FName PackageName)
+FRequestCluster::FGraphSearch::FindOrAddVertex(FName PackageName, FGenerationHelper* ParentGenerationHelper)
 {
 	// Only called from scheduler thread
 	FVertexData*& ExistingVertex = Vertices.FindOrAdd(PackageName);
@@ -1747,6 +1838,21 @@ FRequestCluster::FGraphSearch::FindOrAddVertex(FName PackageName)
 	if (!FPackageName::IsScriptPackage(NameBuffer))
 	{
 		PackageData = Cluster.COTFS.PackageDatas->TryAddPackageDataByPackageName(PackageName);
+		if (!PackageData && ParentGenerationHelper && ICookPackageSplitter::IsUnderGeneratedPackageSubPath(NameBuffer))
+		{
+			const FAssetPackageData* PreviousPackageData =
+				ParentGenerationHelper->GetIncrementalCookAssetPackageData(PackageName);
+			if (PreviousPackageData)
+			{
+				bool bIsMap = PreviousPackageData->Extension == EPackageExtension::Map;
+				PackageData = Cluster.COTFS.PackageDatas->TryAddPackageDataByPackageName(PackageName,
+					false /* bRequireExists */, bIsMap);
+				if (PackageData)
+				{
+					PackageData->SetGenerated(ParentGenerationHelper->GetOwner().GetPackageName());
+				}
+			}
+		}
 	}
 
 	ExistingVertex = AllocateVertex(PackageName, PackageData);
@@ -2138,9 +2244,20 @@ void FRequestCluster::IsRequestCookable(const ITargetPlatform* Platform, FName P
 	if (InCOTFS.bCookFilter)
 	{
 		IAssetRegistry& AssetRegistry = IAssetRegistry::GetChecked();
+		FName PackageNameToTest = PackageName;
+		if (PackageData.IsGenerated())
+		{
+			FName ParentName = PackageData.GetParentGenerator();
+			FPackageData* ParentData = InCOTFS.PackageDatas->FindPackageDataByPackageName(ParentName);
+			if (ParentData)
+			{
+				PackageNameToTest = ParentName;
+			}
+		}
+
 		if (!InCOTFS.CookFilterIncludedClasses.IsEmpty())
 		{
-			TOptional<FAssetPackageData> AssetData = AssetRegistry.GetAssetPackageDataCopy(PackageName);
+			TOptional<FAssetPackageData> AssetData = AssetRegistry.GetAssetPackageDataCopy(PackageNameToTest);
 			bool bIncluded = false;
 			if (AssetData)
 			{
@@ -2163,7 +2280,7 @@ void FRequestCluster::IsRequestCookable(const ITargetPlatform* Platform, FName P
 		if (!InCOTFS.CookFilterIncludedAssetClasses.IsEmpty())
 		{
 			TArray<FAssetData> AssetDatas;
-			AssetRegistry.GetAssetsByPackageName(PackageName, AssetDatas, true /* bIncludeOnlyDiskAssets */);
+			AssetRegistry.GetAssetsByPackageName(PackageNameToTest, AssetDatas, true /* bIncludeOnlyDiskAssets */);
 			bool bIncluded = false;
 			for (FAssetData& AssetData : AssetDatas)
 			{
diff --git a/Engine/Source/Editor/UnrealEd/Private/Cooker/CookRequestCluster.h b/Engine/Source/Editor/UnrealEd/Private/Cooker/CookRequestCluster.h
--- a/Engine/Source/Editor/UnrealEd/Private/Cooker/CookRequestCluster.h
+++ b/Engine/Source/Editor/UnrealEd/Private/Cooker/CookRequestCluster.h
@@ -25,6 +25,7 @@
 
 #include <atomic>
 
+class FAssetPackageData;
 class IAssetRegistry;
 class ICookedPackageWriter;
 class ITargetPlatform;
@@ -138,6 +139,7 @@ private:
 	struct FVertexData
 	{
 		FVertexData(FName InPackageName, UE::Cook::FPackageData* InPackageData, FGraphSearch& GraphSearch);
+		const FAssetPackageData* GetGeneratedAssetPackageData();
 
 		/* Async thread is not allowed to access PackageData, so store its name.The name is immutable for vertex lifetime. */
 		FName PackageName;
@@ -279,7 +281,8 @@ private:
 				FFetchPlatformData& FetchPlatformData, FPackagePlatformData& PackagePlatformData,
 				UE::TargetDomain::FCookAttachments& PlatformAttachments, bool bExploreDependencies);
 
-			void SetIsIterativelyUnmodified(int32 PlatformIndex, bool bIterativelyUnmodified);
+			void SetIsIterativelyUnmodified(int32 PlatformIndex, bool bIterativelyUnmodified,
+				FPackagePlatformData& PackagePlatformData);
 
 		private:
 			FRequestCluster& Cluster;
@@ -318,7 +321,7 @@ private:
 		void ResolveTransitiveBuildDependencyCycle();
 
 		/** Find or add a Vertex for PackageName. If PackageData is provided, use it, otherwise look it up. */
-		FVertexData& FindOrAddVertex(FName PackageName);
+		FVertexData& FindOrAddVertex(FName PackageName, FGenerationHelper* ParentGenerationHelper = nullptr);
 		FVertexData& FindOrAddVertex(FName PackageName, FPackageData& PackageData);
 		/** Batched allocation for vertices. */
 		FVertexData* AllocateVertex(FName PackageName, FPackageData* PackageData);
diff --git a/Engine/Source/Editor/UnrealEd/Private/Cooker/CookWorkerClient.cpp b/Engine/Source/Editor/UnrealEd/Private/Cooker/CookWorkerClient.cpp
--- a/Engine/Source/Editor/UnrealEd/Private/Cooker/CookWorkerClient.cpp
+++ b/Engine/Source/Editor/UnrealEd/Private/Cooker/CookWorkerClient.cpp
@@ -2,6 +2,7 @@
 
 #include "CookWorkerClient.h"
 
+#include "AssetRegistry/AssetData.h"
 #include "Cooker/CompactBinaryTCP.h"
 #include "Cooker/CookDirector.h"
 #include "Cooker/CookGenerationHelper.h"
@@ -309,7 +310,7 @@ void FCookWorkerClient::ReportPackageMessage(FName PackageName, TUniquePtr<FPack
 }
 
 void FCookWorkerClient::ReportDiscoveredPackage(const FPackageData& PackageData, const FInstigator& Instigator,
-	FDiscoveredPlatformSet&& ReachablePlatforms)
+	FDiscoveredPlatformSet&& ReachablePlatforms, FGenerationHelper* ParentGenerationHelper)
 {
 	FDiscoveredPackageReplication& Discovered = PendingDiscoveredPackages.Emplace_GetRef();
 	Discovered.PackageName = PackageData.GetPackageName();
@@ -319,6 +320,13 @@ void FCookWorkerClient::ReportDiscoveredPackage(const FPackageData& PackageData,
 	Discovered.Platforms = MoveTemp(ReachablePlatforms);
 	Discovered.Platforms.ConvertToBitfield(OrderedSessionAndSpecialPlatforms);
 	Discovered.DoesGeneratedRequireGenerator = PackageData.DoesGeneratedRequireGenerator();
+	if (ParentGenerationHelper)
+	{
+		if (FCookGenerationInfo* Info = ParentGenerationHelper->FindInfo(PackageData))
+		{
+			Discovered.GeneratedPackageHash = Info->PackageHash;
+		}
+	}
 }
 
 void FCookWorkerClient::ReportGeneratorQueuedGeneratedPackages(FGenerationHelper& GenerationHelper)
@@ -338,7 +346,7 @@ void FCookWorkerClient::HandleGeneratorMessage(FGeneratorEventMessage&& Generato
 			switch (GeneratorMessage.Event)
 			{
 			case EGeneratorEvent::QueuedGeneratedPackagesFencePassed:
-				GenerationHelper->OnRequestFencePassed(COTFS);
+				GenerationHelper->OnQueuedGeneratedPackagesFencePassed(COTFS);
 				break;
 			case EGeneratorEvent::AllSavesCompleted:
 				GenerationHelper->OnAllSavesCompleted(COTFS);
diff --git a/Engine/Source/Editor/UnrealEd/Private/Cooker/CookWorkerClient.h b/Engine/Source/Editor/UnrealEd/Private/Cooker/CookWorkerClient.h
--- a/Engine/Source/Editor/UnrealEd/Private/Cooker/CookWorkerClient.h
+++ b/Engine/Source/Editor/UnrealEd/Private/Cooker/CookWorkerClient.h
@@ -68,7 +68,7 @@ public:
 	void ReportPromoteToSaveComplete(FPackageData& PackageData);
 	/** Queue a message to the server that a package was discovered as needed in the cook. Will be sent during Tick. */
 	void ReportDiscoveredPackage(const FPackageData& PackageData, const FInstigator& Instigator,
-		FDiscoveredPlatformSet&& ReachablePlatforms);
+		FDiscoveredPlatformSet&& ReachablePlatforms, FGenerationHelper* ParentGenerationHelper);
 	/**
 	 * Queue a message to the server that a PackageGenerator queued its generated packages for cooking, and will
 	 * keep itself in memory until it receives an EGeneratorEvent::QueuedGeneratedPackagesFencePassed.
diff --git a/Engine/Source/Editor/UnrealEd/Private/Cooker/CookWorkerServer.cpp b/Engine/Source/Editor/UnrealEd/Private/Cooker/CookWorkerServer.cpp
--- a/Engine/Source/Editor/UnrealEd/Private/Cooker/CookWorkerServer.cpp
+++ b/Engine/Source/Editor/UnrealEd/Private/Cooker/CookWorkerServer.cpp
@@ -703,7 +703,7 @@ void FCookWorkerServer::SendPendingPackages()
 		FAssignPackageExtraData* ExtraData = PackagesToAssignExtraDatas.Find(PackageData);
 		if (ExtraData)
 		{
-			AssignData.GeneratorPreviousGeneratedPackages = ExtraData->GeneratorPreviousGeneratedPackages;
+			AssignData.GeneratorPreviousGeneratedPackages = MoveTemp(ExtraData->GeneratorPreviousGeneratedPackages);
 		}
 	}
 	for (FPackageData* PackageData : PackagesToAssignInfoPackages)
@@ -1075,7 +1075,7 @@ void FCookWorkerServer::QueueDiscoveredPackage(FDiscoveredPackageReplication&& D
 			TRefCountPtr<FGenerationHelper> GenerationHelper =
 				GeneratorPackageData->CreateUninitializedGenerationHelper();
 			GenerationHelper->NotifyStartQueueGeneratedPackages(COTFS, WorkerId);
-			GenerationHelper->TrackGeneratedPackageListedRemotely(COTFS, PackageData);
+			GenerationHelper->TrackGeneratedPackageListedRemotely(COTFS, PackageData, DiscoveredPackage.GeneratedPackageHash);
 		}
 	}
 
@@ -1216,7 +1216,17 @@ void FAssignPackageData::Write(FCbWriter& Writer,
 	Writer << ParentGenerator;
 	Writer << Instigator;
 	WriteToCompactBinary(Writer, NeedCookPlatforms, OrderedSessionPlatforms);
-	Writer << GeneratorPreviousGeneratedPackages;
+	{
+		Writer.BeginArray();
+		for (const TPair<FName, FAssetPackageData>& Pair : GeneratorPreviousGeneratedPackages)
+		{
+			Writer.BeginArray();
+			Writer << Pair.Key;
+			Pair.Value.NetworkWrite(Writer);
+			Writer.EndArray();
+		}
+		Writer.EndArray();
+	}
 	static_assert(sizeof(ICookPackageSplitter::EGeneratedRequiresGenerator) <= sizeof(uint8), "We are storing it in a uint8");
 	Writer << static_cast<uint8>(DoesGeneratedRequireGenerator);
 	Writer.EndArray();
@@ -1230,7 +1240,37 @@ bool FAssignPackageData::TryRead(FCbFieldView Field, TConstArrayView<const ITarg
 	bOk = LoadFromCompactBinary(*It++, ParentGenerator) & bOk;
 	bOk = LoadFromCompactBinary(*It++, Instigator) & bOk;
 	bOk = LoadFromCompactBinary(*It++, NeedCookPlatforms, OrderedSessionPlatforms) & bOk;
-	bOk = LoadFromCompactBinary(*It++, GeneratorPreviousGeneratedPackages) & bOk;
+	{
+		FCbFieldView ArrayFieldView = *It++;
+		bool bGeneratorPreviousGeneratedPackagesOk = false;
+		const uint64 Length = ArrayFieldView.AsArrayView().Num();
+		if (Length <= MAX_int32)
+		{
+			GeneratorPreviousGeneratedPackages.Empty((int32)Length);
+			bGeneratorPreviousGeneratedPackagesOk = !ArrayFieldView.HasError();
+			for (const FCbFieldView& ElementField : ArrayFieldView)
+			{
+				FCbFieldViewIterator PairIt = ElementField.CreateViewIterator();
+				bool bElementOk = false;
+				FName Key;
+				FAssetPackageData Value;
+				if (LoadFromCompactBinary(*PairIt++, Key))
+				{
+					if (Value.TryNetworkRead(*PairIt++))
+					{
+						GeneratorPreviousGeneratedPackages.Add(Key, MoveTemp(Value));
+						bElementOk = true;
+					}
+				}
+				bGeneratorPreviousGeneratedPackagesOk &= bElementOk;
+			}
+		}
+		else
+		{
+			GeneratorPreviousGeneratedPackages.Empty();
+		}
+		bOk &= bGeneratorPreviousGeneratedPackagesOk;
+	}
 	uint8 DoesGeneratedRequireGeneratorInt = It->AsUInt8();
 	if (!(It++)->HasError() && DoesGeneratedRequireGeneratorInt
 		< static_cast<uint8>(ICookPackageSplitter::EGeneratedRequiresGenerator::Count))
@@ -1421,6 +1461,12 @@ void FDiscoveredPackageReplication::Write(FCbWriter& Writer,
 	Writer << static_cast<uint8>(Instigator.Category);
 	Writer << Instigator.Referencer;
 	Writer << static_cast<uint8>(DoesGeneratedRequireGenerator);
+	bool bGeneratedPackageHash = !GeneratedPackageHash.IsZero();
+	Writer << bGeneratedPackageHash;
+	if (bGeneratedPackageHash)
+	{
+		Writer << GeneratedPackageHash;
+	}
 	static_assert(sizeof(ICookPackageSplitter::EGeneratedRequiresGenerator) <= sizeof(uint8), "We are storing it in a uint8");
 	WriteToCompactBinary(Writer, Platforms, OrderedSessionAndSpecialPlatforms);
 	Writer.EndArray();
@@ -1462,6 +1508,16 @@ bool FDiscoveredPackageReplication::TryRead(FCbFieldView Field,
 	{
 		bOk = false;
 	}
+	bool bGeneratedPackageHash = false;
+	bOk = LoadFromCompactBinary(Iter++, bGeneratedPackageHash) & bOk;
+	if (bGeneratedPackageHash)
+	{
+		bOk = LoadFromCompactBinary(Iter++, GeneratedPackageHash) & bOk;
+	}
+	else
+	{
+		GeneratedPackageHash = FIoHash::Zero;
+	}
 	bOk = LoadFromCompactBinary(Iter++, Platforms, OrderedSessionAndSpecialPlatforms) & bOk;
 	if (!bOk)
 	{
diff --git a/Engine/Source/Editor/UnrealEd/Private/Cooker/CookWorkerServer.h b/Engine/Source/Editor/UnrealEd/Private/Cooker/CookWorkerServer.h
--- a/Engine/Source/Editor/UnrealEd/Private/Cooker/CookWorkerServer.h
+++ b/Engine/Source/Editor/UnrealEd/Private/Cooker/CookWorkerServer.h
@@ -2,6 +2,7 @@
 
 #pragma once
 
+#include "AssetRegistry/AssetData.h"
 #include "CompactBinaryTCP.h"
 #include "Containers/Array.h"
 #include "Containers/ArrayView.h"
@@ -49,7 +50,7 @@ enum class ENotifyRemote
  */
 struct FAssignPackageExtraData
 {
-	TMap<FName, FIoHash> GeneratorPreviousGeneratedPackages;
+	TMap<FName, FAssetPackageData> GeneratorPreviousGeneratedPackages;
 };
 
 /** Class in a Director process that communicates over a Socket with FCookWorkerClient in a CookWorker process. */
@@ -253,7 +254,7 @@ struct FAssignPackageData
 	FName ParentGenerator;
 	FInstigator Instigator;
 	FDiscoveredPlatformSet NeedCookPlatforms;
-	TMap<FName, FIoHash> GeneratorPreviousGeneratedPackages;
+	TMap<FName, FAssetPackageData> GeneratorPreviousGeneratedPackages;
 	ICookPackageSplitter::EGeneratedRequiresGenerator DoesGeneratedRequireGenerator
 		= ICookPackageSplitter::EGeneratedRequiresGenerator::None;
 
@@ -402,6 +403,7 @@ struct FDiscoveredPackageReplication
 	FName PackageName;
 	FName NormalizedFileName;
 	FName ParentGenerator;
+	FIoHash GeneratedPackageHash;
 	FInstigator Instigator;
 	FDiscoveredPlatformSet Platforms;
 	ICookPackageSplitter::EGeneratedRequiresGenerator DoesGeneratedRequireGenerator =
diff --git a/Engine/Source/Editor/UnrealEd/Private/Cooker/IWorkerRequests.h b/Engine/Source/Editor/UnrealEd/Private/Cooker/IWorkerRequests.h
--- a/Engine/Source/Editor/UnrealEd/Private/Cooker/IWorkerRequests.h
+++ b/Engine/Source/Editor/UnrealEd/Private/Cooker/IWorkerRequests.h
@@ -62,7 +62,8 @@ public:
 
 	// Writing Packages during Cook
 	virtual void QueueDiscoveredPackage(UCookOnTheFlyServer& COTFS, FPackageData& PackageData,
-		FInstigator&& Instigator, FDiscoveredPlatformSet&& ReachablePlatforms, bool bUrgent) = 0;
+		FInstigator&& Instigator, FDiscoveredPlatformSet&& ReachablePlatforms, bool bUrgent,
+		FGenerationHelper* ParentGenerationHelper) = 0;
 	virtual void EndQueueGeneratedPackages(UCookOnTheFlyServer& COTFS, FGenerationHelper& GenerationHelper) = 0;
 
 	// Writing Packages from CookByTheBook
diff --git a/Engine/Source/Editor/UnrealEd/Private/Cooker/IterativeValidatePackageWriter.cpp b/Engine/Source/Editor/UnrealEd/Private/Cooker/IterativeValidatePackageWriter.cpp
--- a/Engine/Source/Editor/UnrealEd/Private/Cooker/IterativeValidatePackageWriter.cpp
+++ b/Engine/Source/Editor/UnrealEd/Private/Cooker/IterativeValidatePackageWriter.cpp
@@ -2,7 +2,10 @@
 
 #include "Cooker/IterativeValidatePackageWriter.h"
 
+#include "Cooker/CookGenerationHelper.h"
+#include "Cooker/CookPackageData.h"
 #include "CookOnTheSide/CookLog.h"
+#include "CookOnTheSide/CookOnTheFlyServer.h"
 #include "HAL/FileManager.h"
 #include "Logging/LogMacros.h"
 #include "Misc/CommandLine.h"
@@ -16,10 +19,11 @@ DEFINE_LOG_CATEGORY_STATIC(LogIterativeValidate, Log, All);
 
 constexpr FStringView IterativeValidateFilename(TEXTVIEW("IterativeValidate.bin"));
 
-FIterativeValidatePackageWriter::FIterativeValidatePackageWriter(TUniquePtr<ICookedPackageWriter>&& InInner,
-	EPhase InPhase, const FString& ResolvedMetadataPath)
+FIterativeValidatePackageWriter::FIterativeValidatePackageWriter(UCookOnTheFlyServer& InCOTFS,
+	TUniquePtr<ICookedPackageWriter>&& InInner, EPhase InPhase, const FString& ResolvedMetadataPath)
 	: FDiffPackageWriter(MoveTemp(InInner))
 	, MetadataPath(ResolvedMetadataPath)
+	, COTFS(InCOTFS)
 	, Phase(InPhase)
 {
 	Indent = FCString::Spc(FOutputDeviceHelper::FormatLogLine(ELogVerbosity::Warning,
@@ -58,16 +62,19 @@ void FIterativeValidatePackageWriter::BeginPackage(const FBeginPackageInfo& Info
 		Super::BeginPackage(Info);
 		break;
 	case EPhase::Phase2:
-		if (IterativeFailed.Contains(Info.PackageName))
+		if (IterativeValidated.Contains(Info.PackageName))
+		{
+			// Already saved in Phase 1; no need to diff it or save it now
+			SaveAction = ESaveAction::IgnoreResults;
+		}
+		else if (IterativeFailed.Contains(Info.PackageName))
 		{
 			SaveAction = ESaveAction::CheckForDiffs;
 			Super::BeginPackage(Info);
 		}
 		else
 		{
-			// This is not an IterativeValidated package (because we don't save those; UpdatePackageModificationStatus
-			// prevents saving it). And it is not an IterativeFailed package (checked above), so it is an
-			// IterativeModified package. It was found during Phase1 to be modified and would in a
+			// This is an IterativeModified package. It was found during Phase1 to be modified and would in a
 			// normal iterative cook be resaved rather than iteratively skipped. Resave it as normal.
 			SaveAction = ESaveAction::SaveToInner;
 			Inner->BeginPackage(Info);
@@ -291,10 +298,28 @@ void FIterativeValidatePackageWriter::Initialize(const FCookInfo& CookInfo)
 void FIterativeValidatePackageWriter::UpdatePackageModificationStatus(FName PackageName, bool bIterativelyUnmodified,
 	bool& bInOutShouldIterativelySkip)
 {
+	// We need to not skip previously cooked generator packages, if they were modified and we're read only we still
+	// need to cook them so we can investigate their generated packages. Look up whether they were a generator in
+	// the previous cook results.
+	bool bKnownGenerator = false;
+	if (bIterativelyUnmodified)
+	{
+		// GenerationHelpers were created for previously cooked generators by UCOTFS::PopulateCookedPackages.
+		UE::Cook::FPackageData* PackageData = COTFS.PackageDatas->FindPackageDataByPackageName(PackageName);
+		if (PackageData)
+		{
+			TRefCountPtr<UE::Cook::FGenerationHelper> GenerationHelper = PackageData->GetGenerationHelper();
+			if (GenerationHelper)
+			{
+				bKnownGenerator = true;
+			}
+		}
+	}
+
 	switch (Phase)
 	{
-	case EPhase::AllInOnePhase:
-		// Save the input value for bIterativelyUnmodified, and report skippable if and only if !unmodified and we're readonly.
+	case EPhase::AllInOnePhase:		
+		// Save the input value for bIterativelyUnmodified, and report skippable for the modified if possible
 		if (bIterativelyUnmodified)
 		{
 			IterativelyUnmodified.Add(PackageName);
@@ -302,7 +327,7 @@ void FIterativeValidatePackageWriter::UpdatePackageModificationStatus(FName Pack
 		}
 		else
 		{
-			if (bReadOnly)
+			if (!bKnownGenerator && bReadOnly)
 			{
 				bInOutShouldIterativelySkip = true;
 				++ModifiedCount; // Only increment here if we're skipping it. Otherwise it is incremented in BeginPackage
@@ -310,17 +335,21 @@ void FIterativeValidatePackageWriter::UpdatePackageModificationStatus(FName Pack
 		}
 		break;
 	case EPhase::Phase1:
-		// Invert what gets skipped: save the iteratively skipped files to record their diffs, but skip the regular files
-		bInOutShouldIterativelySkip = !bIterativelyUnmodified;
+		// Save the iteratively unmodified packages to verify their diffs. Skip the non-generator iteratively
+		// modified packages. Do not skip generator packages; we need to save them to test their generated packages.
+		bInOutShouldIterativelySkip = !bIterativelyUnmodified && !bKnownGenerator;
 		if (!bIterativelyUnmodified)
 		{
 			++ModifiedCount;
 		}
 		break;
 	case EPhase::Phase2:
-		// Ignore the Unmodified flag from this cook phase. Skip only the packages that were found to 
-		// be IterativeValidated from Phase1.
-		bInOutShouldIterativelySkip = IterativeValidated.Contains(PackageName);
+		// Ignore the Unmodified flag from this cook phase. Skip the packages that were found to 
+		// be IterativeValidated from Phase1. Save the packages that phase1 found modified; this phase
+		// is responsible for getting those resaved. Reexecute save for the packages that were IterativeFailed
+		// from phase1, so we can test whether they are indeterministic. Always save generators so we can
+		// test their generated packages.
+		bInOutShouldIterativelySkip = IterativeValidated.Contains(PackageName) && !bKnownGenerator;
 		break;
 	default:
 		checkNoEntry();
diff --git a/Engine/Source/Editor/UnrealEd/Private/Cooker/IterativeValidatePackageWriter.h b/Engine/Source/Editor/UnrealEd/Private/Cooker/IterativeValidatePackageWriter.h
--- a/Engine/Source/Editor/UnrealEd/Private/Cooker/IterativeValidatePackageWriter.h
+++ b/Engine/Source/Editor/UnrealEd/Private/Cooker/IterativeValidatePackageWriter.h
@@ -4,6 +4,8 @@
 
 #include "DiffPackageWriter.h"
 
+class UCookOnTheFlyServer;
+
 /**
  * A CookedPackageWriter that diffs the cook results of iteratively-unmodified packages between their last cook
  * results and the current cook.
@@ -18,8 +20,8 @@ public:
 		Phase1,
 		Phase2,
 	};
-	FIterativeValidatePackageWriter(TUniquePtr<ICookedPackageWriter>&& InInner, EPhase InPhase,
-		const FString& ResolvedMetadataPath);
+	FIterativeValidatePackageWriter(UCookOnTheFlyServer& InCOTFS, TUniquePtr<ICookedPackageWriter>&& InInner,
+		EPhase InPhase, const FString& ResolvedMetadataPath);
 
 	// IPackageWriter
 	virtual void BeginPackage(const FBeginPackageInfo& Info) override;
@@ -98,6 +100,7 @@ protected:
 	TSet<FName> IterativelyUnmodified;
 
 	FString MetadataPath;
+	UCookOnTheFlyServer& COTFS;
 	int32 ModifiedCount = 0;
 	EPhase Phase = EPhase::AllInOnePhase;
 	ESaveAction SaveAction = ESaveAction::IgnoreResults;
diff --git a/Engine/Source/Editor/UnrealEd/Private/Cooker/WorkerRequestsLocal.cpp b/Engine/Source/Editor/UnrealEd/Private/Cooker/WorkerRequestsLocal.cpp
--- a/Engine/Source/Editor/UnrealEd/Private/Cooker/WorkerRequestsLocal.cpp
+++ b/Engine/Source/Editor/UnrealEd/Private/Cooker/WorkerRequestsLocal.cpp
@@ -46,8 +46,12 @@ void FWorkerRequestsLocal::DequeueAllExternal(TArray<FSchedulerCallback>& OutCal
 }
 
 void FWorkerRequestsLocal::QueueDiscoveredPackage(UCookOnTheFlyServer& COTFS, FPackageData& PackageData,
-	FInstigator&& Instigator, FDiscoveredPlatformSet&& ReachablePlatforms, bool bUrgent)
+	FInstigator&& Instigator, FDiscoveredPlatformSet&& ReachablePlatforms, bool bUrgent,
+	FGenerationHelper* ParentGenerationHelper)
 {
+	// FWorkerRequestsRemote needs to send ParentGenerationHelper data to the director, but in the local case we
+	// already consumed the data in FGenerationHelper::StartQueueGeneratedPackages, so we don't use the value here.
+	(void)ParentGenerationHelper;
 	COTFS.QueueDiscoveredPackageOnDirector(PackageData, MoveTemp(Instigator), MoveTemp(ReachablePlatforms), bUrgent);
 }
 
diff --git a/Engine/Source/Editor/UnrealEd/Private/Cooker/WorkerRequestsLocal.h b/Engine/Source/Editor/UnrealEd/Private/Cooker/WorkerRequestsLocal.h
--- a/Engine/Source/Editor/UnrealEd/Private/Cooker/WorkerRequestsLocal.h
+++ b/Engine/Source/Editor/UnrealEd/Private/Cooker/WorkerRequestsLocal.h
@@ -33,7 +33,8 @@ public:
 	virtual void RemapTargetPlatforms(const TMap<ITargetPlatform*, ITargetPlatform*>& Remap) override;
 	virtual void OnRemoveSessionPlatform(const ITargetPlatform* TargetPlatform) override;
 	virtual void QueueDiscoveredPackage(UCookOnTheFlyServer& COTFS, FPackageData& PackageData,
-		FInstigator&& Instigator, FDiscoveredPlatformSet&& ReachablePlatforms, bool bUrgent) override;
+		FInstigator&& Instigator, FDiscoveredPlatformSet&& ReachablePlatforms, bool bUrgent,
+		FGenerationHelper* ParentGenerationHelper) override;
 	virtual void EndQueueGeneratedPackages(UCookOnTheFlyServer& COTFS, FGenerationHelper& GenerationHelper) override;
 	virtual void AddStartCookByTheBookRequest(FFilePlatformRequest&& Request) override;
 	virtual void InitializeCookOnTheFly() override;
diff --git a/Engine/Source/Editor/UnrealEd/Private/Cooker/WorkerRequestsRemote.cpp b/Engine/Source/Editor/UnrealEd/Private/Cooker/WorkerRequestsRemote.cpp
--- a/Engine/Source/Editor/UnrealEd/Private/Cooker/WorkerRequestsRemote.cpp
+++ b/Engine/Source/Editor/UnrealEd/Private/Cooker/WorkerRequestsRemote.cpp
@@ -55,10 +55,12 @@ void FWorkerRequestsRemote::DequeueAllExternal(TArray<FSchedulerCallback>& OutCa
 }
 
 void FWorkerRequestsRemote::QueueDiscoveredPackage(UCookOnTheFlyServer& COTFS, FPackageData& PackageData,
-	FInstigator&& Instigator, FDiscoveredPlatformSet&& ReachablePlatforms, bool bUrgent)
+	FInstigator&& Instigator, FDiscoveredPlatformSet&& ReachablePlatforms, bool bUrgent,
+	FGenerationHelper* ParentGenerationHelper)
 {
 	(void)bUrgent; // Tracking urgency on CookWorkers is not supported
-	CookWorkerClient.ReportDiscoveredPackage(PackageData, MoveTemp(Instigator), MoveTemp(ReachablePlatforms));
+	CookWorkerClient.ReportDiscoveredPackage(PackageData, MoveTemp(Instigator), MoveTemp(ReachablePlatforms),
+		ParentGenerationHelper);
 }
 
 void FWorkerRequestsRemote::EndQueueGeneratedPackages(UCookOnTheFlyServer& COTFS,
diff --git a/Engine/Source/Editor/UnrealEd/Private/Cooker/WorkerRequestsRemote.h b/Engine/Source/Editor/UnrealEd/Private/Cooker/WorkerRequestsRemote.h
--- a/Engine/Source/Editor/UnrealEd/Private/Cooker/WorkerRequestsRemote.h
+++ b/Engine/Source/Editor/UnrealEd/Private/Cooker/WorkerRequestsRemote.h
@@ -41,7 +41,8 @@ public:
 	virtual void RemapTargetPlatforms(const TMap<ITargetPlatform*, ITargetPlatform*>& Remap) override;
 	virtual void OnRemoveSessionPlatform(const ITargetPlatform* TargetPlatform) override;
 	virtual void QueueDiscoveredPackage(UCookOnTheFlyServer& COTFS, FPackageData& PackageData,
-		FInstigator&& Instigator, FDiscoveredPlatformSet&& ReachablePlatforms, bool bUrgent) override;
+		FInstigator&& Instigator, FDiscoveredPlatformSet&& ReachablePlatforms, bool bUrgent,
+		FGenerationHelper* ParentGenerationHelper) override;
 	virtual void EndQueueGeneratedPackages(UCookOnTheFlyServer& COTFS, FGenerationHelper& GenerationHelper) override;
 	virtual void AddStartCookByTheBookRequest(FFilePlatformRequest&& Request) override;
 	virtual void InitializeCookOnTheFly() override;
diff --git a/Engine/Source/Editor/UnrealEd/Private/EditorDomain/EditorDomainUtils.cpp b/Engine/Source/Editor/UnrealEd/Private/EditorDomain/EditorDomainUtils.cpp
--- a/Engine/Source/Editor/UnrealEd/Private/EditorDomain/EditorDomainUtils.cpp
+++ b/Engine/Source/Editor/UnrealEd/Private/EditorDomain/EditorDomainUtils.cpp
@@ -326,12 +326,16 @@ void FKnownCustomVersions::FindGuidsChecked(TArray<FGuid>& OutGuids, TConstArray
 FPackageDigest CalculatePackageDigest(IAssetRegistry& AssetRegistry, FName PackageName)
 {
 	AssetRegistry.WaitForPackage(PackageName.ToString());
-	TOptional<FAssetPackageData> PackageDataOptional = AssetRegistry.GetAssetPackageDataCopy(PackageName);
-	if (!PackageDataOptional)
+	TOptional<FAssetPackageData> PackageData = AssetRegistry.GetAssetPackageDataCopy(PackageName);
+	if (!PackageData)
 	{
 		return FPackageDigest(FPackageDigest::EStatus::DoesNotExistInAssetRegistry);
 	}
-	FAssetPackageData& PackageData = *PackageDataOptional;
+	return CalculatePackageDigest(*PackageData, PackageName);
+}
+
+FPackageDigest CalculatePackageDigest(const FAssetPackageData& PackageData, FName PackageName)
+{
 	FPackageDigest Result;
 	Result.DomainUse = EDomainUse::LoadEnabled | EDomainUse::SaveEnabled;;
 	EnumSetFlagsAnd(Result.DomainUse, EDomainUse::LoadEnabled | EDomainUse::SaveEnabled,
@@ -1917,7 +1921,8 @@ bool TrySavePackage(UPackage* Package)
 	if (bStorageResultValid)
 	{
 		TArray<IPackageWriter::FCommitAttachmentInfo, TInlineAllocator<2>> Attachments;
-		UE::TargetDomain::CollectAndStoreCookAttachments(Package, nullptr, nullptr, TArray<FName>(), Attachments);
+		UE::TargetDomain::CollectAndStoreCookAttachments(Package, nullptr, nullptr, nullptr, TArray<FName>(),
+			Attachments);
 		if (!Attachments.IsEmpty())
 		{
 			UE::TargetDomain::CommitEditorDomainCookAttachments(Package->GetFName(), Attachments);
diff --git a/Engine/Source/Editor/UnrealEd/Private/EditorDomain/EditorDomainUtils.h b/Engine/Source/Editor/UnrealEd/Private/EditorDomain/EditorDomainUtils.h
--- a/Engine/Source/Editor/UnrealEd/Private/EditorDomain/EditorDomainUtils.h
+++ b/Engine/Source/Editor/UnrealEd/Private/EditorDomain/EditorDomainUtils.h
@@ -85,6 +85,7 @@ ESaveStorageResult SaveStorageResultFromString(FUtf8StringView Text);
  * Reads information from the AssetRegistry to compute the digest.
  */
 FPackageDigest CalculatePackageDigest(IAssetRegistry& AssetRegistry, FName PackageName);
+FPackageDigest CalculatePackageDigest(const FAssetPackageData& PackageData, FName PackageName);
 
 /** For any ClassNames not already in ClassDigests, look up their UStruct and add them. */
 void PrecacheClassDigests(TConstArrayView<FTopLevelAssetPath> ClassNames);
diff --git a/Engine/Source/Editor/UnrealEd/Private/TargetDomain/TargetDomainUtils.cpp b/Engine/Source/Editor/UnrealEd/Private/TargetDomain/TargetDomainUtils.cpp
--- a/Engine/Source/Editor/UnrealEd/Private/TargetDomain/TargetDomainUtils.cpp
+++ b/Engine/Source/Editor/UnrealEd/Private/TargetDomain/TargetDomainUtils.cpp
@@ -13,6 +13,7 @@
 #include "Cooker/CookDependency.h"
 #include "Cooker/PackageBuildDependencyTracker.h"
 #include "CookOnTheSide/CookLog.h"
+#include "CookPackageSplitter.h"
 #include "DerivedDataBuildDefinition.h"
 #include "DerivedDataBuildKey.h"
 #include "DerivedDataSharedString.h"
@@ -100,7 +101,7 @@ bool FCookDependencies::IsValid() const
 	return bValid;
 }
 
-bool FCookDependencies::HasKeyMatch()
+bool FCookDependencies::HasKeyMatch(const FAssetPackageData* OverrideAssetPackageData)
 {
 	if (!bValid)
 	{
@@ -112,7 +113,7 @@ bool FCookDependencies::HasKeyMatch()
 	}
 	if (CurrentKey.IsZero())
 	{
-		if (!TryCalculateCurrentKey())
+		if (!TryCalculateCurrentKey(OverrideAssetPackageData))
 		{
 			return false;
 		}
@@ -120,7 +121,8 @@ bool FCookDependencies::HasKeyMatch()
 	return CurrentKey == StoredKey;
 }
 
-bool FCookDependencies::TryCalculateCurrentKey(FString* OutErrorMessage)
+bool FCookDependencies::TryCalculateCurrentKey(const FAssetPackageData* OverrideAssetPackageData,
+	FString* OutErrorMessage)
 {
 	IAssetRegistry* AssetRegistry = IAssetRegistry::Get();
 	if (PackageName.IsNone())
@@ -140,7 +142,10 @@ bool FCookDependencies::TryCalculateCurrentKey(FString* OutErrorMessage)
 		return false;
 	}
 	FBlake3 KeyBuilder;
-	UE::EditorDomain::FPackageDigest PackageDigest = EditorDomain->GetPackageDigest(PackageName);
+	UE::EditorDomain::FPackageDigest PackageDigest = OverrideAssetPackageData ?
+		UE::EditorDomain::CalculatePackageDigest(*OverrideAssetPackageData, PackageName) :
+		EditorDomain->GetPackageDigest(PackageName);
+
 	if (!PackageDigest.IsSuccessful())
 	{
 		if (OutErrorMessage) *OutErrorMessage = PackageDigest.GetStatusString();
@@ -240,19 +245,31 @@ void FCookDependencies::Empty()
 }
 
 FCookDependencies FCookDependencies::Collect(UPackage* Package, const ITargetPlatform* TargetPlatform,
-	FSavePackageResultStruct* SaveResult, TArray<FName>&& RuntimeDependencies, FString* OutErrorMessage)
+	FSavePackageResultStruct* SaveResult, const FGeneratedPackageResultStruct* GeneratedResult, 
+	TArray<FName>&& RuntimeDependencies, FString* OutErrorMessage)
 {
 	TStringBuilder<256> StringBuffer;
 	FName TransientPackageName = GetTransientPackage()->GetFName();
-	auto IsTransientPackageName = [&StringBuffer, TransientPackageName](FName InPackageName)
+	auto IsDisallowedDependencyName = [&StringBuffer, TransientPackageName](FName InPackageName,
+		bool bAllowGeneratedPackages)
 		{
 			if (InPackageName == TransientPackageName)
 			{
 				return true;
 			}
 			InPackageName.ToString(StringBuffer);
-			return FPackageName::IsMemoryPackage(StringBuffer) ||
-				FPackageName::IsScriptPackage(StringBuffer);
+			FStringView PackageNameStr = StringBuffer.ToView();
+			return FPackageName::IsMemoryPackage(PackageNameStr) ||
+				FPackageName::IsScriptPackage(PackageNameStr) ||
+				(!bAllowGeneratedPackages && ICookPackageSplitter::IsUnderGeneratedPackageSubPath(PackageNameStr));
+		};
+	auto IsDisallowedBuildDependencyName = [&IsDisallowedDependencyName](FName InPackageName)
+		{
+			return IsDisallowedDependencyName(InPackageName, false /* bAllowGeneratedPackages */);
+		};
+	auto IsDisallowedRuntimeDependencyName = [&IsDisallowedDependencyName](FName InPackageName)
+		{
+			return IsDisallowedDependencyName(InPackageName, true /* bAllowGeneratedPackages */);
 		};
 
 	if (!Package)
@@ -278,9 +295,17 @@ FCookDependencies FCookDependencies::Collect(UPackage* Package, const ITargetPla
 	TSet<FName> BuildDependenciesSet;
 
 	TArray<FName> AssetDependencies;
-	AssetRegistry->GetDependencies(Result.PackageName, AssetDependencies,
-		UE::AssetRegistry::EDependencyCategory::Package, UE::AssetRegistry::EDependencyQuery::Game);
-	RuntimeDependencies.Append(MoveTemp(AssetDependencies));
+	if (!GeneratedResult)
+	{
+		AssetRegistry->GetDependencies(Result.PackageName, AssetDependencies,
+			UE::AssetRegistry::EDependencyCategory::Package, UE::AssetRegistry::EDependencyQuery::Game);
+		RuntimeDependencies.Append(MoveTemp(AssetDependencies));
+	}
+	else
+	{
+		// GeneratedResult->PackageDependencies are already incorporated into the PackageHash of the generated
+		// package by FCookGenerationInfo::CreatePackageHash, so we do not need to add them into the BuildDependencies
+	}
 
 	FPackageBuildDependencyTracker& Tracker = FPackageBuildDependencyTracker::Get();
 
@@ -306,7 +331,7 @@ FCookDependencies FCookDependencies::Collect(UPackage* Package, const ITargetPla
 	}
 
 	Result.BuildPackageDependencies = BuildDependenciesSet.Array();
-	Result.BuildPackageDependencies.RemoveAllSwap(IsTransientPackageName, EAllowShrinking::Yes);
+	Result.BuildPackageDependencies.RemoveAllSwap(IsDisallowedBuildDependencyName, EAllowShrinking::Yes);
 	Result.BuildPackageDependencies.Sort(FNameLexicalLess());
 
 #if UE_WITH_CONFIG_TRACKING
@@ -357,13 +382,14 @@ FCookDependencies FCookDependencies::Collect(UPackage* Package, const ITargetPla
 		Algo::Sort(Result.TransitiveBuildDependencies);
 	}
 
-	RuntimeDependencies.RemoveAllSwap(IsTransientPackageName, EAllowShrinking::No);
+	RuntimeDependencies.RemoveAllSwap(IsDisallowedRuntimeDependencyName, EAllowShrinking::No);
 	RuntimeDependencies.Sort(FNameLexicalLess());
 	RuntimeDependencies.SetNum(Algo::Unique(RuntimeDependencies), EAllowShrinking::Yes);
 
 	Result.RuntimePackageDependencies = MoveTemp(RuntimeDependencies);
 
-	if (!Result.TryCalculateCurrentKey(OutErrorMessage))
+	const FAssetPackageData* AssetPackageData = GeneratedResult ? &GeneratedResult->AssetPackageData : nullptr;
+	if (!Result.TryCalculateCurrentKey(AssetPackageData, OutErrorMessage))
 	{
 		return FCookDependencies();
 	}
@@ -577,15 +603,15 @@ void FCookAttachments::Empty()
 }
 
 bool TryCollectAndStoreCookDependencies(UPackage* Package, const ITargetPlatform* TargetPlatform,
-	FSavePackageResultStruct* SaveResult, TArray<FName>&& RuntimeDependencies,
-	IPackageWriter::FCommitAttachmentInfo& OutResult)
+	FSavePackageResultStruct* SaveResult, const FGeneratedPackageResultStruct* GeneratedResult,
+	TArray<FName>&& RuntimeDependencies, IPackageWriter::FCommitAttachmentInfo& OutResult)
 {
 	FString ErrorMessage;
-	FCookDependencies CookDependencies = FCookDependencies::Collect(Package, TargetPlatform, SaveResult, 
-		MoveTemp(RuntimeDependencies), &ErrorMessage);
+	FCookDependencies CookDependencies = FCookDependencies::Collect(Package, TargetPlatform, SaveResult,
+		GeneratedResult, MoveTemp(RuntimeDependencies), &ErrorMessage);
 	if (!CookDependencies.IsValid())
 	{
-		// CookPackageSplitterTODO: This error occurs for generated packages. Need to register them with EditorDomain.
+		// IterativeTODO: This error occurs due to dependencies on _Verse and on some transient packages.
 #if 0
 		UE_LOG(LogCook, Error, TEXT("Could not collect CookDependencies for package '%s': %s"),
 			*Package->GetName(), *ErrorMessage);
@@ -654,19 +680,24 @@ void FCookAttachments::Fetch(TArrayView<FName> PackageNames, const ITargetPlatfo
 	}
 }
 
-bool IsIterativeEnabled(FName PackageName, bool bAllowAllClasses)
+bool IsIterativeEnabled(FName PackageName, bool bAllowAllClasses, const FAssetPackageData* OverrideAssetPackageData)
 {
 	IAssetRegistry* AssetRegistry = IAssetRegistry::Get();
 	if (!AssetRegistry)
 	{
 		return false;
 	}
-	TOptional<FAssetPackageData> PackageDataOpt = AssetRegistry->GetAssetPackageDataCopy(PackageName);
-	if (!PackageDataOpt)
+	TOptional<FAssetPackageData> PackageDataOpt;
+	if (!OverrideAssetPackageData)
 	{
-		return false;
+		PackageDataOpt = AssetRegistry->GetAssetPackageDataCopy(PackageName);
+		if (!PackageDataOpt)
+		{
+			return false;
+		}
+		OverrideAssetPackageData = PackageDataOpt.GetPtrOrNull();
 	}
-	FAssetPackageData& PackageData = *PackageDataOpt;
+	const FAssetPackageData& PackageData = *OverrideAssetPackageData;
 
 	if (!bAllowAllClasses)
 	{
diff --git a/Engine/Source/Editor/UnrealEd/Private/TargetDomain/TargetDomainUtils.h b/Engine/Source/Editor/UnrealEd/Private/TargetDomain/TargetDomainUtils.h
--- a/Engine/Source/Editor/UnrealEd/Private/TargetDomain/TargetDomainUtils.h
+++ b/Engine/Source/Editor/UnrealEd/Private/TargetDomain/TargetDomainUtils.h
@@ -2,6 +2,8 @@
 
 #pragma once
 
+#include "AssetRegistry/AssetData.h"
+#include "AssetRegistry/IAssetRegistry.h"
 #include "Containers/Array.h"
 #include "Containers/ArrayView.h"
 #include "Containers/ContainersFwd.h"
@@ -36,6 +38,16 @@ namespace UE::TargetDomain
 /** Call during Startup to initialize global data used by TargetDomain functions. */
 void CookInitialize();
 
+/**
+ * Information collected from a CookPackageSplitter Generated package after it is saved; this information is needed to collect
+ * the CookDependencies for the generated package.
+ */
+struct FGeneratedPackageResultStruct
+{
+	FAssetPackageData AssetPackageData;
+	TArray<FAssetDependency> PackageDependencies;
+};
+
 /**
  * Recording of the dependencies of a package discovered during cook, used in incremental cooks. All dependencies
  * except for those marked Runtime contribute to the packages TargetDomain Key. If HasKeyMatch returns false after
@@ -74,10 +86,10 @@ public:
 	FName GetPackageName() const;
 	const FIoHash& GetStoredKey() const;
 	const FIoHash& GetCurrentKey() const;
-	bool HasKeyMatch();
+	bool HasKeyMatch(const FAssetPackageData* OverrideAssetPackageData);
 
 	// Modifying the structure data
-	bool TryCalculateCurrentKey(FString* OutErrorMessage = nullptr);
+	bool TryCalculateCurrentKey(const FAssetPackageData* OverrideAssetPackageData, FString* OutErrorMessage = nullptr);
 	void Reset();
 	void Empty();
 
@@ -87,7 +99,8 @@ public:
 	 * that have recorded its data during the package's load/other/save operations in the current cook session.
 	 */
 	static FCookDependencies Collect(UPackage* Package, const ITargetPlatform* TargetPlatform,
-		FSavePackageResultStruct* SaveResult, TArray<FName>&& RuntimeDependencies, FString* OutErrorMessage = nullptr);
+		FSavePackageResultStruct* SaveResult, const FGeneratedPackageResultStruct* GeneratedResult,
+		TArray<FName>&& RuntimeDependencies, FString* OutErrorMessage = nullptr);
 
 	// Fetch function to load the dependencies from a PackageStore is not yet implemented independently for
 	// this structure. Use FCookAttachments instead. 
@@ -144,17 +157,20 @@ struct FCookAttachments
 };
 
 bool TryCollectAndStoreCookDependencies(UPackage* Package, const ITargetPlatform* TargetPlatform,
-	FSavePackageResultStruct* SaveResult, TArray<FName>&& RuntimeDependencies,
-	IPackageWriter::FCommitAttachmentInfo& OutResult);
+	FSavePackageResultStruct* SaveResult, const FGeneratedPackageResultStruct* GeneratedResult,
+	TArray<FName>&& RuntimeDependencies, IPackageWriter::FCommitAttachmentInfo& OutResult);
 bool TryCollectAndStoreBuildDefinitionList(UPackage* Package, const ITargetPlatform* TargetPlatform,
 	IPackageWriter::FCommitAttachmentInfo& OutResult);
 
 template <typename ArrayType>
-void CollectAndStoreCookAttachments(UPackage* Package, const ITargetPlatform* TargetPlatform,
-	FSavePackageResultStruct* SaveResult, TArray<FName>&& RuntimeDependencies, ArrayType& Output)
+void CollectAndStoreCookAttachments(UPackage* Package,
+	const ITargetPlatform* TargetPlatform, FSavePackageResultStruct* SaveResult,
+	const FGeneratedPackageResultStruct* GeneratedResult, TArray<FName>&& RuntimeDependencies,
+	ArrayType& Output)
 {
 	IPackageWriter::FCommitAttachmentInfo Result;
-	if (TryCollectAndStoreCookDependencies(Package, TargetPlatform, SaveResult, MoveTemp(RuntimeDependencies), Result))
+	if (TryCollectAndStoreCookDependencies(Package, TargetPlatform, SaveResult, GeneratedResult,
+		MoveTemp(RuntimeDependencies), Result))
 	{
 		Output.Add(MoveTemp(Result));
 	}
@@ -165,7 +181,8 @@ void CollectAndStoreCookAttachments(UPackage* Package, const ITargetPlatform* Ta
 }
 
 /** Return whether iterative cook is enabled for the given packagename, based on used-class allowlist/blocklist. */
-bool IsIterativeEnabled(FName PackageName, bool bAllowAllClasses);
+bool IsIterativeEnabled(FName PackageName, bool bAllowAllClasses,
+	const FAssetPackageData* OverrideAssetPackageData = nullptr);
 
 /** Store extra information derived during save and used by the cooker for the given EditorDomain package. */
 void CommitEditorDomainCookAttachments(FName PackageName, TArrayView<IPackageWriter::FCommitAttachmentInfo> Attachments);
diff --git a/Engine/Source/Editor/UnrealEd/Public/Commandlets/AssetRegistryGenerator.h b/Engine/Source/Editor/UnrealEd/Public/Commandlets/AssetRegistryGenerator.h
--- a/Engine/Source/Editor/UnrealEd/Public/Commandlets/AssetRegistryGenerator.h
+++ b/Engine/Source/Editor/UnrealEd/Public/Commandlets/AssetRegistryGenerator.h
@@ -81,7 +81,7 @@ public:
 	/** Info about a GeneratorPackage (see ICookPackageSplitter) loaded from previous iterative cooks. */
 	struct FGeneratorPackageInfo
 	{
-		TMap<FName, FIoHash> Generated;
+		TMap<FName, FAssetPackageData> Generated;
 	};
 
 	enum EDifference
@@ -287,6 +287,12 @@ private:
 
 	bool ComputePackageDifferences_IsPackageFileUnchanged(const FComputeDifferenceOptions& Options, FName PackageName,
 		const FAssetPackageData& CurrentPackageData, const FAssetPackageData& PreviousPackageData);
+	/**
+	 * When copying AssetPackageData from the previous saved development assetregistry into the
+	 * FGeneratorPackageInfo returned from ComputePackageRemovals that are used for incremental cook,
+	 * we only need a subset of the data; we skip copying the portions we don't need to improve performance.
+	 */
+	static FAssetPackageData CopyAssetPackageDataForIncrementalCook(const FAssetPackageData& Source);
 
 	/** State of the asset registry that is being built for this platform */
 	FAssetRegistryState State;
diff --git a/Engine/Source/Editor/UnrealEd/Public/CookPackageSplitter.h b/Engine/Source/Editor/UnrealEd/Public/CookPackageSplitter.h
--- a/Engine/Source/Editor/UnrealEd/Public/CookPackageSplitter.h
+++ b/Engine/Source/Editor/UnrealEd/Public/CookPackageSplitter.h
@@ -296,6 +296,8 @@ public:
 
 	/** The name of the _Generated_ subdirectory that is the parent directory of a splitter's generated packages. */
 	UNREALED_API static const TCHAR* GetGeneratedPackageSubPath();
+	/** Return true if the given path is a _Generated_ directory, or a subpath under it. */
+	UNREALED_API static bool IsUnderGeneratedPackageSubPath(FStringView FileOrLongPackagePath);
 
 	/**
 	 * Return the full packagename that will be used for a GeneratedPackage, based on the GeneratorPackage's name and
diff --git a/Engine/Source/Runtime/CoreUObject/Private/AssetRegistry/AssetData.cpp b/Engine/Source/Runtime/CoreUObject/Private/AssetRegistry/AssetData.cpp
--- a/Engine/Source/Runtime/CoreUObject/Private/AssetRegistry/AssetData.cpp
+++ b/Engine/Source/Runtime/CoreUObject/Private/AssetRegistry/AssetData.cpp
@@ -1028,6 +1028,13 @@ void FAssetPackageData::NetworkWrite(FCbWriter& Writer) const
 	{
 		Writer << CookedHash;
 	}
+	FIoHash PackageSavedHash = GetPackageSavedHash();
+	bool bPackageSavedHash = !PackageSavedHash.IsZero();
+	Writer << bPackageSavedHash;
+	if (bPackageSavedHash)
+	{
+		Writer << PackageSavedHash;
+	}
 	Writer << ChunkHashes.Array();
 	Writer << ImportedClasses;
 	Writer << DiskSize;
@@ -1054,6 +1061,14 @@ bool FAssetPackageData::TryNetworkRead(FCbFieldView Field)
 	{
 		CookedHash = FMD5Hash();
 	}
+	bool bPackageSavedHash = false;
+	FIoHash PackageSavedHash;
+	bOk = LoadFromCompactBinary(*Iter++, bPackageSavedHash) & bOk;
+	if (bPackageSavedHash)
+	{
+		bOk = LoadFromCompactBinary(*Iter++, PackageSavedHash) & bOk;
+	}
+	SetPackageSavedHash(PackageSavedHash);
 	TArray<TPair<FIoChunkId, FIoHash>> ChunkHashesArray;
 	if (LoadFromCompactBinary(*Iter++, ChunkHashesArray))
 	{
