Fixed ensure firing on ortho zoom.
- CL 31854437 added ensures to FViewportCameraTransform::SetOrthoZoom
- previously a second call to SetOrthoZoom was used to set zoom to clamped value
#rnx

diff --git a/Engine/Source/Editor/UnrealEd/Private/EditorViewportClient.cpp b/Engine/Source/Editor/UnrealEd/Private/EditorViewportClient.cpp
--- a/Engine/Source/Editor/UnrealEd/Private/EditorViewportClient.cpp
+++ b/Engine/Source/Editor/UnrealEd/Private/EditorViewportClient.cpp
@@ -3741,8 +3741,8 @@ void FEditorViewportClient::OnOrthoZoom( const struct FInputEventState& InputSta
 	const float OldUnitsPerPixel = GetOrthoUnitsPerPixel(Viewport);
 
 	//update zoom based on input
-	SetOrthoZoom( GetOrthoZoom() + (GetOrthoZoom() / CAMERA_ZOOM_DAMPEN) * Delta );
-	SetOrthoZoom( FMath::Clamp<float>( GetOrthoZoom(), GetMinimumOrthoZoom(), MAX_ORTHOZOOM ) );
+	const float Zoom = GetOrthoZoom() + (GetOrthoZoom() / CAMERA_ZOOM_DAMPEN) * Delta;
+	SetOrthoZoom( FMath::Clamp<float>( Zoom, GetMinimumOrthoZoom(), MAX_ORTHOZOOM ) );
 
 	if (bCenterZoomAroundCursor)
 	{
@@ -5398,8 +5398,8 @@ void FEditorViewportClient::MoveViewportCamera(const FVector& InDrag, const FRot
 
 			if( ( LeftMouseButtonDown || bIsUsingTrackpad ) && RightMouseButtonDown )
 			{
-				SetOrthoZoom( GetOrthoZoom() + (GetOrthoZoom() / CAMERA_ZOOM_DAMPEN) * InDrag.Z );
-				SetOrthoZoom( FMath::Clamp<float>( GetOrthoZoom(), GetMinimumOrthoZoom(), MAX_ORTHOZOOM ) );
+				const float Zoom = GetOrthoZoom() + (GetOrthoZoom() / CAMERA_ZOOM_DAMPEN) * InDrag.Z;
+				SetOrthoZoom( FMath::Clamp<float>( Zoom, GetMinimumOrthoZoom(), MAX_ORTHOZOOM ) );
 			}
 			else
 			{
