Fix CreateProjectFromTemplate
* Better support archetype reinstantiation mixed with instances (sort)

diff --git a/Engine/Source/Editor/UnrealEd/Private/Kismet2/KismetReinstanceUtilities.cpp b/Engine/Source/Editor/UnrealEd/Private/Kismet2/KismetReinstanceUtilities.cpp
--- a/Engine/Source/Editor/UnrealEd/Private/Kismet2/KismetReinstanceUtilities.cpp
+++ b/Engine/Source/Editor/UnrealEd/Private/Kismet2/KismetReinstanceUtilities.cpp
@@ -42,6 +42,7 @@
 #include "Subsystems/AssetEditorSubsystem.h"
 #include "Engine/ScopedMovementUpdate.h"
 #include "InstancedReferenceSubobjectHelper.h"
+#include "Algo/TopologicalSort.h"
 #include "UObject/OverridableManager.h"
 #include "UObject/PropertyOptional.h"
 #include "UObject/PropertyBagRepository.h"
@@ -2176,7 +2177,7 @@ static void ReplaceObjectHelper(UObject*& OldObject, UClass* OldClass, UObject*&
 			int32 ArchetypeIndex = ObjectsToReplace.Find(OldArchetype);
 			if (ArchetypeIndex != INDEX_NONE)
 			{
-				if (ensure(ArchetypeIndex > OldObjIndex))
+				if (!ensure(ArchetypeIndex < OldObjIndex))
 				{
 					// if this object has an archetype, but it hasn't been 
 					// reinstanced yet (but is queued to) then we need to swap out 
@@ -2668,23 +2669,7 @@ void FBlueprintCompileReinstancer::ReplaceInstancesOfClass_Inner(const TMap<UCla
 			check(OldClass && NewClass);
 			check(OldClass != NewClass || IsReloadActive());
 			{
-				auto IsScriptComponent = [](UClass* InClass) -> bool
-				{
-					bool bIsScriptComponent = false;
-					// Hacky way of dtecting ScriptComponents
-					static FName NAME_ScriptComponent(TEXT("ScriptComponent"));
-					for (UClass* CurrentClass = InClass; CurrentClass && !bIsScriptComponent; CurrentClass = CurrentClass->GetSuperClass())
-					{
-						bIsScriptComponent = CurrentClass->GetFName() == NAME_ScriptComponent;
-					}					
-					return bIsScriptComponent;
-				};
-				
 				const bool bIsComponent = NewClass->IsChildOf<UActorComponent>();
-				// Keeping script component separate from bIsComponent as there's extra rules for replacing actor components
-				// that may not apply to ScriptComponents.
-				// We need to replace ScriptComponents that are on Blueprint CDOs when they're being edited
-				const bool bIsScriptComponent = IsScriptComponent(NewClass);
 
 				// If any of the class changes are of an actor component to scene component or reverse then we will fixup SCS of all actors affected
 				if (bIsComponent && !bFixupSCS)
@@ -2700,6 +2685,33 @@ void FBlueprintCompileReinstancer::ReplaceInstancesOfClass_Inner(const TMap<UCla
 				const bool bIncludeDerivedClasses = false;
 				ObjectsToReplace.Reset();
 				GetObjectsOfClass(OldClass, ObjectsToReplace, bIncludeDerivedClasses);
+
+				if (InstancesThatShouldUseOldClass)
+				{
+					// Pre-remove instance that will not need to be replaced. That way ReplaceObjectHelper will not barf on the archetype not being replaced before its instances.
+					for (auto It = ObjectsToReplace.CreateIterator(); It; ++It)
+					{
+						if (InstancesThatShouldUseOldClass->Contains(*It))
+						{
+							It.RemoveCurrentSwap();
+						}
+					}
+				}
+
+				if(!bArchetypesAreUpToDate)
+				{
+					Algo::TopologicalSort(ObjectsToReplace, [&ObjectsToReplace](const UObject* OldObject)
+					{
+						TArray<UObject*> Dependencies;
+						UObject* Archetype = OldObject->GetArchetype();
+						if (Archetype && ObjectsToReplace.Contains(Archetype) && !Archetype->HasAnyFlags(RF_ClassDefaultObject))
+						{
+							Dependencies.Add(Archetype);
+						}
+						return Dependencies;
+					});
+				}
+				
 				// Then fix 'real' (non archetype) instances of the class
 				for (int32 OldObjIndex = 0; OldObjIndex < ObjectsToReplace.Num(); ++OldObjIndex)
 				{
@@ -2717,9 +2729,13 @@ void FBlueprintCompileReinstancer::ReplaceInstancesOfClass_Inner(const TMap<UCla
 
 					// Skip archetype instances, EXCEPT for component templates and child actor templates
 					const bool bIsChildActorTemplate = OldActor && OldActor->GetOuter()->IsA<UChildActorComponent>();
-					if ((!bIsValid && !bIsScriptComponent) || // @todo: why do we need to replace PendingKill script components?
-						(!bIsComponent && !bIsChildActorTemplate && OldObject->IsTemplate() && !bIsScriptComponent) ||
-						(InstancesThatShouldUseOldClass && InstancesThatShouldUseOldClass->Contains(OldObject)))
+					if (!bIsValid || 
+						  (!bIsComponent && !bIsChildActorTemplate &&
+						    ( (bArchetypesAreUpToDate && OldObject->IsTemplate()) ||
+						      (!bArchetypesAreUpToDate && OldObject->HasAnyFlags(RF_ClassDefaultObject))
+						    )
+						  )
+						)
 					{
 						continue;
 					}
