In FEditorViewportClient::FocusViewportOnBox fix a division by 0 when calculating the Zoom for ortho viewports.

In the case where the ortho viewports have not yet been opened, the X and Y components of ViewportSize are both 0.  We now use DEFAULT_ORTHOZOOM in this case.

diff --git a/Engine/Source/Editor/UnrealEd/Private/EditorViewportClient.cpp b/Engine/Source/Editor/UnrealEd/Private/EditorViewportClient.cpp
--- a/Engine/Source/Editor/UnrealEd/Private/EditorViewportClient.cpp
+++ b/Engine/Source/Editor/UnrealEd/Private/EditorViewportClient.cpp
@@ -806,7 +806,7 @@ void FEditorViewportClient::FocusViewportOnBox( const FBox& BoundingBox, bool bI
 	float Radius = FMath::Max<FVector::FReal>(BoundingBox.GetExtent().Size(), 10.f);
 
 	float AspectToUse = AspectRatio;
-	FIntPoint ViewportSize = Viewport->GetSizeXY();
+	const FIntPoint ViewportSize = Viewport->GetSizeXY();
 	if (!bUseControllingActorViewInfo && ViewportSize.X > 0 && ViewportSize.Y > 0)
 	{
 		AspectToUse = Viewport->GetDesiredAspectRatio();
@@ -860,11 +860,15 @@ void FEditorViewportClient::FocusViewportOnBox( const FBox& BoundingBox, bool bI
 				* size of the viewport.  It then calculates backwards from what the view size should be (The radius of the bounding volume),
 				* to find the new OrthoZoom value for the viewport. The 15.0f is a fudge factor.
 				*/
-				float NewOrthoZoom;
-				uint32 MinAxisSize = (AspectToUse > 1.0f) ? Viewport->GetSizeXY().Y : Viewport->GetSizeXY().X;
-				float Zoom = Radius / (MinAxisSize / 2.0f);
+				float NewOrthoZoom = DEFAULT_ORTHOZOOM;
+				
+				if (ViewportSize.X > 0 && ViewportSize.Y > 0)
+				{
+					uint32 MinAxisSize = (AspectToUse > 1.0f) ? ViewportSize.Y : ViewportSize.X;
+					float Zoom = Radius / (MinAxisSize / 2.0f);
+					NewOrthoZoom = Zoom * (ViewportSize.X * 15.0f);
+				}
 
-				NewOrthoZoom = Zoom * (Viewport->GetSizeXY().X * 15.0f);
 				NewOrthoZoom = FMath::Clamp<float>(NewOrthoZoom, GetMinimumOrthoZoom(), MAX_ORTHOZOOM);
 				ViewTransform.SetOrthoZoom(NewOrthoZoom);
 			}
