Support for configuring the world initialisation values

#rnx

diff --git a/Engine/Source/Editor/UnrealEd/Classes/Commandlets/BaseIteratePackagesCommandlet.h b/Engine/Source/Editor/UnrealEd/Classes/Commandlets/BaseIteratePackagesCommandlet.h
--- a/Engine/Source/Editor/UnrealEd/Classes/Commandlets/BaseIteratePackagesCommandlet.h
+++ b/Engine/Source/Editor/UnrealEd/Classes/Commandlets/BaseIteratePackagesCommandlet.h
@@ -98,7 +98,7 @@ protected:
 	bool bUseWorldPartitionBuilder;
 
 	/** Initialisation for the FScopedEditorWorld if bUseWorldPartitionBuilder is enabled */
-	UWorld::InitializationValues WorldPartitionInitialisationValues;
+	UWorld::InitializationValues WorldInitialisationValues;
 
 	/** Should we build lighting for the packages we are saving? **/
 	//bool bShouldBuildLighting;
diff --git a/Engine/Source/Editor/UnrealEd/Private/Commandlets/BaseIteratePackagesCommandlet.cpp b/Engine/Source/Editor/UnrealEd/Private/Commandlets/BaseIteratePackagesCommandlet.cpp
--- a/Engine/Source/Editor/UnrealEd/Private/Commandlets/BaseIteratePackagesCommandlet.cpp
+++ b/Engine/Source/Editor/UnrealEd/Private/Commandlets/BaseIteratePackagesCommandlet.cpp
@@ -664,7 +664,7 @@ void UBaseIteratePackagesCommandlet::LoadAndSaveOnePackage(const FString& Filena
 						PerformAdditionalOperations(Object, bSavePackage);
 					});
 
-				FScopedEditorWorld ScopeEditorWorld(World, WorldPartitionInitialisationValues);
+				FScopedEditorWorld ScopeEditorWorld(World, WorldInitialisationValues);
 
 				/*TUniquePtr<FLoaderAdapterShape> LoaderAdapterShape;
 				FBox Bounds = FBox(FVector(-HALF_WORLD_MAX, -HALF_WORLD_MAX, -HALF_WORLD_MAX), FVector(HALF_WORLD_MAX, HALF_WORLD_MAX, HALF_WORLD_MAX));
@@ -1060,13 +1060,13 @@ int32 UBaseIteratePackagesCommandlet::Main( const FString& Params )
 	bForceFinishAllCompilationBeforeGC = false;
 
 	bUseWorldPartitionBuilder = false;
-	WorldPartitionInitialisationValues.RequiresHitProxies(false);
-	WorldPartitionInitialisationValues.ShouldSimulatePhysics(false);
-	WorldPartitionInitialisationValues.EnableTraceCollision(false);
-	WorldPartitionInitialisationValues.CreateNavigation(false);
-	WorldPartitionInitialisationValues.CreateAISystem(false);
-	WorldPartitionInitialisationValues.AllowAudioPlayback(false);
-	WorldPartitionInitialisationValues.CreatePhysicsScene(true);
+	WorldInitialisationValues.RequiresHitProxies(false);
+	WorldInitialisationValues.ShouldSimulatePhysics(false);
+	WorldInitialisationValues.EnableTraceCollision(false);
+	WorldInitialisationValues.CreateNavigation(false);
+	WorldInitialisationValues.CreateAISystem(false);
+	WorldInitialisationValues.AllowAudioPlayback(false);
+	WorldInitialisationValues.CreatePhysicsScene(true);
 
 	bKeepPackageGUIDOnSave = Switches.Contains(TEXT("KeepPackageGUIDOnSave"));
 	
