Cooker: CookRequestCluster: Mark that the initial requests into a cluster have been pulled into the cluster. This fixes a bug where an initially requested vertex that is also referenced by other vertices in the cluster is pulled in a second time, and when pulled in clobbers some data we discovered about the vertex such as the CookSuppressionReason if we found it was not cookable. That bug broke the diagnostic rejectionreason for why the package was not cooked.
#rnx

diff --git a/Engine/Source/Editor/UnrealEd/Private/Cooker/CookRequestCluster.cpp b/Engine/Source/Editor/UnrealEd/Private/Cooker/CookRequestCluster.cpp
--- a/Engine/Source/Editor/UnrealEd/Private/Cooker/CookRequestCluster.cpp
+++ b/Engine/Source/Editor/UnrealEd/Private/Cooker/CookRequestCluster.cpp
@@ -684,6 +684,7 @@ void FRequestCluster::FGraphSearch::StartSearch()
 	{
 		FVertexData& Vertex = FindOrAddVertex(Pair.Key->GetPackageName(), *Pair.Key);
 		check(Vertex.PackageData);
+		Vertex.bPulledIntoCluster = true;
 		AddToVisitVertexQueue(Vertex);
 	}
 }
