Ensure that multi-pass package writers like the one used with -iterativevalidate don't stomp over the data in the canonical FSavePackageResultStruct.

diff --git a/Engine/Source/Editor/UnrealEd/Private/CookOnTheFlyServer.cpp b/Engine/Source/Editor/UnrealEd/Private/CookOnTheFlyServer.cpp
--- a/Engine/Source/Editor/UnrealEd/Private/CookOnTheFlyServer.cpp
+++ b/Engine/Source/Editor/UnrealEd/Private/CookOnTheFlyServer.cpp
@@ -6121,7 +6121,9 @@ void UCookOnTheFlyServer::SaveCookedPackage(UE::Cook::FSaveCookedPackageContext&
 			SaveArgs.InOutSaveOverrides = &SaveOverrides;
 
 			Context.PackageWriter->UpdateSaveArguments(SaveArgs);
-			for(;;)
+			FSavePackageResultStruct AuthoritativeResult = ESavePackageResult::Error;
+			bool IsFirstPass = true;
+			for (;;) 
 			{
 				try
 				{
@@ -6143,7 +6145,13 @@ void UCookOnTheFlyServer::SaveCookedPackage(UE::Cook::FSaveCookedPackageContext&
 					Context.SavePackageResult = ESavePackageResult::Error;
 				}
 
-				if (Context.PackageWriter->IsAnotherSaveNeeded(Context.SavePackageResult, SaveArgs))
+				bool IsAnotherSaveNeeded = Context.PackageWriter->IsAnotherSaveNeeded(Context.SavePackageResult, SaveArgs);
+				if (IsFirstPass)
+				{
+					AuthoritativeResult = MoveTemp(Context.SavePackageResult);
+					IsFirstPass = false;
+				}
+				if (IsAnotherSaveNeeded)
 				{
 					// We must not try a second save of a package while the first save is still in flight.
 					// The optimal solution is to wait for ONLY the package that needs a second save, but we don't
@@ -6155,6 +6163,7 @@ void UCookOnTheFlyServer::SaveCookedPackage(UE::Cook::FSaveCookedPackageContext&
 					break;
 				}
 			}
+			Context.SavePackageResult = MoveTemp(AuthoritativeResult);
 
 			// If package was actually saved check with asset manager to make sure it wasn't excluded for being a
 			// development or never cook package. But skip sending the warnings from this check if it was editor-only.
