DiffAssetBulkData now shows a more sane number for new and delete: previously it was only chunks within the shared packages, which was misleading, now it's across all the packages. Also, way faster now through the power of unintentional N^2
#rnx

diff --git a/Engine/Source/Editor/UnrealEd/Private/Commandlets/DiffAssetBulkDataCommandlet.cpp b/Engine/Source/Editor/UnrealEd/Private/Commandlets/DiffAssetBulkDataCommandlet.cpp
--- a/Engine/Source/Editor/UnrealEd/Private/Commandlets/DiffAssetBulkDataCommandlet.cpp
+++ b/Engine/Source/Editor/UnrealEd/Private/Commandlets/DiffAssetBulkDataCommandlet.cpp
@@ -150,17 +150,34 @@ int32 UDiffAssetBulkDataCommandlet::Main(const FString& FullCommandLine)
 	const TMap<FName, const FAssetPackageData*>& BasePackages = BaseState.GetAssetPackageDataMap();
 	const TMap<FName, const FAssetPackageData*>& CurrentPackages = CurrentState.GetAssetPackageDataMap();
 
-	TArray<FName> NewPackages, DeletedPackages, SharedPackages;
+	TArray<FName> NewPackages, DeletedPackages;
+
+	struct FIteratedPackage
+	{
+		FName Name = NAME_None;
+		const FAssetPackageData* Base = nullptr;
+		const FAssetPackageData* Current = nullptr;
+		FIteratedPackage() = default;
+		FIteratedPackage(FName _Name, const FAssetPackageData* _Base, const FAssetPackageData* _Current) :
+			Name(_Name),
+			Base(_Base),
+			Current(_Current) {}
+
+	};
+	TArray<FIteratedPackage> UnionedPackages;
+
 	{
 		for (const TPair<FName, const FAssetPackageData*>& NamePackageDataPair : BasePackages)
 		{
-			if (CurrentState.GetAssetPackageData(NamePackageDataPair.Key) == nullptr)
+			const FAssetPackageData* Current = CurrentState.GetAssetPackageData(NamePackageDataPair.Key);
+
+			UnionedPackages.Emplace(FIteratedPackage(NamePackageDataPair.Key, NamePackageDataPair.Value, Current));
+
+			if (Current == nullptr)
 			{
 				DeletedPackages.Add(NamePackageDataPair.Key);
 				continue;
 			}
-
-			SharedPackages.Add(NamePackageDataPair.Key);
 		}
 
 		for (const TPair<FName, const FAssetPackageData*>& NamePackageDataPair : CurrentPackages)
@@ -168,6 +185,8 @@ int32 UDiffAssetBulkDataCommandlet::Main(const FString& FullCommandLine)
 			if (BaseState.GetAssetPackageData(NamePackageDataPair.Key) == nullptr)
 			{
 				NewPackages.Add(NamePackageDataPair.Key);
+
+				UnionedPackages.Emplace(FIteratedPackage(NamePackageDataPair.Key, nullptr, NamePackageDataPair.Value));
 			}
 		}
 	}
@@ -177,48 +196,60 @@ int32 UDiffAssetBulkDataCommandlet::Main(const FString& FullCommandLine)
 	// This whole thing assumes that the index parameter of CreateIoChunkId is always 0. This is likely not going
 	// to be true with FDerivedData, once that gets turned on, but should be easy to update when the time comes.
 	//
-	TArray<FName> ChangedChunksByType[(uint32)EIoChunkType::MAX], NewChunksByType[(uint32)EIoChunkType::MAX], DeletedChunksByType[(uint32)EIoChunkType::MAX];
-	for (const FName& SharedPackage : SharedPackages)
+	TSet<FName> ChangedChunksByType[(uint32)EIoChunkType::MAX], NewChunksByType[(uint32)EIoChunkType::MAX], DeletedChunksByType[(uint32)EIoChunkType::MAX];
+	for (const FIteratedPackage& IteratedPackage : UnionedPackages)
 	{
-		const FAssetPackageData* BasePackage = BasePackages[SharedPackage];
-		const FAssetPackageData* CurrentPackage = CurrentPackages[SharedPackage];
+		const FAssetPackageData* BasePackage = IteratedPackage.Base;
+		const FAssetPackageData* CurrentPackage = IteratedPackage.Current;
 
-		for (const TPair<FIoChunkId, FIoHash>& ChunkHashPair : BasePackage->ChunkHashes)
+		if (BasePackage)
 		{
-			if (ChunkHashPair.Key.GetChunkType() != EIoChunkType::BulkData &&
-				ChunkHashPair.Key.GetChunkType() != EIoChunkType::OptionalBulkData &&
-				ChunkHashPair.Key.GetChunkType() != EIoChunkType::MemoryMappedBulkData)
-				continue;
-
-			const FIoHash* CurrentHash = CurrentPackage->ChunkHashes.Find(ChunkHashPair.Key);
-			if (CurrentHash == nullptr)
+			for (const TPair<FIoChunkId, FIoHash>& ChunkHashPair : BasePackage->ChunkHashes)
 			{
-				TArray<FName>& Deleted = DeletedChunksByType[(uint32)ChunkHashPair.Key.GetChunkType()];
-				check(Deleted.Contains(SharedPackage) == false); // Because only 0 chunk index
-				Deleted.Add(SharedPackage);
-				continue;
-			}
+				if (ChunkHashPair.Key.GetChunkType() != EIoChunkType::BulkData &&
+					ChunkHashPair.Key.GetChunkType() != EIoChunkType::OptionalBulkData &&
+					ChunkHashPair.Key.GetChunkType() != EIoChunkType::MemoryMappedBulkData)
+					continue;
 
-			if (*CurrentHash != ChunkHashPair.Value)
-			{
-				TArray<FName>& Changed = ChangedChunksByType[(uint32)ChunkHashPair.Key.GetChunkType()];
-				check(Changed.Contains(SharedPackage) == false); // Because only 0 chunk index
-				Changed.Add(SharedPackage);
+				const FIoHash* CurrentHash = nullptr;
+				if (CurrentPackage)
+				{
+					CurrentHash = CurrentPackage->ChunkHashes.Find(ChunkHashPair.Key);
+				}
+
+				if (CurrentHash == nullptr)
+				{
+					TSet<FName>& Deleted = DeletedChunksByType[(uint32)ChunkHashPair.Key.GetChunkType()];
+					check(Deleted.Contains(IteratedPackage.Name) == false); // Because only 0 chunk index
+					Deleted.Add(IteratedPackage.Name);
+					continue;
+				}
+
+				if (*CurrentHash != ChunkHashPair.Value)
+				{
+					TSet<FName>& Changed = ChangedChunksByType[(uint32)ChunkHashPair.Key.GetChunkType()];
+					check(Changed.Contains(IteratedPackage.Name) == false); // Because only 0 chunk index
+					Changed.Add(IteratedPackage.Name);
+				}
 			}
 		}
 
-		for (const TPair<FIoChunkId, FIoHash>& ChunkHashPair : CurrentPackage->ChunkHashes)
+		if (CurrentPackage)
 		{
-			if (ChunkHashPair.Key.GetChunkType() != EIoChunkType::BulkData &&
-				ChunkHashPair.Key.GetChunkType() != EIoChunkType::OptionalBulkData &&
-				ChunkHashPair.Key.GetChunkType() != EIoChunkType::MemoryMappedBulkData)
-				continue;
-
-			if (BasePackage->ChunkHashes.Contains(ChunkHashPair.Key) == false)
+			for (const TPair<FIoChunkId, FIoHash>& ChunkHashPair : CurrentPackage->ChunkHashes)
 			{
-				TArray<FName>& New = NewChunksByType[(uint32)ChunkHashPair.Key.GetChunkType()];
-				check(New.Contains(SharedPackage) == false); // Because only 0 chunk index
-				New.Add(SharedPackage);
+				if (ChunkHashPair.Key.GetChunkType() != EIoChunkType::BulkData &&
+					ChunkHashPair.Key.GetChunkType() != EIoChunkType::OptionalBulkData &&
+					ChunkHashPair.Key.GetChunkType() != EIoChunkType::MemoryMappedBulkData)
+					continue;
+
+				if (!BasePackage ||
+					BasePackage->ChunkHashes.Contains(ChunkHashPair.Key) == false)
+				{
+					TSet<FName>& New = NewChunksByType[(uint32)ChunkHashPair.Key.GetChunkType()];
+					check(New.Contains(IteratedPackage.Name) == false); // Because only 0 chunk index
+					New.Add(IteratedPackage.Name);
+				}
 			}
 		}
 	}
@@ -254,7 +285,7 @@ int32 UDiffAssetBulkDataCommandlet::Main(const FString& FullCommandLine)
 	
 	for (const FName& ChangedPackageName : ChangedPackages)
 	{
-		TConstArrayView<FAssetData const*> BaseAssetDatas = BaseState.GetAssetsByPackageName(ChangedPackageName);		
+		TConstArrayView<FAssetData const*> BaseAssetDatas = BaseState.GetAssetsByPackageName(ChangedPackageName);
 		TConstArrayView<FAssetData const*> CurrentAssetDatas = CurrentState.GetAssetsByPackageName(ChangedPackageName);
 
 		struct FDiffTag
@@ -413,25 +444,26 @@ int32 UDiffAssetBulkDataCommandlet::Main(const FString& FullCommandLine)
 			continue;
 		}
 
-		const TArray<FName>& NewChunksForType = NewChunksByType[ChunkTypeIndex];
-		const TArray<FName>& DeletedChunksForType = DeletedChunksByType[ChunkTypeIndex];
-		const TArray<FName>& ChangedChunksForType = ChangedChunksByType[ChunkTypeIndex];
+		const TSet<FName>& NewChunksForType = NewChunksByType[ChunkTypeIndex];
+		const TSet<FName>& DeletedChunksForType = DeletedChunksByType[ChunkTypeIndex];
+		const TSet<FName>& ChangedChunksForType = ChangedChunksByType[ChunkTypeIndex];
 
 		TotalNewChunks += NewChunksForType.Num();
 		TotalChangedChunks += ChangedChunksForType.Num();
 		TotalDeletedChunks += DeletedChunksForType.Num();
 		UE_LOG(LogDiffAssetBulk, Display, TEXT("    %-20s %10d %10d %10d"), *LexToString(ChunkType), NewChunksForType.Num(), DeletedChunksForType.Num(), ChangedChunksForType.Num());
 	}
-	UE_LOG(LogDiffAssetBulk, Display, TEXT("    =================================================="));
+	UE_LOG(LogDiffAssetBulk, Display, TEXT("    ====================================================="));
 
 	UE_LOG(LogDiffAssetBulk, Display, TEXT("    %-20s %10d %10d %10d"), TEXT("Total"), TotalNewChunks, TotalDeletedChunks, TotalChangedChunks);
 
 	UE_LOG(LogDiffAssetBulk, Display, TEXT(""));
 
-	UE_LOG(LogDiffAssetBulk, Display, TEXT("    Packages Added:     %8d"), NewPackages.Num());
-	UE_LOG(LogDiffAssetBulk, Display, TEXT("    Packages Deleted:   %8d"), DeletedPackages.Num());
-	UE_LOG(LogDiffAssetBulk, Display, TEXT("    Packages Changed:   %8d"), ChangedPackages.Num());
-	UE_LOG(LogDiffAssetBulk, Display, TEXT("    Packages Unmodified:%8d"), SharedPackages.Num() - ChangedPackages.Num());
+	UE_LOG(LogDiffAssetBulk, Display, TEXT("    Base Packages:                %8d"), BasePackages.Num());
+	UE_LOG(LogDiffAssetBulk, Display, TEXT("    Current Packages:             %8d"), CurrentPackages.Num());
+	UE_LOG(LogDiffAssetBulk, Display, TEXT("    Bulk Data Packages Added:     %8d"), NewPackages.Num());
+	UE_LOG(LogDiffAssetBulk, Display, TEXT("    Bulk Data Packages Deleted:   %8d"), DeletedPackages.Num());
+	UE_LOG(LogDiffAssetBulk, Display, TEXT("    Bulk Data Packages Changed:   %8d"), ChangedPackages.Num());
 
 	if (ChangedPackages.Num() == 0)
 	{
@@ -485,16 +517,16 @@ int32 UDiffAssetBulkDataCommandlet::Main(const FString& FullCommandLine)
 		}
 		
 
-		UE_LOG(LogDiffAssetBulk, Warning, TEXT("    Can't determine blame:        : %-7d // Assets had blame tags but all matched - check determinism! -ListDeterminism"), TotalUnassignablePackages);
+		UE_LOG(LogDiffAssetBulk, Display, TEXT("    Can't determine blame:        : %-7d // Assets had blame tags but all matched - check determinism! -ListDeterminism"), TotalUnassignablePackages);
 		for (TPair<FTopLevelAssetPath, TArray<FName>>& ClassPackages : PackagesWithUnassignableDiffsByAssumedClass)
 		{
-			UE_LOG(LogDiffAssetBulk, Warning, TEXT("        %s : %d"), *ClassPackages.Key.ToString(), ClassPackages.Value.Num());
+			UE_LOG(LogDiffAssetBulk, Display, TEXT("        %s : %d"), *ClassPackages.Key.ToString(), ClassPackages.Value.Num());
 			Algo::Sort(ClassPackages.Value, FNameLexicalLess());
 			if (bListDeterminism)
 			{
 				for (const FName& PackageName : ClassPackages.Value)
 				{
-					UE_LOG(LogDiffAssetBulk, Warning, TEXT("            %s"), *PackageName.ToString());
+					UE_LOG(LogDiffAssetBulk, Display, TEXT("            %s"), *PackageName.ToString());
 				}
 			}
 		}
