Interchange: add LOD import Asynchronous parameter to allow synchronous import when the caller cannot support asynchronous.

#rnx

diff --git a/Engine/Plugins/Tests/EditorTests/Source/EditorTests/Private/UnrealEd/FbxAutomationTests.cpp b/Engine/Plugins/Tests/EditorTests/Source/EditorTests/Private/UnrealEd/FbxAutomationTests.cpp
--- a/Engine/Plugins/Tests/EditorTests/Source/EditorTests/Private/UnrealEd/FbxAutomationTests.cpp
+++ b/Engine/Plugins/Tests/EditorTests/Source/EditorTests/Private/UnrealEd/FbxAutomationTests.cpp
@@ -480,6 +480,7 @@ bool FFbxImportAssetsAutomationTest::RunTest(const FString& Parameters)
 					CurTestSuccessful = false;
 					continue;
 				}
+				constexpr bool bAsyncFalse = false;
 
 				if (GlobalImportedObjects[0]->IsA(UStaticMesh::StaticClass()))
 				{
@@ -508,7 +509,7 @@ bool FFbxImportAssetsAutomationTest::RunTest(const FString& Parameters)
 					ImportData->ImportUniformScale = TestPlan->ImportUI->StaticMeshImportData->ImportUniformScale;
 					ImportData->bImportAsScene = TestPlan->ImportUI->StaticMeshImportData->bImportAsScene;
 
-					FbxMeshUtils::ImportStaticMeshLOD(ExistingStaticMesh, LodFile, TestPlan->LodIndex);
+					FbxMeshUtils::ImportStaticMeshLOD(ExistingStaticMesh, LodFile, TestPlan->LodIndex, bAsyncFalse);
 				}
 				else if (GlobalImportedObjects[0]->IsA(USkeletalMesh::StaticClass()))
 				{
@@ -540,7 +541,7 @@ bool FFbxImportAssetsAutomationTest::RunTest(const FString& Parameters)
 					ImportData->ImportUniformScale = TestPlan->ImportUI->SkeletalMeshImportData->ImportUniformScale;
 					ImportData->bImportAsScene = TestPlan->ImportUI->SkeletalMeshImportData->bImportAsScene;
 
-					FbxMeshUtils::ImportSkeletalMeshLOD(ExistingSkeletalMesh, LodFile, TestPlan->LodIndex);
+					FbxMeshUtils::ImportSkeletalMeshLOD(ExistingSkeletalMesh, LodFile, TestPlan->LodIndex, bAsyncFalse);
 				}
 			}
 			break;
diff --git a/Engine/Source/Editor/SkeletalMeshEditor/Private/SkeletalMeshEditorSubsystem.cpp b/Engine/Source/Editor/SkeletalMeshEditor/Private/SkeletalMeshEditorSubsystem.cpp
--- a/Engine/Source/Editor/SkeletalMeshEditor/Private/SkeletalMeshEditorSubsystem.cpp
+++ b/Engine/Source/Editor/SkeletalMeshEditor/Private/SkeletalMeshEditorSubsystem.cpp
@@ -670,8 +670,8 @@ int32 USkeletalMeshEditorSubsystem::ImportLOD(USkeletalMesh* BaseMesh, const int
 
 		BaseMesh->GetLODInfo(LODIndex)->bHasBeenSimplified = false;
 	}
-
-	if (!FbxMeshUtils::ImportSkeletalMeshLOD(BaseMesh, ResolveFilename, LODIndex))
+	constexpr bool bAsyncFalse = false;
+	if (!FbxMeshUtils::ImportSkeletalMeshLOD(BaseMesh, ResolveFilename, LODIndex, bAsyncFalse))
 	{
 		UE_LOG(LogSkeletalMeshEditorSubsystem, Error, TEXT("SkeletalMesh ImportLOD: Cannot import mesh LOD."));
 		return INDEX_NONE;
diff --git a/Engine/Source/Editor/StaticMeshEditor/Private/StaticMeshEditorSubsystem.cpp b/Engine/Source/Editor/StaticMeshEditor/Private/StaticMeshEditorSubsystem.cpp
--- a/Engine/Source/Editor/StaticMeshEditor/Private/StaticMeshEditorSubsystem.cpp
+++ b/Engine/Source/Editor/StaticMeshEditor/Private/StaticMeshEditorSubsystem.cpp
@@ -668,8 +668,8 @@ int32 UStaticMeshEditorSubsystem::ImportLOD(UStaticMesh* BaseStaticMesh, const i
 		return INDEX_NONE;
 	}
 
-
-	if (!FbxMeshUtils::ImportStaticMeshLOD(BaseStaticMesh, ResolveFilename, LODIndex))
+	constexpr bool bAsyncFalse = false;
+	if (!FbxMeshUtils::ImportStaticMeshLOD(BaseStaticMesh, ResolveFilename, LODIndex, bAsyncFalse))
 	{
 		UE_LOG(LogStaticMeshEditorSubsystem, Error, TEXT("StaticMesh ImportLOD: Cannot import mesh LOD."));
 		return INDEX_NONE;
diff --git a/Engine/Source/Editor/UnrealEd/Private/FbxMeshUtils.cpp b/Engine/Source/Editor/UnrealEd/Private/FbxMeshUtils.cpp
--- a/Engine/Source/Editor/UnrealEd/Private/FbxMeshUtils.cpp
+++ b/Engine/Source/Editor/UnrealEd/Private/FbxMeshUtils.cpp
@@ -186,7 +186,7 @@ namespace FbxMeshUtils
 		}
 	}
 
-	bool ImportStaticMeshLOD( UStaticMesh* BaseStaticMesh, const FString& Filename, int32 LODLevel )
+	bool ImportStaticMeshLOD( UStaticMesh* BaseStaticMesh, const FString& Filename, int32 LODLevel, bool bAsync)
 	{
 		if (!BaseStaticMesh)
 		{
@@ -200,10 +200,10 @@ namespace FbxMeshUtils
 		if (bInterchangeCanImportSourceData)
 		{
 			//Call interchange mesh utilities to import custom LOD
-			UInterchangeMeshUtilities::ImportCustomLod(BaseStaticMesh, LODLevel, SourceData).Then([BaseStaticMesh, LODLevel](TFuture<bool> Result)
+			UInterchangeMeshUtilities::ImportCustomLod(BaseStaticMesh, LODLevel, SourceData, bAsync).Then([BaseStaticMesh, LODLevel, bAsync](TFuture<bool> Result)
 				{
 					bool bResult = Result.Get();
-					Async(EAsyncExecution::TaskGraphMainThread, [BaseStaticMesh, LODLevel, bResult]()
+					auto OnImportCustomLodDone = [BaseStaticMesh, LODLevel, bResult]()
 						{
 							if (bResult)
 							{
@@ -221,7 +221,17 @@ namespace FbxMeshUtils
 								NotificationInfo.ExpireDuration = 5.0f;
 								FSlateNotificationManager::Get().AddNotification(NotificationInfo);
 							}
-						});
+						};
+
+					if (IsInGameThread())
+					{
+						OnImportCustomLodDone();
+					}
+					else
+					{
+						ensure(bAsync);
+						Async(EAsyncExecution::TaskGraphMainThread, OnImportCustomLodDone);
+					}
 				});
 
 			return true;
@@ -358,7 +368,7 @@ namespace FbxMeshUtils
 		return bSuccess;
 	}
 
-	bool ImportStaticMeshHiResSourceModel(UStaticMesh* BaseStaticMesh, const FString& Filename)
+	bool ImportStaticMeshHiResSourceModel(UStaticMesh* BaseStaticMesh, const FString& Filename, bool bAsync)
 	{
 		if (!BaseStaticMesh)
 		{
@@ -376,39 +386,48 @@ namespace FbxMeshUtils
 			//Set the asset import data to pass the correct import options
 			TempStaticMesh->SetAssetImportData(BaseStaticMesh->GetAssetImportData());
 			//Call interchange mesh utilities to import custom LOD
-			UInterchangeMeshUtilities::ImportCustomLod(TempStaticMesh, 0, SourceData).Then([BaseStaticMesh,TempStaticMesh, Filename](TFuture<bool> Result)
+			UInterchangeMeshUtilities::ImportCustomLod(TempStaticMesh, 0, SourceData, bAsync).Then([BaseStaticMesh,TempStaticMesh, Filename, bAsync](TFuture<bool> Result)
 			{
 				bool bResult = Result.Get();
-				Async(EAsyncExecution::TaskGraphMainThread, [BaseStaticMesh, TempStaticMesh, Filename, bResult]()
-				{
-					// Copy high res mesh from temporary static mesh to targeted one
-					if (bResult && Private::CopyHighResMeshDescription(TempStaticMesh, BaseStaticMesh))
-					{
-						FStaticMeshSourceModel& SourceModel = BaseStaticMesh->GetHiResSourceModel();
-						SourceModel.SourceImportFilename = UAssetImportData::SanitizeImportFilename(Filename, nullptr);
-						SourceModel.bImportWithBaseMesh = false;
-
-						// Notification of success
-						FNotificationInfo NotificationInfo(FText::GetEmpty());
-						NotificationInfo.Text = NSLOCTEXT("UnrealEd", "ImportStaticMeshHiResSourceModelSuccessful", "High res mesh imported successfully!");
-						NotificationInfo.ExpireDuration = 5.0f;
-						FSlateNotificationManager::Get().AddNotification(NotificationInfo);
-					}
-					else
+				auto OnImportCustomLodDone = [BaseStaticMesh, TempStaticMesh, Filename, bResult]()
 					{
-						// Notification of failure
-						FNotificationInfo NotificationInfo(FText::GetEmpty());
-						NotificationInfo.Text = NSLOCTEXT("UnrealEd", "ImportStaticMeshHiResSourceModelFail", "Failed to import high res mesh!");
-						NotificationInfo.ExpireDuration = 5.0f;
-						FSlateNotificationManager::Get().AddNotification(NotificationInfo);
-					}
+						// Copy high res mesh from temporary static mesh to targeted one
+						if (bResult && Private::CopyHighResMeshDescription(TempStaticMesh, BaseStaticMesh))
+						{
+							FStaticMeshSourceModel& SourceModel = BaseStaticMesh->GetHiResSourceModel();
+							SourceModel.SourceImportFilename = UAssetImportData::SanitizeImportFilename(Filename, nullptr);
+							SourceModel.bImportWithBaseMesh = false;
+
+							// Notification of success
+							FNotificationInfo NotificationInfo(FText::GetEmpty());
+							NotificationInfo.Text = NSLOCTEXT("UnrealEd", "ImportStaticMeshHiResSourceModelSuccessful", "High res mesh imported successfully!");
+							NotificationInfo.ExpireDuration = 5.0f;
+							FSlateNotificationManager::Get().AddNotification(NotificationInfo);
+						}
+						else
+						{
+							// Notification of failure
+							FNotificationInfo NotificationInfo(FText::GetEmpty());
+							NotificationInfo.Text = NSLOCTEXT("UnrealEd", "ImportStaticMeshHiResSourceModelFail", "Failed to import high res mesh!");
+							NotificationInfo.ExpireDuration = 5.0f;
+							FSlateNotificationManager::Get().AddNotification(NotificationInfo);
+						}
 
-					if(TempStaticMesh)
-					{
-						TempStaticMesh->ClearFlags(RF_Public | RF_Standalone);
-						TempStaticMesh->MarkAsGarbage();
-					}
-				});
+						if (TempStaticMesh)
+						{
+							TempStaticMesh->ClearFlags(RF_Public | RF_Standalone);
+							TempStaticMesh->MarkAsGarbage();
+						}
+					};
+				if (IsInGameThread())
+				{
+					OnImportCustomLodDone();
+				}
+				else
+				{
+					ensure(bAsync);
+					Async(EAsyncExecution::TaskGraphMainThread, OnImportCustomLodDone);
+				}
 			});
 
 			return true;
@@ -494,7 +513,7 @@ namespace FbxMeshUtils
 		return bSuccess;
 	}
 
-	bool ImportSkeletalMeshLOD( class USkeletalMesh* SelectedSkelMesh, const FString& Filename, int32 LODLevel)
+	bool ImportSkeletalMeshLOD( class USkeletalMesh* SelectedSkelMesh, const FString& Filename, int32 LODLevel, bool bAsync)
 	{
 		//Make sure skeletal mesh is valid
 		if (!SelectedSkelMesh)
@@ -509,15 +528,15 @@ namespace FbxMeshUtils
 		if (bInterchangeCanImportSourceData)
 		{
 			//Call interchange mesh utilities to import custom LOD
-			UInterchangeMeshUtilities::ImportCustomLod(SelectedSkelMesh, LODLevel, SourceData).Then([SelectedSkelMesh, LODLevel](TFuture<bool> Result)
+			UInterchangeMeshUtilities::ImportCustomLod(SelectedSkelMesh, LODLevel, SourceData, bAsync).Then([SelectedSkelMesh, LODLevel, bAsync](TFuture<bool> Result)
 				{
 					bool bResult = Result.Get();
-					Async(EAsyncExecution::TaskGraphMainThread, [SelectedSkelMesh, LODLevel, bResult]()
+					auto OnImportCustomLodDone = [SelectedSkelMesh, LODLevel, bResult]()
 						{
 							if (bResult)
 							{
 								//If we use alternate skinweight, we must re-import all profile for this LOD
-								if(SelectedSkelMesh && !SelectedSkelMesh->GetSkinWeightProfiles().IsEmpty())
+								if (SelectedSkelMesh && !SelectedSkelMesh->GetSkinWeightProfiles().IsEmpty())
 								{
 									//Enqueue the re-import alternate skinning
 									TSharedPtr<FInterchangeSkeletalMeshAlternateSkinWeightPostImportTask> SkeletalMeshPostImportTask = MakeShared<FInterchangeSkeletalMeshAlternateSkinWeightPostImportTask>(SelectedSkelMesh);
@@ -528,7 +547,7 @@ namespace FbxMeshUtils
 									SkeletalMeshPostImportTask->AddLodToReimportAlternate(LODLevel);
 									UInterchangeManager::GetInterchangeManager().EnqueuePostImportTask(SkeletalMeshPostImportTask);
 								}
-								
+
 
 								// Notification of success
 								FNotificationInfo NotificationInfo(FText::GetEmpty());
@@ -544,7 +563,17 @@ namespace FbxMeshUtils
 								NotificationInfo.ExpireDuration = 5.0f;
 								FSlateNotificationManager::Get().AddNotification(NotificationInfo);
 							}
-						});
+						};
+
+					if (IsInGameThread())
+					{
+						OnImportCustomLodDone();
+					}
+					else
+					{
+						ensure(bAsync);
+						Async(EAsyncExecution::TaskGraphMainThread, OnImportCustomLodDone);
+					}
 				});
 			return true;
 		}
@@ -1064,7 +1093,8 @@ namespace FbxMeshUtils
 
 		if (bInterchangeCanImportSourceData)
 		{
-			TFuture<bool> Result = UInterchangeMeshUtilities::ImportCustomLod(SelectedMesh, LODLevel, SourceData);
+			constexpr bool bAsyncTrue = true;
+			TFuture<bool> Result = UInterchangeMeshUtilities::ImportCustomLod(SelectedMesh, LODLevel, SourceData, bAsyncTrue);
 
 			Result.Then([Promise, bNotifyCB, SkeletalMesh, StaticMesh, LODLevel](TFuture<bool> FutureResult)
 				{
@@ -1110,13 +1140,14 @@ namespace FbxMeshUtils
 		bool bImportSuccess = false;
 		if(!FilenameToImport.IsEmpty())
 		{
+			constexpr bool bAsyncFalse = false;
 			if(SkeletalMesh)
 			{
-				bImportSuccess = ImportSkeletalMeshLOD(SkeletalMesh, FilenameToImport, LODLevel);
+				bImportSuccess = ImportSkeletalMeshLOD(SkeletalMesh, FilenameToImport, LODLevel, bAsyncFalse);
 			}
 			else if(StaticMesh)
 			{
-				bImportSuccess = ImportStaticMeshLOD(StaticMesh, FilenameToImport, LODLevel);
+				bImportSuccess = ImportStaticMeshLOD(StaticMesh, FilenameToImport, LODLevel, bAsyncFalse);
 			}
 		}
 
@@ -1182,10 +1213,11 @@ namespace FbxMeshUtils
 		}
 
 		bool bImportSuccess = false;
+		constexpr bool bAsyncFalse = false;
 
 		if (!FilenameToImport.IsEmpty())
 		{
-			bImportSuccess = ImportStaticMeshHiResSourceModel(StaticMesh, FilenameToImport);
+			bImportSuccess = ImportStaticMeshHiResSourceModel(StaticMesh, FilenameToImport, bAsyncFalse);
 		}
 
 		if (!bImportSuccess && bPromptOnFail)
@@ -1197,7 +1229,7 @@ namespace FbxMeshUtils
 
 			if (FilenameToImport.Len() > 0 && FPaths::FileExists(FilenameToImport))
 			{
-				bImportSuccess = ImportStaticMeshHiResSourceModel(StaticMesh, FilenameToImport);
+				bImportSuccess = ImportStaticMeshHiResSourceModel(StaticMesh, FilenameToImport, bAsyncFalse);
 			}
 		}
 
diff --git a/Engine/Source/Editor/UnrealEd/Public/FbxMeshUtils.h b/Engine/Source/Editor/UnrealEd/Public/FbxMeshUtils.h
--- a/Engine/Source/Editor/UnrealEd/Public/FbxMeshUtils.h
+++ b/Engine/Source/Editor/UnrealEd/Public/FbxMeshUtils.h
@@ -20,7 +20,7 @@ namespace FbxMeshUtils
 	 * @param LODLevel			The level of the lod to import
 	 * @return					Whether or not the import succeeded
 	 */
-	UNREALED_API bool ImportStaticMeshLOD( UStaticMesh* BaseStaticMesh, const FString& Filename, int32 LODLevel );
+	UNREALED_API bool ImportStaticMeshLOD( UStaticMesh* BaseStaticMesh, const FString& Filename, int32 LODLevel, bool bAsync);
 
 	/**
 	 * Imports a mesh as the high res source model of the given static mesh
@@ -29,7 +29,7 @@ namespace FbxMeshUtils
 	 * @param Filename			The filename of the FBX file containing the mesh
 	 * @return					Whether or not the import succeeded
 	 */
-	UNREALED_API bool ImportStaticMeshHiResSourceModel( UStaticMesh* BaseStaticMesh, const FString& Filename );
+	UNREALED_API bool ImportStaticMeshHiResSourceModel( UStaticMesh* BaseStaticMesh, const FString& Filename, bool bAsync);
 
 	/**
 	 * Imports a skeletal mesh LOD to the given skeletal mesh
@@ -39,7 +39,7 @@ namespace FbxMeshUtils
 	 * @param LODLevel			The level of the lod to import
 	 * @return					Whether or not the import succeeded
 	 */
-	UNREALED_API bool ImportSkeletalMeshLOD( USkeletalMesh* Mesh, const FString& Filename, int32 LODLevel );
+	UNREALED_API bool ImportSkeletalMeshLOD( USkeletalMesh* Mesh, const FString& Filename, int32 LODLevel, bool bAsync);
 
 	/**
 	 * Imports a mesh LOD to the given mesh
diff --git a/Engine/Source/Runtime/Interchange/Engine/Private/InterchangeMeshUtilities.cpp b/Engine/Source/Runtime/Interchange/Engine/Private/InterchangeMeshUtilities.cpp
--- a/Engine/Source/Runtime/Interchange/Engine/Private/InterchangeMeshUtilities.cpp
+++ b/Engine/Source/Runtime/Interchange/Engine/Private/InterchangeMeshUtilities.cpp
@@ -99,14 +99,14 @@ bool FInterchangeSkeletalMeshAlternateSkinWeightPostImportTask::AddLodToReimport
 	return true;
 }
 
-TFuture<bool> UInterchangeMeshUtilities::ImportCustomLod(UObject* MeshObject, const int32 LodIndex, const UInterchangeSourceData* SourceData)
+TFuture<bool> UInterchangeMeshUtilities::ImportCustomLod(UObject* MeshObject, const int32 LodIndex, const UInterchangeSourceData* SourceData, bool bAsync)
 {
 	TSharedPtr<TPromise<bool>> Promise = MakeShared<TPromise<bool>>();
 	
-	return InternalImportCustomLod(Promise, MeshObject, LodIndex, SourceData);
+	return InternalImportCustomLod(Promise, MeshObject, LodIndex, SourceData, bAsync);
 }
 
-TFuture<bool> UInterchangeMeshUtilities::InternalImportCustomLod(TSharedPtr<TPromise<bool>> Promise, UObject* MeshObject, const int32 LodIndex, const UInterchangeSourceData* SourceData)
+TFuture<bool> UInterchangeMeshUtilities::InternalImportCustomLod(TSharedPtr<TPromise<bool>> Promise, UObject* MeshObject, const int32 LodIndex, const UInterchangeSourceData* SourceData, bool bAsync)
 {
 #if WITH_EDITOR
 	UInterchangeManager& InterchangeManager = UInterchangeManager::GetInterchangeManager();
@@ -240,7 +240,10 @@ TFuture<bool> UInterchangeMeshUtilities::InternalImportCustomLod(TSharedPtr<TPro
 			}
 		};
 
-	UE::Interchange::FAssetImportResultRef AssetImportResult = InterchangeManager.ImportAssetAsync(ImportAssetPath, SourceData, ImportAssetParameters);
+	UE::Interchange::FAssetImportResultRef AssetImportResult = bAsync 
+		? InterchangeManager.ImportAssetAsync(ImportAssetPath, SourceData, ImportAssetParameters)
+		: InterchangeManager.ImportAssetWithResult(ImportAssetPath, SourceData, ImportAssetParameters);
+
 	FString SourceDataFilename = SourceData->GetFilename();
 	if (SkeletalMesh)
 	{
diff --git a/Engine/Source/Runtime/Interchange/Engine/Public/InterchangeMeshUtilities.h b/Engine/Source/Runtime/Interchange/Engine/Public/InterchangeMeshUtilities.h
--- a/Engine/Source/Runtime/Interchange/Engine/Public/InterchangeMeshUtilities.h
+++ b/Engine/Source/Runtime/Interchange/Engine/Public/InterchangeMeshUtilities.h
@@ -52,12 +52,12 @@ public:
 	 * @Note - Since this is the async version, the SourceData parameter is required.
 	 * @Note - This function will search for an available interchange asset factory that can do the job for the MeshObject class
 	 */
-	static INTERCHANGEENGINE_API TFuture<bool> ImportCustomLod(UObject* MeshObject, const int32 LodIndex, const UInterchangeSourceData* SourceData);
+	static INTERCHANGEENGINE_API TFuture<bool> ImportCustomLod(UObject* MeshObject, const int32 LodIndex, const UInterchangeSourceData* SourceData, bool bAsync);
 
 
 private:
 
-	static TFuture<bool> InternalImportCustomLod(TSharedPtr<TPromise<bool>> Promise, UObject* MeshObject, const int32 LodIndex, const UInterchangeSourceData* SourceData);
+	static TFuture<bool> InternalImportCustomLod(TSharedPtr<TPromise<bool>> Promise, UObject* MeshObject, const int32 LodIndex, const UInterchangeSourceData* SourceData, bool bAsync);
 };
 
 
