Deprecate bOneHeadsetEachProcess and replace with bHMDForPrimaryProcessOnly. The original PR simply replaces the variable but does not take into account to UE deprecation rule, so this is a modified version of the PR.

diff --git a/Engine/Source/Editor/UnrealEd/Classes/Settings/LevelEditorPlaySettings.h b/Engine/Source/Editor/UnrealEd/Classes/Settings/LevelEditorPlaySettings.h
--- a/Engine/Source/Editor/UnrealEd/Classes/Settings/LevelEditorPlaySettings.h
+++ b/Engine/Source/Editor/UnrealEd/Classes/Settings/LevelEditorPlaySettings.h
@@ -425,8 +425,19 @@ private:
 	* The usefullness of this will vary by XR platform.  
 	* The PIE instances may get special -HMDSimulator behavior from an XR plugin, they may successfully make connections to the HMD hardware, their attempt to connect to hardware may be rejected by the runtime.
 	*/
-	UPROPERTY(config, EditAnywhere, Category = "Multiplayer Options", meta = (EditCondition = "!RunUnderOneProcess"))
-	bool bOneHeadsetEachProcess;
+	UE_DEPRECATED(5.5, "Deprecated. Use bHMDForPrimaryProcessOnly instead")
+	UPROPERTY(config)
+	bool bOneHeadsetEachProcess_DEPRECATED;
+
+	/** 
+	* When running multiple player windows in multiple processes, this option determines how HMD(s)/Simulator(s) is/are attached to processes.
+	* 
+	* If checked (default), there is only one HMD/Simulator and it is only attached to the primary process.
+	* 
+	* If unchecked, there could be multiple HMDs/Simulators if HMD/XRSimulator supports it, and each process could have its own attached HMD or XRSimulator, or all processes could share the same HMD by switching the HMD to the active process dynamically.
+	*/
+	UPROPERTY(config, EditAnywhere, Category = "Multiplayer Options", meta = (DisplayName = "HMD For Primary Process Only", EditCondition = "!RunUnderOneProcess"))
+	bool bHMDForPrimaryProcessOnly;
 
 private:
 	UNREALED_API void PushDebugDrawingSettings();
@@ -521,7 +532,15 @@ public:
 		}
 	}
 
-	bool IsOneHeadsetEachProcess() const { return bOneHeadsetEachProcess; }
+	UE_DEPRECATED(5.5, "Deprecated. Use IsHMDForPrimaryProcessOnly function instead")
+	bool IsOneHeadsetEachProcess() const
+	{ 
+		PRAGMA_DISABLE_DEPRECATION_WARNINGS
+		return bOneHeadsetEachProcess_DEPRECATED;
+		PRAGMA_ENABLE_DEPRECATION_WARNINGS
+	}
+
+	bool IsHMDForPrimaryProcessOnly() const { return bHMDForPrimaryProcessOnly; }
 
 public:
 
diff --git a/Engine/Source/Editor/UnrealEd/Private/PlayLevelNewProcess.cpp b/Engine/Source/Editor/UnrealEd/Private/PlayLevelNewProcess.cpp
--- a/Engine/Source/Editor/UnrealEd/Private/PlayLevelNewProcess.cpp
+++ b/Engine/Source/Editor/UnrealEd/Private/PlayLevelNewProcess.cpp
@@ -181,7 +181,7 @@ void UEditorEngine::LaunchNewProcess(const FRequestPlaySessionParams& InParams,
 	// VRPreview handling
 	if (InParams.SessionPreviewTypeOverride.Get(EPlaySessionPreviewType::NoPreview) == EPlaySessionPreviewType::VRPreview)
 	{
-		if (!InParams.EditorPlaySettings->IsOneHeadsetEachProcess())
+		if (InParams.EditorPlaySettings->IsHMDForPrimaryProcessOnly())
 		{
 			// If they're trying to launch a new process (from the editor) in VR, this will fail because the editor
 			// owns the HMD resource, so we warn, and then fall back. They will need to use single-process for VR preview.
diff --git a/Engine/Source/Editor/UnrealEd/Private/Settings/SettingsClasses.cpp b/Engine/Source/Editor/UnrealEd/Private/Settings/SettingsClasses.cpp
--- a/Engine/Source/Editor/UnrealEd/Private/Settings/SettingsClasses.cpp
+++ b/Engine/Source/Editor/UnrealEd/Private/Settings/SettingsClasses.cpp
@@ -513,7 +513,10 @@ ULevelEditorPlaySettings::ULevelEditorPlaySettings( const FObjectInitializer& Ob
 	ServerDebugDrawingColorTintStrength = 0.0f;
 	ServerDebugDrawingColorTint = FLinearColor(0.0f, 0.0f, 0.0f, 1.0f);
 
-	bOneHeadsetEachProcess = false;
+	PRAGMA_DISABLE_DEPRECATION_WARNINGS
+	bOneHeadsetEachProcess_DEPRECATED = false;
+	PRAGMA_ENABLE_DEPRECATION_WARNINGS
+	bHMDForPrimaryProcessOnly = true;
 }
 
 void ULevelEditorPlaySettings::PushDebugDrawingSettings()
@@ -583,6 +586,15 @@ void ULevelEditorPlaySettings::PostInitProperties()
 {
 	Super::PostInitProperties();
 
+	// If the deprecated bOneHeadsetEachProcess setting is on, translate it to bHMDForPrimaryProcessOnly value
+	PRAGMA_DISABLE_DEPRECATION_WARNINGS
+	if (bOneHeadsetEachProcess_DEPRECATED && !GetClass()->HasAnyClassFlags(EClassFlags::CLASS_Abstract))
+	PRAGMA_ENABLE_DEPRECATION_WARNINGS
+	{
+		bHMDForPrimaryProcessOnly = false;
+		UpdateSinglePropertyInConfigFile(GetClass()->FindPropertyByName(GET_MEMBER_NAME_CHECKED(ULevelEditorPlaySettings, bHMDForPrimaryProcessOnly)), GEditorPerProjectIni);
+	};
+
 	NewWindowWidth = FMath::Max(0, NewWindowWidth);
 	NewWindowHeight = FMath::Max(0, NewWindowHeight);
 
