IncrementalCook: Add global guid to invalidate targetdomain keys. Change existing editordomain globalguid to more performant guid.
Bump both editordomain key and targetdomain key, which will invalidate all existing targetdomain; this is not required currently but it is not a problem and would be difficult to avoid.
#rnx

diff --git a/Engine/Source/Editor/UnrealEd/Private/EditorDomain/EditorDomainUtils.cpp b/Engine/Source/Editor/UnrealEd/Private/EditorDomain/EditorDomainUtils.cpp
--- a/Engine/Source/Editor/UnrealEd/Private/EditorDomain/EditorDomainUtils.cpp
+++ b/Engine/Source/Editor/UnrealEd/Private/EditorDomain/EditorDomainUtils.cpp
@@ -30,6 +30,7 @@
 #include "Misc/CoreDelegates.h"
 #include "Misc/DelayedAutoRegister.h"
 #include "Misc/FileHelper.h"
+#include "Misc/Guid.h"
 #include "Misc/PackageAccessTracking.h"
 #include "Misc/PackageAccessTrackingOps.h"
 #include "Misc/PackagePath.h"
@@ -57,6 +58,15 @@
 #include "UObject/UE5MainStreamObjectVersion.h"
 #include "UObject/UObjectIterator.h"
 
+namespace UE::EditorDomain
+{
+
+// Change to a new guid when EditorDomain needs to be invalidated. Note this will also invalidate
+// every TargetDomain key since TargetDomain keys include EditorDomain keys.
+FGuid EditorDomainVersion(TEXT("ED92BBE49F6F4F2E94E9D8BC10AD59A7"));
+
+}
+
 #if !defined(EDITORDOMAINTIMEPROFILERTRACE_ENABLED)
 #if UE_TRACE_ENABLED && !UE_BUILD_SHIPPING
 #define EDITORDOMAINTIMEPROFILERTRACE_ENABLED 1
@@ -219,9 +229,6 @@ bool bGUtilsCookInitialized = false;
 FBlake3Hash GGlobalConstructClassesHash;
 int64 GMaxBulkDataSize = -1;
 
-// Change to a new guid when EditorDomain needs to be invalidated
-const TCHAR* EditorDomainVersion = TEXT("4132358BA4F34EFA8294F50D76F1C94F");
-
 // Identifier of the CacheBuckets for EditorDomain tables
 const TCHAR* EditorDomainPackageBucketName = TEXT("EditorDomainPackage");
 const TCHAR* BulkDataListBucketName = TEXT("BulkDataList");
@@ -344,7 +351,7 @@ FPackageDigest CalculatePackageDigest(const FAssetPackageData& PackageData, FNam
 	FBlake3 Writer;
 	FStringView ProjectName(FApp::GetProjectName());
 	Writer.Update(ProjectName.GetData(), ProjectName.Len() * sizeof(ProjectName[0]));
-	Writer.Update(EditorDomainVersion, FCString::Strlen(EditorDomainVersion)*sizeof(EditorDomainVersion[0]));
+	Writer.Update(&EditorDomainVersion, sizeof(EditorDomainVersion));
 	uint8 EditorDomainSaveUnversioned = GetEditorDomainSaveUnversioned() ? 1 : 0;
 	Writer.Update(&EditorDomainSaveUnversioned, sizeof(EditorDomainSaveUnversioned));
 	Writer.Update(&PackageData.GetPackageSavedHash().GetBytes(), sizeof(PackageData.GetPackageSavedHash().GetBytes()));
diff --git a/Engine/Source/Editor/UnrealEd/Private/TargetDomain/TargetDomainUtils.cpp b/Engine/Source/Editor/UnrealEd/Private/TargetDomain/TargetDomainUtils.cpp
--- a/Engine/Source/Editor/UnrealEd/Private/TargetDomain/TargetDomainUtils.cpp
+++ b/Engine/Source/Editor/UnrealEd/Private/TargetDomain/TargetDomainUtils.cpp
@@ -25,6 +25,7 @@
 #include "IO/IoDispatcher.h"
 #include "IO/IoHash.h"
 #include "Misc/App.h"
+#include "Misc/Guid.h"
 #include "Misc/ScopeRWLock.h"
 #include "Misc/StringBuilder.h"
 #include "Serialization/CompactBinary.h"
@@ -37,7 +38,13 @@
 namespace UE::TargetDomain
 {
 
+// Change TargetDomainVersion to a new guid when all TargetDomain keys in an incremental cook need to be invalidated
+FGuid TargetDomainVersion(TEXT("C9B0281696234067A3A888CEEAAA50F9"));
+
+// Bump CookDependenciesVersion version when the serialization of CookDependencies has changed and we want to add
+// backwards compatibility rather than invalidating everything.
 constexpr uint32 CookDependenciesVersion = 0x00000002;
+
 static const FUtf8StringView CookDependenciesAttachmentKey = UTF8TEXTVIEW("CookDependencies");
 static const FUtf8StringView BuildDefinitionsAttachmentKey = UTF8TEXTVIEW("BuildDefinitionsAttachmentKey");
 /**
@@ -143,6 +150,8 @@ bool FCookDependencies::TryCalculateCurrentKey(const FAssetPackageData* Override
 		return false;
 	}
 	FBlake3 KeyBuilder;
+	KeyBuilder.Update(&TargetDomainVersion, sizeof(TargetDomainVersion));
+
 	UE::EditorDomain::FPackageDigest PackageDigest = OverrideAssetPackageData ?
 		UE::EditorDomain::CalculatePackageDigest(*OverrideAssetPackageData, PackageName) :
 		EditorDomain->GetPackageDigest(PackageName);
