IncrementalCook: Change name of -fullcook parameter to -forcerecook, to reduce confusion between full as in non-incremental and full as in every package in the depot.
#rnx

diff --git a/Engine/Source/Editor/UnrealEd/Private/CookOnTheFlyServer.cpp b/Engine/Source/Editor/UnrealEd/Private/CookOnTheFlyServer.cpp
--- a/Engine/Source/Editor/UnrealEd/Private/CookOnTheFlyServer.cpp
+++ b/Engine/Source/Editor/UnrealEd/Private/CookOnTheFlyServer.cpp
@@ -11005,7 +11005,8 @@ void UCookOnTheFlyServer::LoadBeginCookIterativeFlags(FBeginCookContext& BeginCo
 void UCookOnTheFlyServer::LoadBeginCookIterativeFlagsLocal(FBeginCookContext& BeginContext) const
 {
 	const bool bIsDiffOnly = FParse::Param(FCommandLine::Get(), TEXT("DIFFONLY"));
-	const bool bIterative = !FParse::Param(FCommandLine::Get(), TEXT("fullcook")) && (bHybridIterativeEnabled || IsCookFlagSet(ECookInitializationFlags::Iterative));
+	bool bForceRecook = FParse::Param(FCommandLine::Get(), TEXT("fullcook")) || FParse::Param(FCommandLine::Get(), TEXT("forcerecook"));
+	const bool bIterative = !bForceRecook && (bHybridIterativeEnabled || IsCookFlagSet(ECookInitializationFlags::Iterative));
 	const bool bIsSharedIterativeCook = IsCookFlagSet(ECookInitializationFlags::IterateSharedBuild);
 
 	for (FBeginCookContextPlatform& PlatformContext : BeginContext.PlatformContexts)
@@ -11063,11 +11064,20 @@ void UCookOnTheFlyServer::LoadBeginCookIterativeFlagsLocal(FBeginCookContext& Be
 			bool bIterativeAllowed = true;
 			if (!bIterative && !PlatformData->bIsSandboxInitialized)
 			{
-				UE_LOG(LogCook, Display,
-					TEXT("FULL COOK: Neither -iterative nor -cookincremental were specified. Deleting previously cooked packages for platform %s and executing a full cook."),
-					*TargetPlatform->PlatformName());
+				if (bForceRecook)
+				{
+					UE_LOG(LogCook, Display,
+						TEXT("FULL COOK: -forcerecook was specified. Deleting previously cooked packages for platform %s and recooking all packages discovered in the current cook."),
+						*TargetPlatform->PlatformName());
+				}
+				else
+				{
+					UE_LOG(LogCook, Display,
+						TEXT("FULL COOK: Neither -iterative nor -cookincremental were specified. Deleting previously cooked packages for platform %s and recooking all packages discovered in the current cook."),
+						*TargetPlatform->PlatformName());
+				}
 				bIterativeAllowed = false;
-				if (bRunningAsShaderServer)
+				if (!bForceRecook && bRunningAsShaderServer)
 				{
 					UE_LOG(LogCook, Display,
 						TEXT("'-odsc' was passed on commandline, but '-iterative' was not, so the cooker as a side effect is clearing cook results. The build will need to be recooked before it can be staged. Add the commandline argument '-iterative' to avoid this unnecessary clear.")
@@ -11077,7 +11087,7 @@ void UCookOnTheFlyServer::LoadBeginCookIterativeFlagsLocal(FBeginCookContext& Be
 			else if (!ArePreviousCookSettingsCompatible(PlatformContext.CurrentCookSettings, TargetPlatform))
 			{
 				UE_LOG(LogCook, Display,
-					TEXT("FULL COOK: %s was specified, but global settings have changed and all previously cook packages are invalidated. Deleting previously cooked packages for platform %s and executing a full cook."),
+					TEXT("FULL COOK: %s was specified, but global settings have changed and all previously cook packages are invalidated. Deleting previously cooked packages for platform %s and recooking all packages discovered in the current cook."),
 					bHybridIterativeEnabled ? TEXT("-cookincremental") : TEXT("-iterative"),
 					*TargetPlatform->PlatformName());
 				if (bRunningAsShaderServer)
