Allow showing dialog non-modally and relax requirement on result type being default-constructable.

diff --git a/Engine/Source/Editor/UnrealEd/Private/Dialogs/Dialogs.cpp b/Engine/Source/Editor/UnrealEd/Private/Dialogs/Dialogs.cpp
--- a/Engine/Source/Editor/UnrealEd/Private/Dialogs/Dialogs.cpp
+++ b/Engine/Source/Editor/UnrealEd/Private/Dialogs/Dialogs.cpp
@@ -5,6 +5,7 @@
 
 #include "Dialog/DialogUtils.h"
 #include "Dialogs/DialogsPrivate.h"
+#include "Interfaces/IMainFrameModule.h"
 #include "Misc/App.h"
 #include "Misc/AssertionMacros.h"
 #include "Misc/MessageDialog.h"
@@ -1153,4 +1154,17 @@ void UE::Private::ShowModalDialogWindow(TSharedRef<SWindow> Window)
 	GEditor->EditorAddModalWindow(Window);
 }
 
+void UE::Private::ShowNonModalDialogWindow(TSharedRef<SWindow> Window)
+{
+	IMainFrameModule& MainFrameModule = FModuleManager::LoadModuleChecked<IMainFrameModule>(TEXT("MainFrame"));
+	if (MainFrameModule.GetParentWindow().IsValid())
+	{
+		FSlateApplication::Get().AddWindowAsNativeChild(Window, MainFrameModule.GetParentWindow().ToSharedRef());
+	}
+	else
+	{
+		FSlateApplication::Get().AddWindow(Window);
+	}
+}
+
 #undef LOCTEXT_NAMESPACE 
diff --git a/Engine/Source/Editor/UnrealEd/Public/Dialogs/Dialogs.h b/Engine/Source/Editor/UnrealEd/Public/Dialogs/Dialogs.h
--- a/Engine/Source/Editor/UnrealEd/Public/Dialogs/Dialogs.h
+++ b/Engine/Source/Editor/UnrealEd/Public/Dialogs/Dialogs.h
@@ -255,6 +255,7 @@ namespace UE::Private
 {
 UNREALED_API TSharedRef<SWindow> CreateModalDialogWindow(const FText& InTitle, TSharedRef<SWidget> Contents, ESizingRule Sizing, FVector2D MinDimensions);
 UNREALED_API void ShowModalDialogWindow(TSharedRef<SWindow> Window);
+UNREALED_API void ShowNonModalDialogWindow(TSharedRef<SWindow> Window);
 } // namespace UE::Private
 
 /**
@@ -271,24 +272,39 @@ public:
 
 	ResultType ShowModalDialog(const FText& InTitle)
 	{
-		static_assert(std::is_default_constructible_v<ResultType>, "ResultType must be default constructable");
 		Window = UE::Private::CreateModalDialogWindow(InTitle, AsShared(), Sizing, MinDimensions);
 		Window->SetWidgetToFocusOnActivate(GetWidgetToFocusOnActivate());
-		ResultType Result;
-		ResultPointer = &Result;
+		TOptional<ResultType> Result;
+		OnFinished.BindLambda([&Result](ResultType InResult) { Result.Emplace(MoveTemp(InResult)); });
 		UE::Private::ShowModalDialogWindow(Window.ToSharedRef());
-		Window.Reset();
-		ResultPointer = nullptr;
-		return MoveTemp(Result);
+		return MoveTemp(Result.GetValue());
+	}
+
+	void ShowNonModal(const FText& InTitle, TDelegate<void(ResultType)> InOnFinished)
+	{
+		Window = UE::Private::CreateModalDialogWindow(InTitle, AsShared(), Sizing, MinDimensions);
+		Window->SetWidgetToFocusOnActivate(GetWidgetToFocusOnActivate());
+		OnFinished = InOnFinished;
+		UE::Private::ShowNonModalDialogWindow(Window.ToSharedRef());
+	}
+
+	template<typename CallbackType>
+	void ShowNonModal(const FText& InTitle, CallbackType InOnFinished)
+	{
+		Window = UE::Private::CreateModalDialogWindow(InTitle, AsShared(), Sizing, MinDimensions);
+		Window->SetWidgetToFocusOnActivate(GetWidgetToFocusOnActivate());
+		OnFinished.BindLambda(InOnFinished);
+		UE::Private::ShowNonModalDialogWindow(Window.ToSharedRef());
 	}
 
 protected:
 	// Derived classes call this function from their widget events to close the dialog and return the result to the calling context
 	void ProvideResult(ResultType InResult)
 	{
-		// Close owning window and move result into space where ShowDialog can return it
-		*ResultPointer = MoveTemp(InResult);
 		Window->RequestDestroyWindow();
+		Window.Reset();
+		OnFinished.ExecuteIfBound(MoveTemp(InResult));
+		OnFinished = {};
 	}
 	
 	virtual TSharedPtr<SWidget> GetWidgetToFocusOnActivate() 
@@ -302,7 +318,7 @@ protected:
 
 private:
 	TSharedPtr<SWindow> Window;
-	ResultType* ResultPointer = nullptr;
+	TDelegate<void(ResultType)> OnFinished;
 };
 
 UE_DEPRECATED(4.26, "Creating groups (nested packages) is no longer supported. Use PromptUserIfExistingObject overload that does not take the Group paramater.")
