ConfigAccessTracking: Move FConfigAccessData into Core so it can be used by FCookDependency.
#rnx

diff --git a/Engine/Source/Editor/UnrealEd/Private/Cooker/CookConfigAccessTracker.cpp b/Engine/Source/Editor/UnrealEd/Private/Cooker/CookConfigAccessTracker.cpp
--- a/Engine/Source/Editor/UnrealEd/Private/Cooker/CookConfigAccessTracker.cpp
+++ b/Engine/Source/Editor/UnrealEd/Private/Cooker/CookConfigAccessTracker.cpp
@@ -17,90 +17,6 @@
 
 #endif
 
-namespace UE::ConfigAccessTracking
-{
-
-void EscapeConfigTrackingTokenToString(FName Token, FStringBuilderBase& Result)
-{
-	Result.Reset();
-	EscapeConfigTrackingTokenAppendString(Token, Result);
-}
-
-void EscapeConfigTrackingTokenAppendString(FName Token, FStringBuilderBase& Result)
-{
-	int32 InitialLength = Result.Len();
-	Result << Token;
-	FStringView AddedView = Result.ToView().RightChop(InitialLength);
-	if (AddedView.Contains(TEXTVIEW(":")))
-	{
-		FString ReplaceText(AddedView);
-		ReplaceText.ReplaceInline(TEXT(":"), TEXT("::"));
-		Result.RemoveSuffix(AddedView.Len());
-		Result << ReplaceText;
-	}
-}
-
-bool TryTokenizeConfigTrackingString(FStringView Text, TArrayView<FStringBuilderBase*> OutTokens)
-{
-	int32 TextLen = Text.Len();
-	if (TextLen == 0)
-	{
-		for (FStringBuilderBase* Token : OutTokens) Token->Reset();
-		return false;
-	}
-	const TCHAR* TextData = Text.GetData();
-
-	int32 NumTokens = OutTokens.Num();
-	check(NumTokens > 0);
-	int32 NextTokenIndex = 0;
-	FStringBuilderBase* OutToken = OutTokens[NextTokenIndex++];
-	OutToken->Reset();
-	for (int32 Index = 0; Index < TextLen; ++Index)
-	{
-		TCHAR C = TextData[Index];
-		if (C != ':')
-		{
-			OutToken->AppendChar(C);
-		}
-		else if (Index < TextLen - 1 && TextData[Index + 1] == ':')
-		{
-			++Index;
-			OutToken->AppendChar(':');
-		}
-		else
-		{
-			if (OutToken->Len() == 0)
-			{
-				// An empty token and therefore the string is invalid; abandon anything that follows it
-				while (NextTokenIndex < NumTokens) OutTokens[NextTokenIndex++]->Reset();
-				return false;
-			}
-			if (NextTokenIndex >= NumTokens)
-			{
-				// Too many tokens
-				return false;
-			}
-			OutToken = OutTokens[NextTokenIndex++];
-			OutToken->Reset();
-		}
-	}
-	if (OutToken->Len() == 0)
-	{
-		// Empty token
-		while (NextTokenIndex < NumTokens) OutTokens[NextTokenIndex++]->Reset();
-		return false;
-	}
-	if (NextTokenIndex < NumTokens)
-	{
-		// too few tokens
-		while (NextTokenIndex < NumTokens) OutTokens[NextTokenIndex++]->Reset();
-		return false;
-	}
-	return true;
-}
-
-} // namespace UE::ConfigAccessTracking
-
 #if UE_WITH_CONFIG_TRACKING
 
 DEFINE_LOG_CATEGORY_STATIC(LogConfigBuildDependencyTracker, Log, All);
@@ -110,144 +26,6 @@ namespace UE::ConfigAccessTracking
 
 FCookConfigAccessTracker FCookConfigAccessTracker::Singleton;
 
-FConfigAccessData::FConfigAccessData(ELoadType InLoadType, FNameEntryId InConfigPlatform, FNameEntryId InFileName)
-	: ConfigPlatform(InConfigPlatform)
-	, FileName(InFileName)
-	, LoadType(InLoadType)
-{
-}
-
-FConfigAccessData::FConfigAccessData(ELoadType InLoadType, FNameEntryId InConfigPlatform, FNameEntryId InFileName,
-	FNameEntryId InSectionName, FMinimalName InValueName, const ITargetPlatform* InRequestingPlatform)
-	: ConfigPlatform(InConfigPlatform)
-	, FileName(InFileName)
-	, SectionName(InSectionName)
-	, ValueName(InValueName)
-	, RequestingPlatform(InRequestingPlatform)
-	, LoadType(InLoadType)
-{
-}
-
-FConfigAccessData FConfigAccessData::GetFileOnlyData() const
-{
-	return FConfigAccessData(LoadType, ConfigPlatform, FileName);
-}
-
-FConfigAccessData FConfigAccessData::GetPathOnlyData() const
-{
-	return FConfigAccessData(LoadType, ConfigPlatform, FileName, SectionName, ValueName, nullptr);
-}
-
-FString FConfigAccessData::FullPathToString() const
-{
-	TStringBuilder<256> Out;
-	AppendFullPath(Out);
-	return FString(*Out);
-}
-
-void FConfigAccessData::AppendFullPath(FStringBuilderBase& Out) const
-{
-	if (LoadType == ELoadType::Uninitialized || FileName.IsNone())
-	{
-		Out << TEXTVIEW("<Invalid>");
-		return;
-	}
-
-	Out << LexToString(LoadType) << TEXTVIEW(".");
-	if (ConfigPlatform.IsNone())
-	{
-		Out << PlatformAgnosticName;
-	}
-	else
-	{
-		EscapeConfigTrackingTokenAppendString(GetConfigPlatform(), Out);
-	}
-	Out << TEXTVIEW(".");
-	EscapeConfigTrackingTokenAppendString(GetFileName(), Out);
-	if (!SectionName.IsNone())
-	{
-		Out << TEXTVIEW(":[");
-		EscapeConfigTrackingTokenAppendString(GetSectionName(), Out);
-		Out << TEXTVIEW("]");
-		if (!ValueName.IsNone())
-		{
-			Out << TEXTVIEW(":");
-			EscapeConfigTrackingTokenAppendString(GetValueName(), Out);
-		}
-	}
-}
-
-FConfigAccessData FConfigAccessData::Parse(FStringView Text)
-{
-	// ConfigSystem.<Editor>.../../../Engine/Config/ConsoleVariables.ini:Section:ValueName
-	//   -> "ConfigSystem", "<Editor>", "../../../Engine/Config/ConsoleVariables.ini", "Section", "ValueName"
-	// No tokens can contain a single colon, but they might contain double colon, which is the escape code for a single colon.
-	// 3rd token might have dots, first two cannot.
-
-	TStringBuilder<128> FullFilePath;
-	TStringBuilder<64> SectionNameStr;
-	TStringBuilder<64> ValueNameStr;
-	FStringBuilderBase* TokenBuffer[] = { &FullFilePath, &SectionNameStr, &ValueNameStr };
-	TArrayView<FStringBuilderBase*> ColonDelimitedTokens = TokenBuffer;
-	TryTokenizeConfigTrackingString(Text, ColonDelimitedTokens);
-	if (FullFilePath.Len() == 0)
-	{
-		return FConfigAccessData();
-	}
-
-	FStringView FullFilePathTokens[3];
-	FullFilePathTokens[0] = FullFilePath;
-	int32 NumFilePathTokens = 1;
-	for (NumFilePathTokens = 1; NumFilePathTokens < UE_ARRAY_COUNT(FullFilePathTokens); ++NumFilePathTokens)
-	{
-		FStringView& Current = FullFilePathTokens[NumFilePathTokens - 1];
-		FStringView& Next = FullFilePathTokens[NumFilePathTokens];
-		int32 DotIndex = Current.Find(TEXTVIEW("."));
-		if (DotIndex == INDEX_NONE)
-		{
-			break;
-		}
-		Next = Current.RightChop(DotIndex + 1);
-		Current.LeftInline(DotIndex);
-		if (Next.IsEmpty())
-		{
-			// break before incrementing NumFilePathTokens so that we do not count the empty Next as a token
-			break;
-		}
-	}
-	// FullFilePath is of the form 
-	//   LoadType.Platform.ConfigName
-	// Platform is <Editor> if the configfile was an editor config file rather than a platform-specific config file
-	if (NumFilePathTokens < 3)
-	{
-		// Not a valid FConfigAccessData text; we always require all 3 of LoadType,Platform,ConfigName
-		return FConfigAccessData();
-	}
-
-	FConfigAccessData Result;
-	LexFromString(Result.LoadType, FullFilePathTokens[0]);
-	if (Result.LoadType == ELoadType::Uninitialized)
-	{
-		return FConfigAccessData();
-	}
-
-	FName LocalConfigPlatform = FullFilePathTokens[1] == PlatformAgnosticName
-		? NAME_None : FName(FullFilePathTokens[1]);
-	FStringView SectionNameView(SectionNameStr);
-	// We write out sectionnames with [] for readability; remove them if they exist.
-	if (SectionNameView.StartsWith('[')) SectionNameView.RightChopInline(1);
-	if (SectionNameView.EndsWith(']')) SectionNameView.LeftChopInline(1);
-	FName LocalSectionName = SectionNameView.Len() == 0 ? NAME_None : FName(SectionNameView, NAME_NO_NUMBER);
-	FName LocalValueName = ValueNameStr.Len() == 0 ? NAME_None : FName(ValueNameStr);
-
-	Result.ConfigPlatform = LocalConfigPlatform.GetComparisonIndex();
-	Result.FileName = FName(FullFilePathTokens[2]).GetComparisonIndex();
-	Result.SectionName = LocalSectionName.GetComparisonIndex();
-	Result.ValueName = FMinimalName(LocalValueName);
-
-	return Result;
-}
-
 TStringBuilder<256> ToString(FNameEntryId FileName, FNameEntryId SectionName, FMinimalName ValueName)
 {
 	return TStringBuilder<256>(InPlace, FileName, TEXT(":["), SectionName, TEXT("]:"), FName(ValueName));
@@ -692,92 +470,6 @@ bool IsLoadableLoadType(ELoadType LoadType)
 
 } // namespace UE::ConfigAccessTracking
 
-const TCHAR* LexToString(UE::ConfigAccessTracking::ELoadType LoadType)
-{
-	using namespace UE::ConfigAccessTracking;
-
-	switch (LoadType)
-	{
-	case ELoadType::ConfigSystem: return TEXT("ConfigSystem");
-	case ELoadType::LocalIniFile: return TEXT("LocalIniFile");
-	case ELoadType::LocalSingleIniFile: return TEXT("LocalSingleIniFile");
-	case ELoadType::ExternalIniFile: return TEXT("ExternalIniFile");
-	case ELoadType::ExternalSingleIniFile: return TEXT("ExternalSingleIniFile");
-	case ELoadType::Manual: return TEXT("Manual");
-	case ELoadType::SuppressReporting: return TEXT("SuppressReporting");
-	case ELoadType::Uninitialized:
-	default:
-		return TEXT("Uninitialized");
-	}
-}
-
-void LexFromString(UE::ConfigAccessTracking::ELoadType& OutLoadType, FStringView Text)
-{
-	using namespace UE::ConfigAccessTracking;
-
-	if (Text.IsEmpty())
-	{
-		OutLoadType = ELoadType::Uninitialized;
-		return;
-	}
-	switch (Text[0])
-	{
-	case 'C':
-	case 'c':
-		if (Text.Equals(TEXT("ConfigSystem"), ESearchCase::IgnoreCase))
-		{
-			OutLoadType = ELoadType::ConfigSystem;
-			return;
-		}
-		break;
-	case 'L':
-	case 'l':
-		if (Text.Equals(TEXT("LocalIniFile"), ESearchCase::IgnoreCase))
-		{
-			OutLoadType = ELoadType::LocalIniFile;
-			return;
-		}
-		if (Text.Equals(TEXT("LocalSingleIniFile"), ESearchCase::IgnoreCase))
-		{
-			OutLoadType = ELoadType::LocalSingleIniFile;
-			return;
-		}
-		break;
-	case 'E':
-	case 'e':
-		if (Text.Equals(TEXT("ExternalIniFile"), ESearchCase::IgnoreCase))
-		{
-			OutLoadType = ELoadType::ExternalIniFile;
-			return;
-		}
-		if (Text.Equals(TEXT("ExternalSingleIniFile"), ESearchCase::IgnoreCase))
-		{
-			OutLoadType = ELoadType::ExternalSingleIniFile;
-			return;
-		}
-		break;
-	case 'S':
-	case 's':
-		if (Text.Equals(TEXT("SuppressReporting"), ESearchCase::IgnoreCase))
-		{
-			OutLoadType = ELoadType::SuppressReporting;
-			return;
-		}
-		break;
-	case 'M':
-	case 'm':
-		if (Text.Equals(TEXT("Manual"), ESearchCase::IgnoreCase))
-		{
-			OutLoadType = ELoadType::Manual;
-			return;
-		}
-		break;
-	default:
-		break;
-	}
-	OutLoadType = ELoadType::Uninitialized;
-}
-
 namespace UE::ConfigAccessTracking
 {
 
diff --git a/Engine/Source/Editor/UnrealEd/Private/Cooker/CookConfigAccessTracker.h b/Engine/Source/Editor/UnrealEd/Private/Cooker/CookConfigAccessTracker.h
--- a/Engine/Source/Editor/UnrealEd/Private/Cooker/CookConfigAccessTracker.h
+++ b/Engine/Source/Editor/UnrealEd/Private/Cooker/CookConfigAccessTracker.h
@@ -16,68 +16,18 @@
 #include "Containers/Map.h"
 #include "Containers/Set.h"
 #include "Interfaces/ITargetPlatform.h"
+#include "Misc/ConfigAccessData.h"
 #include "Misc/ConfigCacheIni.h"
 #include "Misc/PackageAccessTracking.h"
 #include "Templates/TypeHash.h"
 #include "Templates/UnrealTemplate.h"
 #endif
 
-namespace UE::ConfigAccessTracking
-{
-
-void EscapeConfigTrackingTokenToString(FName Token, FStringBuilderBase& Result);
-void EscapeConfigTrackingTokenAppendString(FName Token, FStringBuilderBase& Result);
-bool TryTokenizeConfigTrackingString(FStringView Text, TArrayView<FStringBuilderBase*> OutTokens);
-
-} // namespace UE::ConfigAccessTracking
-
 #if UE_WITH_CONFIG_TRACKING
 
-namespace UE::ConfigAccessTracking { struct FConfigAccessData; }
-
 namespace UE::ConfigAccessTracking
 {
 
-constexpr FStringView PlatformAgnosticName = TEXTVIEW("<Editor>");
-
-/** Full path of a FConfigValue that was reported read. */
-struct FConfigAccessData
-{
-	FNameEntryId ConfigPlatform;
-	FNameEntryId FileName;
-	FNameEntryId SectionName;
-	FMinimalName ValueName;
-	const ITargetPlatform* RequestingPlatform = nullptr;
-	ELoadType LoadType = ELoadType::Uninitialized;
-
-	FConfigAccessData() = default;
-	FConfigAccessData(ELoadType InLoadType, FNameEntryId InConfigPlatform, FNameEntryId InFileName,
-		FNameEntryId InSectionName, FMinimalName InValueName, const ITargetPlatform* InRequestingPlatform);
-	FConfigAccessData(ELoadType InLoadType, FNameEntryId InConfigPlatform, FNameEntryId InFileName);
-
-	FName GetConfigPlatform() const { return FName(ConfigPlatform, ConfigPlatform, NAME_NO_NUMBER_INTERNAL); }
-	FName GetFileName() const { return FName(FileName, FileName, NAME_NO_NUMBER_INTERNAL); }
-	FName GetSectionName() const { return FName(SectionName, SectionName, NAME_NO_NUMBER_INTERNAL); }
-	FName GetValueName() const { return FName(ValueName); }
-
-	FConfigAccessData GetFileOnlyData() const;
-	FConfigAccessData GetPathOnlyData() const;
-	FString FullPathToString() const;
-	void AppendFullPath(FStringBuilderBase& Out) const;
-
-	/**
-	 * "ConfigSystem.<Editor>.../../../Engine/Config/ConsoleVariables.ini:[Section]:Value"
-	 *   -> "ConfigSystem", "<Editor>", "../../../Engine/Config/ConsoleVariables.ini", "Section", "Value"
-	 */
-	static FConfigAccessData Parse(FStringView Text);
-
-	friend uint32 GetTypeHash(const FConfigAccessData& Data);
-	bool IsSameConfigFile(const FConfigAccessData& Other) const;
-	bool operator==(const FConfigAccessData& Other) const;
-	bool operator!=(const FConfigAccessData& Other) const;
-	bool operator<(const FConfigAccessData& Other) const;
-};
-
 /**
  * Tracker that subscribes to AddConfigValueReadCallback and for each access records the access associated with
  * the package that is currently in scope according to PackageAccessTracking_Private.
@@ -159,12 +109,6 @@ bool IsLoadableLoadType(ELoadType LoadType);
 
 } // namespace UE::ConfigAccessTracking
 
-/** Convert ELoadType -> text */
-const TCHAR* LexToString(UE::ConfigAccessTracking::ELoadType LoadType);
-
-/** Convert text -> ELoadType */
-void LexFromString(UE::ConfigAccessTracking::ELoadType& OutLoadType, FStringView Text);
-
 namespace UE::ConfigAccessTracking
 {
 
@@ -184,56 +128,6 @@ public:
 
 }
 
-namespace UE::ConfigAccessTracking
-{
-
-inline uint32 GetTypeHash(const UE::ConfigAccessTracking::FConfigAccessData& Data)
-{
-	uint32 Hash = static_cast<uint32>(Data.LoadType);
-	Hash = HashCombineFast(Hash, Data.ConfigPlatform.ToUnstableInt());
-	Hash = HashCombineFast(Hash, Data.FileName.ToUnstableInt());
-	Hash = HashCombineFast(Hash, Data.SectionName.ToUnstableInt());
-	Hash = HashCombineFast(Hash, GetTypeHash(Data.ValueName));
-	Hash = HashCombineFast(Hash, GetTypeHash(Data.RequestingPlatform));
-	return Hash;
-}
-
-
-inline bool FConfigAccessData::IsSameConfigFile(const FConfigAccessData& Other) const
-{
-	return LoadType == Other.LoadType && ConfigPlatform == Other.ConfigPlatform && FileName == Other.FileName;
-}
-
-inline bool FConfigAccessData::operator==(const FConfigAccessData& Other) const
-{
-	return LoadType == Other.LoadType && ConfigPlatform == Other.ConfigPlatform &&
-		FileName == Other.FileName && SectionName == Other.SectionName && ValueName == Other.ValueName &&
-		RequestingPlatform == Other.RequestingPlatform;
-}
-
-inline bool FConfigAccessData::operator!=(const FConfigAccessData& Other) const
-{
-	return !(*this == Other);
-}
-
-inline bool FConfigAccessData::operator<(const FConfigAccessData& Other) const
-{
-	if (LoadType != Other.LoadType) return static_cast<uint32>(LoadType) < static_cast<uint32>(Other.LoadType);
-	if (ConfigPlatform != Other.ConfigPlatform) return ConfigPlatform.LexicalLess(Other.ConfigPlatform);
-	if (FileName != Other.FileName) return FileName.LexicalLess(Other.FileName);
-	if (SectionName != Other.SectionName) return SectionName.LexicalLess(Other.SectionName);
-	if (ValueName != Other.ValueName) return FName(ValueName).LexicalLess(FName(Other.ValueName));
-	if (RequestingPlatform != Other.RequestingPlatform)
-	{
-		if (RequestingPlatform == nullptr) return true;
-		if (Other.RequestingPlatform == nullptr) return false;
-		return RequestingPlatform->PlatformName() < Other.RequestingPlatform->PlatformName();
-	}
-	return false;
-}
-
-} // namespace UE::ConfigAccessTracking
-
 #else // !UE_WITH_CONFIG_TRACKING
 
 namespace UE::ConfigAccessTracking
diff --git a/Engine/Source/Runtime/Core/Private/Misc/ConfigAccessData.cpp b/Engine/Source/Runtime/Core/Private/Misc/ConfigAccessData.cpp
new file mode 100644
--- /dev/null
+++ b/Engine/Source/Runtime/Core/Private/Misc/ConfigAccessData.cpp
@@ -0,0 +1,317 @@
+// Copyright Epic Games, Inc. All Rights Reserved.
+
+#include "Misc/ConfigAccessData.h"
+
+#if UE_WITH_CONFIG_TRACKING || WITH_EDITOR
+
+#include "Misc/StringBuilder.h"
+
+namespace UE::ConfigAccessTracking
+{
+
+FConfigAccessData::FConfigAccessData(ELoadType InLoadType, FNameEntryId InConfigPlatform, FNameEntryId InFileName)
+	: ConfigPlatform(InConfigPlatform)
+	, FileName(InFileName)
+	, LoadType(InLoadType)
+{
+}
+
+FConfigAccessData::FConfigAccessData(ELoadType InLoadType, FNameEntryId InConfigPlatform, FNameEntryId InFileName,
+	FNameEntryId InSectionName, FMinimalName InValueName, const ITargetPlatform* InRequestingPlatform)
+	: ConfigPlatform(InConfigPlatform)
+	, FileName(InFileName)
+	, SectionName(InSectionName)
+	, ValueName(InValueName)
+	, RequestingPlatform(InRequestingPlatform)
+	, LoadType(InLoadType)
+{
+}
+
+FConfigAccessData FConfigAccessData::GetFileOnlyData() const
+{
+	return FConfigAccessData(LoadType, ConfigPlatform, FileName);
+}
+
+FConfigAccessData FConfigAccessData::GetPathOnlyData() const
+{
+	return FConfigAccessData(LoadType, ConfigPlatform, FileName, SectionName, ValueName, nullptr);
+}
+
+FString FConfigAccessData::FullPathToString() const
+{
+	TStringBuilder<256> Out;
+	AppendFullPath(Out);
+	return FString(*Out);
+}
+
+void FConfigAccessData::AppendFullPath(FStringBuilderBase& Out) const
+{
+	if (LoadType == ELoadType::Uninitialized || FileName.IsNone())
+	{
+		Out << TEXTVIEW("<Invalid>");
+		return;
+	}
+
+	Out << LexToString(LoadType) << TEXTVIEW(".");
+	if (ConfigPlatform.IsNone())
+	{
+		Out << PlatformAgnosticName;
+	}
+	else
+	{
+		EscapeConfigTrackingTokenAppendString(GetConfigPlatform(), Out);
+	}
+	Out << TEXTVIEW(".");
+	EscapeConfigTrackingTokenAppendString(GetFileName(), Out);
+	if (!SectionName.IsNone())
+	{
+		Out << TEXTVIEW(":[");
+		EscapeConfigTrackingTokenAppendString(GetSectionName(), Out);
+		Out << TEXTVIEW("]");
+		if (!ValueName.IsNone())
+		{
+			Out << TEXTVIEW(":");
+			EscapeConfigTrackingTokenAppendString(GetValueName(), Out);
+		}
+	}
+}
+
+FConfigAccessData FConfigAccessData::Parse(FStringView Text)
+{
+	// ConfigSystem.<Editor>.../../../Engine/Config/ConsoleVariables.ini:Section:ValueName
+	//   -> "ConfigSystem", "<Editor>", "../../../Engine/Config/ConsoleVariables.ini", "Section", "ValueName"
+	// No tokens can contain a single colon, but they might contain double colon, which is the escape code for a single colon.
+	// 3rd token might have dots, first two cannot.
+
+	TStringBuilder<128> FullFilePath;
+	TStringBuilder<64> SectionNameStr;
+	TStringBuilder<64> ValueNameStr;
+	FStringBuilderBase* TokenBuffer[] = { &FullFilePath, &SectionNameStr, &ValueNameStr };
+	TArrayView<FStringBuilderBase*> ColonDelimitedTokens = TokenBuffer;
+	TryTokenizeConfigTrackingString(Text, ColonDelimitedTokens);
+	if (FullFilePath.Len() == 0)
+	{
+		return FConfigAccessData();
+	}
+
+	FStringView FullFilePathTokens[3];
+	FullFilePathTokens[0] = FullFilePath;
+	int32 NumFilePathTokens = 1;
+	for (NumFilePathTokens = 1; NumFilePathTokens < UE_ARRAY_COUNT(FullFilePathTokens); ++NumFilePathTokens)
+	{
+		FStringView& Current = FullFilePathTokens[NumFilePathTokens - 1];
+		FStringView& Next = FullFilePathTokens[NumFilePathTokens];
+		int32 DotIndex = Current.Find(TEXTVIEW("."));
+		if (DotIndex == INDEX_NONE)
+		{
+			break;
+		}
+		Next = Current.RightChop(DotIndex + 1);
+		Current.LeftInline(DotIndex);
+		if (Next.IsEmpty())
+		{
+			// break before incrementing NumFilePathTokens so that we do not count the empty Next as a token
+			break;
+		}
+	}
+	// FullFilePath is of the form 
+	//   LoadType.Platform.ConfigName
+	// Platform is <Editor> if the configfile was an editor config file rather than a platform-specific config file
+	if (NumFilePathTokens < 3)
+	{
+		// Not a valid FConfigAccessData text; we always require all 3 of LoadType,Platform,ConfigName
+		return FConfigAccessData();
+	}
+
+	FConfigAccessData Result;
+	LexFromString(Result.LoadType, FullFilePathTokens[0]);
+	if (Result.LoadType == ELoadType::Uninitialized)
+	{
+		return FConfigAccessData();
+	}
+
+	FName LocalConfigPlatform = FullFilePathTokens[1] == PlatformAgnosticName
+		? NAME_None : FName(FullFilePathTokens[1]);
+	FStringView SectionNameView(SectionNameStr);
+	// We write out sectionnames with [] for readability; remove them if they exist.
+	if (SectionNameView.StartsWith('[')) SectionNameView.RightChopInline(1);
+	if (SectionNameView.EndsWith(']')) SectionNameView.LeftChopInline(1);
+	FName LocalSectionName = SectionNameView.Len() == 0 ? NAME_None : FName(SectionNameView, NAME_NO_NUMBER);
+	FName LocalValueName = ValueNameStr.Len() == 0 ? NAME_None : FName(ValueNameStr);
+
+	Result.ConfigPlatform = LocalConfigPlatform.GetComparisonIndex();
+	Result.FileName = FName(FullFilePathTokens[2]).GetComparisonIndex();
+	Result.SectionName = LocalSectionName.GetComparisonIndex();
+	Result.ValueName = FMinimalName(LocalValueName);
+
+	return Result;
+}
+
+void EscapeConfigTrackingTokenToString(FName Token, FStringBuilderBase& Result)
+{
+	Result.Reset();
+	EscapeConfigTrackingTokenAppendString(Token, Result);
+}
+
+void EscapeConfigTrackingTokenAppendString(FName Token, FStringBuilderBase& Result)
+{
+	int32 InitialLength = Result.Len();
+	Result << Token;
+	FStringView AddedView = Result.ToView().RightChop(InitialLength);
+	if (AddedView.Contains(TEXTVIEW(":")))
+	{
+		FString ReplaceText(AddedView);
+		ReplaceText.ReplaceInline(TEXT(":"), TEXT("::"));
+		Result.RemoveSuffix(AddedView.Len());
+		Result << ReplaceText;
+	}
+}
+
+bool TryTokenizeConfigTrackingString(FStringView Text, TArrayView<FStringBuilderBase*> OutTokens)
+{
+	int32 TextLen = Text.Len();
+	if (TextLen == 0)
+	{
+		for (FStringBuilderBase* Token : OutTokens) Token->Reset();
+		return false;
+	}
+	const TCHAR* TextData = Text.GetData();
+
+	int32 NumTokens = OutTokens.Num();
+	check(NumTokens > 0);
+	int32 NextTokenIndex = 0;
+	FStringBuilderBase* OutToken = OutTokens[NextTokenIndex++];
+	OutToken->Reset();
+	for (int32 Index = 0; Index < TextLen; ++Index)
+	{
+		TCHAR C = TextData[Index];
+		if (C != ':')
+		{
+			OutToken->AppendChar(C);
+		}
+		else if (Index < TextLen - 1 && TextData[Index + 1] == ':')
+		{
+			++Index;
+			OutToken->AppendChar(':');
+		}
+		else
+		{
+			if (OutToken->Len() == 0)
+			{
+				// An empty token and therefore the string is invalid; abandon anything that follows it
+				while (NextTokenIndex < NumTokens) OutTokens[NextTokenIndex++]->Reset();
+				return false;
+			}
+			if (NextTokenIndex >= NumTokens)
+			{
+				// Too many tokens
+				return false;
+			}
+			OutToken = OutTokens[NextTokenIndex++];
+			OutToken->Reset();
+		}
+	}
+	if (OutToken->Len() == 0)
+	{
+		// Empty token
+		while (NextTokenIndex < NumTokens) OutTokens[NextTokenIndex++]->Reset();
+		return false;
+	}
+	if (NextTokenIndex < NumTokens)
+	{
+		// too few tokens
+		while (NextTokenIndex < NumTokens) OutTokens[NextTokenIndex++]->Reset();
+		return false;
+	}
+	return true;
+}
+
+} // namespace UE::ConfigAccessTracking
+
+const TCHAR* LexToString(UE::ConfigAccessTracking::ELoadType LoadType)
+{
+	using namespace UE::ConfigAccessTracking;
+
+	switch (LoadType)
+	{
+	case ELoadType::ConfigSystem: return TEXT("ConfigSystem");
+	case ELoadType::LocalIniFile: return TEXT("LocalIniFile");
+	case ELoadType::LocalSingleIniFile: return TEXT("LocalSingleIniFile");
+	case ELoadType::ExternalIniFile: return TEXT("ExternalIniFile");
+	case ELoadType::ExternalSingleIniFile: return TEXT("ExternalSingleIniFile");
+	case ELoadType::Manual: return TEXT("Manual");
+	case ELoadType::SuppressReporting: return TEXT("SuppressReporting");
+	case ELoadType::Uninitialized:
+	default:
+		return TEXT("Uninitialized");
+	}
+}
+
+void LexFromString(UE::ConfigAccessTracking::ELoadType& OutLoadType, FStringView Text)
+{
+	using namespace UE::ConfigAccessTracking;
+
+	if (Text.IsEmpty())
+	{
+		OutLoadType = ELoadType::Uninitialized;
+		return;
+	}
+	switch (Text[0])
+	{
+	case 'C':
+	case 'c':
+		if (Text.Equals(TEXT("ConfigSystem"), ESearchCase::IgnoreCase))
+		{
+			OutLoadType = ELoadType::ConfigSystem;
+			return;
+		}
+		break;
+	case 'L':
+	case 'l':
+		if (Text.Equals(TEXT("LocalIniFile"), ESearchCase::IgnoreCase))
+		{
+			OutLoadType = ELoadType::LocalIniFile;
+			return;
+		}
+		if (Text.Equals(TEXT("LocalSingleIniFile"), ESearchCase::IgnoreCase))
+		{
+			OutLoadType = ELoadType::LocalSingleIniFile;
+			return;
+		}
+		break;
+	case 'E':
+	case 'e':
+		if (Text.Equals(TEXT("ExternalIniFile"), ESearchCase::IgnoreCase))
+		{
+			OutLoadType = ELoadType::ExternalIniFile;
+			return;
+		}
+		if (Text.Equals(TEXT("ExternalSingleIniFile"), ESearchCase::IgnoreCase))
+		{
+			OutLoadType = ELoadType::ExternalSingleIniFile;
+			return;
+		}
+		break;
+	case 'S':
+	case 's':
+		if (Text.Equals(TEXT("SuppressReporting"), ESearchCase::IgnoreCase))
+		{
+			OutLoadType = ELoadType::SuppressReporting;
+			return;
+		}
+		break;
+	case 'M':
+	case 'm':
+		if (Text.Equals(TEXT("Manual"), ESearchCase::IgnoreCase))
+		{
+			OutLoadType = ELoadType::Manual;
+			return;
+		}
+		break;
+	default:
+		break;
+	}
+	OutLoadType = ELoadType::Uninitialized;
+}
+
+#endif // UE_WITH_CONFIG_TRACKING || WITH_EDITOR
diff --git a/Engine/Source/Runtime/Core/Public/Misc/ConfigAccessData.h b/Engine/Source/Runtime/Core/Public/Misc/ConfigAccessData.h
new file mode 100644
--- /dev/null
+++ b/Engine/Source/Runtime/Core/Public/Misc/ConfigAccessData.h
@@ -0,0 +1,176 @@
+// Copyright Epic Games, Inc. All Rights Reserved.
+
+#pragma once
+
+#include "Misc/ConfigAccessTracking.h"
+
+#if UE_WITH_CONFIG_TRACKING || WITH_EDITOR
+#include "Containers/ArrayView.h"
+#include "Containers/StringFwd.h"
+#include "Containers/StringView.h"
+#include "Containers/UnrealString.h"
+#include "Interfaces/ITargetPlatform.h"
+#include "Templates/TypeHash.h"
+
+namespace UE::ConfigAccessTracking
+{
+
+/**
+ * Full path of a read of an FConfigValue; used to persistently track which config values were read during
+ * an invocation of the editor.
+ */
+struct FConfigAccessData
+{
+	FNameEntryId ConfigPlatform;
+	FNameEntryId FileName;
+	FNameEntryId SectionName;
+	FMinimalName ValueName;
+	const ITargetPlatform* RequestingPlatform = nullptr;
+	ELoadType LoadType = ELoadType::Uninitialized;
+
+	FConfigAccessData() = default;
+	FConfigAccessData(ELoadType InLoadType, FName InConfigPlatform, FName InFileName,
+		FName InSectionName, FName InValueName, const ITargetPlatform* InRequestingPlatform);
+	CORE_API FConfigAccessData(ELoadType InLoadType, FNameEntryId InConfigPlatform, FNameEntryId InFileName,
+		FNameEntryId InSectionName, FMinimalName InValueName, const ITargetPlatform* InRequestingPlatform);
+	CORE_API FConfigAccessData(ELoadType InLoadType, FNameEntryId InConfigPlatform, FNameEntryId InFileName);
+
+	FName GetConfigPlatform() const;
+	FName GetFileName() const;
+	FName GetSectionName() const;
+	FName GetValueName() const;
+
+	CORE_API FConfigAccessData GetFileOnlyData() const;
+	CORE_API FConfigAccessData GetPathOnlyData() const;
+	CORE_API FString FullPathToString() const;
+	CORE_API void AppendFullPath(FStringBuilderBase& Out) const;
+
+	/**
+	 * "ConfigSystem.<Editor>.../../../Engine/Config/ConsoleVariables.ini:[Section]:Value"
+	 *   -> "ConfigSystem", "<Editor>", "../../../Engine/Config/ConsoleVariables.ini", "Section", "Value"
+	 */
+	CORE_API static FConfigAccessData Parse(FStringView Text);
+
+	friend uint32 GetTypeHash(const FConfigAccessData& Data);
+	bool IsSameConfigFile(const FConfigAccessData& Other) const;
+	bool operator==(const FConfigAccessData& Other) const;
+	bool operator!=(const FConfigAccessData& Other) const;
+	bool operator<(const FConfigAccessData& Other) const;
+};
+
+CORE_API void EscapeConfigTrackingTokenToString(FName Token, FStringBuilderBase& Result);
+CORE_API void EscapeConfigTrackingTokenAppendString(FName Token, FStringBuilderBase& Result);
+CORE_API bool TryTokenizeConfigTrackingString(FStringView Text, TArrayView<FStringBuilderBase*> OutTokens);
+
+constexpr FStringView PlatformAgnosticName = TEXTVIEW("<Editor>");
+
+}
+
+/** Convert ELoadType -> text */
+const TCHAR* LexToString(UE::ConfigAccessTracking::ELoadType LoadType);
+
+/** Convert text -> ELoadType */
+void LexFromString(UE::ConfigAccessTracking::ELoadType& OutLoadType, FStringView Text);
+
+///////////////////////////////////////////////////////
+// Inline implementations
+///////////////////////////////////////////////////////
+
+namespace UE::ConfigAccessTracking
+{
+
+inline FConfigAccessData::FConfigAccessData(ELoadType InLoadType, FName InConfigPlatform, FName InFileName,
+	FName InSectionName, FName InValueName, const ITargetPlatform* InRequestingPlatform)
+	: FConfigAccessData(InLoadType, InConfigPlatform.GetComparisonIndex(), InFileName.GetComparisonIndex(),
+		InSectionName.GetComparisonIndex(), FMinimalName(InValueName), InRequestingPlatform)
+{
+}
+
+inline FName FConfigAccessData::GetConfigPlatform() const
+{
+	return FName(ConfigPlatform, ConfigPlatform, NAME_NO_NUMBER_INTERNAL);
+}
+
+inline FName FConfigAccessData::GetFileName() const
+{
+	return FName(FileName, FileName, NAME_NO_NUMBER_INTERNAL);
+}
+
+inline FName FConfigAccessData::GetSectionName() const
+{
+	return FName(SectionName, SectionName, NAME_NO_NUMBER_INTERNAL);
+}
+
+inline FName FConfigAccessData::GetValueName() const
+{
+	return FName(ValueName);
+}
+
+inline uint32 GetTypeHash(const UE::ConfigAccessTracking::FConfigAccessData& Data)
+{
+	uint32 Hash = static_cast<uint32>(Data.LoadType);
+	Hash = HashCombineFast(Hash, Data.ConfigPlatform.ToUnstableInt());
+	Hash = HashCombineFast(Hash, Data.FileName.ToUnstableInt());
+	Hash = HashCombineFast(Hash, Data.SectionName.ToUnstableInt());
+	Hash = HashCombineFast(Hash, GetTypeHash(Data.ValueName));
+	Hash = HashCombineFast(Hash, GetTypeHash(Data.RequestingPlatform));
+	return Hash;
+}
+
+inline bool FConfigAccessData::IsSameConfigFile(const FConfigAccessData& Other) const
+{
+	return LoadType == Other.LoadType && ConfigPlatform == Other.ConfigPlatform && FileName == Other.FileName;
+}
+
+inline bool FConfigAccessData::operator==(const FConfigAccessData& Other) const
+{
+	return LoadType == Other.LoadType && ConfigPlatform == Other.ConfigPlatform &&
+		FileName == Other.FileName && SectionName == Other.SectionName && ValueName == Other.ValueName &&
+		RequestingPlatform == Other.RequestingPlatform;
+}
+
+inline bool FConfigAccessData::operator!=(const FConfigAccessData& Other) const
+{
+	return !(*this == Other);
+}
+
+inline bool FConfigAccessData::operator<(const FConfigAccessData& Other) const
+{
+	if (LoadType != Other.LoadType)
+	{
+		return static_cast<uint32>(LoadType) < static_cast<uint32>(Other.LoadType);
+	}
+	if (ConfigPlatform != Other.ConfigPlatform)
+	{
+		return ConfigPlatform.LexicalLess(Other.ConfigPlatform);
+	}
+	if (FileName != Other.FileName)
+	{
+		return FileName.LexicalLess(Other.FileName);
+	}
+	if (SectionName != Other.SectionName)
+	{
+		return SectionName.LexicalLess(Other.SectionName);
+	}
+	if (ValueName != Other.ValueName)
+	{
+		return FName(ValueName).LexicalLess(FName(Other.ValueName));
+	}
+	if (RequestingPlatform != Other.RequestingPlatform)
+	{
+		if (RequestingPlatform == nullptr)
+		{
+			return true;
+		}
+		if (Other.RequestingPlatform == nullptr)
+		{
+			return false;
+		}
+		return RequestingPlatform->PlatformName() < Other.RequestingPlatform->PlatformName();
+	}
+	return false;
+}
+
+} // namespace UE::ConfigAccessTracking
+
+#endif // UE_WITH_CONFIG_TRACKING || WITH_EDITOR
\ No newline at end of file
diff --git a/Engine/Source/Runtime/Core/Public/Misc/ConfigAccessTracking.h b/Engine/Source/Runtime/Core/Public/Misc/ConfigAccessTracking.h
--- a/Engine/Source/Runtime/Core/Public/Misc/ConfigAccessTracking.h
+++ b/Engine/Source/Runtime/Core/Public/Misc/ConfigAccessTracking.h
@@ -24,7 +24,7 @@ class FConfigSection;
 struct FConfigValue;
 #endif
 
-#if UE_WITH_CONFIG_TRACKING 
+#if WITH_EDITOR || UE_WITH_CONFIG_TRACKING
 
 namespace UE::ConfigAccessTracking
 {
@@ -45,6 +45,15 @@ enum class ELoadType : uint8
 	Uninitialized,
 };
 
+} // namespace UE::ConfigAccessTracking
+
+#endif // WITH_EDITOR || UE_WITH_CONFIG_TRACKING
+
+#if UE_WITH_CONFIG_TRACKING
+
+namespace UE::ConfigAccessTracking
+{
+
 /**
  * A refcounted struct of data about an FConfigFile, including a backpointer to the FConfigFile if it is still 
  * in memory. References to this data are held by all the FConfigSections in the FConfigFile so that they can lookup
