Partial Backout - CL30419904
Original CL Desc
-----------------------------------------------------------------
Fix NormalizePackageNames so it always outputs relative paths, because we are comparing relative paths later on.

diff --git a/Engine/Source/Editor/UnrealEd/Private/Commandlets/PackageUtilities.cpp b/Engine/Source/Editor/UnrealEd/Private/Commandlets/PackageUtilities.cpp
--- a/Engine/Source/Editor/UnrealEd/Private/Commandlets/PackageUtilities.cpp
+++ b/Engine/Source/Editor/UnrealEd/Private/Commandlets/PackageUtilities.cpp
@@ -124,15 +124,13 @@ bool NormalizePackageNames( TArray<FString> PackageNames, TArray<FString>& Packa
 			TStringBuilder<256> UnusedRelPath;
 			for ( const FString& Path : Paths)
 			{
-				// Make sure paths are relative so SearchDirectoryRecursive will output relative paths
-				const FString RelativePath = FPaths::CreateStandardFilename(Path);
-				if (!FPackageName::TryGetMountPointForPath(RelativePath, UnusedPackagePath, UnusedFilePath, UnusedRelPath))
+				if (!FPackageName::TryGetMountPointForPath(Path, UnusedPackagePath, UnusedFilePath, UnusedRelPath))
 				{
 					UE_LOG(LogPackageUtilities, Warning,
-						TEXT("Engine.ini:[Core.System]:Paths entry '%s' is not mounted. Skipping it."), *RelativePath);
+						TEXT("Engine.ini:[Core.System]:Paths entry '%s' is not mounted. Skipping it."), *Path);
 					continue;
 				}
-				FString SearchWildcard = RelativePath / PackageWildcard;
+				FString SearchWildcard = Path / PackageWildcard;
 				UE_LOG(LogPackageUtilities, Log, TEXT("Searching using wildcard: '%s'"), *SearchWildcard);
 				SearchDirectoryRecursive(SearchWildcard, PackageNames, PackagePathNames);
 			}
