Cooker: Fix error that package requested from project packaging settings is a redirector to NOT fire when the package is a multiple-asset package that happens to contain a redirector subasset internally, such as a map with a redirected levelblueprint. The internal redirector will be removed when the package is saved and the package itself will be valid in the cooked output.
#rnx

diff --git a/Engine/Source/Editor/UnrealEd/Private/CookOnTheFlyServer.cpp b/Engine/Source/Editor/UnrealEd/Private/CookOnTheFlyServer.cpp
--- a/Engine/Source/Editor/UnrealEd/Private/CookOnTheFlyServer.cpp
+++ b/Engine/Source/Editor/UnrealEd/Private/CookOnTheFlyServer.cpp
@@ -9287,15 +9287,31 @@ void UCookOnTheFlyServer::CollectFilesToCook(TArray<FName>& FilesInPath, TMap<FN
 				TArray<FAssetData> Assets;
 				if (!AssetRegistry->GetAssetsByPackageName(PackagePath, Assets))
 				{
-					const FText ErrorMessage = FText::Format(LOCTEXT("GameMapSettingsMissing", "{0} contains a path to a missing asset '{1}'. The intended asset will fail to load in a packaged build. Select the intended asset again in Project Settings to fix this issue."),
+					const FText ErrorMessage = FText::Format(LOCTEXT("GameMapSettingsMissing",
+						"{0} contains a path to a missing asset '{1}'. "
+						"The intended asset will fail to load in a packaged build. "
+						"Select the intended asset again in Project Settings to fix this issue."),
 						FText::FromName(GameDefaultSet.Key), FText::FromName(PackagePath));
 					LogCookerMessage(ErrorMessage.ToString(), EMessageSeverity::Error);
 				}
-				else if (Algo::AnyOf(Assets, [](const FAssetData& Asset) { return Asset.IsRedirector(); }))
+				else
 				{
-					const FText ErrorMessage = FText::Format(LOCTEXT("GameMapSettingsRedirectorDetected", "{0} contains a redirected reference '{1}'. The intended asset will fail to load in a packaged build. Select the intended asset again in Project Settings to fix this issue."),
-						FText::FromName(GameDefaultSet.Key), FText::FromName(PackagePath));
-					LogCookerMessage(ErrorMessage.ToString(), EMessageSeverity::Error);
+					TArray<const FAssetData*, TInlineAllocator<1>> AssetPtrs;
+					AssetPtrs.Reserve(Assets.Num());
+					for (const FAssetData& AssetData : Assets)
+					{
+						AssetPtrs.Add(&AssetData);
+					}
+					const FAssetData* PrimaryAssetData = UE::AssetRegistry::GetMostImportantAsset(AssetPtrs);
+					if (PrimaryAssetData && PrimaryAssetData->IsRedirector())
+					{
+						const FText ErrorMessage = FText::Format(LOCTEXT("GameMapSettingsRedirectorDetected",
+							"{0} contains a redirected reference '{1}'. "
+							"The intended asset will fail to load in a packaged build. "
+							"Select the intended asset again in Project Settings to fix this issue."),
+							FText::FromName(GameDefaultSet.Key), FText::FromName(PackagePath));
+						LogCookerMessage(ErrorMessage.ToString(), EMessageSeverity::Error);
+					}
 				}
 
 				AddFileToCook(FilesInPath, Instigators, PackagePath.ToString(),
