Fixing crash caused by referencing an invalid UStatusBarSubsystem to create a StatusBarWidget when running a test through the command line as there is no valid Editor initialized.

#rnx

diff --git a/Engine/Source/Editor/UnrealEd/Private/Toolkits/SStandaloneAssetEditorToolkitHost.cpp b/Engine/Source/Editor/UnrealEd/Private/Toolkits/SStandaloneAssetEditorToolkitHost.cpp
--- a/Engine/Source/Editor/UnrealEd/Private/Toolkits/SStandaloneAssetEditorToolkitHost.cpp
+++ b/Engine/Source/Editor/UnrealEd/Private/Toolkits/SStandaloneAssetEditorToolkitHost.cpp
@@ -188,9 +188,12 @@ void SStandaloneAssetEditorToolkitHost::RestoreFromLayout( const TSharedRef<FTab
 	TSharedPtr<SWindow> ParentWindow = FSlateApplication::Get().FindWidgetWindow( HostTab );
 	TSharedPtr<SWidget> RestoredUI = MyTabManager->RestoreFrom( NewLayout, ParentWindow );
 
-	StatusBarWidget = GEditor->GetEditorSubsystem<UStatusBarSubsystem>()->MakeStatusBarWidget(StatusBarName, HostTab);
+	checkf(RestoredUI.IsValid(), TEXT("The layout must have a primary dock area"));
 
-	checkf(RestoredUI.IsValid(), TEXT("The layout must have a primary dock area") );
+	if (UStatusBarSubsystem* StatusBarSubsystem = GEditor ? GEditor->GetEditorSubsystem<UStatusBarSubsystem>() : nullptr)
+	{
+		StatusBarWidget = StatusBarSubsystem->MakeStatusBarWidget(StatusBarName, HostTab);
+	}
 
 	this->ChildSlot
 	[
@@ -209,7 +212,7 @@ void SStandaloneAssetEditorToolkitHost::RestoreFromLayout( const TSharedRef<FTab
 		.Padding(0.0f, 2.0f, 0.0f, 0.0f)
 		.AutoHeight()
 		[
-			StatusBarWidget.ToSharedRef()
+			StatusBarWidget.IsValid() ? StatusBarWidget.ToSharedRef() : SNullWidget::NullWidget
 		]
 	];
 }
