IncrementalCook: Add FCookDependency::SettingsObject.
#rnx

diff --git a/Engine/Source/Editor/UnrealEd/Private/EditorDomain/EditorDomainUtils.cpp b/Engine/Source/Editor/UnrealEd/Private/EditorDomain/EditorDomainUtils.cpp
--- a/Engine/Source/Editor/UnrealEd/Private/EditorDomain/EditorDomainUtils.cpp
+++ b/Engine/Source/Editor/UnrealEd/Private/EditorDomain/EditorDomainUtils.cpp
@@ -456,6 +456,58 @@ FPackageDigest CalculatePackageDigest(const FAssetPackageData& PackageData, FNam
 	return Result;
 }
 
+bool TryAppendClassDigests(FBlake3& Writer, TConstArrayView<FTopLevelAssetPath> ClassPaths, FString* OutErrorMessage)
+{
+	FClassDigestMap& ClassDigests = GetClassDigests();
+	bool bHasTriedPrecacheClassDigests = false;
+	int32 NextClass = 0;
+	while (NextClass < ClassPaths.Num())
+	{
+		{
+			FReadScopeLock ClassDigestsScopeLock(ClassDigests.Lock);
+			while (NextClass < ClassPaths.Num())
+			{
+				const FTopLevelAssetPath& ClassPath = ClassPaths[NextClass];
+				FClassDigestData* ExistingData = ClassDigests.Map.Find(ClassPath);
+				if (!ExistingData)
+				{
+					break;
+				}
+				NextClass++;
+
+				// We only support hashes for the ClosestNative class
+				FClassDigestData* NativeData = ExistingData;
+				if (ExistingData->ClosestNative != ClassPath)
+				{
+					NativeData = ClassDigests.Map.Find(ExistingData->ClosestNative);
+					checkf(NativeData, TEXT("Classes are only stored in a ClosestNative field if they exist"));
+				}
+				check(NativeData->bNative);
+				Writer.Update(&NativeData->InclusiveSchemaHash, sizeof(NativeData->InclusiveSchemaHash));
+			}
+		}
+
+		if (NextClass < ClassPaths.Num())
+		{
+			// EDITORDOMAIN_TODO: Remove the clauses !IsInGameThread || GIsSavingPackage once FindObject no longer asserts if GIsSavingPackage
+			if (bHasTriedPrecacheClassDigests || !IsInGameThread() || GIsSavingPackage)
+			{
+				if (OutErrorMessage)
+				{
+					*OutErrorMessage = FString::Printf(TEXT("%s is not a valid class."),
+						*ClassPaths[NextClass].ToString());
+				}
+				return false;
+			}
+			TConstArrayView<FTopLevelAssetPath> RemainingClasses =
+				TConstArrayView<FTopLevelAssetPath>(ClassPaths).RightChop(NextClass);
+			PrecacheClassDigests(RemainingClasses);
+			bHasTriedPrecacheClassDigests = true;
+		}
+	}
+	return true;
+}
+
 /**
  * Holds context data for a call to PrecacheClassDigests, which needs to recursively
  * traverse a graph of of class parents and construction classes
diff --git a/Engine/Source/Editor/UnrealEd/Private/EditorDomain/EditorDomainUtils.h b/Engine/Source/Editor/UnrealEd/Private/EditorDomain/EditorDomainUtils.h
--- a/Engine/Source/Editor/UnrealEd/Private/EditorDomain/EditorDomainUtils.h
+++ b/Engine/Source/Editor/UnrealEd/Private/EditorDomain/EditorDomainUtils.h
@@ -86,6 +86,8 @@ ESaveStorageResult SaveStorageResultFromString(FUtf8StringView Text);
  */
 FPackageDigest CalculatePackageDigest(IAssetRegistry& AssetRegistry, FName PackageName);
 FPackageDigest CalculatePackageDigest(const FAssetPackageData& PackageData, FName PackageName);
+/** FindOrAdd class digest for the given classpath, and append its hash onto the writer. Returns false if not found. */
+bool TryAppendClassDigests(FBlake3& Writer, TConstArrayView<FTopLevelAssetPath> ClassPaths, FString* OutErrorMessage);
 
 /** For any ClassNames not already in ClassDigests, look up their UStruct and add them. */
 void PrecacheClassDigests(TConstArrayView<FTopLevelAssetPath> ClassNames);
diff --git a/Engine/Source/Editor/UnrealEd/Private/TargetDomain/TargetDomainUtils.cpp b/Engine/Source/Editor/UnrealEd/Private/TargetDomain/TargetDomainUtils.cpp
--- a/Engine/Source/Editor/UnrealEd/Private/TargetDomain/TargetDomainUtils.cpp
+++ b/Engine/Source/Editor/UnrealEd/Private/TargetDomain/TargetDomainUtils.cpp
@@ -155,15 +155,42 @@ bool FCookDependencies::TryCalculateCurrentKey(const FAssetPackageData* Override
 	UE::EditorDomain::FPackageDigest PackageDigest = OverrideAssetPackageData ?
 		UE::EditorDomain::CalculatePackageDigest(*OverrideAssetPackageData, PackageName) :
 		EditorDomain->GetPackageDigest(PackageName);
-
 	if (!PackageDigest.IsSuccessful())
 	{
 		if (OutErrorMessage) *OutErrorMessage = PackageDigest.GetStatusString();
 		return false;
 	}
-
 	KeyBuilder.Update(&PackageDigest.Hash, sizeof(PackageDigest.Hash));
 
+	if (!ClassDependencies.IsEmpty())
+	{
+		TArray<FTopLevelAssetPath, TInlineAllocator<1>> ClassPaths;
+		ClassPaths.Reserve(ClassDependencies.Num());
+		for (const FString& ClassPathStr : ClassDependencies)
+		{
+			FTopLevelAssetPath& ClassPath = ClassPaths.Emplace_GetRef(ClassPathStr);
+			if (!ClassPath.IsValid())
+			{
+				if (OutErrorMessage)
+				{
+					*OutErrorMessage = FString::Printf(
+						TEXT("ClassDependency failed: %s is not a valid TopLevelAssetPath."), *ClassPathStr);
+				}
+				return false;
+			}
+		}
+		FString AppendDigestError;
+		if (!UE::EditorDomain::TryAppendClassDigests(KeyBuilder, ClassPaths, &AppendDigestError))
+		{
+			if (OutErrorMessage)
+			{
+				*OutErrorMessage = FString::Printf(
+					TEXT("ClassDependency failed: %s"), *AppendDigestError);
+			}
+			return false;
+		}
+	}
+
 	for (FName PackageDependency : BuildPackageDependencies)
 	{
 		PackageDigest = EditorDomain->GetPackageDigest(PackageDependency);
@@ -367,41 +394,120 @@ FCookDependencies FCookDependencies::Collect(UPackage* Package, const ITargetPla
 #endif
 	if (SaveResult)
 	{
-		// Move most types of FCookDependency into this->CookDependencies, but separate out
-		// TransitiveBuildDependencies; they require different processing than other types.
+		// Convert some FCookDependency types into our specially-handled dependencies, and put the more
+		// generic ones into this->CookDependencies
 		Result.CookDependencies.Reserve(SaveResult->CookDependencies.Num());
-		for (UE::Cook::FCookDependency& CookDependency : SaveResult->CookDependencies)
+
+		// Additional dependencies can be added from SettingsObject dependencies. Add these included dependencies
+		// to a followup list that we process in the second round, if they exist. Handle detecting cycles in the
+		// included dependencies; cycle detection is currently simple: nested includes are not yet required
+		// so give an error if we ever make it to a third round.
+		TArray<UE::Cook::FCookDependency> NextIncludedDependenciesBuffer[2];
+		TArray<UE::Cook::FCookDependency>* RemainingDependencies = &SaveResult->CookDependencies;
+		int32 IncludedDependenciesRound = 0;
+		constexpr int32 MaxNumberOfIncludedDependenciesRounds = 2;
+		while (!RemainingDependencies->IsEmpty())
 		{
-			if (CookDependency.GetType() == UE::Cook::ECookDependency::TransitiveBuild)
+			TArray<UE::Cook::FCookDependency>* NextIncludedDependencies =
+				&NextIncludedDependenciesBuffer[IncludedDependenciesRound % UE_ARRAY_COUNT(NextIncludedDependenciesBuffer)];
+			if (IncludedDependenciesRound++ > 0)
 			{
-				// Build dependencies from a package to itself have a performance cost and serve no
-				// purpose, so remove them. They can occur in some systems that naively add a build dependency
-				// from one object to another without checking whether the second object is in the same package.
-				FName DependencyPackageName = CookDependency.GetPackageName();
-				if (DependencyPackageName == Result.PackageName)
-				{
-					continue;
-				}
-				if (CookDependency.IsAlsoAddRuntimeDependency())
+				NextIncludedDependencies->Reset();
+				if (IncludedDependenciesRound > MaxNumberOfIncludedDependenciesRounds)
 				{
-					RuntimeDependencies.Add(DependencyPackageName);
+					if (OutErrorMessage)
+					{
+						*OutErrorMessage = TEXT("More rounds than expected when handling included dependencies.");
+					}
+					return FCookDependencies();
 				}
-				Result.TransitiveBuildDependencies.Add(MoveTemp(CookDependency));
 			}
-			else if (CookDependency.GetType() == UE::Cook::ECookDependency::Package)
-			{
-				BuildDependenciesSet.Add(CookDependency.GetPackageName());
-			}
-			else
+
+			for (UE::Cook::FCookDependency& CookDependency : *RemainingDependencies)
 			{
-				Result.CookDependencies.Add(MoveTemp(CookDependency));
+				switch (CookDependency.GetType())
+				{
+				case UE::Cook::ECookDependency::TransitiveBuild:
+				{
+					// Build dependencies from a package to itself have a performance cost and serve no
+					// purpose, so remove them. They can occur in some systems that naively add a build dependency
+					// from one object to another without checking whether the second object is in the same package.
+					FName DependencyPackageName = CookDependency.GetPackageName();
+					if (DependencyPackageName == Result.PackageName)
+					{
+						continue;
+					}
+					if (CookDependency.IsAlsoAddRuntimeDependency())
+					{
+						RuntimeDependencies.Add(DependencyPackageName);
+					}
+					Result.TransitiveBuildDependencies.Add(MoveTemp(CookDependency));
+					break;
+				}
+				case UE::Cook::ECookDependency::Package:
+					BuildDependenciesSet.Add(CookDependency.GetPackageName());
+					break;
+				case UE::Cook::ECookDependency::Config:
+					Result.ConfigDependencies.Add(FString(CookDependency.GetConfigPath()));
+					break;
+				case UE::Cook::ECookDependency::SettingsObject:
+				{
+					const UObject* SettingsObject = CookDependency.GetSettingsObject();
+					if (SettingsObject)
+					{
+						// We rely on the object to be rooted because we use its pointer as a key for the lifetime of
+						// the cook process, so it being garbage collected and something else allocated on the same
+						// pointer would break our key. IsRooted should have been validated by
+						// FCookDependency::SettingsObject.
+						check(SettingsObject->IsRooted());
+						FCookDependencyGroups::FRecordedDependencies& IncludeDependencies =
+							FCookDependencyGroups::Get().FindOrCreate(reinterpret_cast<UPTRINT>(SettingsObject));
+						if (!IncludeDependencies.bInitialized)
+						{
+							IncludeDependencies.Dependencies = CollectSettingsObject(SettingsObject,
+								&IncludeDependencies.ErrorMessage);
+							IncludeDependencies.bInitialized = true;
+						}
+						if (!IncludeDependencies.Dependencies.IsValid())
+						{
+							if (OutErrorMessage)
+							{
+								*OutErrorMessage = FString::Printf(
+									TEXT("Dependencies for SettingsObject %s are unavailable: %s."),
+									*SettingsObject->GetPathName(), *IncludeDependencies.ErrorMessage);
+							}
+							return FCookDependencies();
+						}
+
+						for (const UE::Cook::FCookDependency& IncludeDependency :
+							IncludeDependencies.Dependencies.GetCookDependencies())
+						{
+							NextIncludedDependencies->Add(IncludeDependency);
+						}
+					}
+					break;
+				}
+				case UE::Cook::ECookDependency::NativeClass:
+					Result.ClassDependencies.Add(FString(CookDependency.GetClassPath()));
+					break;
+				default:
+					Result.CookDependencies.Add(MoveTemp(CookDependency));
+					break;
+				}
 			}
+			RemainingDependencies = NextIncludedDependencies;
 		}
-
-		Algo::Sort(Result.CookDependencies);
-		Algo::Sort(Result.TransitiveBuildDependencies);
 	}
 
+	Algo::Sort(Result.CookDependencies);
+	Result.CookDependencies.SetNum(Algo::Unique(Result.CookDependencies));
+	Algo::Sort(Result.TransitiveBuildDependencies);
+	Result.TransitiveBuildDependencies.SetNum(Algo::Unique(Result.TransitiveBuildDependencies));
+	Algo::Sort(Result.ConfigDependencies);
+	Result.ConfigDependencies.SetNum(Algo::Unique(Result.ConfigDependencies));
+	Algo::Sort(Result.ClassDependencies);
+	Result.ClassDependencies.SetNum(Algo::Unique(Result.ClassDependencies));
+
 	FName TransientPackageName = GetTransientPackage()->GetFName();
 	Result.BuildPackageDependencies = BuildDependenciesSet.Array();
 	Result.BuildPackageDependencies.RemoveAllSwap([TransientPackageName](FName InPackageName)
@@ -444,6 +550,65 @@ FCookDependencies FCookDependencies::Collect(UPackage* Package, const ITargetPla
 	return Result;
 }
 
+FCookDependencies FCookDependencies::CollectSettingsObject(const UObject* Object, FString* OutErrorMessage)
+{
+	if (!Object)
+	{
+		if (OutErrorMessage)
+		{
+			*OutErrorMessage = TEXT("Invalid null Object.");
+		}
+		return FCookDependencies();
+	}
+
+	UClass* Class = Object->GetClass();
+	if (!Class->HasAnyClassFlags(CLASS_Config | CLASS_PerObjectConfig))
+	{
+		if (OutErrorMessage)
+		{
+			*OutErrorMessage = FString::Printf(TEXT("Class %s is not a config class."), *Class->GetPathName());
+		}
+		return FCookDependencies();
+	}
+	if (!Class->HasAnyClassFlags(CLASS_PerObjectConfig) && Object != Class->GetDefaultObject())
+	{
+		if (OutErrorMessage)
+		{
+			*OutErrorMessage = FString::Printf(TEXT("Class %s is not a per-object-config class."),
+				*Class->GetPathName());
+		}
+		return FCookDependencies();
+	}
+
+	FCookDependencies Result;
+	TArray<UE::ConfigAccessTracking::FConfigAccessData> ConfigDatas;
+	const_cast<UObject*>(Object)->LoadConfig(nullptr /* ConfigClass */, nullptr /* Filename */,
+		UE::LCPF_None /* PropagationFlags */, nullptr /* PropertyToLoad */, &ConfigDatas);
+	Result.CookDependencies.Reserve(ConfigDatas.Num() + 1);
+	for (const UE::ConfigAccessTracking::FConfigAccessData& ConfigData : ConfigDatas)
+	{
+		Result.CookDependencies.Add(UE::Cook::FCookDependency::Config(ConfigData));
+	}
+
+	// In addition to adding the config dependencies, add a dependency on the class schema. If the current class has
+	// config fields A,B,C, we add dependencies on those config values. But if the class header is modified to have
+	// additional config field D then we need to rebuild packages that depend on it to record the new dependency on D.
+	UClass* NativeClass = Class;
+	while (NativeClass && !NativeClass->IsNative())
+	{
+		NativeClass = NativeClass->GetSuperClass();
+	}
+	if (NativeClass)
+	{
+		Result.CookDependencies.Add(UE::Cook::FCookDependency::NativeClass(NativeClass));
+	}
+
+	Algo::Sort(Result.CookDependencies);
+	Result.CookDependencies.SetNum(Algo::Unique(Result.CookDependencies));
+	Result.bValid = true;
+	return Result;
+}
+
 }
 
 bool LoadFromCompactBinary(FCbObjectView ObjectView, UE::TargetDomain::FCookDependencies& Dependencies)
@@ -513,6 +678,13 @@ bool LoadFromCompactBinary(FCbObjectView ObjectView, UE::TargetDomain::FCookDepe
 				return false;
 			}
 		}
+		if (FieldView.GetName().Equals(UTF8TEXTVIEW("ClassDependencies")))
+		{
+			if (!LoadFromCompactBinary(FieldView++, Dependencies.ClassDependencies))
+			{
+				return false;
+			}
+		}
 		if (FieldView == Last)
 		{
 			++FieldView;
@@ -557,6 +729,10 @@ FCbWriter& operator<<(FCbWriter& Writer, const UE::TargetDomain::FCookDependenci
 	{
 		Writer << "TransitiveBuildDependencies" << CookDependencies.TransitiveBuildDependencies;
 	}
+	if (!CookDependencies.ClassDependencies.IsEmpty())
+	{
+		Writer << "ClassDependencies" << CookDependencies.ClassDependencies;
+	}
 
 	Writer.EndObject();
 	return Writer;
@@ -565,6 +741,17 @@ FCbWriter& operator<<(FCbWriter& Writer, const UE::TargetDomain::FCookDependenci
 namespace UE::TargetDomain
 {
 
+FCookDependencyGroups& FCookDependencyGroups::Get()
+{
+	static FCookDependencyGroups Singleton;
+	return Singleton;
+}
+
+FCookDependencyGroups::FRecordedDependencies& FCookDependencyGroups::FindOrCreate(UPTRINT Key)
+{
+	return Groups.FindOrAdd(Key);
+}
+
 FBuildDefinitionList FBuildDefinitionList::Collect(UPackage* Package, const ITargetPlatform* TargetPlatform,
 	FString* OutErrorMessage)
 {
diff --git a/Engine/Source/Editor/UnrealEd/Private/TargetDomain/TargetDomainUtils.h b/Engine/Source/Editor/UnrealEd/Private/TargetDomain/TargetDomainUtils.h
--- a/Engine/Source/Editor/UnrealEd/Private/TargetDomain/TargetDomainUtils.h
+++ b/Engine/Source/Editor/UnrealEd/Private/TargetDomain/TargetDomainUtils.h
@@ -103,6 +103,7 @@ public:
 	static FCookDependencies Collect(UPackage* Package, const ITargetPlatform* TargetPlatform,
 		FSavePackageResultStruct* SaveResult, const FGeneratedPackageResultStruct* GeneratedResult,
 		TArray<FName>&& RuntimeDependencies, FString* OutErrorMessage = nullptr);
+	static FCookDependencies CollectSettingsObject(const UObject* Object, FString* OutErrorMessage);
 
 	// Fetch function to load the dependencies from a PackageStore is not yet implemented independently for
 	// this structure. Use FCookAttachments instead. 
@@ -114,6 +115,7 @@ private:
 	TArray<FName> ScriptPackageDependencies;
 	TArray<UE::Cook::FCookDependency> CookDependencies;
 	TArray<UE::Cook::FCookDependency> TransitiveBuildDependencies;
+	TArray<FString> ClassDependencies;
 	FName PackageName;
 	FIoHash StoredKey;
 	FIoHash CurrentKey;
@@ -124,6 +126,31 @@ private:
 	friend FCookAttachments;
 };
 
+/**
+ * Non-persistent cache of groups of cookdependencies. Dependencies to a CookDependencyGroup are not persistently
+ * recorded into the oplog, instead we make a copy of all of their dependencies and append those dependencies onto
+ * the CookDependencies that are written for a package.
+ *
+ * Example: The cookdependencies used by the CDO of a settings object that itself is configured by config values.
+ *          The settings object's class's schema and the list of config settings are included in the cookdependencies.
+ */
+class FCookDependencyGroups
+{
+public:
+	struct FRecordedDependencies
+	{
+		FCookDependencies Dependencies;
+		FString ErrorMessage;
+		bool bInitialized = false;
+	};
+
+	static FCookDependencyGroups& Get();
+	FRecordedDependencies& FindOrCreate(UPTRINT Key);
+
+private:
+	TMap<UPTRINT, FRecordedDependencies> Groups;
+};
+
 /** Wrapper around TArray<FBuildDefinition>, used to provide custom functions for compactbinary, collection, and fetch */
 struct FBuildDefinitionList
 {
diff --git a/Engine/Source/Runtime/CoreUObject/Private/Cooker/CookDependency.cpp b/Engine/Source/Runtime/CoreUObject/Private/Cooker/CookDependency.cpp
--- a/Engine/Source/Runtime/CoreUObject/Private/Cooker/CookDependency.cpp
+++ b/Engine/Source/Runtime/CoreUObject/Private/Cooker/CookDependency.cpp
@@ -11,12 +11,15 @@
 #include "HAL/PlatformMath.h"
 #include "Hash/Blake3.h"
 #include "Misc/AssertionMacros.h"
+#include "Misc/ConfigAccessData.h"
 #include "Misc/StringBuilder.h"
 #include "Serialization/Archive.h"
 #include "Serialization/CompactBinary.h"
 #include "Serialization/CompactBinarySerialization.h"
 #include "Serialization/CompactBinaryWriter.h"
+#include "Templates/Casts.h"
 #include "Templates/UniquePtr.h"
+#include "UObject/Class.h"
 #endif
 
 #if WITH_EDITOR
@@ -73,6 +76,96 @@ FCookDependency FCookDependency::ConsoleVariable(FStringView VariableName)
 	return Result;
 }
 
+FCookDependency FCookDependency::Config(UE::ConfigAccessTracking::FConfigAccessData AccessData)
+{
+	FCookDependency Result(ECookDependency::Config);
+	Result.ConfigAccessData.Reset(new UE::ConfigAccessTracking::FConfigAccessData(MoveTemp(AccessData)));
+	return Result;
+}
+
+FCookDependency FCookDependency::Config(UE::ConfigAccessTracking::ELoadType LoadType, FName Platform,
+	FName FileName, FName SectionName, FName ValueName)
+{
+	FCookDependency Result(ECookDependency::Config);
+	Result.ConfigAccessData.Reset(new UE::ConfigAccessTracking::FConfigAccessData(
+		LoadType, Platform, FileName, SectionName, ValueName, nullptr /* RequestPlatform */));
+	return Result;
+}
+
+FCookDependency FCookDependency::Config(FName FileName, FName SectionName, FName ValueName)
+{
+	return Config(UE::ConfigAccessTracking::ELoadType::ConfigSystem, NAME_None,
+		FileName, SectionName, ValueName);
+}
+
+FCookDependency FCookDependency::SettingsObject(const UObject* InObject)
+{
+	FCookDependency Result(ECookDependency::SettingsObject);
+	if (InObject)
+	{
+		const UClass* Class = Cast<const UClass>(InObject);
+		if (Class)
+		{
+			InObject = Class->GetDefaultObject();
+		}
+		else
+		{
+			Class = InObject->GetClass();
+		}
+
+		if (!InObject->IsRooted())
+		{
+			UE_LOG(LogCore, Error, TEXT("Invalid FCookDependency::SettingsObject(%s). The object is not in the root set and may be garbage collected. ")
+				TEXT("FCookDependency keeps a raw pointer to SettingsObjects and does not support pointers to objects that are not in the root set. ")
+				TEXT("The dependency will be ignored."),
+				*InObject->GetPathName());
+			InObject = nullptr;
+		}
+		else
+		{
+			if (!Class->HasAnyClassFlags(CLASS_Config | CLASS_PerObjectConfig))
+			{
+				UE_LOG(LogCore, Error, TEXT("Invalid FCookDependency::SettingsObject(%s). The object's class %s is not a config class. CookDependency::SettingsObject only supports config classes. ")
+					TEXT("The dependency will be ignored."),
+					*InObject->GetPathName(), *Class->GetPathName());
+				InObject = nullptr;
+			}
+			else if (!Class->HasAnyClassFlags(CLASS_PerObjectConfig) && InObject != Class->GetDefaultObject())
+			{
+				UE_LOG(LogCore, Error, TEXT("Invalid FCookDependency::SettingsObject(%s). The object is not the ClassDefaultObject and its class %s is not a per-object-config class. ")
+					TEXT("CookDependency::SettingsObject only supports the CDO or per-object-config objects. ")
+					TEXT("The dependency will be ignored."),
+					*InObject->GetPathName(), *Class->GetPathName());
+				InObject = nullptr;
+			}
+		}
+	}
+	Result.ObjectPtr = InObject;
+	return Result;
+}
+
+FCookDependency FCookDependency::NativeClass(const UClass* InClass)
+{
+	if (InClass)
+	{
+		if (!InClass->IsNative())
+		{
+			UE_LOG(LogCore, Error, TEXT("Invalid FCookDependency::NativeClass(%s). The class is not native. ")
+				TEXT("The dependency will be ignored."),
+				InClass ? *InClass->GetPathName() : TEXT("<null>"));
+			InClass = nullptr;
+		}
+	}
+	return NativeClass(InClass ? InClass->GetPathName() : FStringView());
+}
+
+FCookDependency FCookDependency::NativeClass(FStringView ClassPath)
+{
+	FCookDependency Result(ECookDependency::NativeClass);
+	Result.StringData = ClassPath;
+	return Result;
+}
+
 FCookDependency::FCookDependency()
 	: Type(ECookDependency::None)
 {
@@ -112,6 +205,7 @@ FCookDependency& FCookDependency::operator=(const FCookDependency& Other)
 		break;
 	case ECookDependency::File:
 	case ECookDependency::ConsoleVariable:
+	case ECookDependency::NativeClass:
 		StringData = Other.StringData;
 		break;
 	case ECookDependency::Function:
@@ -123,6 +217,14 @@ FCookDependency& FCookDependency::operator=(const FCookDependency& Other)
 	case ECookDependency::Package:
 		NameData = Other.NameData;
 		break;
+	case ECookDependency::Config:
+		ConfigAccessData.Reset(Other.ConfigAccessData.IsValid() ?
+			new UE::ConfigAccessTracking::FConfigAccessData(*Other.ConfigAccessData) :
+			nullptr);
+		break;
+	case ECookDependency::SettingsObject:
+		ObjectPtr = Other.ObjectPtr;
+		break;
 	default:
 		checkNoEntry();
 		break;
@@ -141,6 +243,7 @@ FCookDependency& FCookDependency::operator=(FCookDependency&& Other)
 		break;
 	case ECookDependency::File:
 	case ECookDependency::ConsoleVariable:
+	case ECookDependency::NativeClass:
 		StringData = MoveTemp(Other.StringData);
 		break;
 	case ECookDependency::Function:
@@ -152,6 +255,12 @@ FCookDependency& FCookDependency::operator=(FCookDependency&& Other)
 	case ECookDependency::Package:
 		NameData = MoveTemp(Other.NameData);
 		break;
+	case ECookDependency::Config:
+		ConfigAccessData = MoveTemp(Other.ConfigAccessData);
+		break;
+	case ECookDependency::SettingsObject:
+		ObjectPtr = Other.ObjectPtr;
+		break;
 	default:
 		checkNoEntry();
 		break;
@@ -159,6 +268,12 @@ FCookDependency& FCookDependency::operator=(FCookDependency&& Other)
 	return *this;
 }
 
+FString FCookDependency::GetConfigPath() const
+{
+	return Type == ECookDependency::Config && ConfigAccessData.IsValid() ?
+		ConfigAccessData->FullPathToString() : FString();
+}
+
 void FCookDependency::UpdateHash(FCookDependencyContext& Context) const
 {
 	switch (GetType())
@@ -236,6 +351,21 @@ void FCookDependency::UpdateHash(FCookDependencyContext& Context) const
 		Context.Update(VariableAsString.GetCharArray().GetData(), VariableAsString.GetCharArray().Num() * sizeof(FString::ElementType));
 		return;
 	}
+	case ECookDependency::Config:
+		Context.LogError(FString::Printf(
+			TEXT("FCookDependency::Config('%s') failed to UpdateHash: Config dependencies do not implement UpdateHash and it should not be called on them."),
+			*GetConfigPath()));
+		return;
+	case ECookDependency::SettingsObject:
+		Context.LogError(FString::Printf(
+			TEXT("FCookDependency::SettingsObject('%s') failed to UpdateHash: SettingsObject dependencies do not implement UpdateHash and it should not be called on them."),
+			ObjectPtr ? *ObjectPtr->GetPathName() : TEXT("<null>")));
+		return;
+	case ECookDependency::NativeClass:
+		Context.LogError(FString::Printf(
+			TEXT("FCookDependency::NativeClass('%s') failed to UpdateHash: NativeClass dependencies do not implement UpdateHash and it should not be called on them."),
+			*StringData));
+		return;
 	default:
 		checkNoEntry();
 		return;
@@ -250,6 +380,7 @@ void FCookDependency::Construct()
 		break;
 	case ECookDependency::File:
 	case ECookDependency::ConsoleVariable:
+	case ECookDependency::NativeClass:
 		new(&StringData) FString();
 		break;
 	case ECookDependency::Function:
@@ -261,6 +392,13 @@ void FCookDependency::Construct()
 	case ECookDependency::Package:
 		new(&NameData) FName();
 		break;
+	case ECookDependency::Config:
+		// Set the union's bytes equal to a TUniquePtr with a null internal pointer
+		new(&ConfigAccessData) TUniquePtr<UE::ConfigAccessTracking::FConfigAccessData>(nullptr);
+		break;
+	case ECookDependency::SettingsObject:
+		ObjectPtr = nullptr;
+		break;
 	default:
 		checkNoEntry();
 		break;
@@ -275,6 +413,7 @@ void FCookDependency::Destruct()
 		break;
 	case ECookDependency::File:
 	case ECookDependency::ConsoleVariable:
+	case ECookDependency::NativeClass:
 		StringData.~FString();
 		break;
 	case ECookDependency::Function:
@@ -286,6 +425,13 @@ void FCookDependency::Destruct()
 	case ECookDependency::Package:
 		NameData.~FName();
 		break;
+	case ECookDependency::Config:
+		// TUniquePtr's destructor will also destruct its internal pointer
+		ConfigAccessData.~TUniquePtr<UE::ConfigAccessTracking::FConfigAccessData>();
+		break;
+	case ECookDependency::SettingsObject:
+		ObjectPtr = nullptr;
+		break;
 	default:
 		checkNoEntry();
 		break;
@@ -302,6 +448,7 @@ void FCookDependency::Save(FCbWriter& Writer) const
 		break;
 	case ECookDependency::File:
 	case ECookDependency::ConsoleVariable:
+	case ECookDependency::NativeClass:
 		Writer << StringData;
 		break;
 	case ECookDependency::Function:
@@ -318,6 +465,12 @@ void FCookDependency::Save(FCbWriter& Writer) const
 	case ECookDependency::Package:
 		Writer << NameData;
 		break;
+	case ECookDependency::Config:
+		Writer << GetConfigPath();
+		break;
+	case ECookDependency::SettingsObject:
+		// Settings objects are not persistable; save out an empty SettingsObject dependency
+		break;
 	default:
 		checkNoEntry();
 		break;
@@ -402,6 +555,39 @@ bool FCookDependency::Load(FCbFieldView Value)
 		*this = FCookDependency::ConsoleVariable(LocalConsoleVariableName);
 		return true;
 	}
+	case ECookDependency::Config:
+	{
+		FString LocalConfigPath;
+		if (!LoadFromCompactBinary(Field++, LocalConfigPath))
+		{
+			return false;
+		}
+		if (LocalConfigPath.IsEmpty())
+		{
+			*this = FCookDependency(ECookDependency::Config);
+		}
+		else
+		{
+			*this = FCookDependency::Config(UE::ConfigAccessTracking::FConfigAccessData::Parse(LocalConfigPath));
+		}
+		return true;
+	}
+	case ECookDependency::SettingsObject:
+	{
+		// Settings objects are not persistable; construct an empty SettingsObject dependency
+		*this = FCookDependency::SettingsObject(nullptr);
+		return true;
+	}
+	case ECookDependency::NativeClass:
+	{
+		FString LocalClassPath;
+		if (!LoadFromCompactBinary(Field++, LocalClassPath))
+		{
+			return false;
+		}
+		*this = FCookDependency::NativeClass(LocalClassPath);
+		return true;
+	}
 	default:
 		break;
 	}
@@ -409,6 +595,18 @@ bool FCookDependency::Load(FCbFieldView Value)
 	return false;
 }
 
+bool FCookDependency::ConfigAccessDataLessThan(const UE::ConfigAccessTracking::FConfigAccessData& A,
+	const UE::ConfigAccessTracking::FConfigAccessData& B)
+{
+	return A < B;
+}
+
+bool FCookDependency::ConfigAccessDataEqual(const UE::ConfigAccessTracking::FConfigAccessData& A,
+	const UE::ConfigAccessTracking::FConfigAccessData& B)
+{
+	return A == B;
+}
+
 void FCookDependencyContext::Update(const void* Data, uint64 Size)
 {
 	FBlake3& Blake3Hasher = *(reinterpret_cast<FBlake3*>(Hasher));
diff --git a/Engine/Source/Runtime/CoreUObject/Public/Cooker/CookDependency.h b/Engine/Source/Runtime/CoreUObject/Public/Cooker/CookDependency.h
--- a/Engine/Source/Runtime/CoreUObject/Public/Cooker/CookDependency.h
+++ b/Engine/Source/Runtime/CoreUObject/Public/Cooker/CookDependency.h
@@ -8,10 +8,14 @@
 #include "HAL/PreprocessorHelpers.h"
 #include "Serialization/CompactBinary.h"
 #include "Templates/Function.h"
+#include "Templates/UniquePtr.h"
 #include "UObject/NameTypes.h"
 
 class FCbFieldView;
 class FCbWriter;
+class UObject;
+namespace UE::ConfigAccessTracking { enum class ELoadType : uint8; }
+namespace UE::ConfigAccessTracking { struct FConfigAccessData; }
 #endif
 
 #if WITH_EDITOR
@@ -80,12 +84,15 @@ private:
  */
 enum class ECookDependency : uint8
 {
-	None = 0x00,
-	File = 0x01,
-	Function = 0x02,
-	TransitiveBuild = 0x03,
-	Package = 0x4,
-	ConsoleVariable = 0x5,
+	None					= 0,
+	File					= 1,
+	Function				= 2,
+	TransitiveBuild			= 3,
+	Package					= 4,
+	ConsoleVariable			= 5,
+	Config					= 6,
+	SettingsObject			= 7,
+	NativeClass				= 8,
 
 	Count,
 };
@@ -140,6 +147,27 @@ public:
 	 */
 	COREUOBJECT_API static FCookDependency ConsoleVariable(FStringView VariableName);
 
+	/** Create a dependency on the value of a config variable. */
+	COREUOBJECT_API static FCookDependency Config(UE::ConfigAccessTracking::FConfigAccessData AccessData);
+	COREUOBJECT_API static FCookDependency Config(UE::ConfigAccessTracking::ELoadType LoadType, FName Platform,
+		FName FileName, FName SectionName, FName ValueName);
+	/** Create a dependency on the value of a config variable, with LoadType=ConfigSystem and Platform=NAME_None. */
+	COREUOBJECT_API static FCookDependency Config(FName FileName, FName SectionName, FName ValueName);
+
+	/**
+	 * Adds a dependency on the config values and class schema of a settings object. Gives an error and ignores the object
+	 * if the object is not a config-driven settings object, such as the CDO of a config UClass or a perObjectConfig
+	 * object.
+	 *
+	 * SettingsObject dependencies are not directly persistable; all of the dependencies reported by the SettingsObject are
+	 * copied onto the dependencies of the package declaring the SettingsObject dependency.
+	 */
+	COREUOBJECT_API static FCookDependency SettingsObject(const UObject* InObject);
+
+	/** Adds a dependency on the class schema of a nativeclass. */
+	COREUOBJECT_API static FCookDependency NativeClass(const UClass* InClass);
+	COREUOBJECT_API static FCookDependency NativeClass(FStringView ClassPath);
+
 	/** Construct an empty dependency; it will never be invalidated. */
 	COREUOBJECT_API FCookDependency();
 
@@ -165,11 +193,32 @@ public:
 	/** If GetType() == TransitiveBuild, whether AlsoAddRuntimeDependency was selected, otherwise false. */
 	bool IsAlsoAddRuntimeDependency() const;
 
+	/**
+	 * Returns the full path of the config access (e.g. Platform.Filename.Section.ValueName)
+	 * if GetType() == Config, else empty.
+	 */
+	COREUOBJECT_API FString GetConfigPath() const;
+
+	/**
+	 * Returns the SettingsObject pointer if GetType() == SettingsObject, else nullptr. Can also be null for
+	 * SettingsObject that was found to be invalid.
+	 */
+	const UObject* GetSettingsObject() const;
+
+	/**
+	 * Returns the classpath if GetType() == NativeClass, else empty string. Can also be empty for
+	 * NativeClass that was found to be invalid.
+	 */
+	FStringView GetClassPath() const;
+
 	/**
 	 * Comparison operator for e.g. deterministic ordering of dependencies.
 	 * Uses persistent comparison data and is somewhat expensive.
 	 */
 	bool operator<(const FCookDependency& Other) const;
+	/** Equality operator for uniqueness testing */
+	bool operator==(const FCookDependency& Other) const;
+	bool operator!=(const FCookDependency& Other) const;
 
 	/** Calculate the current hash of this CookDependency, and add it into Context. */
 	COREUOBJECT_API void UpdateHash(FCookDependencyContext& Context) const;
@@ -180,6 +229,10 @@ private:
 	void Destruct();
 	COREUOBJECT_API void Save(FCbWriter& Writer) const;
 	COREUOBJECT_API bool Load(FCbFieldView Value);
+	static COREUOBJECT_API bool ConfigAccessDataLessThan(const UE::ConfigAccessTracking::FConfigAccessData& A,
+		const UE::ConfigAccessTracking::FConfigAccessData& B);
+	static COREUOBJECT_API bool ConfigAccessDataEqual(const UE::ConfigAccessTracking::FConfigAccessData& A,
+		const UE::ConfigAccessTracking::FConfigAccessData& B);
 
 	/** Public hidden friend for operator<< into an FCbWriter. */
 	friend FCbWriter& operator<<(FCbWriter& Writer, const FCookDependency& CookDependencies)
@@ -211,6 +264,8 @@ private:
 		FFunctionData FunctionData;
 		FTransitiveBuildData TransitiveBuildData;
 		FName NameData;
+		const UObject* ObjectPtr;
+		TUniquePtr<UE::ConfigAccessTracking::FConfigAccessData> ConfigAccessData;
 	};
 };
 
@@ -326,6 +381,16 @@ inline bool FCookDependency::IsAlsoAddRuntimeDependency() const
 	return Type == ECookDependency::TransitiveBuild ? TransitiveBuildData.bAlsoAddRuntimeDependency : false;
 }
 
+inline const UObject* FCookDependency::GetSettingsObject() const
+{
+	return Type == ECookDependency::SettingsObject ? ObjectPtr : nullptr;
+}
+
+inline FStringView FCookDependency::GetClassPath() const
+{
+	return Type == ECookDependency::NativeClass ? StringData : FStringView();
+}
+
 inline bool FCookDependency::operator<(const FCookDependency& Other) const
 {
 	if (static_cast<uint8>(Type) != static_cast<uint8>(Other.Type))
@@ -339,6 +404,7 @@ inline bool FCookDependency::operator<(const FCookDependency& Other) const
 		return false;
 	case ECookDependency::File:
 	case ECookDependency::ConsoleVariable:
+	case ECookDependency::NativeClass:
 		return StringData.Compare(Other.StringData, ESearchCase::IgnoreCase) < 0;
 	case ECookDependency::Function:
 	{
@@ -372,13 +438,93 @@ inline bool FCookDependency::operator<(const FCookDependency& Other) const
 		return false;
 	}
 	case ECookDependency::Package:
-		return NameData.Compare(Other.NameData) <  0;
+		return NameData.Compare(Other.NameData) < 0;
+	case ECookDependency::Config:
+		if (ConfigAccessData.IsValid() != Other.ConfigAccessData.IsValid())
+		{
+			return !ConfigAccessData.IsValid();
+		}
+		if (!ConfigAccessData.IsValid())
+		{
+			return false; // equal
+		}
+		return ConfigAccessDataLessThan(*ConfigAccessData, *Other.ConfigAccessData);
+	case ECookDependency::SettingsObject:
+		// SettingsObjects are not persistable, so we do not use a persistent sort key; just the object ptr.
+		return ObjectPtr < Other.ObjectPtr;
 	default:
 		checkNoEntry();
 		return false;
 	}
 }
 
+inline bool FCookDependency::operator==(const FCookDependency& Other) const
+{
+	if (static_cast<uint8>(Type) != static_cast<uint8>(Other.Type))
+	{
+		return false;
+	}
+
+	switch (Type)
+	{
+	case ECookDependency::None:
+		return true;
+	case ECookDependency::File:
+	case ECookDependency::ConsoleVariable:
+	case ECookDependency::NativeClass:
+		return StringData.Compare(Other.StringData, ESearchCase::IgnoreCase) == 0;
+	case ECookDependency::Function:
+	{
+		if (FunctionData.Name.Compare(Other.FunctionData.Name) != 0)
+		{
+			return false;
+		}
+		FMemoryView ViewA;
+		FMemoryView ViewB;
+		bool bHasViewA = FunctionData.Args.TryGetRangeView(ViewA);
+		bool bHasViewB = FunctionData.Args.TryGetRangeView(ViewB);
+		if (bHasViewA != bHasViewB)
+		{
+			return false;
+		}
+		return ViewA.CompareBytes(ViewB) == 0;
+	}
+	case ECookDependency::TransitiveBuild:
+	{
+		// FName.Compare is lexical and case-insensitive, which is what we want
+		int32 Compare = TransitiveBuildData.PackageName.Compare(Other.TransitiveBuildData.PackageName);
+		if (Compare != 0)
+		{
+			return false;
+		}
+		return TransitiveBuildData.bAlsoAddRuntimeDependency == Other.TransitiveBuildData.bAlsoAddRuntimeDependency;
+	}
+	case ECookDependency::Package:
+		return NameData.Compare(Other.NameData) == 0;
+	case ECookDependency::Config:
+		if (ConfigAccessData.IsValid() != Other.ConfigAccessData.IsValid())
+		{
+			return false;
+		}
+		if (!ConfigAccessData.IsValid())
+		{
+			return true;
+		}
+		return ConfigAccessDataEqual(*ConfigAccessData, *Other.ConfigAccessData);
+	case ECookDependency::SettingsObject:
+		// SettingsObjects are not persistable, so we do not use a persistent sort key; just the object ptr.
+		return ObjectPtr == Other.ObjectPtr;
+	default:
+		checkNoEntry();
+		return false;
+	}
+}
+
+inline bool FCookDependency::operator!=(const FCookDependency& Other) const
+{
+	return !(*this == Other);
+}
+
 inline FCookDependencyContext::FCookDependencyContext(void* InHasher, TUniqueFunction<void(FString&&)>&& InOnLogError, FName InPackageName)
 	: OnLogError(MoveTemp(InOnLogError))
 	, PackageName(InPackageName)
diff --git a/Engine/Source/Runtime/Engine/Private/Texture2D.cpp b/Engine/Source/Runtime/Engine/Private/Texture2D.cpp
--- a/Engine/Source/Runtime/Engine/Private/Texture2D.cpp
+++ b/Engine/Source/Runtime/Engine/Private/Texture2D.cpp
@@ -48,6 +48,11 @@
 #include "UObject/StrongObjectPtr.h"
 #include "ProfilingDebugging/AssetMetadataTrace.h"
 
+#if WITH_EDITOR
+#include "Cooker/CookDependency.h"
+#include "Interfaces/ITargetPlatform.h"
+#endif
+
 #include UE_INLINE_GENERATED_CPP_BY_NAME(Texture2D)
 
 #if WITH_EDITORONLY_DATA
@@ -681,6 +686,14 @@ void UTexture2D::PreSave(FObjectPreSaveContext ObjectSaveContext)
 		UpdateResource();
 	}
 
+	if (ObjectSaveContext.IsCooking())
+	{
+		const ITargetPlatform* TargetPlatform = ObjectSaveContext.GetTargetPlatform();
+		check(TargetPlatform);
+		ObjectSaveContext.AddCookBuildDependency(UE::Cook::FCookDependency::SettingsObject(
+			&TargetPlatform->GetTextureLODSettings()));
+	}
+
 	// #TODO DC This is redundant code coming from UTexture::Presave that can be removed once we remove the above streaming code.
 
 	// Ensure that compilation has finished before saving the package
