[UnrealEd] Fixed rotation handle for combined translate/rotate in local space

diff --git a/Engine/Source/Editor/UnrealEd/Private/UnrealWidgetRender.cpp b/Engine/Source/Editor/UnrealEd/Private/UnrealWidgetRender.cpp
--- a/Engine/Source/Editor/UnrealEd/Private/UnrealWidgetRender.cpp
+++ b/Engine/Source/Editor/UnrealEd/Private/UnrealWidgetRender.cpp
@@ -889,8 +889,21 @@ void FWidget::Render_TranslateRotateZ( const FSceneView* View, FPrimitiveDrawInt
 		{
 			PDI->SetHitProxy( new HWidgetAxis(EAxisList::ZRotation, bDisabled) );
 			{
-				FVector XAxis = CustomCoordSystem.TransformPosition( FVector(1,0,0).RotateAngleAxis( (EditorModeTools ? EditorModeTools->TranslateRotateXAxisAngle : 0 ), FVector(0,0,1)) );
-				FVector YAxis = CustomCoordSystem.TransformPosition( FVector(0,1,0).RotateAngleAxis( (EditorModeTools ? EditorModeTools->TranslateRotateXAxisAngle : 0 ), FVector(0,0,1)) );
+				FVector XAxis;
+				FVector YAxis;
+				if (Space.bIsLocalSpace)
+				{
+					// In local space, only need to set axis
+					XAxis = CustomCoordSystem.TransformPosition(FVector::ForwardVector);
+					YAxis = CustomCoordSystem.TransformPosition(FVector::RightVector);
+				}
+				else
+				{
+					// In world space, need to both set and modify axis
+					XAxis = CustomCoordSystem.TransformPosition(FVector::ForwardVector.RotateAngleAxis((EditorModeTools ? EditorModeTools->TranslateRotateXAxisAngle : 0), FVector::UpVector));
+					YAxis = CustomCoordSystem.TransformPosition(FVector::RightVector.RotateAngleAxis((EditorModeTools ? EditorModeTools->TranslateRotateXAxisAngle : 0), FVector::UpVector));
+				}
+
 				FVector BaseArrowPoint = InLocation + XAxis * ScaledRadius;
 				DrawFlatArrow(PDI, BaseArrowPoint, XAxis, YAxis, ZRotateColor, ScaledRadius, ScaledRadius*.5f, ZRotateMaterial->GetRenderProxy(), SDPG_Foreground);
 			}
