[UnrealEd] Removed "Jump to bookmark" toast notification (add and remove notifications remain)

diff --git a/Engine/Source/Editor/UnrealEd/Private/Bookmarks/BookmarkTypeToolsImpl.cpp b/Engine/Source/Editor/UnrealEd/Private/Bookmarks/BookmarkTypeToolsImpl.cpp
--- a/Engine/Source/Editor/UnrealEd/Private/Bookmarks/BookmarkTypeToolsImpl.cpp
+++ b/Engine/Source/Editor/UnrealEd/Private/Bookmarks/BookmarkTypeToolsImpl.cpp
@@ -184,9 +184,6 @@ public:
 					{
 						Actions->JumpToBookmark(Bookmark, InSettings, *InViewportClient);
 					}
-
-					const FText NotifyMessage = LOCTEXT("JumpToBookmarkSuccessFormat", "Jumped to bookmark {0}");
-					Notify(FText::Format(NotifyMessage, InIndex), true);
 				}
 				else
 				{
