ODSC server fixes / startup time improvements:
* iterative was only specified for cook by the book, but ODSC is using COTF
* GetNeverCookPackageFileNames is called at startup before the asset registry is done. This ends up calling FPackageName::FindPackagesInDirectories which can be quite slow. Since ODSC is not supposed to cook anything in the first place, we should be able to skip the function call

#rnx

diff --git a/Engine/Source/Editor/UnrealEd/Private/Commandlets/CookCommandlet.cpp b/Engine/Source/Editor/UnrealEd/Private/Commandlets/CookCommandlet.cpp
--- a/Engine/Source/Editor/UnrealEd/Private/Commandlets/CookCommandlet.cpp
+++ b/Engine/Source/Editor/UnrealEd/Private/Commandlets/CookCommandlet.cpp
@@ -245,6 +245,13 @@ int32 UCookCommandlet::Main(const FString& CmdLineParams)
 			Timeout = 180;
 		}
 
+		if (Switches.Contains(TEXT("ODSC")))
+		{
+			// ODSC piggybacks on the cook commandlet but does not cook any packages. Turn -iterative on to prevent
+			// an unnecessary clearing of the cook results.
+			bIterativeCooking = true;
+		}
+
 		CookOnTheFly( InstanceId, Port, Timeout, bForceClose, TargetPlatforms);
 	}
 	else if (Switches.Contains(TEXT("COOKWORKER")))
diff --git a/Engine/Source/Editor/UnrealEd/Private/CookOnTheFlyServer.cpp b/Engine/Source/Editor/UnrealEd/Private/CookOnTheFlyServer.cpp
--- a/Engine/Source/Editor/UnrealEd/Private/CookOnTheFlyServer.cpp
+++ b/Engine/Source/Editor/UnrealEd/Private/CookOnTheFlyServer.cpp
@@ -12372,6 +12372,11 @@ TArray<FName> UCookOnTheFlyServer::GetNeverCookPackageFileNames(TArrayView<const
 {
 	TArray<FString> NeverCookDirectories(ExtraNeverCookDirectories);
 
+	if (bRunningAsShaderServer)
+	{
+		return TArray<FName>();
+	}
+
 	auto AddDirectoryPathArray = [&NeverCookDirectories](const TArray<FDirectoryPath>& DirectoriesToNeverCook, const TCHAR* SettingName)
 	{
 		for (const FDirectoryPath& DirToNotCook : DirectoriesToNeverCook)
