Avoid crash due to FPropertySoftPath::ResolvePath attempting to resolve against the incorrect object when using sparse class data
[RN] Fixed crash when diffing class objects with sparse data

diff --git a/Engine/Source/Editor/UnrealEd/Private/DiffUtils.cpp b/Engine/Source/Editor/UnrealEd/Private/DiffUtils.cpp
--- a/Engine/Source/Editor/UnrealEd/Private/DiffUtils.cpp
+++ b/Engine/Source/Editor/UnrealEd/Private/DiffUtils.cpp
@@ -489,6 +489,8 @@ FPropertyPath FPropertySoftPath::ResolvePath(const UObject* Object) const
 				if (SparseClassDataStruct && SparseClassDataStruct->IsChildOf(RootTypeHint))
 				{
 					ResolvedProperty = UEDiffUtils_Private::Resolve(RootTypeHint, PropertyIdentifier);
+					// return if null won't mutate...
+					ContainerAddress = const_cast<UClass*>(AsClass)->GetSparseClassData(EGetSparseClassDataMethod::ReturnIfNull);
 				}
 			}
 		}
