[EditorObject ImportProperties] fix for importing deep object hierarchies in SceneGraph
When importing hierarchies of objects with instanced sub-objects we don't want to use the instancing graph of our outer objects since we need our dedicated instances.

diff --git a/Engine/Source/Editor/UnrealEd/Private/EditorObject.cpp b/Engine/Source/Editor/UnrealEd/Private/EditorObject.cpp
--- a/Engine/Source/Editor/UnrealEd/Private/EditorObject.cpp
+++ b/Engine/Source/Editor/UnrealEd/Private/EditorObject.cpp
@@ -12,6 +12,7 @@
 #include "UObject/Object.h"
 #include "UObject/Class.h"
 #include "UObject/UnrealType.h"
+#include "UObject/OverridableManager.h"
 #include "UObject/PropertyPortFlags.h"
 #include "Serialization/ArchiveReplaceObjectRef.h"
 #include "Misc/AsciiSet.h"
@@ -622,6 +623,21 @@ static const TCHAR* ImportProperties(
 
 				if (!ComponentTemplate)
 				{
+					// When importing hierarchies of objects with instanced sub-objects we don't want to use the
+					// instancing graph of our outer objects since we need our dedicated instances.
+					// For example:
+					//
+					//	A_CDO { SubObj }
+					//	PastedObj { A_1{ SubObj_1 } A_2{ SubObj_1 } }
+					//
+					//	by sharing the instancing graph, we end up with both SubObj_1 being the same pointer
+					//	by using a dedicated instancing graph, both SubObj_1 are newly created object pointer 
+					//
+					// Ideally this rule should be applied more widely but for now we limit it to the Prefab system
+					// which we could identify using the fact it uses overridable serialization.
+					const bool bShouldUseDedicatedInstancingGraph = FOverridableManager::Get().IsEnabled(*SubobjectOuter);
+					FObjectInstancingGraph InstanceGraphForNewObject;
+
 					ComponentTemplate = NewObject<UObject>(
 						SubobjectOuter,
 						TemplateClass,
@@ -629,7 +645,7 @@ static const TCHAR* ImportProperties(
 						NewFlags,
 						Archetype,
 						!!SubobjectOuter,
-						&InstanceGraph
+						bShouldUseDedicatedInstancingGraph ? &InstanceGraphForNewObject : &InstanceGraph
 						);
 				}
 				else
