Adding multiprocess support for -iterativevalidate (2-phase and 1-phase modes). Added additional hooks to IMPCollector for this and for licensee use (ServerTick and ServerTickPackage) and provided the package name to FMPCollectorClientMessageContext.

diff --git a/Engine/Source/Editor/UnrealEd/Private/Cooker/CompactBinaryTCP.cpp b/Engine/Source/Editor/UnrealEd/Private/Cooker/CompactBinaryTCP.cpp
--- a/Engine/Source/Editor/UnrealEd/Private/Cooker/CompactBinaryTCP.cpp
+++ b/Engine/Source/Editor/UnrealEd/Private/Cooker/CompactBinaryTCP.cpp
@@ -366,3 +366,25 @@ bool LoadFromCompactBinary(FCbFieldView Field, FSoftObjectPathSerializationWrapp
 	Path.Inner.SetPath(PathString);
 	return true;
 }
+
+bool LoadFromCompactBinary(FCbFieldView Field, UE::CompactBinaryTCP::FMarshalledMessage& Value)
+{
+	bool bOk = !Field.HasError();
+	bOk &= LoadFromCompactBinary(Field["T"], Value.MessageType);
+
+	if (bOk)
+	{
+		FCbObjectView ValueView = Field["V"].AsObjectView();
+		bOk &= !Field.HasError();
+		if (bOk)
+		{
+			Value.Object = FCbObject::Clone(ValueView);
+		}
+	}
+	if (!bOk)
+	{
+		Value.MessageType.Invalidate();
+		Value.Object.Reset();
+	}
+	return bOk;
+}
diff --git a/Engine/Source/Editor/UnrealEd/Private/Cooker/CompactBinaryTCP.h b/Engine/Source/Editor/UnrealEd/Private/Cooker/CompactBinaryTCP.h
--- a/Engine/Source/Editor/UnrealEd/Private/Cooker/CompactBinaryTCP.h
+++ b/Engine/Source/Editor/UnrealEd/Private/Cooker/CompactBinaryTCP.h
@@ -194,6 +194,17 @@ inline FCbWriter& operator<<(FCbWriter& Writer, const TMap<KeyType, ValueType, S
 	return Writer;
 }
 
+inline FCbWriter& operator<<(FCbWriter& Writer, const UE::CompactBinaryTCP::FMarshalledMessage& Value)
+{
+	Writer.BeginObject();
+	Writer << "T" << Value.MessageType;
+	Writer << "V" << Value.Object;
+	Writer.EndObject();
+	return Writer;
+}
+
+UNREALED_API bool LoadFromCompactBinary(FCbFieldView Field, UE::CompactBinaryTCP::FMarshalledMessage& Value);
+
 // FSoftObjectPath has an implicit constructor from FString for backwards compatibility; if we
 // try to create an operator<< for it, it will cause operator<< to be ambiguous.
 // This prevents us from using it directly in containers. To hack around this,
diff --git a/Engine/Source/Editor/UnrealEd/Private/Cooker/CookDirector.cpp b/Engine/Source/Editor/UnrealEd/Private/Cooker/CookDirector.cpp
--- a/Engine/Source/Editor/UnrealEd/Private/Cooker/CookDirector.cpp
+++ b/Engine/Source/Editor/UnrealEd/Private/Cooker/CookDirector.cpp
@@ -365,14 +365,45 @@ void FCookDirector::AssignRequests(TArrayView<FPackageData*> Requests, TArray<FW
 TMap<FPackageData*, FAssignPackageExtraData> FCookDirector::GetAssignPackageExtraDatas(
 	TConstArrayView<FPackageData*> Requests) const
 {
+	FMPCollectorServerTickPackageContext Context;
+	Context.Platforms = COTFS.PlatformManager->GetSessionPlatforms();
+
 	TMap<FPackageData*, FAssignPackageExtraData> Results;
 	for (FPackageData* Request : Requests)
 	{
+		FAssignPackageExtraData* ExtraData = nullptr;
+
+		auto GetOrAllocateExtraData = [&Results, &ExtraData](FPackageData* InRequest)
+		{
+			if (!ExtraData)
+			{
+				ExtraData = &Results.FindOrAdd(InRequest);
+			}
+			return ExtraData;
+		};
+
 		TRefCountPtr<FGenerationHelper> GenerationHelper = Request->GetGenerationHelper();
 		if (GenerationHelper && !GenerationHelper->GetPreviousGeneratedPackages().IsEmpty())
 		{
-			FAssignPackageExtraData& ExtraData = Results.FindOrAdd(Request);
-			ExtraData.GeneratorPreviousGeneratedPackages = GenerationHelper->GetPreviousGeneratedPackages();
+			ExtraData = GetOrAllocateExtraData(Request);
+			ExtraData->GeneratorPreviousGeneratedPackages = GenerationHelper->GetPreviousGeneratedPackages();
+		}
+
+		Context.PackageName = Request->GetPackageName();
+		for (const TPair<FGuid, TRefCountPtr<IMPCollector>>& CollectorPair : Collectors)
+		{
+			IMPCollector* Collector = CollectorPair.Value.GetReference();
+			Collector->ServerTickPackage(Context);
+			if (!Context.Messages.IsEmpty())
+			{
+				ExtraData = GetOrAllocateExtraData(Request);
+				FGuid MessageType = Collector->GetMessageType();
+				for (FCbObject& Object : Context.Messages)
+				{
+					ExtraData->PerPackageCollectorMessages.Add({ MessageType, MoveTemp(Object) });
+				}
+				Context.Messages.Reset();
+			}
 		}
 	}
 	return Results;
@@ -1619,15 +1650,26 @@ void FCookDirector::ConstructReadonlyThreadVariables()
 	CommandletExecutablePath = FUnrealEdMisc::Get().GetProjectEditorBinaryPath();
 
 	InitialConfigMessage = MakeUnique<FInitialConfigMessage>();
-	const TArray<const ITargetPlatform*>& SessionPlatforms = COTFS.PlatformManager->GetSessionPlatforms();
-	TArray<ITargetPlatform*> OrderedSessionPlatforms;
-	OrderedSessionPlatforms.Reset(SessionPlatforms.Num());
-	for (const ITargetPlatform* TargetPlatform : SessionPlatforms)
+	FMPCollectorServerTickContext StartupContext(FMPCollectorServerTickContext::EServerEventType::WorkerStartup);
+	StartupContext.Platforms = COTFS.PlatformManager->GetSessionPlatforms();
+	InitialConfigMessage->ReadFromLocal(COTFS, StartupContext.Platforms,
+		*COTFS.CookByTheBookOptions, *COTFS.CookOnTheFlyOptions, BeginCookContext);
+
+
+	for (const TPair<FGuid, TRefCountPtr<IMPCollector>>& CollectorPair : Collectors)
 	{
-		OrderedSessionPlatforms.Add(const_cast<ITargetPlatform*>(TargetPlatform));
+		IMPCollector* Collector = CollectorPair.Value.GetReference();
+		Collector->ServerTick(StartupContext);
+		if (!StartupContext.Messages.IsEmpty())
+		{
+			FGuid MessageType = Collector->GetMessageType();
+			for (FCbObject& Object : StartupContext.Messages)
+			{
+				InitialConfigMessage->AddMessage({ MessageType, MoveTemp(Object) });
+			}
+			StartupContext.Messages.Reset();
+		}
 	}
-	InitialConfigMessage->ReadFromLocal(COTFS, OrderedSessionPlatforms,
-		*COTFS.CookByTheBookOptions, *COTFS.CookOnTheFlyOptions, BeginCookContext);
 }
 
 const FInitialConfigMessage& FCookDirector::GetInitialConfigMessage()
diff --git a/Engine/Source/Editor/UnrealEd/Private/Cooker/CookWorkerClient.cpp b/Engine/Source/Editor/UnrealEd/Private/Cooker/CookWorkerClient.cpp
--- a/Engine/Source/Editor/UnrealEd/Private/Cooker/CookWorkerClient.cpp
+++ b/Engine/Source/Editor/UnrealEd/Private/Cooker/CookWorkerClient.cpp
@@ -169,6 +169,8 @@ const TArray<ITargetPlatform*>& FCookWorkerClient::GetTargetPlatforms() const
 void FCookWorkerClient::DoneWithInitialSettings()
 {
 	InitialConfigMessage.Reset();
+	// Process remaining deferred initialization messages and discard them
+	HandleReceiveMessages(MoveTemp(DeferredInitializationMessages));
 }
 
 void FCookWorkerClient::ReportDemoteToIdle(const FPackageData& PackageData, ESuppressCookReason Reason)
@@ -591,6 +593,8 @@ void FCookWorkerClient::PollReceiveConfigMessage()
 		return;
 	}
 
+	HandleReceiveMessages(InitialConfigMessage->ConsumeCollectorMessages());
+
 	UE_LOG(LogCook, Display, TEXT("Initialization from CookDirector complete."));
 	SendToState(EConnectStatus::Connected);
 	Messages.RemoveAt(0);
@@ -667,7 +671,7 @@ void FCookWorkerClient::PumpReceiveMessages()
 	HandleReceiveMessages(MoveTemp(Messages));
 }
 
-void FCookWorkerClient::HandleReceiveMessages(TArray<UE::CompactBinaryTCP::FMarshalledMessage>&& Messages)
+void FCookWorkerClient::HandleReceiveMessages(TArray<UE::CompactBinaryTCP::FMarshalledMessage>&& Messages, FName OptionalPackageName /*= NAME_None*/)
 {
 	for (UE::CompactBinaryTCP::FMarshalledMessage& Message : Messages)
 	{
@@ -742,8 +746,18 @@ void FCookWorkerClient::HandleReceiveMessages(TArray<UE::CompactBinaryTCP::FMars
 					check(*Collector);
 					FMPCollectorClientMessageContext Context;
 					Context.Platforms = OrderedSessionPlatforms;
+					Context.PackageName = OptionalPackageName;
 					(*Collector)->ClientReceiveMessage(Context, Message.Object);
 				}
+				else if (InitialConfigMessage.IsValid())
+				{
+					ensureMsgf(Messages.GetData() != DeferredInitializationMessages.GetData(), 
+						TEXT("HandleReceiveMessages may not be called with the deferred initialization message array until after calling DoneWithInitialSettings()"));
+
+					// If we are still running our initialization, then we may not have the relevant collectors registered yet
+					// Defer the message and try again at the end of initialization
+					DeferredInitializationMessages.Add(Message);
+				}
 				else
 				{
 					UE_LOG(LogCook, Error,
@@ -872,6 +886,10 @@ void FCookWorkerClient::AssignPackages(FAssignPackagesMessage& Message)
 				TRefCountPtr<FGenerationHelper> GenerationHelper = PackageData.CreateUninitializedGenerationHelper();
 				GenerationHelper->SetPreviousGeneratedPackages(MoveTemp(AssignData.GeneratorPreviousGeneratedPackages));
 			}
+			if (!AssignData.PerPackageCollectorMessages.IsEmpty())
+			{
+				HandleReceiveMessages(MoveTemp(AssignData.PerPackageCollectorMessages), AssignData.ConstructData.PackageName);
+			}
 			TConstArrayView<const ITargetPlatform*> NeedCookPlatforms =
 				AssignData.NeedCookPlatforms.GetPlatforms(COTFS, nullptr, OrderedSessionPlatforms,
 					&NeedCookPlatformsBuffer);
diff --git a/Engine/Source/Editor/UnrealEd/Private/Cooker/CookWorkerClient.h b/Engine/Source/Editor/UnrealEd/Private/Cooker/CookWorkerClient.h
--- a/Engine/Source/Editor/UnrealEd/Private/Cooker/CookWorkerClient.h
+++ b/Engine/Source/Editor/UnrealEd/Private/Cooker/CookWorkerClient.h
@@ -113,7 +113,7 @@ private:
 	/** Helper for Tick, pump receive messages from the Server. */
 	void PumpReceiveMessages();
 	/** Helper for PumpReceiveMessages: dispatch the messages received from the socket. */
-	void HandleReceiveMessages(TArray<UE::CompactBinaryTCP::FMarshalledMessage>&& Messages);
+	void HandleReceiveMessages(TArray<UE::CompactBinaryTCP::FMarshalledMessage>&& Messages, FName OptionalPackageName = NAME_None);
 	/** Helper for Tick, pump Send/Receive and check for whether we are done shutting down. */
 	void PumpDisconnect(FTickStackData& StackData);
 	/** Send the message immediately to the Socket. If cannot complete immediately, it will be finished during Tick. */
@@ -154,6 +154,7 @@ private:
 	// Variables Read/Write only from the Scheduler thread
 	TSharedPtr<FInternetAddr> DirectorAddr;
 	TUniquePtr<FInitialConfigMessage> InitialConfigMessage;
+	TArray<UE::CompactBinaryTCP::FMarshalledMessage> DeferredInitializationMessages;
 	TRefCountPtr<FLogMessagesMessageHandler> LogMessageHandler;
 	TArray<ITargetPlatform*> OrderedSessionPlatforms;
 	TArray<ITargetPlatform*> OrderedSessionAndSpecialPlatforms;
diff --git a/Engine/Source/Editor/UnrealEd/Private/Cooker/CookWorkerServer.cpp b/Engine/Source/Editor/UnrealEd/Private/Cooker/CookWorkerServer.cpp
--- a/Engine/Source/Editor/UnrealEd/Private/Cooker/CookWorkerServer.cpp
+++ b/Engine/Source/Editor/UnrealEd/Private/Cooker/CookWorkerServer.cpp
@@ -704,6 +704,7 @@ void FCookWorkerServer::SendPendingPackages()
 		if (ExtraData)
 		{
 			AssignData.GeneratorPreviousGeneratedPackages = MoveTemp(ExtraData->GeneratorPreviousGeneratedPackages);
+			AssignData.PerPackageCollectorMessages = MoveTemp(ExtraData->PerPackageCollectorMessages);
 		}
 	}
 	for (FPackageData* PackageData : PackagesToAssignInfoPackages)
@@ -1229,6 +1230,7 @@ void FAssignPackageData::Write(FCbWriter& Writer,
 	}
 	static_assert(sizeof(ICookPackageSplitter::EGeneratedRequiresGenerator) <= sizeof(uint8), "We are storing it in a uint8");
 	Writer << static_cast<uint8>(DoesGeneratedRequireGenerator);
+	Writer << PerPackageCollectorMessages;
 	Writer.EndArray();
 }
 
@@ -1282,6 +1284,7 @@ bool FAssignPackageData::TryRead(FCbFieldView Field, TConstArrayView<const ITarg
 	{
 		bOk = false;
 	}
+	bOk = LoadFromCompactBinary(*It++, PerPackageCollectorMessages) & bOk;
 	return bOk;
 }
 
@@ -1365,13 +1368,17 @@ bool FAbortWorkerMessage::TryRead(FCbObjectView Object)
 FGuid FAbortWorkerMessage::MessageType(TEXT("83FD99DFE8DB4A9A8E71684C121BE6F3"));
 
 void FInitialConfigMessage::ReadFromLocal(const UCookOnTheFlyServer& COTFS,
-	const TArray<ITargetPlatform*>& InOrderedSessionPlatforms, const FCookByTheBookOptions& InCookByTheBookOptions,
+	const TConstArrayView<const ITargetPlatform*>& InOrderedSessionPlatforms, const FCookByTheBookOptions& InCookByTheBookOptions,
 	const FCookOnTheFlyOptions& InCookOnTheFlyOptions, const FBeginCookContextForWorker& InBeginContext)
 {
 	InitialSettings.CopyFromLocal(COTFS);
 	BeginCookSettings.CopyFromLocal(COTFS);
 	BeginCookContext = InBeginContext;
-	OrderedSessionPlatforms = InOrderedSessionPlatforms;
+	OrderedSessionPlatforms.Reset(InOrderedSessionPlatforms.Num());
+	for (const ITargetPlatform* Platform : InOrderedSessionPlatforms)
+	{
+		OrderedSessionPlatforms.Add(const_cast<ITargetPlatform*>(Platform));
+	}
 	DirectorCookMode = COTFS.GetCookMode();
 	CookInitializationFlags = COTFS.GetCookFlags();
 	CookByTheBookOptions = InCookByTheBookOptions;
@@ -1398,6 +1405,7 @@ void FInitialConfigMessage::Write(FCbWriter& Writer) const
 	Writer << "BeginCookContext" << BeginCookContext;
 	Writer << "CookByTheBookOptions" << CookByTheBookOptions;
 	Writer << "CookOnTheFlyOptions" << CookOnTheFlyOptions;
+	Writer << "MPCollectorMessages" << MPCollectorMessages;
 }
 
 bool FInitialConfigMessage::TryRead(FCbObjectView Object)
@@ -1446,6 +1454,8 @@ bool FInitialConfigMessage::TryRead(FCbObjectView Object)
 	bOk = LoadFromCompactBinary(Object["BeginCookContext"], BeginCookContext) & bOk;
 	bOk = LoadFromCompactBinary(Object["CookByTheBookOptions"], CookByTheBookOptions) & bOk;
 	bOk = LoadFromCompactBinary(Object["CookOnTheFlyOptions"], CookOnTheFlyOptions) & bOk;
+	bOk = LoadFromCompactBinary(Object["MPCollectorMessages"], MPCollectorMessages) & bOk;
+
 	return bOk;
 }
 
diff --git a/Engine/Source/Editor/UnrealEd/Private/Cooker/CookWorkerServer.h b/Engine/Source/Editor/UnrealEd/Private/Cooker/CookWorkerServer.h
--- a/Engine/Source/Editor/UnrealEd/Private/Cooker/CookWorkerServer.h
+++ b/Engine/Source/Editor/UnrealEd/Private/Cooker/CookWorkerServer.h
@@ -51,6 +51,7 @@ enum class ENotifyRemote
 struct FAssignPackageExtraData
 {
 	TMap<FName, FAssetPackageData> GeneratorPreviousGeneratedPackages;
+	TArray<UE::CompactBinaryTCP::FMarshalledMessage> PerPackageCollectorMessages;
 };
 
 /** Class in a Director process that communicates over a Socket with FCookWorkerClient in a CookWorker process. */
@@ -255,6 +256,7 @@ struct FAssignPackageData
 	FInstigator Instigator;
 	FDiscoveredPlatformSet NeedCookPlatforms;
 	TMap<FName, FAssetPackageData> GeneratorPreviousGeneratedPackages;
+	TArray<UE::CompactBinaryTCP::FMarshalledMessage> PerPackageCollectorMessages;
 	ICookPackageSplitter::EGeneratedRequiresGenerator DoesGeneratedRequireGenerator
 		= ICookPackageSplitter::EGeneratedRequiresGenerator::None;
 
@@ -369,9 +371,10 @@ public:
 	virtual FGuid GetMessageType() const override { return MessageType; }
 	virtual const TCHAR* GetDebugName() const override { return TEXT("InitialConfigMessage"); }
 
-	void ReadFromLocal(const UCookOnTheFlyServer& COTFS, const TArray<ITargetPlatform*>& InOrderedSessionPlatforms,
+	void ReadFromLocal(const UCookOnTheFlyServer& COTFS, const TConstArrayView<const ITargetPlatform*>& InOrderedSessionPlatforms,
 		const FCookByTheBookOptions& InCookByTheBookOptions, const FCookOnTheFlyOptions& InCookOnTheFlyOptions,
 		const FBeginCookContextForWorker& InBeginContext);
+	void AddMessage(UE::CompactBinaryTCP::FMarshalledMessage&& InMarshalledMessage) { MPCollectorMessages.Add(MoveTemp(InMarshalledMessage)); }
 
 	ECookMode::Type GetDirectorCookMode() const { return DirectorCookMode; }
 	ECookInitializationFlags GetCookInitializationFlags() const { return CookInitializationFlags; }
@@ -379,6 +382,8 @@ public:
 	FBeginCookConfigSettings&& ConsumeBeginCookConfigSettings() { return MoveTemp(BeginCookSettings); }
 	FCookByTheBookOptions&& ConsumeCookByTheBookOptions() { return MoveTemp(CookByTheBookOptions); }
 	FCookOnTheFlyOptions&& ConsumeCookOnTheFlyOptions() { return MoveTemp(CookOnTheFlyOptions); }
+	TArray<UE::CompactBinaryTCP::FMarshalledMessage>&& ConsumeCollectorMessages() { return MoveTemp(MPCollectorMessages); }
+
 	const FBeginCookContextForWorker& GetBeginCookContext() const { return BeginCookContext; }
 	const TArray<ITargetPlatform*>& GetOrderedSessionPlatforms() const { return OrderedSessionPlatforms; }
 	bool IsZenStore() const { return bZenStore; }
@@ -392,6 +397,7 @@ private:
 	FCookByTheBookOptions CookByTheBookOptions;
 	FCookOnTheFlyOptions CookOnTheFlyOptions;
 	TArray<ITargetPlatform*> OrderedSessionPlatforms;
+	TArray<UE::CompactBinaryTCP::FMarshalledMessage> MPCollectorMessages;
 	ECookMode::Type DirectorCookMode = ECookMode::CookByTheBook;
 	ECookInitializationFlags CookInitializationFlags = ECookInitializationFlags::None;
 	bool bZenStore = false;
diff --git a/Engine/Source/Editor/UnrealEd/Private/Cooker/IterativeValidatePackageWriter.cpp b/Engine/Source/Editor/UnrealEd/Private/Cooker/IterativeValidatePackageWriter.cpp
--- a/Engine/Source/Editor/UnrealEd/Private/Cooker/IterativeValidatePackageWriter.cpp
+++ b/Engine/Source/Editor/UnrealEd/Private/Cooker/IterativeValidatePackageWriter.cpp
@@ -1,9 +1,10 @@
 // Copyright Epic Games, Inc. All Rights Reserved.
 
 #include "Cooker/IterativeValidatePackageWriter.h"
-
+#include "Cooker/CompactBinaryTCP.h"
 #include "Cooker/CookGenerationHelper.h"
 #include "Cooker/CookPackageData.h"
+#include "Cooker/MPCollector.h"
 #include "CookOnTheSide/CookLog.h"
 #include "CookOnTheSide/CookOnTheFlyServer.h"
 #include "HAL/FileManager.h"
@@ -19,6 +20,261 @@ DEFINE_LOG_CATEGORY_STATIC(LogIterativeValidate, Log, All);
 
 constexpr FStringView IterativeValidateFilename(TEXTVIEW("IterativeValidate.bin"));
 
+class FIterativeValidateMPCollector : public UE::Cook::IMPCollector
+{
+public:
+	FIterativeValidateMPCollector(FIterativeValidatePackageWriter* InOwner) : Owner(InOwner) {}
+
+	virtual FGuid GetMessageType() const { return MessageType; }
+	virtual const TCHAR* GetDebugName() const { return TEXT("IterativeValidateMPCollector"); }
+
+	virtual void ServerTick(UE::Cook::FMPCollectorServerTickContext& Context) override;
+	virtual void ClientTickPackage(UE::Cook::FMPCollectorClientTickPackageContext& Context) override;
+	virtual void ServerTickPackage(UE::Cook::FMPCollectorServerTickPackageContext& Context);
+
+	virtual void ClientReceiveMessage(UE::Cook::FMPCollectorClientMessageContext& Context, FCbObjectView Message) override;
+	virtual void ServerReceiveMessage(UE::Cook::FMPCollectorServerMessageContext& Context, FCbObjectView Message) override;
+
+public:
+	static FGuid MessageType;
+
+private:
+	enum class EMessageSubtype : uint8
+	{
+		ServerToClient_WorkerStartup,
+		ClientToServer_ReplIsAnotherSaveNeeded,
+		ServerToClient_ReplUpdatePackageModificationStatus,
+		Invalid
+	};
+
+	bool TryWritePackageStatus(FCbWriter& Writer, FName PackageName);
+	void ReadAndSyncPackageStatus(FCbObjectView Message, FName PackageName);
+	FIterativeValidatePackageWriter* Owner;
+};
+
+FCbWriter& operator<<(FCbWriter& Writer, const FIterativeValidatePackageWriter::FMessage& Message)
+{
+	Writer.BeginObject();
+	Writer << "Text" << Message.Text;
+
+	static_assert(sizeof(Message.Verbosity) == sizeof(uint8));
+	Writer << "Verbosity" << (uint8)Message.Verbosity;
+
+	Writer.EndObject();
+	return Writer;
+}
+
+bool LoadFromCompactBinary(FCbFieldView Field, FIterativeValidatePackageWriter::FMessage& Message)
+{
+	bool bOk = !Field.HasError();
+	if (bOk && LoadFromCompactBinary(Field["Text"], Message.Text))
+	{
+		uint8 Verbosity = ELogVerbosity::NumVerbosity;
+		bOk = LoadFromCompactBinary(Field["Verbosity"], Verbosity) && Verbosity < ELogVerbosity::NumVerbosity;
+		if (bOk)
+		{
+			Message.Verbosity = (ELogVerbosity::Type)Verbosity;
+		}
+	}
+	if (!bOk)
+	{
+		Message = FIterativeValidatePackageWriter::FMessage();
+	}
+	return bOk;
+}
+
+FCbWriter& operator<<(FCbWriter& Writer, FIterativeValidatePackageWriter::EPackageStatus Status)
+{
+	Writer << (uint8)Status;
+	return Writer;
+}
+
+bool LoadFromCompactBinary(FCbFieldView Field, FIterativeValidatePackageWriter::EPackageStatus& Status)
+{
+	uint8 StatusInteger = (uint8)FIterativeValidatePackageWriter::EPackageStatus::Count;
+	if (!LoadFromCompactBinary(Field, StatusInteger))
+	{
+		UE_LOG(LogIterativeValidate, Error, TEXT("Failed to deserialize package status."));
+	}
+	else if (StatusInteger >= (uint8)FIterativeValidatePackageWriter::EPackageStatus::Count)
+	{
+		UE_LOG(LogIterativeValidate, Error, TEXT("Unexpected package status deserialized: %d"), StatusInteger);
+	}
+	else
+	{
+		Status = (FIterativeValidatePackageWriter::EPackageStatus)StatusInteger;
+		return true;
+	}
+	return false;
+}
+
+void FIterativeValidateMPCollector::ServerTick(UE::Cook::FMPCollectorServerTickContext& Context)
+{
+	static_assert(sizeof(EMessageSubtype) == sizeof(uint8));
+
+	if (Context.GetEventType() == UE::Cook::FMPCollectorServerTickContext::EServerEventType::WorkerStartup)
+	{
+		FCbWriter Writer;
+		Writer.BeginObject();
+		Writer << "MessageSubtype" << (uint8)EMessageSubtype::ServerToClient_WorkerStartup;
+		Writer << "PackageStatusMap" << Owner->PackageStatusMap;
+		Writer << "PackageMessageMap" << Owner->PackageMessageMap;
+		Writer.EndObject();
+		Context.AddMessage(Writer.Save().AsObject());
+	}
+}
+
+void FIterativeValidateMPCollector::ClientTickPackage(UE::Cook::FMPCollectorClientTickPackageContext& Context)
+{
+	static_assert(sizeof(EMessageSubtype) == sizeof(uint8));
+
+	FCbWriter Writer;
+	Writer.BeginObject();
+	Writer << "MessageSubtype" << (uint8)EMessageSubtype::ClientToServer_ReplIsAnotherSaveNeeded;
+
+	const FName PackageName = Context.GetPackageName();
+	if (PackageName.IsNone())
+	{
+		UE_LOG(LogCook, Error, TEXT("Context does not contain a valid package name."))
+		// It's safe to continue because TryWritePackageStatus will return false if the name is none. 
+		// The error is logged here to make the call site clear
+	}
+
+	if (TryWritePackageStatus(Writer, PackageName))
+	{
+		Writer.EndObject();
+		Context.AddMessage(Writer.Save().AsObject());
+	}
+}
+
+void FIterativeValidateMPCollector::ServerTickPackage(UE::Cook::FMPCollectorServerTickPackageContext& Context)
+{
+	static_assert(sizeof(EMessageSubtype) == sizeof(uint8));
+
+	FCbWriter Writer;
+	Writer.BeginObject();
+	Writer << "MessageSubtype" << (uint8)EMessageSubtype::ServerToClient_ReplUpdatePackageModificationStatus;
+
+	const FName PackageName = Context.GetPackageName();
+	if (PackageName.IsNone())
+	{
+		UE_LOG(LogCook, Error, TEXT("Context does not contain a valid package name."))
+		// It's safe to continue because TryWritePackageStatus will return false if the name is none. 
+		// The error is logged here to make the call site clear
+	}
+
+	if (TryWritePackageStatus(Writer, PackageName))
+	{
+		Writer.EndObject();
+		Context.AddMessage(Writer.Save().AsObject());
+	}
+}
+
+void FIterativeValidateMPCollector::ClientReceiveMessage(UE::Cook::FMPCollectorClientMessageContext& Context, FCbObjectView Message)
+{
+	uint8 MessageSubtypeAsInteger = (uint8)EMessageSubtype::Invalid;
+	if (LoadFromCompactBinary(Message["MessageSubtype"], MessageSubtypeAsInteger))
+	{
+		switch ((EMessageSubtype)MessageSubtypeAsInteger)
+		{
+			case EMessageSubtype::ServerToClient_WorkerStartup:
+			{
+				bool bOk = LoadFromCompactBinary(Message["PackageStatusMap"], Owner->PackageStatusMap);
+				bOk = bOk && LoadFromCompactBinary(Message["PackageMessageMap"], Owner->PackageMessageMap);
+				check(bOk); // If we fail this, we will fail to get anything right during the rest of the validation. Better to terminate quickly.
+				break;
+			}
+			case EMessageSubtype::ServerToClient_ReplUpdatePackageModificationStatus:
+			{
+				FName PackageName = Context.GetPackageName();
+				if (!PackageName.IsNone())
+				{
+					ReadAndSyncPackageStatus(Message, PackageName);
+				}
+				else
+				{
+					UE_LOG(LogCook, Error, TEXT("Cannot process ServerToClient_ReplUpdatePackageModificationStatus without a valid package name in the current context."));
+				}
+				break;
+			}
+			default:
+			{
+				UE_LOG(LogCook, Error, TEXT("Unexpected message type: %d"), MessageSubtypeAsInteger);
+				break;
+			}
+		};
+	}
+}
+
+void FIterativeValidateMPCollector::ServerReceiveMessage(UE::Cook::FMPCollectorServerMessageContext& Context, FCbObjectView Message)
+{
+	uint8 MessageSubtypeAsInteger = (uint8)EMessageSubtype::Invalid;
+	if (LoadFromCompactBinary(Message["MessageSubtype"], MessageSubtypeAsInteger))
+	{
+		FName PackageName = Context.GetPackageName();
+		if (PackageName.IsNone())
+		{
+			UE_LOG(LogCook, Error, TEXT("Cannot process messages on server without a valid package name in the current context."));
+		}
+		else if (MessageSubtypeAsInteger == (uint8)EMessageSubtype::ClientToServer_ReplIsAnotherSaveNeeded)
+		{
+			ReadAndSyncPackageStatus(Message, PackageName);
+		}
+		else
+		{
+			UE_LOG(LogCook, Error, TEXT("Unexpected message received. MessageSubtype == %d"), MessageSubtypeAsInteger);
+		}
+	}
+	else
+	{
+		UE_LOG(LogCook, Error, TEXT("Invalid message received. No MessageSubtype field available."));
+	}
+}
+
+FGuid FIterativeValidateMPCollector::MessageType(TEXT("5E56C5D96F3B455E9452C15ADA601A71"));
+
+bool FIterativeValidateMPCollector::TryWritePackageStatus(FCbWriter& Writer, FName PackageName)
+{
+	FIterativeValidatePackageWriter::EPackageStatus PackageStatus = Owner->GetPackageStatus(PackageName);
+
+	if (PackageStatus != FIterativeValidatePackageWriter::EPackageStatus::NotYetProcessed)
+	{
+		Writer << "Status" << (uint8)PackageStatus;
+
+		const TArray<FIterativeValidatePackageWriter::FMessage>* Messages = Owner->PackageMessageMap.Find(PackageName);
+		if (Messages != nullptr)
+		{
+			Writer << "MessageArray" << *Messages;
+		}
+		return true;
+	}
+	return false;
+}
+
+void FIterativeValidateMPCollector::ReadAndSyncPackageStatus(FCbObjectView Message, FName PackageName)
+{
+	uint8 PackageStatusInteger = (uint8)FIterativeValidatePackageWriter::EPackageStatus::Count;
+	bool bOk = LoadFromCompactBinary(Message["Status"], PackageStatusInteger);
+	bOk = bOk && PackageStatusInteger < (uint8)FIterativeValidatePackageWriter::EPackageStatus::Count;
+
+	if (bOk)
+	{
+		FIterativeValidatePackageWriter::EPackageStatus PackageStatus = (FIterativeValidatePackageWriter::EPackageStatus)PackageStatusInteger;
+		Owner->SetPackageStatus(PackageName, PackageStatus);
+		if (Message.FindView("MessageArray").HasValue())
+		{
+			TArray<FIterativeValidatePackageWriter::FMessage>& MessageArray = Owner->PackageMessageMap.FindOrAdd(PackageName);
+			bOk = LoadFromCompactBinary(Message["MessageArray"], MessageArray);
+		}
+	}
+	
+	if (!bOk)
+	{
+		UE_LOG(LogCook, Error, TEXT("Invalid message received in ReadAndSyncPackageStatus. Status received (%d) for package \"%s\""),
+			PackageStatusInteger, *PackageName.ToString());
+	}
+}
+
 FIterativeValidatePackageWriter::FIterativeValidatePackageWriter(UCookOnTheFlyServer& InCOTFS,
 	TUniquePtr<ICookedPackageWriter>&& InInner, EPhase InPhase, const FString& ResolvedMetadataPath)
 	: FDiffPackageWriter(MoveTemp(InInner))
@@ -26,6 +282,8 @@ FIterativeValidatePackageWriter::FIterativeValidatePackageWriter(UCookOnTheFlySe
 	, COTFS(InCOTFS)
 	, Phase(InPhase)
 {
+	COTFS.RegisterCollector(new FIterativeValidateMPCollector(this));
+
 	Indent = FCString::Spc(FOutputDeviceHelper::FormatLogLine(ELogVerbosity::Warning,
 		LogIterativeValidate.GetCategoryName(), TEXT(""), GPrintLogTimes).Len());
 }
@@ -36,7 +294,7 @@ void FIterativeValidatePackageWriter::BeginPackage(const FBeginPackageInfo& Info
 	switch (Phase)
 	{
 	case EPhase::AllInOnePhase:
-		if (IterativelyUnmodified.Contains(Info.PackageName))
+		if (GetPackageStatus(Info.PackageName) == EPackageStatus::DeclaredUnmodified_NotYetProcessed)
 		{
 			// Save to memory and look for diffs before saving it out to disk
 			SaveAction = ESaveAction::CheckForDiffs;
@@ -45,7 +303,6 @@ void FIterativeValidatePackageWriter::BeginPackage(const FBeginPackageInfo& Info
 		else
 		{
 			// Not iteratively skippable, so we expect it to change. No need to look for diffs, just save to disk.
-			++ModifiedCount;
 			if (bReadOnly)
 			{
 				SaveAction = ESaveAction::IgnoreResults;
@@ -62,24 +319,27 @@ void FIterativeValidatePackageWriter::BeginPackage(const FBeginPackageInfo& Info
 		Super::BeginPackage(Info);
 		break;
 	case EPhase::Phase2:
-		if (IterativeValidated.Contains(Info.PackageName))
 		{
-			// Already saved in Phase 1; no need to diff it or save it now
-			SaveAction = ESaveAction::IgnoreResults;
-		}
-		else if (IterativeFailed.Contains(Info.PackageName))
-		{
-			SaveAction = ESaveAction::CheckForDiffs;
-			Super::BeginPackage(Info);
-		}
-		else
-		{
-			// This is an IterativeModified package. It was found during Phase1 to be modified and would in a
-			// normal iterative cook be resaved rather than iteratively skipped. Resave it as normal.
-			SaveAction = ESaveAction::SaveToInner;
-			Inner->BeginPackage(Info);
+			EPackageStatus PackageStatus = GetPackageStatus(Info.PackageName);
+			if (PackageStatus == EPackageStatus::DeclaredUnmodified_ConfirmedUnmodified)
+			{
+				// Already saved in Phase 1; no need to diff it or save it now
+				SaveAction = ESaveAction::IgnoreResults;
+			}
+			else if (PackageStatus == EPackageStatus::DeclaredUnmodified_FoundModified_IndeterminismOrFalsePositive)
+			{
+				SaveAction = ESaveAction::CheckForDiffs;
+				Super::BeginPackage(Info);
+			}
+			else
+			{
+				// This is an IterativeModified package. It was found during Phase1 to be modified and would in a
+				// normal iterative cook be resaved rather than iteratively skipped. Resave it as normal.
+				SaveAction = ESaveAction::SaveToInner;
+				Inner->BeginPackage(Info);
+			}
+			break;
 		}
-		break;
 	default:
 		checkNoEntry();
 		break;
@@ -322,15 +582,15 @@ void FIterativeValidatePackageWriter::UpdatePackageModificationStatus(FName Pack
 		// Save the input value for bIterativelyUnmodified, and report skippable for the modified if possible
 		if (bIterativelyUnmodified)
 		{
-			IterativelyUnmodified.Add(PackageName);
+			SetPackageStatus(PackageName, EPackageStatus::DeclaredUnmodified_NotYetProcessed);
 			bInOutShouldIterativelySkip = false;
 		}
 		else
 		{
+			SetPackageStatus(PackageName, EPackageStatus::DeclaredModified_WillNotVerify);
 			if (!bKnownGenerator && bReadOnly)
 			{
 				bInOutShouldIterativelySkip = true;
-				++ModifiedCount; // Only increment here if we're skipping it. Otherwise it is incremented in BeginPackage
 			}
 		}
 		break;
@@ -340,7 +600,11 @@ void FIterativeValidatePackageWriter::UpdatePackageModificationStatus(FName Pack
 		bInOutShouldIterativelySkip = !bIterativelyUnmodified && !bKnownGenerator;
 		if (!bIterativelyUnmodified)
 		{
-			++ModifiedCount;
+			SetPackageStatus(PackageName, EPackageStatus::DeclaredModified_WillNotVerify);
+		}
+		else
+		{
+			SetPackageStatus(PackageName, EPackageStatus::DeclaredUnmodified_NotYetProcessed);
 		}
 		break;
 	case EPhase::Phase2:
@@ -349,7 +613,8 @@ void FIterativeValidatePackageWriter::UpdatePackageModificationStatus(FName Pack
 		// is responsible for getting those resaved. Reexecute save for the packages that were IterativeFailed
 		// from phase1, so we can test whether they are indeterministic. Always save generators so we can
 		// test their generated packages.
-		bInOutShouldIterativelySkip = IterativeValidated.Contains(PackageName) && !bKnownGenerator;
+		bInOutShouldIterativelySkip = (GetPackageStatus(PackageName) == EPackageStatus::DeclaredUnmodified_ConfirmedUnmodified) 
+										&& !bKnownGenerator;
 		break;
 	default:
 		checkNoEntry();
@@ -394,13 +659,18 @@ void FIterativeValidatePackageWriter::BeginCook(const FCookInfo& Info)
 		}
 		break;
 	case EPhase::Phase2:
-		Load();
-		UE_LOG(LogIterativeValidate, Display,
-			TEXT("Phase2: %d packages were found during Phase1 to be iteratively unmodified but had differences. Running -diffonly on them again to check whether the differences are due to indeterminism or to IterativeFalsePositives."), IterativeFailed.Num());
-		UE_LOG(LogIterativeValidate, Display,
-			TEXT("%d packages were found during Phase1 to be modified or new and will be resaved."),
-			ModifiedCount);
-		break;
+		{
+			Load();
+			FStatusCounts StatusCounts = CountPackagesByStatus();
+			UE_LOG(LogIterativeValidate, Display,
+				TEXT("Phase2: %d packages were found during Phase1 to be iteratively unmodified but had differences. "
+					"Running -diffonly on them again to check whether the differences are due to indeterminism or to IterativeFalsePositives."), 
+					StatusCounts[EPackageStatus::DeclaredUnmodified_FoundModified_IndeterminismOrFalsePositive]);
+			UE_LOG(LogIterativeValidate, Display,
+				TEXT("%d packages were found during Phase1 to be modified or new and will be resaved."),
+				StatusCounts[EPackageStatus::DeclaredModified_WillNotVerify]);
+			break;
+		}
 	default:
 		checkNoEntry();
 		break;
@@ -411,16 +681,20 @@ void FIterativeValidatePackageWriter::BeginCook(const FCookInfo& Info)
 void FIterativeValidatePackageWriter::EndCook(const FCookInfo& Info)
 {
 	Super::EndCook(Info);
+	FStatusCounts StatusCounts = CountPackagesByStatus();
+
 	switch (Phase)
 	{
 	case EPhase::AllInOnePhase:
 	{
 		UE_LOG(LogIterativeValidate, Display,
 			TEXT("Modified: %d. DetectedUnmodified: %d. ValidatedUnmodified: %d. IterativeSkipFalsePositive: %d."),
-			ModifiedCount, IterativeValidated.Num() + IterativeSkipFalsePositive.Num(), IterativeValidated.Num(),
-			IterativeSkipFalsePositive.Num());
-		FString Message = FString::Printf(TEXT("IterativeSkipFalsePositive: %d."), IterativeSkipFalsePositive.Num());
-		if (IterativeSkipFalsePositive.Num())
+			StatusCounts[EPackageStatus::DeclaredModified_WillNotVerify], 
+			StatusCounts[EPackageStatus::DeclaredUnmodified_ConfirmedUnmodified] + StatusCounts[EPackageStatus::DeclaredUnmodified_FoundModified_FalsePositive],
+			StatusCounts[EPackageStatus::DeclaredUnmodified_ConfirmedUnmodified],
+			StatusCounts[EPackageStatus::DeclaredUnmodified_FoundModified_FalsePositive]);
+		FString Message = FString::Printf(TEXT("IterativeSkipFalsePositive: %d."), StatusCounts[EPackageStatus::DeclaredUnmodified_FoundModified_FalsePositive]);
+		if (StatusCounts[EPackageStatus::DeclaredUnmodified_FoundModified_FalsePositive] > 0)
 		{
 			UE_LOG(LogIterativeValidate, Error, TEXT("%s"), *Message);
 		}
@@ -433,18 +707,24 @@ void FIterativeValidatePackageWriter::EndCook(const FCookInfo& Info)
 	case EPhase::Phase1:
 		UE_LOG(LogIterativeValidate, Display,
 			TEXT("Modified: %d. DetectedUnmodified: %d. ValidatedUnmodified: %d. IterativeSkipFalsePositiveOrIndeterminism: %d."),
-			ModifiedCount, IterativeValidated.Num() + IterativeFailed.Num(), IterativeValidated.Num(),
-			IterativeFailed.Num());
+			StatusCounts[EPackageStatus::DeclaredModified_WillNotVerify],
+			StatusCounts[EPackageStatus::DeclaredUnmodified_ConfirmedUnmodified] 
+			+ StatusCounts[EPackageStatus::DeclaredUnmodified_FoundModified_IndeterminismOrFalsePositive],
+			StatusCounts[EPackageStatus::DeclaredUnmodified_ConfirmedUnmodified],
+			StatusCounts[EPackageStatus::DeclaredUnmodified_FoundModified_IndeterminismOrFalsePositive]);
 		Save();
 		break;
 	case EPhase::Phase2:
 	{
 		UE_LOG(LogIterativeValidate, Display,
 			TEXT("Modified: %d. DetectedUnmodified: %d. ValidatedUnmodified: %d. Indeterminism: %d."),
-			ModifiedCount, IterativeValidated.Num() + IterativeFailed.Num(), IterativeValidated.Num(),
-			IndeterminismFailed.Num());
-		FString Message = FString::Printf(TEXT("IterativeSkipFalsePositive: %d."), IterativeSkipFalsePositive.Num());
-		if (IterativeSkipFalsePositive.Num())
+			StatusCounts[EPackageStatus::DeclaredModified_WillNotVerify],
+			StatusCounts[EPackageStatus::DeclaredUnmodified_ConfirmedUnmodified] 
+			+ StatusCounts[EPackageStatus::DeclaredUnmodified_FoundModified_IndeterminismOrFalsePositive],
+			StatusCounts[EPackageStatus::DeclaredUnmodified_ConfirmedUnmodified],
+			StatusCounts[EPackageStatus::DeclaredUnmodified_FoundModified_Indeterminism]);
+		FString Message = FString::Printf(TEXT("IterativeSkipFalsePositive: %d."), StatusCounts[EPackageStatus::DeclaredUnmodified_FoundModified_FalsePositive]);
+		if (StatusCounts[EPackageStatus::DeclaredUnmodified_FoundModified_FalsePositive] > 0)
 		{
 			UE_LOG(LogIterativeValidate, Error, TEXT("%s"), *Message);
 		}
@@ -505,7 +785,7 @@ bool FIterativeValidatePackageWriter::IsAnotherSaveNeeded(FSavePackageResultStru
 	switch (Phase)
 	{
 	case EPhase::AllInOnePhase:
-		check(IterativelyUnmodified.Contains(BeginInfo.PackageName)); // If !Contains, then we would have set SaveAction=SaveToInner or IgnoreResults and early exited above
+		check(GetPackageStatus(BeginInfo.PackageName) == EPackageStatus::DeclaredUnmodified_NotYetProcessed); // Otherwise we would have set SaveAction=SaveToInner or IgnoreResults and early exited above
 		if (Super::IsAnotherSaveNeeded(PreviousResult, SaveArgs))
 		{
 			return true;
@@ -521,15 +801,16 @@ bool FIterativeValidatePackageWriter::IsAnotherSaveNeeded(FSavePackageResultStru
 
 			if (bIsDifferent && !bNewPackage)
 			{
-				IterativeSkipFalsePositive.Add(BeginInfo.PackageName);
+				SetPackageStatus(BeginInfo.PackageName, EPackageStatus::DeclaredUnmodified_FoundModified_FalsePositive);
 			}
 			else if (!bNewPackage)
 			{
-				IterativeValidated.Add(BeginInfo.PackageName);
+
+				SetPackageStatus(BeginInfo.PackageName, EPackageStatus::DeclaredUnmodified_ConfirmedUnmodified);
 			}
 			else
 			{
-				++ModifiedCount;
+				SetPackageStatus(BeginInfo.PackageName, EPackageStatus::DeclaredModified_WillNotVerify);
 			}
 
 			if (bReadOnly)
@@ -563,7 +844,8 @@ bool FIterativeValidatePackageWriter::IsAnotherSaveNeeded(FSavePackageResultStru
 
 			// Mark that the iterative validation failed if it was not already marked by log or warning messages.
 			// We need to record it for an indeterminism test
-			IterativeFailed.FindOrAdd(BeginInfo.PackageName);
+			SetPackageStatus(BeginInfo.PackageName, EPackageStatus::DeclaredUnmodified_FoundModified_IndeterminismOrFalsePositive);
+
 			return true;
 		}
 		else if (!bNewPackage)
@@ -571,9 +853,10 @@ bool FIterativeValidatePackageWriter::IsAnotherSaveNeeded(FSavePackageResultStru
 			// No differences found, so finish off the superclass's save during CommitPackage, without doing a
 			// SaveToInner pass
 			// Mark that the iterative validation passed
-			IterativeValidated.Add(BeginInfo.PackageName);
+			SetPackageStatus(BeginInfo.PackageName, EPackageStatus::DeclaredUnmodified_ConfirmedUnmodified);
+
 			TArray<FMessage> Messages;
-			IterativeFailed.RemoveAndCopyValue(BeginInfo.PackageName, Messages);
+			PackageMessageMap.RemoveAndCopyValue(BeginInfo.PackageName, Messages);
 			for (FMessage& Message : Messages)
 			{
 				// If no differences were detected, we should not have logged any warning or error messages
@@ -586,7 +869,7 @@ bool FIterativeValidatePackageWriter::IsAnotherSaveNeeded(FSavePackageResultStru
 			// New packages need to be resaved in Phase2; for our purposes they are equivalent
 			// to a package that iteration detected as modified.
 			// Do not add an entry for it in our results for iterative packages, and do not resave it in this pass
-			++ModifiedCount;
+			SetPackageStatus(BeginInfo.PackageName, EPackageStatus::DeclaredModified_WillNotVerify);
 			return false;
 		}
 	case EPhase::Phase2:
@@ -610,7 +893,7 @@ void FIterativeValidatePackageWriter::OnDiffWriterMessage(ELogVerbosity::Type Ve
 			TEXT("%s"), *ResolveText(Message));
 		break;
 	case EPhase::Phase1:
-		IterativeFailed.FindOrAdd(BeginInfo.PackageName).Add(FMessage{ FString(Message), Verbosity });
+		PackageMessageMap.FindOrAdd(BeginInfo.PackageName).Add(FMessage{ FString(Message), Verbosity });
 		break;
 	case EPhase::Phase2:
 		break;
@@ -632,15 +915,15 @@ void FIterativeValidatePackageWriter::LogIterativeDifferences()
 	{
 		UE_LOG(LogIterativeValidate, Display, TEXT("Could not validate %s because it has a non-deterministic save."),
 			*BeginInfo.PackageName.ToString());
-		IndeterminismFailed.Add(BeginInfo.PackageName);
+		SetPackageStatus(BeginInfo.PackageName, EPackageStatus::DeclaredUnmodified_FoundModified_Indeterminism);
 		return;
 	}
 
 	// Otherwise, no determinism issues, so the differences indicate a bug in Diff Package
-	IterativeSkipFalsePositive.Add(BeginInfo.PackageName);
+	SetPackageStatus(BeginInfo.PackageName, EPackageStatus::DeclaredUnmodified_FoundModified_FalsePositive);
 	FMsg::Logf(__FILE__, __LINE__, LogIterativeValidate.GetCategoryName(), ELogVerbosity::Warning,
 		TEXT("IterativeSkipFalsePositive package %s."), *BeginInfo.PackageName.ToString());
-	TArray<FMessage>& Messages = IterativeFailed.FindOrAdd(BeginInfo.PackageName);
+	TArray<FMessage>& Messages = PackageMessageMap.FindOrAdd(BeginInfo.PackageName);
 	for (const FMessage& Message : Messages)
 	{
 		FMsg::Logf(__FILE__, __LINE__, LogIterativeValidate.GetCategoryName(), Message.Verbosity,
@@ -693,9 +976,8 @@ void FIterativeValidatePackageWriter::Serialize(FArchive& Ar)
 		Ar.SetError();
 		return;
 	}
-	Ar << IterativeValidated;
-	Ar << IterativeFailed;
-	Ar << ModifiedCount;
+	Ar << PackageStatusMap;
+	Ar << PackageMessageMap;
 }
 
 FArchive& operator<<(FArchive& Ar, FIterativeValidatePackageWriter::FMessage& Message)
@@ -713,3 +995,29 @@ FString FIterativeValidatePackageWriter::GetIterativeValidatePath() const
 {
 	return FPaths::Combine(MetadataPath, FString(IterativeValidateFilename));
 }
+
+FIterativeValidatePackageWriter::EPackageStatus FIterativeValidatePackageWriter::GetPackageStatus(FName PackageName) const
+{
+	if (const EPackageStatus* Status = PackageStatusMap.Find(PackageName))
+	{
+		return *Status;
+	}
+	return EPackageStatus::NotYetProcessed;
+}
+
+void FIterativeValidatePackageWriter::SetPackageStatus(FName PackageName, EPackageStatus NewStatus)
+{
+	PackageStatusMap.FindOrAdd(PackageName) = NewStatus;
+}
+
+FIterativeValidatePackageWriter::FStatusCounts FIterativeValidatePackageWriter::CountPackagesByStatus()
+{
+	FIterativeValidatePackageWriter::FStatusCounts StatusCounts;
+
+	for (TPair<FName, EPackageStatus>& PackageStatusPair : PackageStatusMap)
+	{
+		StatusCounts[PackageStatusPair.Value]++;
+	}
+
+	return StatusCounts;
+}
diff --git a/Engine/Source/Editor/UnrealEd/Private/Cooker/IterativeValidatePackageWriter.h b/Engine/Source/Editor/UnrealEd/Private/Cooker/IterativeValidatePackageWriter.h
--- a/Engine/Source/Editor/UnrealEd/Private/Cooker/IterativeValidatePackageWriter.h
+++ b/Engine/Source/Editor/UnrealEd/Private/Cooker/IterativeValidatePackageWriter.h
@@ -5,6 +5,7 @@
 #include "DiffPackageWriter.h"
 
 class UCookOnTheFlyServer;
+class FIterativeValidateMPCollector; 
 
 /**
  * A CookedPackageWriter that diffs the cook results of iteratively-unmodified packages between their last cook
@@ -50,12 +51,16 @@ public:
 	virtual bool IsAnotherSaveNeeded(FSavePackageResultStruct& PreviousResult, FSavePackageArgs& SaveArgs) override;
 
 protected:
+	enum class EPackageStatus : uint8;
+
 	virtual void OnDiffWriterMessage(ELogVerbosity::Type Verbosity, FStringView Message) override;
 	void LogIterativeDifferences();
 	void Save();
 	void Load();
 	void Serialize(FArchive& Ar);
 	FString GetIterativeValidatePath() const;
+	EPackageStatus GetPackageStatus(FName PackageName) const;
+	void SetPackageStatus(FName PackageName, EPackageStatus NewStatus);
 
 	enum class ESaveAction : uint8
 	{
@@ -63,6 +68,19 @@ protected:
 		SaveToInner,
 		IgnoreResults,
 	};
+
+	enum class EPackageStatus : uint8
+	{
+		NotYetProcessed,
+		DeclaredUnmodified_ConfirmedUnmodified,
+		DeclaredUnmodified_FoundModified_IndeterminismOrFalsePositive,
+		DeclaredUnmodified_FoundModified_Indeterminism,
+		DeclaredUnmodified_FoundModified_FalsePositive,
+		DeclaredUnmodified_NotYetProcessed,
+		DeclaredModified_WillNotVerify,
+		Count
+	};
+
 	struct FMessage
 	{
 		FString Text;
@@ -70,40 +88,39 @@ protected:
 	};
 	friend FArchive& operator<<(FArchive& Ar, FMessage& Message);
 
-	/**
-	 * Written during Phase1, read during Phase2. Read/Written during AllInOnePhase. Packages that we thought were iteratively
-	 * unchanged and that were confirmed upon resave to be unchanged.
-	 */
-	TSet<FName> IterativeValidated;
-	/**
-	 * Written during Phase1, read during Phase2. Unused in AllInOnePhase. Packages that we thought were iteratively
-	 * unchanged but in which we discovered differences. But the differences might be due to indeterminism. Phase2
-	 * splits this container into IndeterminismFailed entries and IterativeSkipFalsePositive entries.
-	 */
-	TMap<FName, TArray<FMessage>> IterativeFailed;
-	/**
-	 * Read during Phase2. Unused in AllInOnePhase. Packages that we thought were iteratively unchanged but that had
-	 * differences when saved, but those differences turned out to be due to indeterminism.
-	 */
-	TSet<FName> IndeterminismFailed;
-	/**
-	 * Read during Phase2. Read/Written during AllInOnePhase. Packages that we thought were iteratively unchanged but that had
-	 * differences when saved, and no indeterminism detected (AllInOnePhase does not search for indeterminism, so all packages
-	 * with differences end up here), so they must be a bug in the decision to iteratively skip.
-	 */
-	TSet<FName> IterativeSkipFalsePositive;
-	/**
-	 * Unused in Phase1,Phase2. In AllInOnePhase, this records packages that we think are iteratively unchanged.
-	 * We do a 2-pass or 3-pass save for these: look for diffs and then save to disk. For packages not in this list
-	 * we do a 1-pass save: save it to disk without looking for diffs.
-	 */
-	TSet<FName> IterativelyUnmodified;
+	TMap<FName, EPackageStatus> PackageStatusMap;
+	TMap<FName, TArray<FMessage>> PackageMessageMap;
+
+	struct FStatusCounts
+	{
+		FStatusCounts()
+		{
+			Data = MakeUniformStaticArray<uint32, (uint32)EPackageStatus::Count>(0);
+		}
+
+		uint32& operator[](EPackageStatus PackageStatus) 
+		{ 
+			return Data[(uint32)PackageStatus]; 
+		}
+
+	private:
+		TStaticArray<uint32, (uint32)EPackageStatus::Count> Data;
+	};
+
+	FStatusCounts CountPackagesByStatus();
 
 	FString MetadataPath;
 	UCookOnTheFlyServer& COTFS;
-	int32 ModifiedCount = 0;
 	EPhase Phase = EPhase::AllInOnePhase;
 	ESaveAction SaveAction = ESaveAction::IgnoreResults;
 	bool bPackageFirstPass = false;
 	bool bReadOnly = true;
+
+	friend class FIterativeValidateMPCollector;
+	friend FCbWriter& operator<<(FCbWriter& Writer, const FMessage& Message);
+	friend bool LoadFromCompactBinary(FCbFieldView Field, FMessage& Path);
+	friend FCbWriter& operator<<(FCbWriter& Writer, EPackageStatus Status);
+	friend bool LoadFromCompactBinary(FCbFieldView Field, EPackageStatus& Status);
 };
+
+
diff --git a/Engine/Source/Runtime/CoreUObject/Private/Cooker/MPCollector.cpp b/Engine/Source/Runtime/CoreUObject/Private/Cooker/MPCollector.cpp
--- a/Engine/Source/Runtime/CoreUObject/Private/Cooker/MPCollector.cpp
+++ b/Engine/Source/Runtime/CoreUObject/Private/Cooker/MPCollector.cpp
@@ -46,6 +46,21 @@ const ITargetPlatform* FMPCollectorClientTickContext::IntToPlatform(uint8 Platfo
 	return IntToPlatformImpl(PlatformAsInt, Platforms);
 }
 
+void FMPCollectorServerTickContext::AddMessage(FCbObject Object)
+{
+	Messages.Add(Object);
+}
+
+uint8 FMPCollectorServerTickContext::PlatformToInt(const ITargetPlatform* Platform) const
+{
+	return PlatformToIntImpl(Platform, Platforms);
+}
+
+const ITargetPlatform* FMPCollectorServerTickContext::IntToPlatform(uint8 PlatformAsInt) const
+{
+	return IntToPlatformImpl(PlatformAsInt, Platforms);
+}
+
 void FMPCollectorClientTickPackageContext::AddMessage(FCbObject Object)
 {
 	Messages.Emplace(nullptr, MoveTemp(Object));
@@ -78,6 +93,21 @@ const ITargetPlatform* FMPCollectorClientTickPackageContext::IntToPlatform(uint8
 	return IntToPlatformImpl(PlatformAsInt, Platforms);
 }
 
+void FMPCollectorServerTickPackageContext::AddMessage(FCbObject Object)
+{
+	Messages.Emplace(MoveTemp(Object));
+}
+
+uint8 FMPCollectorServerTickPackageContext::PlatformToInt(const ITargetPlatform* Platform) const
+{
+	return PlatformToIntImpl(Platform, Platforms);
+}
+
+const ITargetPlatform* FMPCollectorServerTickPackageContext::IntToPlatform(uint8 PlatformAsInt) const
+{
+	return IntToPlatformImpl(PlatformAsInt, Platforms);
+}
+
 uint8 FMPCollectorClientMessageContext::PlatformToInt(const ITargetPlatform* Platform) const
 {
 	return PlatformToIntImpl(Platform, Platforms);
diff --git a/Engine/Source/Runtime/CoreUObject/Public/Cooker/MPCollector.h b/Engine/Source/Runtime/CoreUObject/Public/Cooker/MPCollector.h
--- a/Engine/Source/Runtime/CoreUObject/Public/Cooker/MPCollector.h
+++ b/Engine/Source/Runtime/CoreUObject/Public/Cooker/MPCollector.h
@@ -84,6 +84,35 @@ private:
 	friend class FCookWorkerClient;
 };
 
+class FMPCollectorServerTickContext
+{
+public:
+	enum class EServerEventType : uint8
+	{
+		WorkerStartup,
+		Count
+	};
+
+	FMPCollectorServerTickContext(EServerEventType InEventType) : EventType(InEventType) { check(InEventType < EServerEventType::Count); }
+
+	TConstArrayView<const ITargetPlatform*> GetPlatforms() const { return Platforms; }
+	
+	COREUOBJECT_API void AddMessage(FCbObject Object);
+
+	COREUOBJECT_API uint8 PlatformToInt(const ITargetPlatform* Platform) const;
+	COREUOBJECT_API const ITargetPlatform* IntToPlatform(uint8 PlatformAsInt) const;
+	
+	EServerEventType GetEventType() const { return EventType; }
+
+private:
+	TConstArrayView<const ITargetPlatform*> Platforms;
+	TArray<FCbObject> Messages;
+	EServerEventType EventType;
+	bool bFlush = false;
+
+	friend class FCookDirector;
+};
+
 class FMPCollectorClientTickPackageContext
 {
 public:
@@ -114,16 +143,43 @@ private:
 	friend class FCookWorkerClient;
 };
 
+class FMPCollectorServerTickPackageContext
+{
+public:
+	struct FPlatformData
+	{
+		const ITargetPlatform* TargetPlatform = nullptr;
+		ECookResult CookResults = ECookResult::NotAttempted;
+	};
+	FName GetPackageName() const { return PackageName; }
+
+	COREUOBJECT_API void AddMessage(FCbObject Object);
+
+	COREUOBJECT_API uint8 PlatformToInt(const ITargetPlatform* Platform) const;
+	COREUOBJECT_API const ITargetPlatform* IntToPlatform(uint8 PlatformAsInt) const;
+
+private:
+	TArray<FCbObject> Messages;
+	TConstArrayView<const ITargetPlatform*> Platforms;
+	TConstArrayView<FPlatformData> PlatformDatas;
+	FName PackageName;
+
+	friend class FCookDirector;
+};
+
 class FMPCollectorClientMessageContext
 {
 public:
 	TConstArrayView<const ITargetPlatform*> GetPlatforms() { return Platforms; }
+	// Name of the relevant package or NAME_None if not available
+	FName GetPackageName() const { return PackageName; }
 
 	COREUOBJECT_API uint8 PlatformToInt(const ITargetPlatform* Platform) const;
 	COREUOBJECT_API const ITargetPlatform* IntToPlatform(uint8 PlatformAsInt) const;
 
 private:
 	TConstArrayView<const ITargetPlatform*> Platforms;
+	FName PackageName;
 
 	friend class FCookWorkerClient;
 };
@@ -167,7 +223,9 @@ public:
 	virtual FGuid GetMessageType() const = 0;
 	virtual const TCHAR* GetDebugName() const = 0;
 
+	virtual void ServerTick(FMPCollectorServerTickContext& Context) {}
 	virtual void ClientTick(FMPCollectorClientTickContext& Context) {}
+	virtual void ServerTickPackage(FMPCollectorServerTickPackageContext& Context) {}
 	virtual void ClientTickPackage(FMPCollectorClientTickPackageContext& Context) {}
 	virtual void ClientReceiveMessage(FMPCollectorClientMessageContext& Context, FCbObjectView Message) {}
 	virtual void ServerReceiveMessage(FMPCollectorServerMessageContext& Context, FCbObjectView Message) {}
