IncrementalCooker: Improve log feedback in Iteratively Skipped log; do not display it if iterativevalidate is running, which always cooks all packages whether they are modified or not.
#rnx

diff --git a/Engine/Source/Editor/UnrealEd/Private/CookOnTheFlyServer.cpp b/Engine/Source/Editor/UnrealEd/Private/CookOnTheFlyServer.cpp
--- a/Engine/Source/Editor/UnrealEd/Private/CookOnTheFlyServer.cpp
+++ b/Engine/Source/Editor/UnrealEd/Private/CookOnTheFlyServer.cpp
@@ -10539,15 +10539,36 @@ void UCookOnTheFlyServer::PrintFinishStats()
 		UE_LOG(LogCook, Display, TEXT("CookWorker total time %f"), TotalCookTime);
 	}
 
+	// Suppress NumPackagesIterativelySkipped display if the PackageWriter is modifying what gets skipped
+	bool bReportIterativeSkips = true;
+	int32 ReportedNumPackagesIterativelySkipped = DetailedCookStats::NumPackagesIterativelySkipped;
+	const ITargetPlatform* FirstTargetPlatform = PlatformManager->GetSessionPlatforms().IsEmpty() ? nullptr :
+		PlatformManager->GetSessionPlatforms()[0];
+	if (FirstTargetPlatform)
+	{
+		if (FindOrCreateSaveContext(FirstTargetPlatform).PackageWriterCapabilities.bOverridesPackageModificationStatus)
+		{
+			bReportIterativeSkips = false;
+			ReportedNumPackagesIterativelySkipped = 0;
+		}
+	}
+	int32 ReportedNumCooked = PackageDatas->GetNumCooked(ECookResult::Succeeded)
+		- ReportedNumPackagesIterativelySkipped - PackageDataFromBaseGameNum;
+	int32 ReportedTotalPackages = PackageDatas->GetNumCooked()
+		- PackageDatas->GetNumCooked(ECookResult::NeverCookPlaceholder) - PackageDataFromBaseGameNum;
 
 	const FPlatformMemoryStats MemStats = FPlatformMemory::GetStats();
 	UE_LOG(LogCook, Display, TEXT("Peak Used virtual %u MiB Peak Used physical %u MiB"), MemStats.PeakUsedVirtual / 1024 / 1024, MemStats.PeakUsedPhysical / 1024 / 1024);
 
-	COOK_STAT(UE_LOG(LogCook, Display, TEXT("Packages Cooked: %d, Packages Iteratively Skipped: %d, Packages Skipped by Platform: %d, Total Packages: %d"),
-		PackageDatas->GetNumCooked(ECookResult::Succeeded) - DetailedCookStats::NumPackagesIterativelySkipped - PackageDataFromBaseGameNum,
-		DetailedCookStats::NumPackagesIterativelySkipped,
+	COOK_STAT(UE_LOG(LogCook, Display,
+		TEXT("Packages Cooked: %d,%s Packages Skipped by Platform: %d, Total Packages: %d"),
+		ReportedNumCooked,
+		(bReportIterativeSkips
+			? *FString::Printf(TEXT(" Packages Iteratively Skipped: %d,"), ReportedNumPackagesIterativelySkipped)
+			: TEXT("")),
 		PackageDatas->GetNumCooked(ECookResult::Failed),
-		PackageDatas->GetNumCooked() - PackageDatas->GetNumCooked(ECookResult::NeverCookPlaceholder) - PackageDataFromBaseGameNum));
+		ReportedTotalPackages
+	));
 }
 
 void UCookOnTheFlyServer::PrintDetailedCookStats()
diff --git a/Engine/Source/Editor/UnrealEd/Private/Cooker/IterativeValidatePackageWriter.cpp b/Engine/Source/Editor/UnrealEd/Private/Cooker/IterativeValidatePackageWriter.cpp
--- a/Engine/Source/Editor/UnrealEd/Private/Cooker/IterativeValidatePackageWriter.cpp
+++ b/Engine/Source/Editor/UnrealEd/Private/Cooker/IterativeValidatePackageWriter.cpp
@@ -513,6 +513,7 @@ ICookedPackageWriter::FCookCapabilities FIterativeValidatePackageWriter::GetCook
 {
 	FCookCapabilities Result = Super::GetCookCapabilities();
 	Result.bReadOnly = bReadOnly;
+	Result.bOverridesPackageModificationStatus = true;
 	return Result;
 }
 
@@ -687,13 +688,17 @@ void FIterativeValidatePackageWriter::EndCook(const FCookInfo& Info)
 	{
 	case EPhase::AllInOnePhase:
 	{
+		int32 DetectedUnmodified = StatusCounts[EPackageStatus::DeclaredUnmodified_ConfirmedUnmodified]
+			+ StatusCounts[EPackageStatus::DeclaredUnmodified_FoundModified_FalsePositive];
 		UE_LOG(LogIterativeValidate, Display,
 			TEXT("Modified: %d. DetectedUnmodified: %d. ValidatedUnmodified: %d. IterativeSkipFalsePositive: %d."),
 			StatusCounts[EPackageStatus::DeclaredModified_WillNotVerify], 
-			StatusCounts[EPackageStatus::DeclaredUnmodified_ConfirmedUnmodified] + StatusCounts[EPackageStatus::DeclaredUnmodified_FoundModified_FalsePositive],
+			DetectedUnmodified,
 			StatusCounts[EPackageStatus::DeclaredUnmodified_ConfirmedUnmodified],
 			StatusCounts[EPackageStatus::DeclaredUnmodified_FoundModified_FalsePositive]);
-		FString Message = FString::Printf(TEXT("IterativeSkipFalsePositive: %d."), StatusCounts[EPackageStatus::DeclaredUnmodified_FoundModified_FalsePositive]);
+		FString Message = FString::Printf(TEXT("Packages Iteratively Skipped: %d: IterativeSkipFalsePositive: %d."),
+			DetectedUnmodified,
+			StatusCounts[EPackageStatus::DeclaredUnmodified_FoundModified_FalsePositive]);
 		if (StatusCounts[EPackageStatus::DeclaredUnmodified_FoundModified_FalsePositive] > 0)
 		{
 			UE_LOG(LogIterativeValidate, Error, TEXT("%s"), *Message);
@@ -705,25 +710,35 @@ void FIterativeValidatePackageWriter::EndCook(const FCookInfo& Info)
 		break;
 	}
 	case EPhase::Phase1:
+	{
+		int32 DetectedUnmodified = StatusCounts[EPackageStatus::DeclaredUnmodified_ConfirmedUnmodified]
+			+ StatusCounts[EPackageStatus::DeclaredUnmodified_FoundModified_IndeterminismOrFalsePositive];
 		UE_LOG(LogIterativeValidate, Display,
 			TEXT("Modified: %d. DetectedUnmodified: %d. ValidatedUnmodified: %d. IterativeSkipFalsePositiveOrIndeterminism: %d."),
 			StatusCounts[EPackageStatus::DeclaredModified_WillNotVerify],
-			StatusCounts[EPackageStatus::DeclaredUnmodified_ConfirmedUnmodified] 
-			+ StatusCounts[EPackageStatus::DeclaredUnmodified_FoundModified_IndeterminismOrFalsePositive],
+			DetectedUnmodified,
 			StatusCounts[EPackageStatus::DeclaredUnmodified_ConfirmedUnmodified],
 			StatusCounts[EPackageStatus::DeclaredUnmodified_FoundModified_IndeterminismOrFalsePositive]);
 		Save();
 		break;
+	}
 	case EPhase::Phase2:
 	{
+		int32 DetectedUnmodified = StatusCounts[EPackageStatus::DeclaredUnmodified_ConfirmedUnmodified]
+			+ StatusCounts[EPackageStatus::DeclaredUnmodified_FoundModified_IndeterminismOrFalsePositive]
+			+ StatusCounts[EPackageStatus::DeclaredUnmodified_FoundModified_Indeterminism]
+			+ StatusCounts[EPackageStatus::DeclaredUnmodified_FoundModified_FalsePositive];
 		UE_LOG(LogIterativeValidate, Display,
-			TEXT("Modified: %d. DetectedUnmodified: %d. ValidatedUnmodified: %d. Indeterminism: %d."),
+			TEXT("Modified: %d. DetectedUnmodified: %d. ValidatedUnmodified: %d. Indeterminism: %d. IterativeSkipFalsePositive: %d."),
 			StatusCounts[EPackageStatus::DeclaredModified_WillNotVerify],
-			StatusCounts[EPackageStatus::DeclaredUnmodified_ConfirmedUnmodified] 
-			+ StatusCounts[EPackageStatus::DeclaredUnmodified_FoundModified_IndeterminismOrFalsePositive],
+			DetectedUnmodified,
 			StatusCounts[EPackageStatus::DeclaredUnmodified_ConfirmedUnmodified],
-			StatusCounts[EPackageStatus::DeclaredUnmodified_FoundModified_Indeterminism]);
-		FString Message = FString::Printf(TEXT("IterativeSkipFalsePositive: %d."), StatusCounts[EPackageStatus::DeclaredUnmodified_FoundModified_FalsePositive]);
+			StatusCounts[EPackageStatus::DeclaredUnmodified_FoundModified_Indeterminism],
+			StatusCounts[EPackageStatus::DeclaredUnmodified_FoundModified_FalsePositive]
+			);
+		FString Message = FString::Printf(TEXT("Packages Iteratively Skipped: %d: IterativeSkipFalsePositive: %d."),
+			DetectedUnmodified,
+			StatusCounts[EPackageStatus::DeclaredUnmodified_FoundModified_FalsePositive]);
 		if (StatusCounts[EPackageStatus::DeclaredUnmodified_FoundModified_FalsePositive] > 0)
 		{
 			UE_LOG(LogIterativeValidate, Error, TEXT("%s"), *Message);
diff --git a/Engine/Source/Runtime/Core/Public/Serialization/PackageWriter.h b/Engine/Source/Runtime/Core/Public/Serialization/PackageWriter.h
--- a/Engine/Source/Runtime/Core/Public/Serialization/PackageWriter.h
+++ b/Engine/Source/Runtime/Core/Public/Serialization/PackageWriter.h
@@ -248,6 +248,13 @@ public:
 		/** If true, the cooker will assume no packages are written and will skip writing non-package data. */
 		bool bReadOnly = false;
 
+		/**
+		 * If true, the PackageWriter can override which packages are iteratively skipped in
+		 * UpdatePackageModificationStatus, and the cooker will therefore avoid assumptions about what is
+		 * iteratively skipped
+		 */
+		bool bOverridesPackageModificationStatus = false;
+
 		/** What header format is produced as output by this writer. */
 		EPackageHeaderFormat HeaderFormat = EPackageHeaderFormat::PackageFileSummary;
 	};
