IncrementalCook: Fix generated packages of generator packages in the iterativedenylist to also be iteratively denied.
#rnx

diff --git a/Engine/Source/Editor/UnrealEd/Private/Cooker/CookRequestCluster.cpp b/Engine/Source/Editor/UnrealEd/Private/Cooker/CookRequestCluster.cpp
--- a/Engine/Source/Editor/UnrealEd/Private/Cooker/CookRequestCluster.cpp
+++ b/Engine/Source/Editor/UnrealEd/Private/Cooker/CookRequestCluster.cpp
@@ -1318,11 +1318,12 @@ bool FRequestCluster::FGraphSearch::FExploreEdgesContext::TryCalculateIterativel
 
 		UE::TargetDomain::FCookDependencies& CookDependencies = QueryPlatformData.CookAttachments.Dependencies;
 		const FAssetPackageData* OverrideAssetPackageData = nullptr;
+		FPackageData* ParentPackageData = nullptr;
 		if (PackageData->IsGenerated())
 		{
 			// If a generator is marked iteratively unmodified, then by contract we are not required to test its
 			// generated packages; they are all marked iteratively unmodified as well
-			FPackageData* ParentPackageData = Cluster.PackageDatas.FindPackageDataByPackageName(
+			ParentPackageData = Cluster.PackageDatas.FindPackageDataByPackageName(
 				PackageData->GetParentGenerator());
 			if (ParentPackageData)
 			{
@@ -1358,6 +1359,15 @@ bool FRequestCluster::FGraphSearch::FExploreEdgesContext::TryCalculateIterativel
 			SetIsIterativelyUnmodified(PlatformIndex, false, PackagePlatformData);
 			continue;
 		}
+		// Generated packages of a generator that is not IterativelyEnabled are also not iteratively enabled, even
+		// if they would otherwise qualify for iterative on their own. e.g. if worlds are iteratively disallowed,
+		// then streamingobject generated packages of the world are also disallowed.
+		if (ParentPackageData && !IsIterativeEnabled(ParentPackageData->GetPackageName(),
+			Cluster.COTFS.bHybridIterativeAllowAllClasses))
+		{
+			SetIsIterativelyUnmodified(PlatformIndex, false, PackagePlatformData);
+			continue;
+		}
 
 		if (!QueryPlatformData.bTransitiveBuildDependenciesResolvedAsNotModified)
 		{
