Cooker: DiffOnly: Fix callstack reported for a diff when the diff location is inside a serialization block that was written from one callstack, and a separate callstack seeked back to within the range of the serialization block and modified some of the bytes before the diff location but did not modify the diff location. In that case, the callstack of the first bigger serialization block should be the one reported, rather than the seek back callstack.
#rnx

diff --git a/Engine/Source/Editor/UnrealEd/Private/Cooker/DiffWriterArchive.cpp b/Engine/Source/Editor/UnrealEd/Private/Cooker/DiffWriterArchive.cpp
--- a/Engine/Source/Editor/UnrealEd/Private/Cooker/DiffWriterArchive.cpp
+++ b/Engine/Source/Editor/UnrealEd/Private/Cooker/DiffWriterArchive.cpp
@@ -239,31 +239,98 @@ void FCallstacks::Add(
 	{
 		uint32 CallstackCRC = 0;
 		bool bSuppressLogging = UE::ArchiveStackTrace::ShouldIgnoreDiff();
-		if (CallstackAtOffsetMap.Num() == 0 || CurrentOffset > CallstackAtOffsetMap.Last().Offset)
+		LastSerializeCallstack = AddUniqueCallstack(bIsCollectingCallstacks, SerializedObject, SerializedProperty, CallstackCRC);
+		check(CallstackCRC != 0 || !bShouldCollectCallstack);
+		FCallstackAtOffset NewBlock{ CurrentOffset, Length, CurrentOffset, Length, CallstackCRC, bSuppressLogging };
+
+		FCallstackAtOffset* LastBlock = CallstackAtOffsetMap.Num() ? &CallstackAtOffsetMap.Last() : nullptr;
+		if (!LastBlock || CurrentOffset >= LastBlock->Offset + LastBlock->Length)
 		{
-			// New data serialized at the end of archive buffer
-			LastSerializeCallstack = AddUniqueCallstack(bIsCollectingCallstacks, SerializedObject, SerializedProperty, CallstackCRC);
-			CallstackAtOffsetMap.Add(FCallstackAtOffset { CurrentOffset, CallstackCRC, bSuppressLogging });
+			// New block serialized at the end of archive buffer
+			CallstackAtOffsetMap.Add(NewBlock);
 		}
 		else
 		{
-			// This happens usually after Seek() so we need to find the exiting offset or insert a new one
-			const int32 CallstackToUpdateIndex = GetCallstackIndexAtOffset(CurrentOffset);
-			check(CallstackToUpdateIndex != -1);
-			FCallstackAtOffset& CallstackToUpdate = CallstackAtOffsetMap[CallstackToUpdateIndex];
-			LastSerializeCallstack = AddUniqueCallstack(bIsCollectingCallstacks, SerializedObject, SerializedProperty, CallstackCRC);
-			if (CallstackToUpdate.Offset == CurrentOffset)
+			// This happens after a Seek(). We need to modify or replace the old block that covered the range written
+			// by this new Serialize call.
+			const int32 OldBlockIndex = GetCallstackIndexAtOffset(CurrentOffset);
+			check(OldBlockIndex != -1);
+
+			FCallstackAtOffset* OldBlock = &CallstackAtOffsetMap[OldBlockIndex];
+
+			int64 OldEnd = OldBlock->Offset + OldBlock->Length;
+			int64 NewEnd = NewBlock.Offset + NewBlock.Length;
+			if (OldEnd <= NewEnd)
 			{
-				CallstackToUpdate.Callstack = CallstackCRC;
+				// The new block overwrites the end of the old block, and possibly overwrites blocks after it
+				check(OldBlock->Offset <= NewBlock.Offset); // GetCallstackIndexAtOffset guarantees this
+				bool bNewEntirelyContainsOld = OldBlock->Offset == NewBlock.Offset;
+				int32 StartRemoveIndex;
+				if (bNewEntirelyContainsOld)
+				{
+					// The new block completely overwrites the old block; delete the old block and replace it with the
+					// the new block. Still need to check whether new also overwrites old blocks after the first.
+					*OldBlock = NewBlock;
+					StartRemoveIndex = OldBlockIndex + 1;
+				}
+				else
+				{
+					// The new block does not overwrite the old block, so keep the old block, clamp it to end at the
+					// new block, and add the new block after it. Still need to check whether new also overwrites old
+					// blocks after the first.
+					// There might be a gap in between the end of old block and the beginning of the new block. This
+					// can occur when we are not recording every serialize call. Leave the old block unmodified in
+					// that case.
+					if (OldEnd > NewBlock.Offset)
+					{
+						OldBlock->Length = NewBlock.Offset - OldBlock->Offset;
+					}
+					CallstackAtOffsetMap.Insert(NewBlock, OldBlockIndex + 1);
+					OldBlock = nullptr; // Our pointer for OldBlock is now possibly invalidated by reallocation.
+					StartRemoveIndex = OldBlockIndex + 2;
+				}
+				int32 EndRemoveIndex = StartRemoveIndex;
+				while (EndRemoveIndex < CallstackAtOffsetMap.Num())
+				{
+					OldBlock = &CallstackAtOffsetMap[EndRemoveIndex];
+					if (OldBlock->Offset >= NewEnd)
+					{
+						break;
+					}
+					OldEnd = OldBlock->Offset + OldBlock->Length;
+					if (OldEnd > NewEnd)
+					{
+						// The beginning of this followup block is overwritten by the new block; shorten it
+						OldBlock->Length = OldEnd - NewEnd;
+						OldBlock->Offset = NewEnd;
+						break;
+					}
+					else
+					{
+						// This followup block is completely inside the new block; mark it for delete and move to next.
+						++EndRemoveIndex;
+					}
+				}
+				if (EndRemoveIndex > StartRemoveIndex)
+				{
+					CallstackAtOffsetMap.RemoveAt(StartRemoveIndex, EndRemoveIndex - StartRemoveIndex,
+						EAllowShrinking::No);
+				}
 			}
-			else
+			else // OldEnd > NewEnd
 			{
-				// Insert a new callstack
-				check(CallstackToUpdate.Offset < CurrentOffset);
-				CallstackAtOffsetMap.Insert(FCallstackAtOffset {CurrentOffset, CallstackCRC, bSuppressLogging }, CallstackToUpdateIndex + 1);
+				// The new block is completely inside the old block, which extends after it. Shorten the beginning of
+				// the old block, and add an additional new block containing the portion of the old block that extends
+				// after the new block.
+				FCallstackAtOffset SegmentAfterNewBlock = *OldBlock;
+				SegmentAfterNewBlock.Offset = NewEnd;
+				SegmentAfterNewBlock.Length = OldEnd - NewEnd;
+				OldBlock->Length = NewBlock.Offset - OldBlock->Offset;
+				CallstackAtOffsetMap.Insert(NewBlock, OldBlockIndex + 1);
+				CallstackAtOffsetMap.Insert(SegmentAfterNewBlock, OldBlockIndex + 2);
+				OldBlock = nullptr; // Our pointer for OldBlock is now possibly invalidated by reallocation.
 			}
 		}
-		check(CallstackCRC != 0 || !bShouldCollectCallstack);
 	}
 	else if (LastSerializeCallstack)
 	{
@@ -339,6 +406,7 @@ void FCallstacks::Append(const FCallstacks& Other, int64 OtherStartOffset)
 	{
 		FCallstackAtOffset& New = CallstackAtOffsetMap.Add_GetRef(OtherOffset);
 		New.Offset += OtherStartOffset;
+		New.SerializeCallOffset += OtherStartOffset;
 	}
 
 	CallstackAtOffsetMap.Sort([](const FCallstackAtOffset& LHS,const FCallstackAtOffset& RHS)
@@ -959,7 +1027,7 @@ void FAccumulator::CompareWithPreviousForSection(const FPackageData& SourcePacka
 				TEXT("%s: Difference at offset %lld (Combined/DiffBreak Offset: %" INT64_FMT "): OnDisk %d != %d InMemory.%s")
 					TEXT("Difference occurs at index %lld within Serialize call at callstack:%s%s%s%s"),
 				*SectionFilename, LocalOffset, DestAbsoluteOffset, SourceByte, DestByte, NewLineToken,
-				DestAbsoluteOffset - CallstackAtOffset.Offset, NewLineToken,
+				DestAbsoluteOffset - CallstackAtOffset.SerializeCallOffset, NewLineToken,
 				*LastDifferenceCallstackDataText, *DiffValues, *DebugDataStackText
 			));
 		}
@@ -1171,10 +1239,8 @@ void FAccumulator::GenerateDiffMapForSection(const FPackageData& SourcePackage,
 					const FCallstacks::FCallstackAtOffset& CallstackAtOffset = Callstacks.GetCallstack(DifferenceCallstackOffsetIndex);
 					if (!CallstackAtOffset.bSuppressLogging)
 					{
-						FDiffInfo OffsetAndSize;
-						OffsetAndSize.Offset = CallstackAtOffset.Offset;
-						OffsetAndSize.Size = Callstacks.GetSerializedDataSizeForOffsetIndex(DifferenceCallstackOffsetIndex);
-						DiffMap.Add(OffsetAndSize);
+						DiffMap.Add(FDiffInfo(CallstackAtOffset.SerializeCallOffset,
+							CallstackAtOffset.SerializeCallLength));
 					}
 				}
 				LastDifferenceCallstackOffsetIndex = DifferenceCallstackOffsetIndex;
@@ -1195,10 +1261,8 @@ void FAccumulator::GenerateDiffMapForSection(const FPackageData& SourcePackage,
 			{
 				if (!CallstackAtOffset.bSuppressLogging)
 				{
-					FDiffInfo OffsetAndSize;
-					OffsetAndSize.Offset = CallstackAtOffset.Offset;
-					OffsetAndSize.Size = Callstacks.GetSerializedDataSizeForOffsetIndex(OffsetIndex);
-					DiffMap.Add(OffsetAndSize);
+					DiffMap.Add(FDiffInfo(CallstackAtOffset.SerializeCallOffset,
+						CallstackAtOffset.SerializeCallLength));
 				}
 			}
 			else
diff --git a/Engine/Source/Editor/UnrealEd/Private/Cooker/DiffWriterArchive.h b/Engine/Source/Editor/UnrealEd/Private/Cooker/DiffWriterArchive.h
--- a/Engine/Source/Editor/UnrealEd/Private/Cooker/DiffWriterArchive.h
+++ b/Engine/Source/Editor/UnrealEd/Private/Cooker/DiffWriterArchive.h
@@ -101,8 +101,20 @@ public:
 	/** Offset and callstack pair */
 	struct FCallstackAtOffset
 	{
-		/** Offset of a Serialize call */
+		/**
+		 * Offset of a block written by Serialize call. Equal to SerializeCallOffset unless the block was split by a
+		 * separate Serialize call.
+		 */
 		int64 Offset = -1;
+		/**
+		 * Length of a block written by Serialize call. Equal to SerializeCallLength unless the block was split by a
+		 * separate Serialize call.
+		 */
+		int64 Length = -1;
+		/** The offset written to by the Serialize call. */
+		int64 SerializeCallOffset = -1;
+		/** The length written by the Serialize call. */
+		int64 SerializeCallLength = -1;
 		/** Callstack CRC for the Serialize call */
 		uint32 Callstack = 0;
 		/** Collected inside of a scope that indicates diff should be recorded but logging should be suppressed */
@@ -182,19 +194,6 @@ public:
 		return UniqueCallstacks[CallstackOffset.Callstack];
 	}
 
-	/** Returns the size of serialized data at the specified offset. */
-	int64 GetSerializedDataSizeForOffsetIndex(int32 InOffsetIndex) const
-	{
-		if (InOffsetIndex < CallstackAtOffsetMap.Num() - 1)
-		{
-			return CallstackAtOffsetMap[InOffsetIndex + 1].Offset - CallstackAtOffsetMap[InOffsetIndex].Offset;
-		}
-		else
-		{
-			return EndOffset - CallstackAtOffsetMap[InOffsetIndex].Offset;
-		}
-	}
-
 	int64 GetEndOffset() const
 	{
 		return EndOffset;
