Create a property bag placeholder type during subobject reinstancing if the previous subobject's type is a class compiler artifact w/ no reinstanced target.

Notes:
- Currently this path is disabled by default. Requires IDO and placeholder type serialization features to be enabled for the subobject instance.
- When enabled, this ensures that edits are preserved across reinstancing, and that those edits will be serialized along with the outer instance as a tagged property set, so that there is no loss of user data.

diff --git a/Engine/Source/Editor/UnrealEd/Private/Kismet2/KismetReinstanceUtilities.cpp b/Engine/Source/Editor/UnrealEd/Private/Kismet2/KismetReinstanceUtilities.cpp
--- a/Engine/Source/Editor/UnrealEd/Private/Kismet2/KismetReinstanceUtilities.cpp
+++ b/Engine/Source/Editor/UnrealEd/Private/Kismet2/KismetReinstanceUtilities.cpp
@@ -3303,6 +3303,22 @@ void FBlueprintCompileReinstancer::PreCreateSubObjectsForReinstantiation_Inner(
 				{
 					SubObjectClass = *NewSubObjectClass;
 				}
+				else if (OldSubObjectClass->HasAnyClassFlags(CLASS_NewerVersionExists))
+				{
+					// The old subobject class is a compiler artifact, with no reinstanced counterpart. Its source may have been renamed (or deleted) and thus it could not be recompiled.
+					// In this case, if property bag features are enabled, we attempt to redirect to a placeholder type so we can serialize to a property bag in order to preserve its data.
+					if (UE::FPropertyBagRepository::IsPropertyBagPlaceholderObjectSupportEnabled())
+					{
+						FTopLevelAssetPath OldSubObjectClassPath = OldSubObjectClass->GetReinstancedClassPathName();
+						if (OldSubObjectClassPath.IsValid())
+						{
+							if (UPackage* OldSubObjectClassPackage = FindPackage(nullptr, *OldSubObjectClassPath.GetPackageName().ToString()))
+							{
+								SubObjectClass = UE::FPropertyBagRepository::CreatePropertyBagPlaceholderClass(OldSubObjectClassPackage, OldSubObjectClass->GetClass(), OldSubObjectClassPath.GetAssetName(), OldSubObjectClass->GetFlags() & ~RF_Transient);
+							}
+						}
+					}
+				}
 
 				// Only pre-create object where the class does not have newer version of the it
 				if(!SubObjectClass->HasAnyClassFlags(CLASS_NewerVersionExists))
diff --git a/Engine/Source/Runtime/CoreUObject/Private/UObject/Class.cpp b/Engine/Source/Runtime/CoreUObject/Private/UObject/Class.cpp
--- a/Engine/Source/Runtime/CoreUObject/Private/UObject/Class.cpp
+++ b/Engine/Source/Runtime/CoreUObject/Private/UObject/Class.cpp
@@ -5454,6 +5454,18 @@ bool UClass::IsSafeToSerializeToStructuredArchives(UClass* InClass)
 	return true;
 }
 
+#if WITH_EDITOR
+FTopLevelAssetPath UClass::GetReinstancedClassPathName() const
+{
+	if (HasAnyClassFlags(CLASS_NewerVersionExists))
+	{
+		return GetReinstancedClassPathName_Impl();
+	}
+
+	return nullptr;
+}
+#endif
+
 #if USTRUCT_FAST_ISCHILDOF_IMPL == USTRUCT_ISCHILDOF_STRUCTARRAY
 
 	FStructBaseChain::FStructBaseChain()
diff --git a/Engine/Source/Runtime/CoreUObject/Public/UObject/Class.h b/Engine/Source/Runtime/CoreUObject/Public/UObject/Class.h
--- a/Engine/Source/Runtime/CoreUObject/Public/UObject/Class.h
+++ b/Engine/Source/Runtime/CoreUObject/Public/UObject/Class.h
@@ -3658,6 +3658,14 @@ public:
 	 */
 	static COREUOBJECT_API bool IsSafeToSerializeToStructuredArchives(UClass* InClass);
 
+#if WITH_EDITOR
+	/**
+	 * If this class was recompiled, returns the path to where we should find the new (i.e. reinstanced) class.
+	 * Note: If this is not a recompiled class, this will return a NULL path and should be handled appropriately.
+	 */
+	COREUOBJECT_API FTopLevelAssetPath GetReinstancedClassPathName() const;
+#endif
+
 private:
 	/** 
 	 * This signature intentionally hides the method declared in UObjectBaseUtility to make it private.
@@ -3708,6 +3716,14 @@ protected:
 	 * @return		the CDO for this class
 	 **/
 	COREUOBJECT_API virtual UObject* CreateDefaultObject();
+
+#if WITH_EDITOR
+	/**
+	 * Internal helper method for GetReinstancedClassPathName(). Subclasses can choose to override the default
+	 * implementation to return a valid class path for recompiled class objects based on how reinstancing is handled.
+	 */
+	virtual FTopLevelAssetPath GetReinstancedClassPathName_Impl() const { return nullptr; }
+#endif
 };
 
 /**
