Cooker: Add more feedback for the error: CookPackageSplitter failure: We are demoting a generated package...
#rnx

diff --git a/Engine/Source/Editor/UnrealEd/Private/Cooker/CookGenerationHelper.cpp b/Engine/Source/Editor/UnrealEd/Private/Cooker/CookGenerationHelper.cpp
--- a/Engine/Source/Editor/UnrealEd/Private/Cooker/CookGenerationHelper.cpp
+++ b/Engine/Source/Editor/UnrealEd/Private/Cooker/CookGenerationHelper.cpp
@@ -1241,10 +1241,11 @@ void FGenerationHelper::ResetSaveState(FCookGenerationInfo& Info, UPackage* Pack
 			UE_LOG(LogCook, Error,
 				TEXT("CookPackageSplitter failure: We are demoting a %s package from save and removing our references that keep its objects loaded.\n")
 				TEXT("This will allow the objects to be garbage collected and cause failures in the splitter which expects them to remain loaded.\n")
-				TEXT("Package=%s, Splitter=%s, ReleaseSaveReason=%s"),
+				TEXT("Package=%s, Splitter=%s, ReleaseSaveReason=%s, NewState=%s"),
 				Info.IsGenerator() ? TEXT("generator") : TEXT("generated"),
 				*Info.GetPackageName(),
-				*GetSplitDataObjectName().ToString(), LexToString(ReleaseSaveReason));
+				*GetSplitDataObjectName().ToString(), LexToString(ReleaseSaveReason), LexToString(NewState));
+			FDebug::DumpStackTraceToLog(ELogVerbosity::Display);
 		}
 		Info.CachedObjectsInOuterInfo.Empty();
 		Info.SetHasTakenOverCachedCookedPlatformData(false);
diff --git a/Engine/Source/Editor/UnrealEd/Private/Cooker/CookTypes.cpp b/Engine/Source/Editor/UnrealEd/Private/Cooker/CookTypes.cpp
--- a/Engine/Source/Editor/UnrealEd/Private/Cooker/CookTypes.cpp
+++ b/Engine/Source/Editor/UnrealEd/Private/Cooker/CookTypes.cpp
@@ -74,6 +74,21 @@ const TCHAR* LexToString(ESuppressCookReason Reason)
 	}
 }
 
+const TCHAR* LexToString(UE::Cook::EPackageState Reason)
+{
+	static_assert(static_cast<int>(EPackageState::Count) == 6);
+	switch (Reason)
+	{
+	case EPackageState::Idle: return TEXT("Idle");
+	case EPackageState::Request: return TEXT("Request");
+	case EPackageState::AssignedToWorker: return TEXT("AssignedToWorker");
+	case EPackageState::LoadPrepare: return TEXT("LoadPrepare");
+	case EPackageState::LoadReady: return TEXT("LoadReady");
+	case EPackageState::Save: return TEXT("Save");
+	default: return TEXT("Invalid");
+	}
+}
+
 EStateChangeReason ConvertToStateChangeReason(ESuppressCookReason Reason)
 {
 	switch (Reason)
diff --git a/Engine/Source/Editor/UnrealEd/Private/Cooker/CookTypes.h b/Engine/Source/Editor/UnrealEd/Private/Cooker/CookTypes.h
--- a/Engine/Source/Editor/UnrealEd/Private/Cooker/CookTypes.h
+++ b/Engine/Source/Editor/UnrealEd/Private/Cooker/CookTypes.h
@@ -205,6 +205,7 @@ namespace UE::Cook
 		/** Number of bits required to store a valid EPackageState */
 		BitCount = 3,
 	};
+	const TCHAR* LexToString(UE::Cook::EPackageState State);
 
 	enum class EPackageStateProperty // Bitfield
 	{
