Cooker: GeneratedPackages: Suppress the message "CookPackageSplitter failure: We are demoting a %s package from save" when it is being demoted due to retraction, for the same reason we suppress it when going to idle: it is not being cooked so it doesn't matter if its objects are garbage collected. Suppression due to retraction, when executed on the director, sends the package into the AssignedToWorker state rather than the Idle state.
#rnx

diff --git a/Engine/Source/Editor/UnrealEd/Private/Cooker/CookGenerationHelper.cpp b/Engine/Source/Editor/UnrealEd/Private/Cooker/CookGenerationHelper.cpp
--- a/Engine/Source/Editor/UnrealEd/Private/Cooker/CookGenerationHelper.cpp
+++ b/Engine/Source/Editor/UnrealEd/Private/Cooker/CookGenerationHelper.cpp
@@ -1236,7 +1236,8 @@ void FGenerationHelper::ResetSaveState(FCookGenerationInfo& Info, UPackage* Pack
 			Info.PackageData->GetCachedObjectsInOuter().Num() != 0 && IsUseInternalReferenceToAvoidGarbageCollect() &&
 			(ReleaseSaveReason != EStateChangeReason::Completed && ReleaseSaveReason != EStateChangeReason::DoneForNow
 				&& ReleaseSaveReason != EStateChangeReason::SaveError
-				&& ReleaseSaveReason != EStateChangeReason::CookerShutdown))
+				&& ReleaseSaveReason != EStateChangeReason::CookerShutdown
+				&& ReleaseSaveReason != EStateChangeReason::Retraction))
 		{
 			UE_LOG(LogCook, Error,
 				TEXT("CookPackageSplitter failure: We are demoting a %s package from save and removing our references that keep its objects loaded.\n")
