Extract the logic for building the name of the View Modes viewport toolbar menu

diff --git a/Engine/Source/Editor/UnrealEd/Private/SEditorViewportViewMenu.cpp b/Engine/Source/Editor/UnrealEd/Private/SEditorViewportViewMenu.cpp
--- a/Engine/Source/Editor/UnrealEd/Private/SEditorViewportViewMenu.cpp
+++ b/Engine/Source/Editor/UnrealEd/Private/SEditorViewportViewMenu.cpp
@@ -12,6 +12,7 @@
 #include "GPUSkinCacheVisualizationMenuCommands.h"
 #include "GPUSkinCache.h"
 #include "RenderResource.h"
+#include "ViewportToolbar/UnrealEdViewportToolbar.h"
 
 #define LOCTEXT_NAMESPACE "EditorViewportViewMenu"
 
@@ -36,50 +37,7 @@ void SEditorViewportViewMenu::Construct( const FArguments& InArgs, TSharedRef<SE
 
 FText SEditorViewportViewMenu::GetViewMenuLabel() const
 {
-	FText Label = LOCTEXT("ViewMenuTitle_Default", "View");
-	TSharedPtr< SEditorViewport > PinnedViewport = Viewport.Pin();
-	if( PinnedViewport.IsValid() )
-	{
-		const TSharedPtr<FEditorViewportClient> ViewportClient = PinnedViewport->GetViewportClient();
-		check(ViewportClient.IsValid());
-		const EViewModeIndex ViewMode = ViewportClient->GetViewMode();
-		// If VMI_VisualizeBuffer, return its subcategory name
-		if (ViewMode == VMI_VisualizeBuffer)
-		{
-			Label = ViewportClient->GetCurrentBufferVisualizationModeDisplayName();
-		}
-		else if (ViewMode == VMI_VisualizeNanite)
-		{
-			Label = ViewportClient->GetCurrentNaniteVisualizationModeDisplayName();
-		}
-		else if (ViewMode == VMI_VisualizeLumen)
-		{
-			Label = ViewportClient->GetCurrentLumenVisualizationModeDisplayName();
-		}
-		else if (ViewMode == VMI_VisualizeSubstrate)
-		{
-			Label = ViewportClient->GetCurrentSubstrateVisualizationModeDisplayName();
-		}
-		else if (ViewMode == VMI_VisualizeGroom)
-		{
-			Label = ViewportClient->GetCurrentGroomVisualizationModeDisplayName();
-		}
-		else if (ViewMode == VMI_VisualizeVirtualShadowMap)
-		{
-			Label = ViewportClient->GetCurrentVirtualShadowMapVisualizationModeDisplayName();
-		}
-		else if (ViewMode == VMI_VisualizeGPUSkinCache)
-		{
-			Label = ViewportClient->GetCurrentGPUSkinCacheVisualizationModeDisplayName();
-		}
-		// For any other category, return its own name
-		else
-		{
-			Label = UViewModeUtils::GetViewModeDisplayName(ViewMode);
-		}
-	}
-
-	return Label;
+	return UE::UnrealEd::GetViewModesSubmenuLabel(Viewport);
 }
 
 const FSlateBrush* SEditorViewportViewMenu::GetViewMenuLabelIcon() const
diff --git a/Engine/Source/Editor/UnrealEd/Private/ViewportToolbar/UnrealEdViewportToolbar.cpp b/Engine/Source/Editor/UnrealEd/Private/ViewportToolbar/UnrealEdViewportToolbar.cpp
new file mode 100644
--- /dev/null
+++ b/Engine/Source/Editor/UnrealEd/Private/ViewportToolbar/UnrealEdViewportToolbar.cpp
@@ -0,0 +1,63 @@
+// Copyright Epic Games, Inc. All Rights Reserved.
+
+#include "ViewportToolbar/UnrealEdViewportToolbar.h"
+
+#include "EditorViewportClient.h"
+#include "SEditorViewport.h"
+#include "Templates/SharedPointer.h"
+
+#define LOCTEXT_NAMESPACE "UnrealEdViewportToolbar"
+
+namespace UE::UnrealEd
+{
+
+FText GetViewModesSubmenuLabel(TWeakPtr<SEditorViewport> InViewport)
+{
+	FText Label = LOCTEXT("ViewMenuTitle_Default", "View");
+	if (TSharedPtr<SEditorViewport> PinnedViewport = InViewport.Pin())
+	{
+		const TSharedPtr<FEditorViewportClient> ViewportClient = PinnedViewport->GetViewportClient();
+		check(ViewportClient.IsValid());
+		const EViewModeIndex ViewMode = ViewportClient->GetViewMode();
+		// If VMI_VisualizeBuffer, return its subcategory name
+		if (ViewMode == VMI_VisualizeBuffer)
+		{
+			Label = ViewportClient->GetCurrentBufferVisualizationModeDisplayName();
+		}
+		else if (ViewMode == VMI_VisualizeNanite)
+		{
+			Label = ViewportClient->GetCurrentNaniteVisualizationModeDisplayName();
+		}
+		else if (ViewMode == VMI_VisualizeLumen)
+		{
+			Label = ViewportClient->GetCurrentLumenVisualizationModeDisplayName();
+		}
+		else if (ViewMode == VMI_VisualizeSubstrate)
+		{
+			Label = ViewportClient->GetCurrentSubstrateVisualizationModeDisplayName();
+		}
+		else if (ViewMode == VMI_VisualizeGroom)
+		{
+			Label = ViewportClient->GetCurrentGroomVisualizationModeDisplayName();
+		}
+		else if (ViewMode == VMI_VisualizeVirtualShadowMap)
+		{
+			Label = ViewportClient->GetCurrentVirtualShadowMapVisualizationModeDisplayName();
+		}
+		else if (ViewMode == VMI_VisualizeGPUSkinCache)
+		{
+			Label = ViewportClient->GetCurrentGPUSkinCacheVisualizationModeDisplayName();
+		}
+		// For any other category, return its own name
+		else
+		{
+			Label = UViewModeUtils::GetViewModeDisplayName(ViewMode);
+		}
+	}
+
+	return Label;
+}
+
+} // namespace UE::UnrealEd
+
+#undef LOCTEXT_NAMESPACE
diff --git a/Engine/Source/Editor/UnrealEd/Public/ViewportToolbar/UnrealEdViewportToolbar.h b/Engine/Source/Editor/UnrealEd/Public/ViewportToolbar/UnrealEdViewportToolbar.h
new file mode 100644
--- /dev/null
+++ b/Engine/Source/Editor/UnrealEd/Public/ViewportToolbar/UnrealEdViewportToolbar.h
@@ -0,0 +1,15 @@
+// Copyright Epic Games, Inc. All Rights Reserved.
+
+#pragma once
+
+#include "Templates/SharedPointerFwd.h"
+
+class FText;
+class SEditorViewport;
+
+namespace UE::UnrealEd
+{
+
+UNREALED_API FText GetViewModesSubmenuLabel(TWeakPtr<SEditorViewport> InViewport);
+
+} // namespace UE::UnrealEd
