PR #11948: Fix and simplify WorldPartitionNavigationDataBuilder clean packages process and removing dead code.
Issues could occur with wrongly generated packages, now the chunks will be deleted regardless of their spatialization.

diff --git a/Engine/Source/Editor/UnrealEd/Private/WorldPartition/WorldPartitionNavigationDataBuilder.cpp b/Engine/Source/Editor/UnrealEd/Private/WorldPartition/WorldPartitionNavigationDataBuilder.cpp
--- a/Engine/Source/Editor/UnrealEd/Private/WorldPartition/WorldPartitionNavigationDataBuilder.cpp
+++ b/Engine/Source/Editor/UnrealEd/Private/WorldPartition/WorldPartitionNavigationDataBuilder.cpp
@@ -83,7 +83,35 @@ bool UWorldPartitionNavigationDataBuilder::RunInternal(UWorld* World, const FCel
 	check(WorldPartition);
 	
 	TSet<UPackage*> NavigationDataChunkActorPackages;
-	TSet<UPackage*> PackagesToClean;
+
+	// Check if we are just in cleaning mode.
+	if (bCleanBuilderPackages)
+	{
+		TSet<UPackage*> PackagesToClean;
+		for (TActorIterator<ANavigationDataChunkActor> ItActor(World); ItActor; ++ItActor)
+		{
+			UPackage* Package = ItActor->GetPackage();
+			check(!PackagesToClean.Find(Package));
+			PackagesToClean.Add(Package);
+		}
+
+		UE_LOG(LogWorldPartitionNavigationDataBuilder, Verbose, TEXT("   Number of packages to clear: %i"), PackagesToClean.Num());
+
+		// Just delete all ANavigationDataChunkActor packages
+		if (!DeletePackages(PackageHelper, PackagesToClean.Array()))
+		{
+			UE_LOG(LogWorldPartitionNavigationDataBuilder, Error, TEXT("Error deleting packages."));
+		}
+
+		// If we had packages to delete we need to notify the delete after the save. Else WP might try to load the deleted descriptors on the next iteration.
+		// Note: this notification is expected to be done by the SavePackages()
+		for (UPackage* Package : PackagesToClean)
+		{
+			WorldPartition->OnPackageDeleted(Package);
+		}
+
+		return true;
+	}
 
 	// Gather all packages before any navigation data chunk actors are deleted
 	for (TActorIterator<ANavigationDataChunkActor> ItActor(World); ItActor; ++ItActor)
@@ -113,12 +141,6 @@ bool UWorldPartitionNavigationDataBuilder::RunInternal(UWorld* World, const FCel
 		
 		if (IsInside2D(GeneratingBounds, Location))
 		{
-			if (bCleanBuilderPackages)
-			{
-				check(!PackagesToClean.Find(Actor->GetPackage()));
-				PackagesToClean.Add(Actor->GetPackage());
-			}
-			
 			UE_LOG(LogWorldPartitionNavigationDataBuilder, Verbose, TEXT("   Destroy actor %s in package %s."), *Actor->GetName(), *Actor->GetPackage()->GetName());
 			const bool bDestroyed = World->DestroyActor(Actor);
 			if (!bDestroyed)
@@ -129,27 +151,6 @@ bool UWorldPartitionNavigationDataBuilder::RunInternal(UWorld* World, const FCel
 	}
 	UE_LOG(LogWorldPartitionNavigationDataBuilder, Verbose, TEXT("   Number of ANavigationDataChunkActor: %i"), Count);
 
-	// Check if we are just in cleaning mode.
-	if (bCleanBuilderPackages)
-	{
-		UE_LOG(LogWorldPartitionNavigationDataBuilder, Verbose, TEXT("   Number of packages to clear: %i"), PackagesToClean.Num());
-		
-		// Just delete all ANavigationDataChunkActor packages
-		if (!DeletePackages(PackageHelper, PackagesToClean.Array()))
-		{
-			UE_LOG(LogWorldPartitionNavigationDataBuilder, Error, TEXT("Error deleting packages."));
-		}
-
-		// If we had packages to delete we need to notify the delete after the save. Else WP might try to load the deleted descriptors on the next iteration.
-		// Note: this notification is expected to be done by the SavePackages()
-		for (UPackage* Package : PackagesToClean)
-		{
-			WorldPartition->OnPackageDeleted(Package);
-		}		
-
-		return true;
-	}
-
 	// Make sure static meshes have compiled before generating navigation data
 	FStaticMeshCompilingManager::Get().FinishAllCompilation();
 	
@@ -393,9 +394,6 @@ bool UWorldPartitionNavigationDataBuilder::GenerateNavigationData(UWorldPartitio
 	// For each cell, gather navmesh and generate a datachunk actor
 	int32 ActorCount = 0;
 
-	// Keep track of all valid navigation data chunk actors
-	TSet<ANavigationDataChunkActor*> ValidNavigationDataChunkActors;
-
 	// A DataChunkActor will be generated for each tile touching the generating bounds.
 	const TSubclassOf<APartitionActor>& NavigationDataActorClass = ANavigationDataChunkActor::StaticClass();
 
@@ -408,7 +406,7 @@ bool UWorldPartitionNavigationDataBuilder::GenerateNavigationData(UWorldPartitio
 
 	const FIntRect GeneratingBounds2D(XMin, YMin, XMax, YMax);
 	FActorPartitionGridHelper::ForEachIntersectingCell(NavigationDataActorClass, GeneratingBounds2D, World->PersistentLevel,
-		[&WorldPartition, &ActorCount, World, &ValidNavigationDataChunkActors, &NavDataBounds, GridSize, this](const UActorPartitionSubsystem::FCellCoord& InCellCoord, const FIntRect& InCellBounds)->bool
+		[&WorldPartition, &ActorCount, World, &NavDataBounds, GridSize, this](const UActorPartitionSubsystem::FCellCoord& InCellCoord, const FIntRect& InCellBounds)->bool
 		{
 			TRACE_CPUPROFILER_EVENT_SCOPE(MakeNavigationDataChunkActorForGridCell);
 		
@@ -453,8 +451,6 @@ bool UWorldPartitionNavigationDataBuilder::GenerateNavigationData(UWorldPartitio
 			const FName CellName = GetCellName(WorldPartition, InCellCoord);
 			DataChunkActor->SetActorLabel(FString::Printf(TEXT("NavDataChunkActor_%s"), *CellName.ToString()));
 
-			ValidNavigationDataChunkActors.Add(DataChunkActor);
-
 			UE_LOG(LogWorldPartitionNavigationDataBuilder, Verbose, TEXT("%i) %s added."), ActorCount, *DataChunkActor->GetName());
 
 			return true;
