Flush Async loading before running the resave commandlet to ensure packages we want to save that are being loaded asynchronously are loaded before we begin
#rnx

diff --git a/Engine/Source/Editor/UnrealEd/Private/Commandlets/ContentCommandlets.cpp b/Engine/Source/Editor/UnrealEd/Private/Commandlets/ContentCommandlets.cpp
--- a/Engine/Source/Editor/UnrealEd/Private/Commandlets/ContentCommandlets.cpp
+++ b/Engine/Source/Editor/UnrealEd/Private/Commandlets/ContentCommandlets.cpp
@@ -1340,6 +1340,10 @@ int32 UResavePackagesCommandlet::Main( const FString& Params )
 	// Make sure any remaining temp files from previous runs are removed
 	CleanTempFiles();
 
+	//Flush Async loading before running the commandlet to ensure packages we want to save that are being loaded asynchronously are loaded before we begin
+	UE_LOG(LogContentCommandlet, Display, TEXT("Flushing Async Loading"));
+	FlushAsyncLoading();
+
 	// Iterate over all packages.
 	for( int32 PackageIndex = 0; PackageIndex < PackageNames.Num(); PackageIndex++ )
 	{
