Bumping down an assertion in Lightmass.cpp to an ensure when adding a material with different export settings

diff --git a/Engine/Source/Editor/UnrealEd/Private/Lightmass/Lightmass.cpp b/Engine/Source/Editor/UnrealEd/Private/Lightmass/Lightmass.cpp
--- a/Engine/Source/Editor/UnrealEd/Private/Lightmass/Lightmass.cpp
+++ b/Engine/Source/Editor/UnrealEd/Private/Lightmass/Lightmass.cpp
@@ -545,7 +545,7 @@ void FLightmassExporter::AddMaterial(UMaterialInterface* InMaterialInterface, co
 
 		if (auto* ExistingExportSettings = MaterialExportSettings.Find(InMaterialInterface))
 		{
-			checkf(ExportSettings == *ExistingExportSettings, TEXT("Attempting to add the same material twice with different export settings, this is not (currently) supported"));
+			ensureMsgf(ExportSettings == *ExistingExportSettings, TEXT("Attempting to add the same material twice with different export settings, this is not (currently) supported"));
 			return;
 		}
 
