Extract UnrealEd code for filling the viewport toolbar View Modes menu to enable others to use it

diff --git a/Engine/Source/Editor/UnrealEd/Private/SEditorViewportViewMenu.cpp b/Engine/Source/Editor/UnrealEd/Private/SEditorViewportViewMenu.cpp
--- a/Engine/Source/Editor/UnrealEd/Private/SEditorViewportViewMenu.cpp
+++ b/Engine/Source/Editor/UnrealEd/Private/SEditorViewportViewMenu.cpp
@@ -2,16 +2,12 @@
 
 
 #include "SEditorViewportViewMenu.h"
-#include "DebugViewModeHelpers.h"
 #include "SEditorViewportViewMenuContext.h"
 #include "Framework/MultiBox/MultiBoxBuilder.h"
 #include "ToolMenus.h"
 #include "Styling/AppStyle.h"
-#include "EditorViewportCommands.h"
-#include "RayTracingDebugVisualizationMenuCommands.h"
-#include "GPUSkinCacheVisualizationMenuCommands.h"
-#include "GPUSkinCache.h"
 #include "RenderResource.h"
+#include "SEditorViewport.h"
 #include "ViewportToolbar/UnrealEdViewportToolbar.h"
 
 #define LOCTEXT_NAMESPACE "EditorViewportViewMenu"
@@ -84,217 +80,16 @@ TSharedRef<SWidget> SEditorViewportViewMenu::GenerateViewMenuContent() const
 
 void SEditorViewportViewMenu::FillViewMenu(UToolMenu* Menu) const
 {
-	const FEditorViewportCommands& BaseViewportActions = FEditorViewportCommands::Get();
-
-	{
-		// View modes
-		{
-			FToolMenuSection& Section = Menu->AddSection("ViewMode", LOCTEXT("ViewModeHeader", "View Mode"));
-			{
-				Section.AddMenuEntry(BaseViewportActions.LitMode, UViewModeUtils::GetViewModeDisplayName(VMI_Lit));
-				Section.AddMenuEntry(BaseViewportActions.UnlitMode, UViewModeUtils::GetViewModeDisplayName(VMI_Unlit));
-				Section.AddMenuEntry(BaseViewportActions.WireframeMode, UViewModeUtils::GetViewModeDisplayName(VMI_BrushWireframe));
-				Section.AddMenuEntry(BaseViewportActions.LitWireframeMode, UViewModeUtils::GetViewModeDisplayName(VMI_Lit_Wireframe));
-				Section.AddMenuEntry(BaseViewportActions.DetailLightingMode, UViewModeUtils::GetViewModeDisplayName(VMI_Lit_DetailLighting));
-				Section.AddMenuEntry(BaseViewportActions.LightingOnlyMode, UViewModeUtils::GetViewModeDisplayName(VMI_LightingOnly));
-				Section.AddMenuEntry(BaseViewportActions.ReflectionOverrideMode, UViewModeUtils::GetViewModeDisplayName(VMI_ReflectionOverride));
-				Section.AddMenuEntry(BaseViewportActions.CollisionPawn, UViewModeUtils::GetViewModeDisplayName(VMI_CollisionPawn));
-				Section.AddMenuEntry(BaseViewportActions.CollisionVisibility, UViewModeUtils::GetViewModeDisplayName(VMI_CollisionVisibility));
-			}
-
-			if (IsRayTracingEnabled())
-			{
-				static auto PathTracingCvar = IConsoleManager::Get().FindTConsoleVariableDataInt(TEXT("r.PathTracing"));
-				const bool bPathTracingSupported = FDataDrivenShaderPlatformInfo::GetSupportsPathTracing(GMaxRHIShaderPlatform);
-				const bool bPathTracingEnabled = PathTracingCvar && PathTracingCvar->GetValueOnAnyThread() != 0;
-				if (bPathTracingSupported && bPathTracingEnabled)
-				{
-					Section.AddMenuEntry(BaseViewportActions.PathTracingMode, UViewModeUtils::GetViewModeDisplayName(VMI_PathTracing));
-				}
-			}
-
-			// Optimization
-			{
-				struct Local
-				{
-					static void BuildOptimizationMenu( UToolMenu* Menu, TWeakPtr< SViewportToolBar > InParentToolBar )
-					{
-						const FEditorViewportCommands& BaseViewportCommands = FEditorViewportCommands::Get();
-
-						UWorld* World = GWorld;
-						const ERHIFeatureLevel::Type FeatureLevel = (IsInGameThread() && World) ? (ERHIFeatureLevel::Type)World->GetFeatureLevel() : GMaxRHIFeatureLevel;
-
-						{
-							FToolMenuSection& Section = Menu->AddSection("OptimizationViewmodes", LOCTEXT("OptimizationSubMenuHeader", "Optimization Viewmodes"));
-							if (FeatureLevel >= ERHIFeatureLevel::SM5)
-							{
-								Section.AddMenuEntry(BaseViewportCommands.LightComplexityMode, UViewModeUtils::GetViewModeDisplayName(VMI_LightComplexity));
-								if (IsStaticLightingAllowed())
-								{
-									Section.AddMenuEntry(BaseViewportCommands.LightmapDensityMode, UViewModeUtils::GetViewModeDisplayName(VMI_LightmapDensity));
-								}
-								Section.AddMenuEntry(BaseViewportCommands.StationaryLightOverlapMode, UViewModeUtils::GetViewModeDisplayName(VMI_StationaryLightOverlap));
-							}
-
-							Section.AddMenuEntry(BaseViewportCommands.ShaderComplexityMode, UViewModeUtils::GetViewModeDisplayName(VMI_ShaderComplexity));
-
-							if (AllowDebugViewShaderMode(DVSM_ShaderComplexityContainedQuadOverhead, GMaxRHIShaderPlatform, FeatureLevel))
-							{
-								Section.AddMenuEntry(BaseViewportCommands.ShaderComplexityWithQuadOverdrawMode, UViewModeUtils::GetViewModeDisplayName(VMI_ShaderComplexityWithQuadOverdraw));
-							}
-							if (AllowDebugViewShaderMode(DVSM_QuadComplexity, GMaxRHIShaderPlatform, FeatureLevel))
-							{
-								Section.AddMenuEntry(BaseViewportCommands.QuadOverdrawMode, UViewModeUtils::GetViewModeDisplayName(VMI_QuadOverdraw));
-							}
-							if (AllowDebugViewShaderMode(DVSM_LWCComplexity, GMaxRHIShaderPlatform, FeatureLevel))
-							{
-								Section.AddMenuEntry(BaseViewportCommands.VisualizeLWCComplexity, 
-								                     UViewModeUtils::GetViewModeDisplayName(VMI_LWCComplexity), 
-								                     TAttribute<FText>(), 
-								                     FSlateIcon(FAppStyle::GetAppStyleSetName(), "EditorViewport.LWCComplexityMode"));
-							}
-						}
-
-						{
-							FToolMenuSection& Section = Menu->AddSection("TextureStreaming", LOCTEXT("TextureStreamingHeader", "Texture Streaming Accuracy"));
-							if (AllowDebugViewShaderMode(DVSM_PrimitiveDistanceAccuracy, GMaxRHIShaderPlatform, FeatureLevel) && (!InParentToolBar.IsValid() || InParentToolBar.Pin()->IsViewModeSupported(VMI_PrimitiveDistanceAccuracy)))
-							{
-								Section.AddMenuEntry(BaseViewportCommands.TexStreamAccPrimitiveDistanceMode, UViewModeUtils::GetViewModeDisplayName(VMI_PrimitiveDistanceAccuracy));
-							}
-							if (AllowDebugViewShaderMode(DVSM_MeshUVDensityAccuracy, GMaxRHIShaderPlatform, FeatureLevel) && (!InParentToolBar.IsValid() || InParentToolBar.Pin()->IsViewModeSupported(VMI_MeshUVDensityAccuracy)))
-							{
-								Section.AddMenuEntry(BaseViewportCommands.TexStreamAccMeshUVDensityMode, UViewModeUtils::GetViewModeDisplayName(VMI_MeshUVDensityAccuracy));
-							}
-							// TexCoordScale accuracy viewmode requires shaders that are only built in the TextureStreamingBuild, which requires the new metrics to be enabled.
-							if (AllowDebugViewShaderMode(DVSM_MaterialTextureScaleAccuracy, GMaxRHIShaderPlatform, FeatureLevel) && CVarStreamingUseNewMetrics.GetValueOnAnyThread() != 0 && (!InParentToolBar.IsValid() || InParentToolBar.Pin()->IsViewModeSupported(VMI_MaterialTextureScaleAccuracy)))
-							{
-								Section.AddMenuEntry(BaseViewportCommands.TexStreamAccMaterialTextureScaleMode, UViewModeUtils::GetViewModeDisplayName(VMI_MaterialTextureScaleAccuracy));
-							}
-							if (AllowDebugViewShaderMode(DVSM_RequiredTextureResolution, GMaxRHIShaderPlatform, FeatureLevel) && (!InParentToolBar.IsValid() || InParentToolBar.Pin()->IsViewModeSupported(VMI_RequiredTextureResolution)))
-							{
-								Section.AddMenuEntry(BaseViewportCommands.RequiredTextureResolutionMode, UViewModeUtils::GetViewModeDisplayName(VMI_RequiredTextureResolution));
-							}
-							if (AllowDebugViewShaderMode(DVSM_RequiredTextureResolution, GMaxRHIShaderPlatform, FeatureLevel) && (!InParentToolBar.IsValid() || InParentToolBar.Pin()->IsViewModeSupported(VMI_VirtualTexturePendingMips)))
-							{
-								Section.AddMenuEntry(BaseViewportCommands.VirtualTexturePendingMipsMode, UViewModeUtils::GetViewModeDisplayName(VMI_VirtualTexturePendingMips));
-							}
-						}
-					}
-				};
-
-				Section.AddSubMenu(
-					"OptimizationSubMenu",
-					LOCTEXT("OptimizationSubMenu", "Optimization Viewmodes"), LOCTEXT("Optimization_ToolTip", "Select optimization visualizer"),
-					FNewToolMenuDelegate::CreateStatic(&Local::BuildOptimizationMenu, ParentToolBar),
-					FUIAction(FExecuteAction(), FCanExecuteAction(),
-						FIsActionChecked::CreateLambda([this]()
-						{
-								const TSharedRef<SEditorViewport> ViewportRef = Viewport.Pin().ToSharedRef();
-								const TSharedPtr<FEditorViewportClient> ViewportClient = ViewportRef->GetViewportClient();
-								check(ViewportClient.IsValid());
-								const EViewModeIndex ViewMode = ViewportClient->GetViewMode();
-								return (
-									// Texture Streaming Accuracy
-									ViewMode == VMI_LightComplexity || ViewMode == VMI_LightmapDensity || ViewMode == VMI_StationaryLightOverlap
-									|| ViewMode == VMI_ShaderComplexity || ViewMode == VMI_ShaderComplexityWithQuadOverdraw || ViewMode == VMI_QuadOverdraw
-									// Texture Streaming Accuracy
-									|| ViewMode == VMI_PrimitiveDistanceAccuracy || ViewMode == VMI_MeshUVDensityAccuracy || ViewMode == VMI_MaterialTextureScaleAccuracy 
-									|| ViewMode == VMI_RequiredTextureResolution || ViewMode == VMI_VirtualTexturePendingMips
-								);
-						})),
-					EUserInterfaceActionType::RadioButton,
-					/* bInOpenSubMenuOnClick = */ false, FSlateIcon(FAppStyle::GetAppStyleSetName(), "EditorViewport.QuadOverdrawMode"));
-			}
-
-			if (IsRayTracingEnabled())
-			{
-				struct Local
-				{
-					static void BuildRayTracingDebugMenu(FMenuBuilder& Menu, TWeakPtr< SViewportToolBar > InParentToolBar)
-					{
-						const FRayTracingDebugVisualizationMenuCommands& RtDebugCommands = FRayTracingDebugVisualizationMenuCommands::Get();
-						RtDebugCommands.BuildVisualisationSubMenu(Menu);
-					}
-				};
-
-				Section.AddSubMenu("RayTracingDebugSubMenu", LOCTEXT("RayTracingDebugSubMenu", "Ray Tracing Debug"), LOCTEXT("RayTracing_ToolTip", "Select ray tracing buffer visualization view modes"), FNewMenuDelegate::CreateStatic(&Local::BuildRayTracingDebugMenu, ParentToolBar), false, FSlateIcon(FAppStyle::GetAppStyleSetName(), "EditorViewport.RayTracingDebugMode"));
-			}
-
-			{
-				struct Local
-				{
-					static void BuildLODMenu(UToolMenu* Menu)
-					{
-						{
-							FToolMenuSection& Section = Menu->AddSection("LevelViewportLODColoration", LOCTEXT("LODModesHeader", "Level of Detail Coloration"));
-							Section.AddMenuEntry(FEditorViewportCommands::Get().LODColorationMode, UViewModeUtils::GetViewModeDisplayName(VMI_LODColoration));
-							Section.AddMenuEntry(FEditorViewportCommands::Get().HLODColorationMode, UViewModeUtils::GetViewModeDisplayName(VMI_HLODColoration));
-						}
-					}
-				};
-
-				Section.AddSubMenu(
-					"VisualizeGroupedLOD",
-					LOCTEXT("VisualizeGroupedLODDisplayName", "Level of Detail Coloration"), LOCTEXT("GroupedLODMenu_ToolTip", "Select a mode for LOD Coloration"),
-					FNewToolMenuDelegate::CreateStatic(&Local::BuildLODMenu),
-					FUIAction(FExecuteAction(), FCanExecuteAction(),
-						FIsActionChecked::CreateLambda([this]()
-						{
-								const TSharedRef<SEditorViewport> ViewportRef = Viewport.Pin().ToSharedRef();
-								const TSharedPtr<FEditorViewportClient> ViewportClient = ViewportRef->GetViewportClient();
-								check(ViewportClient.IsValid());
-								const EViewModeIndex ViewMode = ViewportClient->GetViewMode();
-								return (ViewMode == VMI_LODColoration || ViewMode == VMI_HLODColoration);
-						})),
-					EUserInterfaceActionType::RadioButton,
-					/* bInOpenSubMenuOnClick = */ false, FSlateIcon(FAppStyle::GetAppStyleSetName(), "EditorViewport.GroupLODColorationMode"));
-			}
-
-			if (GEnableGPUSkinCache)
+	UE::UnrealEd::IsViewModeSupportedDelegate IsViewModeSupported = UE::UnrealEd::IsViewModeSupportedDelegate::CreateLambda(
+		[WeakToolBar = ParentToolBar](EViewModeIndex ViewModeIndex) -> bool {
+			if (TSharedPtr<SViewportToolBar> ToolBar = WeakToolBar.Pin())
 			{
-				Section.AddSubMenu(
-					"VisualizeGPUSkinCacheViewMode",
-					LOCTEXT("VisualizeGPUSkinCacheViewModeDisplayName", "GPU Skin Cache"),
-					LOCTEXT("GPUSkinCacheVisualizationMenu_ToolTip", "Select a mode for GPU Skin Cache visualization."),
-					FNewMenuDelegate::CreateStatic(&FGPUSkinCacheVisualizationMenuCommands::BuildVisualisationSubMenu),
-					FUIAction(
-						FExecuteAction(),
-						FCanExecuteAction(),
-						FIsActionChecked::CreateLambda([this]()
-						{
-							const TSharedRef<SEditorViewport> ViewportRef = Viewport.Pin().ToSharedRef();
-							const TSharedPtr<FEditorViewportClient> ViewportClient = ViewportRef->GetViewportClient();
-							check(ViewportClient.IsValid());
-							return ViewportClient->IsViewModeEnabled(VMI_VisualizeGPUSkinCache);
-						})
-					),
-					EUserInterfaceActionType::RadioButton,
-					/* bInOpenSubMenuOnClick = */ false,
-					FSlateIcon(FAppStyle::GetAppStyleSetName(), "EditorViewport.VisualizeGPUSkinCacheMode")
-				);
+				return ToolBar->IsViewModeSupported(ViewModeIndex);
 			}
-		}
-
-		// Auto Exposure
-		{
-			const FEditorViewportCommands& BaseViewportCommands = FEditorViewportCommands::Get();
+			return true;
+		});
 
-			TSharedRef<SWidget> FixedEV100Menu = Viewport.Pin()->BuildFixedEV100Menu();
-			TSharedPtr<FEditorViewportClient> EditorViewPostClient = Viewport.Pin()->GetViewportClient();
-			const bool bIsLevelEditor = EditorViewPostClient.IsValid() && EditorViewPostClient->IsLevelEditorClient();
-
-			FToolMenuSection& Section = Menu->AddSection("Exposure", LOCTEXT("ExposureHeader", "Exposure"));
-			Section.AddMenuEntry(bIsLevelEditor ? BaseViewportCommands.ToggleInGameExposure : BaseViewportCommands.ToggleAutoExposure);
-			Section.AddEntry(FToolMenuEntry::InitWidget("FixedEV100", FixedEV100Menu, LOCTEXT("FixedEV100", "EV100")));
-		}
-
-		// Wireframe Opacity
-		{
-			TSharedRef<SWidget> WireOpacityMenu = Viewport.Pin()->BuildWireframeMenu();
-			FToolMenuSection& Section = Menu->AddSection("Wireframe", LOCTEXT("WireframeHeader", "Wireframe"));
-			Section.AddEntry(FToolMenuEntry::InitWidget("WireframeOpacity", WireOpacityMenu, LOCTEXT("WireframeOpacity", "Opacity")));
-		}
-	}
+	UE::UnrealEd::FillViewMenu(Menu, Viewport.Pin().ToSharedRef(), IsViewModeSupported);
 }
 
 #undef LOCTEXT_NAMESPACE
diff --git a/Engine/Source/Editor/UnrealEd/Private/ViewportToolbar/UnrealEdViewportToolbar.cpp b/Engine/Source/Editor/UnrealEd/Private/ViewportToolbar/UnrealEdViewportToolbar.cpp
--- a/Engine/Source/Editor/UnrealEd/Private/ViewportToolbar/UnrealEdViewportToolbar.cpp
+++ b/Engine/Source/Editor/UnrealEd/Private/ViewportToolbar/UnrealEdViewportToolbar.cpp
@@ -2,9 +2,17 @@
 
 #include "ViewportToolbar/UnrealEdViewportToolbar.h"
 
+#include "DataDrivenShaderPlatformInfo.h"
+#include "DebugViewModeHelpers.h"
 #include "EditorViewportClient.h"
+#include "EditorViewportCommands.h"
+#include "GPUSkinCache.h"
+#include "GPUSkinCacheVisualizationMenuCommands.h"
+#include "RayTracingDebugVisualizationMenuCommands.h"
 #include "SEditorViewport.h"
 #include "Templates/SharedPointer.h"
+#include "ToolMenu.h"
+#include "ToolMenuSection.h"
 
 #define LOCTEXT_NAMESPACE "UnrealEdViewportToolbar"
 
@@ -58,6 +66,259 @@ FText GetViewModesSubmenuLabel(TWeakPtr<SEditorViewport> InViewport)
 	return Label;
 }
 
+void FillViewMenu(
+	UToolMenu* InMenu, TSharedRef<SEditorViewport> InViewport, IsViewModeSupportedDelegate InIsViewModeSupported)
+{
+	const FEditorViewportCommands& BaseViewportActions = FEditorViewportCommands::Get();
+
+	// View modes
+	{
+		FToolMenuSection& Section = InMenu->AddSection("ViewMode", LOCTEXT("ViewModeHeader", "View Mode"));
+		{
+			Section.AddMenuEntry(BaseViewportActions.LitMode, UViewModeUtils::GetViewModeDisplayName(VMI_Lit));
+			Section.AddMenuEntry(BaseViewportActions.UnlitMode, UViewModeUtils::GetViewModeDisplayName(VMI_Unlit));
+			Section.AddMenuEntry(
+				BaseViewportActions.WireframeMode, UViewModeUtils::GetViewModeDisplayName(VMI_BrushWireframe));
+			Section.AddMenuEntry(
+				BaseViewportActions.LitWireframeMode, UViewModeUtils::GetViewModeDisplayName(VMI_Lit_Wireframe));
+			Section.AddMenuEntry(
+				BaseViewportActions.DetailLightingMode, UViewModeUtils::GetViewModeDisplayName(VMI_Lit_DetailLighting));
+			Section.AddMenuEntry(
+				BaseViewportActions.LightingOnlyMode, UViewModeUtils::GetViewModeDisplayName(VMI_LightingOnly));
+			Section.AddMenuEntry(BaseViewportActions.ReflectionOverrideMode,
+				UViewModeUtils::GetViewModeDisplayName(VMI_ReflectionOverride));
+			Section.AddMenuEntry(
+				BaseViewportActions.CollisionPawn, UViewModeUtils::GetViewModeDisplayName(VMI_CollisionPawn));
+			Section.AddMenuEntry(BaseViewportActions.CollisionVisibility,
+				UViewModeUtils::GetViewModeDisplayName(VMI_CollisionVisibility));
+		}
+
+		if (IsRayTracingEnabled())
+		{
+			static auto PathTracingCvar = IConsoleManager::Get().FindTConsoleVariableDataInt(TEXT("r.PathTracing"));
+			const bool bPathTracingSupported = FDataDrivenShaderPlatformInfo::GetSupportsPathTracing(GMaxRHIShaderPlatform);
+			const bool bPathTracingEnabled = PathTracingCvar && PathTracingCvar->GetValueOnAnyThread() != 0;
+			if (bPathTracingSupported && bPathTracingEnabled)
+			{
+				Section.AddMenuEntry(
+					BaseViewportActions.PathTracingMode, UViewModeUtils::GetViewModeDisplayName(VMI_PathTracing));
+			}
+		}
+
+		// Optimization
+		{
+			struct Local
+			{
+				static void BuildOptimizationMenu(UToolMenu* Menu, IsViewModeSupportedDelegate IsViewModeSupported)
+				{
+					const FEditorViewportCommands& BaseViewportCommands = FEditorViewportCommands::Get();
+
+					UWorld* World = GWorld;
+					const ERHIFeatureLevel::Type FeatureLevel = (IsInGameThread() && World)
+																  ? (ERHIFeatureLevel::Type)World->GetFeatureLevel()
+																  : GMaxRHIFeatureLevel;
+
+					{
+						FToolMenuSection& Section = Menu->AddSection(
+							"OptimizationViewmodes", LOCTEXT("OptimizationSubMenuHeader", "Optimization Viewmodes"));
+						if (FeatureLevel >= ERHIFeatureLevel::SM5)
+						{
+							Section.AddMenuEntry(BaseViewportCommands.LightComplexityMode,
+								UViewModeUtils::GetViewModeDisplayName(VMI_LightComplexity));
+							if (IsStaticLightingAllowed())
+							{
+								Section.AddMenuEntry(BaseViewportCommands.LightmapDensityMode,
+									UViewModeUtils::GetViewModeDisplayName(VMI_LightmapDensity));
+							}
+							Section.AddMenuEntry(BaseViewportCommands.StationaryLightOverlapMode,
+								UViewModeUtils::GetViewModeDisplayName(VMI_StationaryLightOverlap));
+						}
+
+						Section.AddMenuEntry(BaseViewportCommands.ShaderComplexityMode,
+							UViewModeUtils::GetViewModeDisplayName(VMI_ShaderComplexity));
+
+						if (AllowDebugViewShaderMode(
+								DVSM_ShaderComplexityContainedQuadOverhead, GMaxRHIShaderPlatform, FeatureLevel))
+						{
+							Section.AddMenuEntry(BaseViewportCommands.ShaderComplexityWithQuadOverdrawMode,
+								UViewModeUtils::GetViewModeDisplayName(VMI_ShaderComplexityWithQuadOverdraw));
+						}
+						if (AllowDebugViewShaderMode(DVSM_QuadComplexity, GMaxRHIShaderPlatform, FeatureLevel))
+						{
+							Section.AddMenuEntry(BaseViewportCommands.QuadOverdrawMode,
+								UViewModeUtils::GetViewModeDisplayName(VMI_QuadOverdraw));
+						}
+						if (AllowDebugViewShaderMode(DVSM_LWCComplexity, GMaxRHIShaderPlatform, FeatureLevel))
+						{
+							Section.AddMenuEntry(BaseViewportCommands.VisualizeLWCComplexity,
+								UViewModeUtils::GetViewModeDisplayName(VMI_LWCComplexity), TAttribute<FText>(),
+								FSlateIcon(FAppStyle::GetAppStyleSetName(), "EditorViewport.LWCComplexityMode"));
+						}
+					}
+
+					{
+						FToolMenuSection& Section = Menu->AddSection(
+							"TextureStreaming", LOCTEXT("TextureStreamingHeader", "Texture Streaming Accuracy"));
+
+						if (AllowDebugViewShaderMode(DVSM_PrimitiveDistanceAccuracy, GMaxRHIShaderPlatform, FeatureLevel)
+							&& (!IsViewModeSupported.IsBound()
+								|| IsViewModeSupported.Execute(VMI_PrimitiveDistanceAccuracy)))
+						{
+							Section.AddMenuEntry(BaseViewportCommands.TexStreamAccPrimitiveDistanceMode,
+								UViewModeUtils::GetViewModeDisplayName(VMI_PrimitiveDistanceAccuracy));
+						}
+						if (AllowDebugViewShaderMode(DVSM_MeshUVDensityAccuracy, GMaxRHIShaderPlatform, FeatureLevel)
+							&& (!IsViewModeSupported.IsBound() || IsViewModeSupported.Execute(VMI_MeshUVDensityAccuracy)))
+						{
+							Section.AddMenuEntry(BaseViewportCommands.TexStreamAccMeshUVDensityMode,
+								UViewModeUtils::GetViewModeDisplayName(VMI_MeshUVDensityAccuracy));
+						}
+						// TexCoordScale accuracy viewmode requires shaders that are only built in the
+						// TextureStreamingBuild, which requires the new metrics to be enabled.
+						if (AllowDebugViewShaderMode(DVSM_MaterialTextureScaleAccuracy, GMaxRHIShaderPlatform, FeatureLevel)
+							&& CVarStreamingUseNewMetrics.GetValueOnAnyThread() != 0
+							&& (!IsViewModeSupported.IsBound()
+								|| IsViewModeSupported.Execute(VMI_MaterialTextureScaleAccuracy)))
+						{
+							Section.AddMenuEntry(BaseViewportCommands.TexStreamAccMaterialTextureScaleMode,
+								UViewModeUtils::GetViewModeDisplayName(VMI_MaterialTextureScaleAccuracy));
+						}
+						if (AllowDebugViewShaderMode(DVSM_RequiredTextureResolution, GMaxRHIShaderPlatform, FeatureLevel)
+							&& (!IsViewModeSupported.IsBound()
+								|| IsViewModeSupported.Execute(VMI_RequiredTextureResolution)))
+						{
+							Section.AddMenuEntry(BaseViewportCommands.RequiredTextureResolutionMode,
+								UViewModeUtils::GetViewModeDisplayName(VMI_RequiredTextureResolution));
+						}
+						if (AllowDebugViewShaderMode(DVSM_RequiredTextureResolution, GMaxRHIShaderPlatform, FeatureLevel)
+							&& (!IsViewModeSupported.IsBound()
+								|| IsViewModeSupported.Execute(VMI_VirtualTexturePendingMips)))
+						{
+							Section.AddMenuEntry(BaseViewportCommands.VirtualTexturePendingMipsMode,
+								UViewModeUtils::GetViewModeDisplayName(VMI_VirtualTexturePendingMips));
+						}
+					}
+				}
+			};
+
+			Section.AddSubMenu("OptimizationSubMenu", LOCTEXT("OptimizationSubMenu", "Optimization Viewmodes"),
+				LOCTEXT("Optimization_ToolTip", "Select optimization visualizer"),
+				FNewToolMenuDelegate::CreateStatic(&Local::BuildOptimizationMenu, InIsViewModeSupported),
+				FUIAction(FExecuteAction(), FCanExecuteAction(),
+					FIsActionChecked::CreateLambda([Viewport = InViewport.ToWeakPtr()]() {
+						const TSharedRef<SEditorViewport> ViewportRef = Viewport.Pin().ToSharedRef();
+						const TSharedPtr<FEditorViewportClient> ViewportClient = ViewportRef->GetViewportClient();
+						check(ViewportClient.IsValid());
+						const EViewModeIndex ViewMode = ViewportClient->GetViewMode();
+						return (
+							// Texture Streaming Accuracy
+							ViewMode == VMI_LightComplexity || ViewMode == VMI_LightmapDensity
+							|| ViewMode == VMI_StationaryLightOverlap || ViewMode == VMI_ShaderComplexity
+							|| ViewMode == VMI_ShaderComplexityWithQuadOverdraw
+							|| ViewMode == VMI_QuadOverdraw
+							// Texture Streaming Accuracy
+							|| ViewMode == VMI_PrimitiveDistanceAccuracy || ViewMode == VMI_MeshUVDensityAccuracy
+							|| ViewMode == VMI_MaterialTextureScaleAccuracy || ViewMode == VMI_RequiredTextureResolution
+							|| ViewMode == VMI_VirtualTexturePendingMips);
+					})),
+				EUserInterfaceActionType::RadioButton,
+				/* bInOpenSubMenuOnClick = */ false,
+				FSlateIcon(FAppStyle::GetAppStyleSetName(), "EditorViewport.QuadOverdrawMode"));
+		}
+
+		if (IsRayTracingEnabled())
+		{
+			struct Local
+			{
+				static void BuildRayTracingDebugMenu(FMenuBuilder& Menu) //, TWeakPtr<SViewportToolBar> InParentToolBar)
+				{
+					const FRayTracingDebugVisualizationMenuCommands& RtDebugCommands =
+						FRayTracingDebugVisualizationMenuCommands::Get();
+					RtDebugCommands.BuildVisualisationSubMenu(Menu);
+				}
+			};
+
+			Section.AddSubMenu("RayTracingDebugSubMenu", LOCTEXT("RayTracingDebugSubMenu", "Ray Tracing Debug"),
+				LOCTEXT("RayTracing_ToolTip", "Select ray tracing buffer visualization view modes"),
+				FNewMenuDelegate::CreateStatic(&Local::BuildRayTracingDebugMenu), //, ParentToolBar)
+				false, FSlateIcon(FAppStyle::GetAppStyleSetName(), "EditorViewport.RayTracingDebugMode"));
+		}
+
+		{
+			struct Local
+			{
+				static void BuildLODMenu(UToolMenu* Menu)
+				{
+					{
+						FToolMenuSection& Section = Menu->AddSection(
+							"LevelViewportLODColoration", LOCTEXT("LODModesHeader", "Level of Detail Coloration"));
+						Section.AddMenuEntry(FEditorViewportCommands::Get().LODColorationMode,
+							UViewModeUtils::GetViewModeDisplayName(VMI_LODColoration));
+						Section.AddMenuEntry(FEditorViewportCommands::Get().HLODColorationMode,
+							UViewModeUtils::GetViewModeDisplayName(VMI_HLODColoration));
+					}
+				}
+			};
+
+			Section.AddSubMenu("VisualizeGroupedLOD",
+				LOCTEXT("VisualizeGroupedLODDisplayName", "Level of Detail Coloration"),
+				LOCTEXT("GroupedLODMenu_ToolTip", "Select a mode for LOD Coloration"),
+				FNewToolMenuDelegate::CreateStatic(&Local::BuildLODMenu),
+				FUIAction(FExecuteAction(), FCanExecuteAction(),
+					FIsActionChecked::CreateLambda([WeakViewport = InViewport.ToWeakPtr()]() {
+						const TSharedRef<SEditorViewport> ViewportRef = WeakViewport.Pin().ToSharedRef();
+						const TSharedPtr<FEditorViewportClient> ViewportClient = ViewportRef->GetViewportClient();
+						check(ViewportClient.IsValid());
+						const EViewModeIndex ViewMode = ViewportClient->GetViewMode();
+						return (ViewMode == VMI_LODColoration || ViewMode == VMI_HLODColoration);
+					})),
+				EUserInterfaceActionType::RadioButton,
+				/* bInOpenSubMenuOnClick = */ false,
+				FSlateIcon(FAppStyle::GetAppStyleSetName(), "EditorViewport.GroupLODColorationMode"));
+		}
+
+		if (GEnableGPUSkinCache)
+		{
+			Section.AddSubMenu("VisualizeGPUSkinCacheViewMode",
+				LOCTEXT("VisualizeGPUSkinCacheViewModeDisplayName", "GPU Skin Cache"),
+				LOCTEXT("GPUSkinCacheVisualizationMenu_ToolTip", "Select a mode for GPU Skin Cache visualization."),
+				FNewMenuDelegate::CreateStatic(&FGPUSkinCacheVisualizationMenuCommands::BuildVisualisationSubMenu),
+				FUIAction(FExecuteAction(), FCanExecuteAction(),
+					FIsActionChecked::CreateLambda([WeakViewport = InViewport.ToWeakPtr()]() {
+						const TSharedRef<SEditorViewport> ViewportRef = WeakViewport.Pin().ToSharedRef();
+						const TSharedPtr<FEditorViewportClient> ViewportClient = ViewportRef->GetViewportClient();
+						check(ViewportClient.IsValid());
+						return ViewportClient->IsViewModeEnabled(VMI_VisualizeGPUSkinCache);
+					})),
+				EUserInterfaceActionType::RadioButton,
+				/* bInOpenSubMenuOnClick = */ false,
+				FSlateIcon(FAppStyle::GetAppStyleSetName(), "EditorViewport.VisualizeGPUSkinCacheMode"));
+		}
+	}
+
+	// Auto Exposure
+	{
+		const FEditorViewportCommands& BaseViewportCommands = FEditorViewportCommands::Get();
+
+		TSharedRef<SWidget> FixedEV100Menu = InViewport->BuildFixedEV100Menu();
+		TSharedPtr<FEditorViewportClient> EditorViewPostClient = InViewport->GetViewportClient();
+		const bool bIsLevelEditor = EditorViewPostClient.IsValid() && EditorViewPostClient->IsLevelEditorClient();
+
+		FToolMenuSection& Section = InMenu->AddSection("Exposure", LOCTEXT("ExposureHeader", "Exposure"));
+		Section.AddMenuEntry(
+			bIsLevelEditor ? BaseViewportCommands.ToggleInGameExposure : BaseViewportCommands.ToggleAutoExposure);
+		Section.AddEntry(FToolMenuEntry::InitWidget("FixedEV100", FixedEV100Menu, LOCTEXT("FixedEV100", "EV100")));
+	}
+
+	// Wireframe Opacity
+	{
+		TSharedRef<SWidget> WireOpacityMenu = InViewport->BuildWireframeMenu();
+		FToolMenuSection& Section = InMenu->AddSection("Wireframe", LOCTEXT("WireframeHeader", "Wireframe"));
+		Section.AddEntry(
+			FToolMenuEntry::InitWidget("WireframeOpacity", WireOpacityMenu, LOCTEXT("WireframeOpacity", "Opacity")));
+	}
+}
+
 } // namespace UE::UnrealEd
 
 #undef LOCTEXT_NAMESPACE
diff --git a/Engine/Source/Editor/UnrealEd/Public/SEditorViewportViewMenu.h b/Engine/Source/Editor/UnrealEd/Public/SEditorViewportViewMenu.h
--- a/Engine/Source/Editor/UnrealEd/Public/SEditorViewportViewMenu.h
+++ b/Engine/Source/Editor/UnrealEd/Public/SEditorViewportViewMenu.h
@@ -5,14 +5,14 @@
 
 #include "CoreMinimal.h"
 #include "Widgets/DeclarativeSyntaxSupport.h"
-#include "Framework/MultiBox/MultiBoxExtender.h"
-#include "Widgets/SWidget.h"
-#include "SEditorViewport.h"
 #include "SEditorViewportToolBarMenu.h"
 #include "Styling/SlateTypes.h"
 
-struct FSlateBrush;
+class FExtender;
+class SEditorViewport;
+class SWidget;
 class UToolMenu;
+struct FSlateBrush;
 
 class SEditorViewportViewMenu : public SEditorViewportToolbarMenu
 {
diff --git a/Engine/Source/Editor/UnrealEd/Public/ViewportToolbar/UnrealEdViewportToolbar.h b/Engine/Source/Editor/UnrealEd/Public/ViewportToolbar/UnrealEdViewportToolbar.h
--- a/Engine/Source/Editor/UnrealEd/Public/ViewportToolbar/UnrealEdViewportToolbar.h
+++ b/Engine/Source/Editor/UnrealEd/Public/ViewportToolbar/UnrealEdViewportToolbar.h
@@ -2,14 +2,30 @@
 
 #pragma once
 
+#include "Delegates/DelegateCombinations.h"
+#include "Engine/EngineBaseTypes.h"
 #include "Templates/SharedPointerFwd.h"
 
 class FText;
 class SEditorViewport;
+class UToolMenu;
 
 namespace UE::UnrealEd
 {
 
 UNREALED_API FText GetViewModesSubmenuLabel(TWeakPtr<SEditorViewport> InViewport);
 
+DECLARE_DELEGATE_RetVal_OneParam(bool, IsViewModeSupportedDelegate, EViewModeIndex);
+
+/**
+ * Populate a given UToolMenu with entries for a View Modes viewport toolbar submenu.
+ *
+ * @param InMenu The menu to poulate with entries.
+ * @param InViewport The viewport associated with this viewport toolbar.
+ * @param InIsViewModeSupported Optional delegate to filter which view modes are added to the list.
+ */
+UNREALED_API void FillViewMenu(UToolMenu* InMenu,
+	TSharedRef<SEditorViewport> InViewport,
+	IsViewModeSupportedDelegate InIsViewModeSupported = IsViewModeSupportedDelegate());
+
 } // namespace UE::UnrealEd
