Remove declaration of unimplemented function

diff --git a/Engine/Source/Editor/UnrealEd/Public/SCommonEditorViewportToolbarBase.h b/Engine/Source/Editor/UnrealEd/Public/SCommonEditorViewportToolbarBase.h
--- a/Engine/Source/Editor/UnrealEd/Public/SCommonEditorViewportToolbarBase.h
+++ b/Engine/Source/Editor/UnrealEd/Public/SCommonEditorViewportToolbarBase.h
@@ -142,14 +142,6 @@ private:
 	 */
 	UNREALED_API FText GetCameraMenuLabel() const;
 
-
-	/**
-	 * Returns the label for the "View" tool bar menu, which changes depending on viewport show flags
-	 *
-	 * @return	Label to use for this menu label
-	 */
-	UNREALED_API FText GetViewMenuLabel() const;
-
 	/**
 	 * Generates the toolbar options menu content 
 	 *
