Fixed up regression introduced when switching to structured logging so that we output the AssetPath as it did previously when using UE_ASSET_LOG (for blueprint compiler errors)

diff --git a/Engine/Source/Editor/UnrealEd/Private/Kismet2/CompilerResultsLog.cpp b/Engine/Source/Editor/UnrealEd/Private/Kismet2/CompilerResultsLog.cpp
--- a/Engine/Source/Editor/UnrealEd/Private/Kismet2/CompilerResultsLog.cpp
+++ b/Engine/Source/Editor/UnrealEd/Private/Kismet2/CompilerResultsLog.cpp
@@ -451,8 +451,8 @@ void FCompilerResultsLog::InternalLogMessage(FName MessageID, const TSharedRef<F
 		{
 			if (IsRunningCommandlet())
 			{				
-				UE_LOGFMT_NSLOC(LogBlueprint, Error, "Blueprint", "CompilerError", "[Compiler] {ErrorMessage} from Source: {SourceFile}",
-												("ErrorMessage", Message->ToText().ToString()), ("SourceFile", SourcePath));
+				UE_LOGFMT_NSLOC(LogBlueprint, Error, "Blueprint", "CompilerError", "[AssetLog] {AssetPath}: [Compiler] {ErrorMessage} from Source: {SourceFile}",
+												("AssetPath", FAssetMsg::FormatPathForAssetLog(*SourcePath)), ("ErrorMessage", Message->ToText().ToString()), ("SourceFile", SourcePath));
 			}
 			else
 			{
