IncrementalCook: Add UE::Cook::IDeterminismHelper to write diagnostics for cook indeterminism and false positive iterative skips.
[RN] Minor, Cooking

diff --git a/Engine/Source/Developer/IoStoreUtilities/Private/ZenStoreWriter.cpp b/Engine/Source/Developer/IoStoreUtilities/Private/ZenStoreWriter.cpp
--- a/Engine/Source/Developer/IoStoreUtilities/Private/ZenStoreWriter.cpp
+++ b/Engine/Source/Developer/IoStoreUtilities/Private/ZenStoreWriter.cpp
@@ -403,6 +403,15 @@ void FZenStoreWriter::WritePackageTrailer(const FPackageTrailerInfo& Info, const
 	checkNoEntry();
 }
 
+void FZenStoreWriter::RegisterDeterminismHelper(UObject* SourceObject,
+	const TRefCountPtr<UE::Cook::IDeterminismHelper>& DeterminismHelper)
+{
+	if (RegisterDeterminismHelperCallback)
+	{
+		RegisterDeterminismHelperCallback(SourceObject, DeterminismHelper);
+	}
+}
+
 void FZenStoreWriter::Initialize(const FCookInfo& Info)
 {
 	CookMode = Info.CookMode;
diff --git a/Engine/Source/Developer/IoStoreUtilities/Public/ZenStoreWriter.h b/Engine/Source/Developer/IoStoreUtilities/Public/ZenStoreWriter.h
--- a/Engine/Source/Developer/IoStoreUtilities/Public/ZenStoreWriter.h
+++ b/Engine/Source/Developer/IoStoreUtilities/Public/ZenStoreWriter.h
@@ -72,6 +72,10 @@ public:
 	{
 		BeginCacheCallback = MoveTemp(InBeginCacheCallback);
 	}
+	void SetRegisterDeterminismHelperCallback(FRegisterDeterminismHelperCallback&& InRegisterDeterminismHelperCallback)
+	{
+		RegisterDeterminismHelperCallback = MoveTemp(InRegisterDeterminismHelperCallback);
+	}
 
 	struct ZenHostInfo
 	{
@@ -81,10 +85,18 @@ public:
 		uint16 HostPort;
 	};
 
-	IOSTOREUTILITIES_API virtual FCookCapabilities GetCookCapabilities() const override
+	virtual FCapabilities GetCapabilities() const override
+	{
+		FCapabilities Result;
+		Result.bDeterminismDebug = (bool)RegisterDeterminismHelperCallback;
+		return Result;
+	}
+
+	virtual FCookCapabilities GetCookCapabilities() const override
 	{
 		FCookCapabilities Result;
 		Result.bDiffModeSupported = true;
+		Result.bOplogAttachments = true;
 		Result.HeaderFormat = EPackageHeaderFormat::ZenPackageSummary;
 		return Result;
 	}
@@ -101,6 +113,9 @@ public:
 
 
 	IOSTOREUTILITIES_API virtual void WriteBulkData(const FBulkDataInfo& Info, const FIoBuffer& BulkData, const TArray<FFileRegion>& FileRegions) override;
+	IOSTOREUTILITIES_API virtual void RegisterDeterminismHelper(UObject* SourceObject,
+		const TRefCountPtr<UE::Cook::IDeterminismHelper>& DeterminismHelper) override;
+
 	IOSTOREUTILITIES_API virtual void Initialize(const FCookInfo& Info) override;
 	IOSTOREUTILITIES_API virtual void BeginCook(const FCookInfo& Info) override;
 	IOSTOREUTILITIES_API virtual void EndCook(const FCookInfo& Info) override;
@@ -237,6 +252,7 @@ private:
 	TMap<FName, TArray<FString>>		PackageAdditionalFiles;
 
 	FBeginCacheCallback					BeginCacheCallback;
+	FRegisterDeterminismHelperCallback RegisterDeterminismHelperCallback;
 
 	FEntryCreatedEvent					EntryCreatedEvent;
 	FCriticalSection					CommitEventCriticalSection;
diff --git a/Engine/Source/Editor/UnrealEd/Classes/CookOnTheSide/CookOnTheFlyServer.h b/Engine/Source/Editor/UnrealEd/Classes/CookOnTheSide/CookOnTheFlyServer.h
--- a/Engine/Source/Editor/UnrealEd/Classes/CookOnTheSide/CookOnTheFlyServer.h
+++ b/Engine/Source/Editor/UnrealEd/Classes/CookOnTheSide/CookOnTheFlyServer.h
@@ -1540,6 +1540,8 @@ private:
 	bool bRunningAsShaderServer = false;
 	/** Whether to skip saving packages that are cooked. When true, the cook will only load and process packages but not write them to disk */
 	bool bSkipSave = false;
+	/** Whether cooked packages should store extra data to debug indeterminism. */
+	bool bDeterminismDebug = false;
 	/** Timers for tracking how long we have been busy, to manage retries and warnings of deadlock */
 	double SaveBusyStartTimeSeconds = MAX_flt;
 	double SaveBusyRetryTimeSeconds = MAX_flt;
diff --git a/Engine/Source/Editor/UnrealEd/Private/CookOnTheFlyServer.cpp b/Engine/Source/Editor/UnrealEd/Private/CookOnTheFlyServer.cpp
--- a/Engine/Source/Editor/UnrealEd/Private/CookOnTheFlyServer.cpp
+++ b/Engine/Source/Editor/UnrealEd/Private/CookOnTheFlyServer.cpp
@@ -5972,7 +5972,7 @@ public:
 	};
 
 	void InitializePackageWriter(UCookOnTheFlyServer& COTFS, ICookedPackageWriter*& CookedPackageWriter,
-		const FString& ResolvedMetadataPath)
+		const FString& ResolvedMetadataPath, UE::Cook::FDeterminismManager* InDeterminismManager)
 	{
 		Initialize();
 		if (DiffMode == EDiffMode::None)
@@ -5992,7 +5992,8 @@ public:
 		switch (DiffMode)
 		{
 		case EDiffMode::DiffOnly:
-			CookedPackageWriter = new FDiffPackageWriter(TUniquePtr<ICookedPackageWriter>(CookedPackageWriter));
+			CookedPackageWriter = new FDiffPackageWriter(TUniquePtr<ICookedPackageWriter>(CookedPackageWriter),
+				InDeterminismManager);
 			break;
 		case EDiffMode::LinkerDiff:
 			CookedPackageWriter = new FLinkerDiffPackageWriter(TUniquePtr<ICookedPackageWriter>(CookedPackageWriter));
@@ -6000,17 +6001,17 @@ public:
 		case EDiffMode::IterativeValidate:
 			CookedPackageWriter = new FIterativeValidatePackageWriter(COTFS,
 				TUniquePtr<ICookedPackageWriter>(CookedPackageWriter), FIterativeValidatePackageWriter::EPhase::AllInOnePhase,
-				ResolvedMetadataPath);
+				ResolvedMetadataPath, InDeterminismManager);
 			break;
 		case EDiffMode::IterativeValidatePhase1:
 			CookedPackageWriter = new FIterativeValidatePackageWriter(COTFS,
 				TUniquePtr<ICookedPackageWriter>(CookedPackageWriter), FIterativeValidatePackageWriter::EPhase::Phase1,
-				ResolvedMetadataPath);
+				ResolvedMetadataPath, InDeterminismManager);
 			break;
 		case EDiffMode::IterativeValidatePhase2:
 			CookedPackageWriter = new FIterativeValidatePackageWriter(COTFS,
 				TUniquePtr<ICookedPackageWriter>(CookedPackageWriter), FIterativeValidatePackageWriter::EPhase::Phase2,
-				ResolvedMetadataPath);
+				ResolvedMetadataPath, InDeterminismManager);
 			break;
 		default:
 			checkNoEntry();
@@ -6018,7 +6019,22 @@ public:
 		}
 	}
 
-private:
+	bool IsDeterminismDebug() const
+	{
+		switch (DiffMode)
+		{
+		case EDiffMode::None: return false;
+		case EDiffMode::DiffOnly: return true;
+		case EDiffMode::LinkerDiff: return false;
+		case EDiffMode::IterativeValidate: return true;
+		case EDiffMode::IterativeValidatePhase1: return true;
+		case EDiffMode::IterativeValidatePhase2: return true;
+		default:
+			checkNoEntry();
+			return false;
+		}
+	}
+
 	void Initialize()
 	{
 		if (bInitialized)
@@ -6065,6 +6081,7 @@ private:
 		bInitialized = true;
 	}
 
+private:
 	bool bInitialized = false;
 	EDiffMode DiffMode = EDiffMode::None;
 };
@@ -6359,6 +6376,10 @@ void FSaveCookedPackageContext::SetupPlatform(const ITargetPlatform* InTargetPla
 	Info.PackageName = Package->GetFName();
 	Info.LooseFilePath = PlatFilename;
 	PackageWriter->BeginPackage(Info);
+	if (CookContext->DeterminismManager)
+	{
+		CookContext->DeterminismManager->BeginPackage(Package, TargetPlatform, PackageWriter);
+	}
 	// Set platform-specific save flags
 	FPackagePlatformData& PlatformData = PackageData.FindOrAddPlatformData(TargetPlatform);
 	uint32 PlatformSaveFlagsMask = SAVE_AllowTimeout;
@@ -6430,6 +6451,10 @@ void FSaveCookedPackageContext::FinishPlatform()
 		Info.WriteOptions = GetCommitWriteOptions();
 
 		PackageWriter->CommitPackage(MoveTemp(Info));
+		if (CookContext->DeterminismManager)
+		{
+			CookContext->DeterminismManager->EndPackage();
+		}
 	}
 
 	// Update asset registry
@@ -6698,6 +6723,11 @@ TArray<IPackageWriter::FCommitAttachmentInfo> FSaveCookedPackageContext::GetComm
 		UE::TargetDomain::CollectAndStoreCookAttachments(Package, TargetPlatform, &SavePackageResult,
 			GeneratedResult, MoveTemp(PlatformRuntimeDependencies), Result);
 	}
+	if (CookContext->DeterminismManager)
+	{
+		CookContext->DeterminismManager->AppendCommitAttachments(Result);
+	}
+
 	return Result;
 }
 
@@ -7237,6 +7267,9 @@ void UCookOnTheFlyServer::SetInitializeConfigSettings(UE::Cook::FInitializeConfi
 	GConfig->GetString(TEXT("CookSettings"), TEXT("MPCookGeneratorSplit"), GeneratorSplit, GEditorIni);
 	FParse::Value(FCommandLine::Get(), TEXT("-MPCookGeneratorSplit="), GeneratorSplit);
 	MPCookGeneratorSplit = UE::Cook::ParseMPCookGeneratorSplitFromString(GeneratorSplit);
+
+	bDeterminismDebug = FParse::Param(FCommandLine::Get(), TEXT("cookdeterminism")) ||
+		FParse::Param(FCommandLine::Get(), TEXT("diffonlybase"));
 }
 
 void UCookOnTheFlyServer::ParseCookFilters()
@@ -11257,6 +11290,7 @@ UE::Cook::FCookSavePackageContext* UCookOnTheFlyServer::CreateSaveContext(const
 	const FString ResolvedRootPath = RootPathSandbox.Replace(TEXT("[Platform]"), *PlatformString);
 	const FString ResolvedMetadataPath = MetadataPathSandbox.Replace(TEXT("[Platform]"), *PlatformString);
 
+	TUniquePtr<FDeterminismManager> DeterminismManager;
 	ICookedPackageWriter* PackageWriter = nullptr;
 	FString WriterDebugName;
 	ICookedPackageWriter::FBeginCacheCallback BeginCacheCallback(
@@ -11265,21 +11299,35 @@ UE::Cook::FCookSavePackageContext* UCookOnTheFlyServer::CreateSaveContext(const
 			return this->SavePackageBeginCacheForCookedPlatformData(Info.PackageName,
 				Info.TargetPlatform, Info.SaveableObjects, Info.SaveFlags);
 		});
+	DiffModeHelper->Initialize();
+	ICookedPackageWriter::FRegisterDeterminismHelperCallback RegisterDeterminismHelperCallback;
+	if (bDeterminismDebug || DiffModeHelper->IsDeterminismDebug())
+	{
+		DeterminismManager = MakeUnique<FDeterminismManager>();
+		RegisterDeterminismHelperCallback =
+			[this, DeterminismManagerPtr = DeterminismManager.Get()]
+			(UObject* SourceObject, const TRefCountPtr<UE::Cook::IDeterminismHelper>& DeterminismHelper)
+			{
+				DeterminismManagerPtr->RegisterDeterminismHelper(SourceObject, DeterminismHelper);
+			};
+	}
 	if (IsUsingZenStore())
 	{
 		FZenStoreWriter* ZenWriter = new FZenStoreWriter(ResolvedRootPath, ResolvedMetadataPath, TargetPlatform);
 		ZenWriter->SetBeginCacheCallback(MoveTemp(BeginCacheCallback));
+		ZenWriter->SetRegisterDeterminismHelperCallback(MoveTemp(RegisterDeterminismHelperCallback));
 		PackageWriter = ZenWriter;
 		WriterDebugName = TEXT("ZenStore");
 	}
 	else
 	{
 		PackageWriter = new FLooseCookedPackageWriter(ResolvedRootPath, ResolvedMetadataPath, TargetPlatform,
-			GetAsyncIODelete(), *SandboxFile, MoveTemp(BeginCacheCallback));
+			GetAsyncIODelete(), *SandboxFile, MoveTemp(BeginCacheCallback),
+			MoveTemp(RegisterDeterminismHelperCallback));
 		WriterDebugName = TEXT("LooseCookedPackageWriter");
 	}
 
-	DiffModeHelper->InitializePackageWriter(*this, PackageWriter, ResolvedMetadataPath);
+	DiffModeHelper->InitializePackageWriter(*this, PackageWriter, ResolvedMetadataPath, DeterminismManager.Get());
 
 	// Setup save package settings (i.e. validation)
 	FSavePackageSettings SavePackageSettings = FSavePackageSettings::GetDefaultSettings();
@@ -11333,7 +11381,8 @@ UE::Cook::FCookSavePackageContext* UCookOnTheFlyServer::CreateSaveContext(const
 			});
 	}
 
-	FCookSavePackageContext* Context = new FCookSavePackageContext(TargetPlatform, PackageWriter, WriterDebugName, MoveTemp(SavePackageSettings));
+	FCookSavePackageContext* Context = new FCookSavePackageContext(TargetPlatform, PackageWriter, WriterDebugName, MoveTemp(SavePackageSettings),
+		MoveTemp(DeterminismManager));
 	return Context;
 
 }
diff --git a/Engine/Source/Editor/UnrealEd/Private/Cooker/CookDeterminismManager.cpp b/Engine/Source/Editor/UnrealEd/Private/Cooker/CookDeterminismManager.cpp
new file mode 100644
--- /dev/null
+++ b/Engine/Source/Editor/UnrealEd/Private/Cooker/CookDeterminismManager.cpp
@@ -0,0 +1,474 @@
+// Copyright Epic Games, Inc. All Rights Reserved.
+
+#include "Cooker/CookDeterminismManager.h"
+
+#include "Misc/StringBuilder.h"
+#include "Serialization/CompactBinarySerialization.h"
+#include "Serialization/CompactBinaryWriter.h"
+#include "UObject/Object.h"
+#include "UObject/Package.h"
+
+namespace UE::Cook
+{
+
+constexpr const ANSICHAR* DeterminismManagerName = "DeterminismManager";
+constexpr int32 DeterminismManagerVersion = 1;
+
+FDeterminismManager::FDeterminismManager()
+	: PackageData(*this)
+{
+}
+
+void FDeterminismManager::BeginPackage(UPackage* InPackage, const ITargetPlatform* TargetPlatform,
+	ICookedPackageWriter* InOplogProvider)
+{
+	check(InPackage);
+	check(InOplogProvider);
+	Package = InPackage;
+	OplogProvider = InOplogProvider;
+	if (!OplogAvailable.IsSet())
+	{
+		OplogAvailable = OplogProvider->GetCookCapabilities().bOplogAttachments;
+	}
+	// We require an empty PackageData to populate; we expect it to be empty from constructor or cleared by EndPackage.
+	check(PackageData.IsEmpty());
+	PackageData.TargetPlatform = TargetPlatform;
+}
+
+void FDeterminismManager::RegisterDeterminismHelper(UObject* SourceObject,
+	const TRefCountPtr<IDeterminismHelper>& DeterminismHelper)
+{
+	if (!DeterminismHelper)
+	{
+		return;
+	}
+
+	FExportDeterminismData& ExportData = PackageData.FindOrAddExportData(SourceObject);
+	ExportData.DeterminismHelpers.Add(DeterminismHelper);
+
+	FDeterminismConstructDiagnosticsContext Context(ExportData);
+	DeterminismHelper->ConstructDiagnostics(Context);
+}
+
+void FDeterminismManager::RecordPackageModified(UObject* InPrimaryAsset)
+{
+	PackageData.bModified = true;
+	if (InPrimaryAsset && InPrimaryAsset->GetPackage() != Package)
+	{
+		InPrimaryAsset = nullptr;
+	}
+	PackageData.PrimaryAsset = InPrimaryAsset;
+	if (PackageData.PrimaryAsset)
+	{
+		FExportDeterminismData& ExportData = PackageData.FindOrAddExportData(PackageData.PrimaryAsset);
+		ExportData.bPrimaryAsset = true;
+	}
+
+	FetchOldDiagnostics();
+}
+
+void FDeterminismManager::RecordExportModified(const FString& ExportPathName)
+{
+	UObject* Export = StaticFindObject(nullptr, nullptr, *ExportPathName);
+	if (!Export || Export->GetPackage() != Package)
+	{
+		return;
+	}
+
+	FExportDeterminismData& ExportData = PackageData.FindOrAddExportData(Export);
+	ExportData.bModified = true;
+}
+
+FString FDeterminismManager::GetCurrentPackageDiagnosticsAsText()
+{
+	if (!OplogAvailable.Get(false /* default value */))
+	{
+		// Comparison text is not implemented when we are unable to fetch attachments from the oplog.
+		return FString();
+	}
+
+	PackageData.Sort();
+
+	for (TPair<UObject*, TUniquePtr<FExportDeterminismData>>& ExportPair : PackageData.Exports)
+	{
+		FExportDeterminismData& ExportData = *ExportPair.Value;
+		ExportData.Sort();
+	}
+
+	TStringBuilder<256> Logger;
+	for (TPair<UObject*, TUniquePtr<FExportDeterminismData>>& ExportPair : PackageData.Exports)
+	{
+		FExportDeterminismData& ExportData = *ExportPair.Value;
+		ExportData.Logger = &Logger;
+		for (TRefCountPtr<IDeterminismHelper>& Helper : ExportData.DeterminismHelpers)
+		{
+			Helper->OnPackageModified(ExportData);
+		}
+		ExportData.Logger = nullptr;
+	}
+	if (Logger.ToView().EndsWith('\n'))
+	{
+		Logger.RemoveSuffix(Logger.ToView().EndsWith(TEXT("\r\n"))? 2 : 1);
+	}
+	return FString(*Logger);
+}
+
+void FDeterminismManager::AppendCommitAttachments(TArray<IPackageWriter::FCommitAttachmentInfo>& OutAttachments)
+{
+	PackageData.Sort();
+	FCbWriter Writer;
+	if (TrySave(Writer))
+	{
+		OutAttachments.Add(IPackageWriter::FCommitAttachmentInfo{ DeterminismManagerName, Writer.Save().AsObject() });
+	}
+}
+
+void FDeterminismManager::EndPackage()
+{
+	Package = nullptr;
+	OplogProvider = nullptr;
+	PackageData = FPackageDeterminismData(*this);
+}
+
+void FDeterminismManager::FetchOldDiagnostics()
+{
+	if (!OplogAvailable.Get(false /* default value */))
+	{
+		return;
+	}
+	FCbObject MarshalledData = OplogProvider->GetOplogAttachment(Package->GetFName(), DeterminismManagerName);
+	TryLoad(MarshalledData.AsFieldView());
+}
+
+bool FDeterminismManager::TrySave(FCbWriter& Writer)
+{
+	bool bHasValues = false;
+
+	PackageData.Sort();
+	Writer.BeginObject();
+	{
+		Writer << "Version" << DeterminismManagerVersion;
+		Writer.BeginArray("Exports");
+		TStringBuilder<256> ExportPackagePath;
+		for (TPair<UObject*, TUniquePtr<FExportDeterminismData>>& ExportPair : PackageData.Exports)
+		{
+			UObject* Export = ExportPair.Key;
+			FExportDeterminismData& ExportData = *ExportPair.Value;
+			if (Export == Package || ExportData.NewDiagnostics.IsEmpty())
+			{
+				continue;
+			}
+			ExportPackagePath.Reset();
+			ExportPair.Key->GetPathName(Package, ExportPackagePath);
+			if (ExportPackagePath.Len() == 0)
+			{
+				continue;
+			}
+
+			bHasValues = true;
+			Writer.BeginArray();
+			{
+				Writer << ExportPackagePath;
+				Writer.BeginArray();
+				for (TPair<FUtf8String, FCbField>& DiagnosticPair : ExportData.NewDiagnostics)
+				{
+					Writer.BeginArray();
+					Writer << DiagnosticPair.Key;
+					Writer << DiagnosticPair.Value;
+					Writer.EndArray();
+				}
+				Writer.EndArray();
+			}
+			Writer.EndArray();
+		}
+		Writer.EndArray();
+	}
+	Writer.EndObject();
+
+	return bHasValues;
+}
+
+bool FDeterminismManager::TryLoad(FCbFieldView Field)
+{
+	int32 Version = Field["Version"].AsInt32();
+	if (Version != DeterminismManagerVersion)
+	{
+		return false;
+	}
+
+	for (FCbFieldView ExportPairField : Field["Exports"])
+	{
+		FCbFieldViewIterator ExportPairIter = ExportPairField.CreateViewIterator();
+		FString ExportPackagePath;
+		if (!LoadFromCompactBinary(*ExportPairIter++, ExportPackagePath))
+		{
+			continue;
+		}
+		if (ExportPackagePath.IsEmpty())
+		{
+			// This could indicate the Package itself, but we don't allow recording diagnostics for the package itself
+			// because LinkerLoad does not record it as a serialized export.
+			continue;
+		}
+
+		UObject* Export = StaticFindObject(nullptr, Package, *ExportPackagePath);
+		if (!Export)
+		{
+			continue;
+		}
+		FExportDeterminismData* AllocatedExportData = nullptr;
+		auto AllocateExportData = [this, Export, &AllocatedExportData]()
+			{
+				if (!AllocatedExportData)
+				{
+					AllocatedExportData = &PackageData.FindOrAddExportData(Export);
+				}
+				return AllocatedExportData;
+			};
+
+		FCbFieldView DiagnosticArrayField = (*ExportPairIter++);
+		for (FCbFieldView DiagnosticPairField : DiagnosticArrayField)
+		{
+			FCbFieldViewIterator DiagnosticPairIter = DiagnosticPairField.CreateViewIterator();
+			FUtf8StringView DiagnosticName = (*DiagnosticPairIter).AsString();
+			if ((*DiagnosticPairIter++).HasError())
+			{
+				continue;
+			}
+			FCbField DiagnosticValue = FCbField::Clone(*DiagnosticPairIter++);
+			AllocateExportData()->AddOldDiagnostic(DiagnosticName, DiagnosticValue);
+		}
+	}
+
+	return true;
+}
+
+FDeterminismConstructDiagnosticsContext::FDeterminismConstructDiagnosticsContext(FExportDeterminismData& InExportData)
+	: ExportData(InExportData)
+{
+}
+
+const ITargetPlatform* FDeterminismConstructDiagnosticsContext::GetTargetPlatform()
+{
+	return ExportData.PackageData.TargetPlatform;
+}
+
+void FDeterminismConstructDiagnosticsContext::AddDiagnostic(FUtf8StringView DiagnosticName, const FCbField& Value)
+{
+	ExportData.AddNewDiagnostic(DiagnosticName, Value);
+}
+
+FExportDeterminismData::FExportDeterminismData(FPackageDeterminismData& InPackageData, UObject* InExport)
+	: PackageData(InPackageData)
+	, Export(InExport)
+{
+}
+
+bool FExportDeterminismData::IsModified()
+{
+	return bModified;
+}
+
+bool FExportDeterminismData::IsPrimaryAsset()
+{
+	return bPrimaryAsset;
+}
+
+const ITargetPlatform* FExportDeterminismData::GetTargetPlatform()
+{
+	return PackageData.TargetPlatform;
+}
+
+const TMap<FUtf8String, FCbField>& FExportDeterminismData::GetOldDiagnostics()
+{
+	Sort();
+	return OldDiagnostics;
+}
+
+const TMap<FUtf8String, FCbField>& FExportDeterminismData::GetNewDiagnostics()
+{
+	Sort();
+	return NewDiagnostics;
+}
+
+void FExportDeterminismData::AppendLog(FStringView LogText)
+{
+	if (Logger)
+	{
+		Logger->Append(LogText);
+	}
+}
+
+void FExportDeterminismData::AppendDiagnostics()
+{
+	if (!bAppendedDiagnostics)
+	{
+		AppendLog(GetCompareText());
+		bAppendedDiagnostics = true;
+	}
+}
+
+IDeterminismModifiedPackageContext& FExportDeterminismData::GetPackageContext()
+{
+	return PackageData;
+}
+
+struct FCompareDiagnosticName
+{
+	bool operator()(const FUtf8String& A, const FUtf8String& B) const
+	{
+		int32 Compare = A.Compare(B, ESearchCase::IgnoreCase);
+		if (Compare != 0)
+		{
+			return Compare < 0;
+		}
+		return A.Compare(B, ESearchCase::CaseSensitive) < 0;
+	}
+};
+
+FString FExportDeterminismData::GetCompareText()
+{
+	if (OldDiagnostics.IsEmpty() && NewDiagnostics.IsEmpty())
+	{
+		return FString();
+	}
+
+	Sort();
+	TStringBuilder<64> ExportRelPath;
+	Export->GetPathName(PackageData.Owner->Package, ExportRelPath);
+	TSet<FUtf8String> Keys;
+	for (const TPair<FUtf8String, FCbField>& Pair : OldDiagnostics)
+	{
+		Keys.Add(Pair.Key);
+	}
+	for (const TPair<FUtf8String, FCbField>& Pair : NewDiagnostics)
+	{
+		Keys.Add(Pair.Key);
+	}
+
+	TStringBuilder<256> Text;
+	for (const FUtf8String& Key : Keys)
+	{
+		Text << TEXT("'") << ExportRelPath << TEXT("':") << Key << TEXT(":Old Value\n");
+		const FCbField* OldValue = OldDiagnostics.Find(Key);
+		if (OldValue)
+		{
+			CompactBinaryToJson(*OldValue, Text);
+			Text << TEXT("\n");
+		}
+		Text << TEXT("'") << ExportRelPath << TEXT("':") << Key << TEXT(":New Value\n");
+		const FCbField* NewValue = NewDiagnostics.Find(Key);
+		if (NewValue)
+		{
+			CompactBinaryToJson(*NewValue, Text);
+			Text << TEXT("\n");
+		}
+	}
+	return FString(Text);
+}
+
+void FExportDeterminismData::AddNewDiagnostic(FUtf8StringView DiagnosticName, const FCbField& Value)
+{
+	FCbField& StoredField = NewDiagnostics.FindOrAdd(FUtf8String(DiagnosticName));
+	StoredField = Value;
+	StoredField.MakeOwned();
+	bSortDirty = true;
+}
+
+void FExportDeterminismData::AddOldDiagnostic(FUtf8StringView DiagnosticName, const FCbField& Value)
+{
+	FCbField& StoredField = OldDiagnostics.FindOrAdd(FUtf8String(DiagnosticName));
+	StoredField = Value;
+	StoredField.MakeOwned();
+	bSortDirty = true;
+}
+
+void FExportDeterminismData::Sort()
+{
+	if (!bSortDirty)
+	{
+		return;
+	}
+	bSortDirty = false;
+	OldDiagnostics.KeySort(FCompareDiagnosticName());
+	NewDiagnostics.KeySort(FCompareDiagnosticName());
+}
+
+FPackageDeterminismData::FPackageDeterminismData(FDeterminismManager& InOwner)
+	: Owner(&InOwner)
+{
+}
+
+const ITargetPlatform* FPackageDeterminismData::GetTargetPlatform()
+{
+	return TargetPlatform;
+}
+
+const TSet<UObject*>& FPackageDeterminismData::GetModifiedExports()
+{
+	Sort();
+	return ModifiedExports;
+}
+
+UObject* FPackageDeterminismData::GetPrimaryAsset()
+{
+	return PrimaryAsset;
+}
+
+const IDeterminismModifiedExportContext& FPackageDeterminismData::GetExportContext(UObject* Export)
+{
+	if (!Export || Export->GetPackage() != Owner->Package)
+	{
+		ensureMsgf(false, TEXT("GetExportContext called with object %s which is not in package %s."),
+			Export ? *Export->GetPathName() : TEXT("<nullptr>"), *Owner->Package->GetName());
+		Export = Owner->Package;
+	}
+	return FindOrAddExportData(Export);
+}
+
+FExportDeterminismData& FPackageDeterminismData::FindOrAddExportData(UObject* Object)
+{
+	check(Object);
+	TUniquePtr<FExportDeterminismData>& Ptr = Exports.FindOrAdd(Object);
+	if (!Ptr.IsValid())
+	{
+		Ptr.Reset(new FExportDeterminismData(*this, Object));
+		bSortDirty = true;
+	}
+	return *Ptr;
+}
+
+bool FPackageDeterminismData::IsEmpty()
+{
+	return Exports.IsEmpty();
+}
+
+void FPackageDeterminismData::Sort()
+{
+	if (!bSortDirty)
+	{
+		return;
+	}
+	bSortDirty = false;
+	UObject* Package = Owner->Package;
+
+	Exports.KeySort([Package](UObject& A, UObject& B)
+		{
+			TStringBuilder<256> ARelPath;
+			TStringBuilder<256> BRelPath;
+			A.GetPathName(Package, ARelPath);
+			B.GetPathName(Package, BRelPath);
+			return ARelPath.ToView().Compare(BRelPath.ToView(), ESearchCase::IgnoreCase) < 0;
+		});
+	ModifiedExports.Reset();
+	ModifiedExports.Reserve(Exports.Num());
+	for (TPair<UObject*, TUniquePtr<FExportDeterminismData>>& Pair : Exports)
+	{
+		if (Pair.Value->bModified)
+		{
+			ModifiedExports.Add(Pair.Key);
+		}
+	}
+}
+
+} // namespace UE::Cook
\ No newline at end of file
diff --git a/Engine/Source/Editor/UnrealEd/Private/Cooker/CookDeterminismManager.h b/Engine/Source/Editor/UnrealEd/Private/Cooker/CookDeterminismManager.h
new file mode 100644
--- /dev/null
+++ b/Engine/Source/Editor/UnrealEd/Private/Cooker/CookDeterminismManager.h
@@ -0,0 +1,139 @@
+// Copyright Epic Games, Inc. All Rights Reserved.
+
+#pragma once
+
+#include "Containers/Array.h"
+#include "Containers/UnrealString.h"
+#include "Cooker/CookDeterminismHelper.h"
+#include "Misc/Optional.h"
+#include "Serialization/PackageWriter.h"
+#include "UObject/NameTypes.h"
+
+class UObject;
+class FCbFieldView;
+class FCbWriter;
+
+namespace UE::Cook
+{
+
+class FDeterminismManager;
+struct FExportDeterminismData;
+struct FPackageDeterminismData;
+
+/**
+ * Helper for FDeterminismManager: the context object that is passed to IDeterminismHelper ConstructDiagnostics
+ * functions.
+ */
+struct FDeterminismConstructDiagnosticsContext : public IDeterminismConstructDiagnosticsContext
+{
+public:
+	FDeterminismConstructDiagnosticsContext(FExportDeterminismData& InExportData);
+
+	virtual const ITargetPlatform* GetTargetPlatform() override;
+	virtual void AddDiagnostic(FUtf8StringView DiagnosticName, const FCbField& Value) override;
+
+public:
+	FExportDeterminismData& ExportData;
+};
+
+/**
+ * Helper for FDeterminismManager: the data about the a UObject export that is passed
+ * to IDeterminismHelper OnPackageModified functions. Also serves as the container
+ * for all diagnostic data gathered by the manager for an export.
+ */
+struct FExportDeterminismData : public IDeterminismModifiedExportContext
+{
+public:
+	FExportDeterminismData(FPackageDeterminismData& InPackageData, UObject* InExport);
+
+	virtual bool IsModified() override;
+	virtual bool IsPrimaryAsset() override;
+	virtual const ITargetPlatform* GetTargetPlatform() override;
+	virtual const TMap<FUtf8String, FCbField>& GetOldDiagnostics() override;
+	virtual const TMap<FUtf8String, FCbField>& GetNewDiagnostics() override;
+	virtual IDeterminismModifiedPackageContext& GetPackageContext() override;
+	virtual FString GetCompareText() override;
+
+	virtual void AppendLog(FStringView LogText) override;
+	virtual void AppendDiagnostics() override;
+
+	void AddOldDiagnostic(FUtf8StringView DiagnosticName, const FCbField& Value);
+	void AddNewDiagnostic(FUtf8StringView DiagnosticName, const FCbField& Value);
+	void Sort();
+public:
+	TMap<FUtf8String, FCbField> OldDiagnostics;
+	TMap<FUtf8String, FCbField> NewDiagnostics;
+	FStringBuilderBase* Logger = nullptr;
+	TArray<TRefCountPtr<IDeterminismHelper>> DeterminismHelpers;
+	FPackageDeterminismData& PackageData;
+	UObject* Export = nullptr;
+	bool bPrimaryAsset = false;
+	bool bModified = false;
+	bool bSortDirty = true;
+	bool bAppendedDiagnostics = false;
+};
+
+/**
+ * Helper for FDeterminismManager: the data about the package that is passed
+ * to IDeterminismHelper OnPackageModified functions. Also serves as the container
+ * for all diagnostic data gathered by the manager for a package.
+ */
+struct FPackageDeterminismData : public IDeterminismModifiedPackageContext
+{
+public:
+	FPackageDeterminismData(FDeterminismManager& InOwner);
+
+	virtual const ITargetPlatform* GetTargetPlatform() override;
+	virtual const TSet<UObject*>& GetModifiedExports() override;
+	virtual UObject* GetPrimaryAsset() override;
+	virtual const IDeterminismModifiedExportContext& GetExportContext(UObject* Export) override;
+
+	bool IsEmpty();
+	void Sort();
+	FExportDeterminismData& FindOrAddExportData(UObject* Object);
+public:
+	TMap<UObject*, TUniquePtr<FExportDeterminismData>> Exports;
+	TSet<UObject*> ModifiedExports;
+	FDeterminismManager* Owner = nullptr; // non-null after constructor
+	UObject* PrimaryAsset = nullptr;
+	const ITargetPlatform* TargetPlatform = nullptr;
+	bool bModified = false;
+	bool bSortDirty = true;
+};
+
+/**
+ * Manager that receives the diagnostics from IDeterminismHelpers declared in UObject PreSave
+ * functions, and saves/loads them as attachments to the package in the oplog, and calls the
+ * DeterminismHelpers callback functions when a package is found to be unexpectedly modified.
+ */
+class FDeterminismManager
+{
+public:
+	FDeterminismManager();
+
+	void BeginPackage(UPackage* InPackage, const ITargetPlatform* TargetPlatform,
+		ICookedPackageWriter* InOplogProvider);
+	void RegisterDeterminismHelper(UObject* SourceObject,
+		const TRefCountPtr<IDeterminismHelper>& DeterminismHelper);
+	void RecordPackageModified(UObject* PrimaryAsset);
+	void RecordExportModified(const FString& ExportPathName);
+	FString GetCurrentPackageDiagnosticsAsText();
+	void AppendCommitAttachments(TArray<IPackageWriter::FCommitAttachmentInfo>& OutAttachments);
+	void EndPackage();
+
+private:
+	void FetchOldDiagnostics();
+	bool TrySave(FCbWriter& Writer);
+	bool TryLoad(FCbFieldView Field);
+
+	ICookedPackageWriter* OplogProvider = nullptr;
+	UPackage* Package = nullptr;
+	FPackageDeterminismData PackageData;
+	TOptional<bool> OplogAvailable;
+
+	friend UE::Cook::FDeterminismConstructDiagnosticsContext;
+	friend UE::Cook::FExportDeterminismData;
+	friend UE::Cook::FPackageDeterminismData;
+};
+
+} // namespace UE::Cook
\ No newline at end of file
diff --git a/Engine/Source/Editor/UnrealEd/Private/Cooker/CookTypes.cpp b/Engine/Source/Editor/UnrealEd/Private/Cooker/CookTypes.cpp
--- a/Engine/Source/Editor/UnrealEd/Private/Cooker/CookTypes.cpp
+++ b/Engine/Source/Editor/UnrealEd/Private/Cooker/CookTypes.cpp
@@ -2,9 +2,11 @@
 
 #include "CookTypes.h"
 
-#include "CompactBinaryTCP.h"
-#include "CookPackageData.h"
 #include "Containers/StringView.h"
+#include "Cooker/CompactBinaryTCP.h"
+#include "Cooker/CookDeterminismManager.h"
+#include "Cooker/CookPackageData.h"
+#include "Cooker/PackageTracker.h"
 #include "DerivedDataRequest.h"
 #include "Editor.h"
 #include "HAL/PlatformTLS.h"
@@ -15,7 +17,6 @@
 #include "Math/UnrealMathUtility.h"
 #include "Misc/CommandLine.h"
 #include "Misc/Parse.h"
-#include "PackageTracker.h"
 #include "Serialization/PackageWriterToSharedBuffer.h"
 
 LLM_DEFINE_TAG(Cooker_CachedPlatformData);
@@ -219,10 +220,12 @@ void SetIsSchedulerThread(bool bValue)
 }
 
 FCookSavePackageContext::FCookSavePackageContext(const ITargetPlatform* InTargetPlatform,
-	ICookedPackageWriter* InPackageWriter, FStringView InWriterDebugName, FSavePackageSettings InSettings)
+	ICookedPackageWriter* InPackageWriter, FStringView InWriterDebugName, FSavePackageSettings InSettings,
+	TUniquePtr<FDeterminismManager>&& InDeterminismManager)
 	: SaveContext(InTargetPlatform, InPackageWriter, MoveTemp(InSettings))
 	, WriterDebugName(InWriterDebugName)
 	, PackageWriter(InPackageWriter)
+	, DeterminismManager(MoveTemp(InDeterminismManager))
 {
 	PackageWriterCapabilities = InPackageWriter->GetCookCapabilities();
 }
diff --git a/Engine/Source/Editor/UnrealEd/Private/Cooker/CookTypes.h b/Engine/Source/Editor/UnrealEd/Private/Cooker/CookTypes.h
--- a/Engine/Source/Editor/UnrealEd/Private/Cooker/CookTypes.h
+++ b/Engine/Source/Editor/UnrealEd/Private/Cooker/CookTypes.h
@@ -90,6 +90,7 @@ class TFastPointerSet : public TSet<KeyType, TFastPointerSetKeyFuncs<KeyType>, S
 
 namespace UE::Cook
 {
+	class FDeterminismManager;
 	struct FPackageData;
 	struct FPlatformData;
 
@@ -297,13 +298,15 @@ namespace UE::Cook
 	struct FCookSavePackageContext
 	{
 		FCookSavePackageContext(const ITargetPlatform* InTargetPlatform,
-			ICookedPackageWriter* InPackageWriter, FStringView InWriterDebugName, FSavePackageSettings InSettings);
+			ICookedPackageWriter* InPackageWriter, FStringView InWriterDebugName, FSavePackageSettings InSettings,
+			TUniquePtr<FDeterminismManager>&& InDeterminismManager);
 		~FCookSavePackageContext();
 
 		FSavePackageContext SaveContext;
 		FString WriterDebugName;
 		ICookedPackageWriter* PackageWriter;
 		ICookedPackageWriter::FCookCapabilities PackageWriterCapabilities;
+		TUniquePtr<FDeterminismManager> DeterminismManager;
 	};
 
 	/* Thread Local Storage access to identify which thread is the SchedulerThread for cooking. */
@@ -392,7 +395,7 @@ namespace UE::Cook
 
 	TConstArrayView<const TCHAR*> GetCommandLineDelimiterStrs();
 	TConstArrayView<TCHAR> GetCommandLineDelimiterChars();
-}
+} // namespace UE::Cook
 
 bool LexTryParseString(FPlatformMemoryStats::EMemoryPressureStatus& OutValue, FStringView Text);
 FString LexToString(FPlatformMemoryStats::EMemoryPressureStatus Value);
diff --git a/Engine/Source/Editor/UnrealEd/Private/Cooker/DiffPackageWriter.cpp b/Engine/Source/Editor/UnrealEd/Private/Cooker/DiffPackageWriter.cpp
--- a/Engine/Source/Editor/UnrealEd/Private/Cooker/DiffPackageWriter.cpp
+++ b/Engine/Source/Editor/UnrealEd/Private/Cooker/DiffPackageWriter.cpp
@@ -28,7 +28,8 @@
 
 DEFINE_LOG_CATEGORY_STATIC(LogDiff, Log, All);
 
-FDiffPackageWriter::FDiffPackageWriter(TUniquePtr<ICookedPackageWriter>&& InInner)
+FDiffPackageWriter::FDiffPackageWriter(TUniquePtr<ICookedPackageWriter>&& InInner,
+	UE::Cook::FDeterminismManager* InDeterminismManager)
 	: Inner(MoveTemp(InInner))
 {
 	AccumulatorGlobals.Reset(new UE::DiffWriter::FAccumulatorGlobals(Inner.Get()));
@@ -56,6 +57,9 @@ FDiffPackageWriter::FDiffPackageWriter(TUniquePtr<ICookedPackageWriter>&& InInne
 	Indent = FCString::Spc(FOutputDeviceHelper::FormatLogLine(ELogVerbosity::Warning,
 		LogDiff.GetCategoryName(), TEXT(""), GPrintLogTimes).Len());
 	NewLine = TEXT("\n"); // OutputDevices are responsible for remapping to LINE_TERMINATOR if desired
+
+	check(InDeterminismManager);
+	DeterminismManager = InDeterminismManager;
 }
 
 void FDiffPackageWriter::ParseCmds()
@@ -215,6 +219,7 @@ void FDiffPackageWriter::WritePackageData(const FPackageInfo& Info, FLargeMemory
 
 		TMap<FName, FArchiveDiffStats> PackageDiffStats;
 		const TCHAR* CutoffString = TEXT("UEditorEngine::Save()");
+		Accumulator.SetDeterminismManager(*DeterminismManager);
 		Accumulator.CompareWithPrevious(CutoffString, PackageDiffStats);
 
 		//COOK_STAT(FSavePackageStats::NumberOfDifferentPackages++);
@@ -350,6 +355,12 @@ bool FDiffPackageWriter::IsAnotherSaveNeeded(FSavePackageResultStruct& PreviousR
 	}
 }
 
+void FDiffPackageWriter::RegisterDeterminismHelper(UObject* SourceObject,
+	const TRefCountPtr<UE::Cook::IDeterminismHelper>& DeterminismHelper)
+{
+	DeterminismManager->RegisterDeterminismHelper(SourceObject, DeterminismHelper);
+}
+
 bool FDiffPackageWriter::FilterPackageName(const FString& InWildcard)
 {
 	bool bInclude = false;
diff --git a/Engine/Source/Editor/UnrealEd/Private/Cooker/DiffPackageWriter.h b/Engine/Source/Editor/UnrealEd/Private/Cooker/DiffPackageWriter.h
--- a/Engine/Source/Editor/UnrealEd/Private/Cooker/DiffPackageWriter.h
+++ b/Engine/Source/Editor/UnrealEd/Private/Cooker/DiffPackageWriter.h
@@ -3,6 +3,7 @@
 #pragma once
 
 #include "AssetRegistry/AssetRegistryState.h"
+#include "Cooker/CookDeterminismManager.h"
 #include "Cooker/DiffWriterArchive.h"
 #include "Serialization/PackageWriter.h"
 
@@ -12,13 +13,14 @@ namespace UE::DiffWriter { struct FAccumulatorGlobals; }
 class FDiffPackageWriter : public ICookedPackageWriter
 {
 public:
-	FDiffPackageWriter(TUniquePtr<ICookedPackageWriter>&& InInner);
+	FDiffPackageWriter(TUniquePtr<ICookedPackageWriter>&& InInner, UE::Cook::FDeterminismManager* InDeterminismManager);
 
 	// IPackageWriter
 	virtual FCapabilities GetCapabilities() const override
 	{
 		FCapabilities Result = Inner->GetCapabilities();
 		Result.bIgnoreHeaderDiffs = bIgnoreHeaderDiffs;
+		Result.bDeterminismDebug = true;
 		return Result;
 	}
 	virtual void BeginPackage(const FBeginPackageInfo& Info) override;
@@ -51,8 +53,10 @@ public:
 	virtual TUniquePtr<FLargeMemoryWriter> CreateLinkerExportsArchive(FName PackageName, UObject* Asset, uint16 MultiOutputIndex) override;
 	virtual bool IsPreSaveCompleted() const override
 	{
-		return bDiffCallstack;
+		return bHasStartedSecondSave;
 	}
+	virtual void RegisterDeterminismHelper(UObject* SourceObject,
+		const TRefCountPtr<UE::Cook::IDeterminismHelper>& DeterminismHelper) override;
 
 	// ICookedPackageWriter
 	virtual FCookCapabilities GetCookCapabilities() const override
@@ -137,6 +141,7 @@ protected:
 	TSet<FTopLevelAssetPath> CompareDenyListClasses;
 	TUniquePtr<ICookedPackageWriter> Inner;
 	TUniquePtr<UE::DiffWriter::FAccumulatorGlobals> AccumulatorGlobals;
+	UE::Cook::FDeterminismManager* DeterminismManager = nullptr;
 	/** Only non-null between BeginPackage and CommitPackage. */
 	UPackage* Package = nullptr;
 	const TCHAR* Indent = nullptr;
@@ -149,7 +154,6 @@ protected:
 	bool bIgnoreHeaderDiffs = false;
 	bool bIsDifferent = false;
 	bool bNewPackage = false;
-	bool bDiffCallstack = false;
 	bool bHasStartedSecondSave = false;
 	bool bDumpObjList = false;
 	bool bDumpObjects = false;
diff --git a/Engine/Source/Editor/UnrealEd/Private/Cooker/DiffWriterArchive.cpp b/Engine/Source/Editor/UnrealEd/Private/Cooker/DiffWriterArchive.cpp
--- a/Engine/Source/Editor/UnrealEd/Private/Cooker/DiffWriterArchive.cpp
+++ b/Engine/Source/Editor/UnrealEd/Private/Cooker/DiffWriterArchive.cpp
@@ -3,6 +3,7 @@
 #include "Cooker/DiffWriterArchive.h"
 
 #include "Compression/CompressionUtil.h"
+#include "Cooker/CookDeterminismManager.h"
 #include "Cooker/DiffWriterLinkerLoadHeader.h"
 #include "Cooker/DiffWriterZenHeader.h"
 #include "HAL/FileManager.h"
@@ -605,6 +606,11 @@ void FAccumulator::SetHeaderSize(int64 InHeaderSize)
 	HeaderSize = InHeaderSize;
 }
 
+void FAccumulator::SetDeterminismManager(UE::Cook::FDeterminismManager& InDeterminismManager)
+{
+	DeterminismManager = &InDeterminismManager;
+}
+
 FName FAccumulator::GetAssetClass() const
 {
 	return Asset != nullptr ? Asset->GetClass()->GetFName() : NAME_None;
@@ -950,6 +956,10 @@ void FAccumulator::CompareWithPreviousForSection(const FPackageData& SourcePacka
 			check(CallstackAtOffsetPtr && DifferenceCallstackDataPtr); // These were set up above.
 			const FCallstacks::FCallstackAtOffset& CallstackAtOffset = *CallstackAtOffsetPtr;
 			const FCallstacks::FCallstackData& DifferenceCallstackData = *DifferenceCallstackDataPtr;
+			if (DeterminismManager)
+			{
+				DeterminismManager->RecordExportModified(DifferenceCallstackData.SerializedObjectName);
+			}
 
 			FString BeforePropertyVal;
 			FString AfterPropertyVal;
@@ -1096,6 +1106,10 @@ void FAccumulator::CompareWithPrevious(const TCHAR* CallstackCutoffText, TMap<FN
 
 	MessageCallback(ELogVerbosity::Display, FString::Printf(TEXT("Comparing: %s"), *Filename));
 	MessageCallback(ELogVerbosity::Warning, FString::Printf(TEXT("Asset class: %s"), *AssetClass.ToString()));
+	if (DeterminismManager)
+	{
+		DeterminismManager->RecordPackageModified(Asset);
+	}
 
 	int32 NumLoggedDiffs = 0;
 	
@@ -1151,6 +1165,17 @@ void FAccumulator::CompareWithPrevious(const TCHAR* CallstackCutoffText, TMap<FN
 	CompareWithPreviousForSection(SourcePackageExports, DestPackageExports, CallstackCutoffText,
 		NumLoggedDiffs, OutStats, ExportsFilename);
 
+	// Log comparison of the DeterminismHelper diagnostics, if we have any
+	if (DeterminismManager)
+	{
+		FString DeterminismLines = DeterminismManager->GetCurrentPackageDiagnosticsAsText();
+		if (!DeterminismLines.IsEmpty())
+		{
+			DeterminismLines = FString(TEXT("DeterminismHelper Diagnostics:\n") + DeterminismLines);
+			MessageCallback(ELogVerbosity::Display, MoveTemp(DeterminismLines));
+		}
+	}
+
 	// Optionally save out any differences we detected.
 	const FArchiveDiffStats& Stats = OutStats.FindOrAdd(AssetClass);
 	if (Stats.NumDiffs > 0)
diff --git a/Engine/Source/Editor/UnrealEd/Private/Cooker/DiffWriterArchive.h b/Engine/Source/Editor/UnrealEd/Private/Cooker/DiffWriterArchive.h
--- a/Engine/Source/Editor/UnrealEd/Private/Cooker/DiffWriterArchive.h
+++ b/Engine/Source/Editor/UnrealEd/Private/Cooker/DiffWriterArchive.h
@@ -24,6 +24,7 @@ class FLinkerLoad;
 class FProperty;
 class FUObjectThreadContext;
 class UObject;
+namespace UE::Cook { class FDeterminismManager; }
 struct FUObjectSerializeContext;
 
 
@@ -261,6 +262,7 @@ public:
 	void CompareWithPrevious(const TCHAR* CallstackCutoffText, TMap<FName,FArchiveDiffStats>& OutStats);
 
 	void SetHeaderSize(int64 InHeaderSize);
+	void SetDeterminismManager(UE::Cook::FDeterminismManager& InDeterminismManager);
 	void SetCollectingCallstacks(bool bInCollectingCallstacks);
 	FName GetAssetClass() const;
 	bool IsWriterUsingPostSaveTransforms() const;
@@ -282,6 +284,7 @@ private:
 	TArray<uint8> FirstSaveLinkerData;
 	int64 FirstSaveLinkerSize = 0;
 	FAccumulatorGlobals& Globals;
+	UE::Cook::FDeterminismManager* DeterminismManager = nullptr;
 
 	FDiffMap DiffMap;
 	FMessageCallback MessageCallback;
diff --git a/Engine/Source/Editor/UnrealEd/Private/Cooker/IterativeValidatePackageWriter.cpp b/Engine/Source/Editor/UnrealEd/Private/Cooker/IterativeValidatePackageWriter.cpp
--- a/Engine/Source/Editor/UnrealEd/Private/Cooker/IterativeValidatePackageWriter.cpp
+++ b/Engine/Source/Editor/UnrealEd/Private/Cooker/IterativeValidatePackageWriter.cpp
@@ -276,8 +276,9 @@ void FIterativeValidateMPCollector::ReadAndSyncPackageStatus(FCbObjectView Messa
 }
 
 FIterativeValidatePackageWriter::FIterativeValidatePackageWriter(UCookOnTheFlyServer& InCOTFS,
-	TUniquePtr<ICookedPackageWriter>&& InInner, EPhase InPhase, const FString& ResolvedMetadataPath)
-	: FDiffPackageWriter(MoveTemp(InInner))
+	TUniquePtr<ICookedPackageWriter>&& InInner, EPhase InPhase, const FString& ResolvedMetadataPath,
+	UE::Cook::FDeterminismManager* InDeterminismManager)
+	: FDiffPackageWriter(MoveTemp(InInner), InDeterminismManager)
 	, MetadataPath(ResolvedMetadataPath)
 	, COTFS(InCOTFS)
 	, Phase(InPhase)
diff --git a/Engine/Source/Editor/UnrealEd/Private/Cooker/IterativeValidatePackageWriter.h b/Engine/Source/Editor/UnrealEd/Private/Cooker/IterativeValidatePackageWriter.h
--- a/Engine/Source/Editor/UnrealEd/Private/Cooker/IterativeValidatePackageWriter.h
+++ b/Engine/Source/Editor/UnrealEd/Private/Cooker/IterativeValidatePackageWriter.h
@@ -22,7 +22,8 @@ public:
 		Phase2,
 	};
 	FIterativeValidatePackageWriter(UCookOnTheFlyServer& InCOTFS, TUniquePtr<ICookedPackageWriter>&& InInner,
-		EPhase InPhase, const FString& ResolvedMetadataPath);
+		EPhase InPhase, const FString& ResolvedMetadataPath,
+		UE::Cook::FDeterminismManager* InDeterminismManager);
 
 	// IPackageWriter
 	virtual void BeginPackage(const FBeginPackageInfo& Info) override;
diff --git a/Engine/Source/Editor/UnrealEd/Private/Cooker/LooseCookedPackageWriter.cpp b/Engine/Source/Editor/UnrealEd/Private/Cooker/LooseCookedPackageWriter.cpp
--- a/Engine/Source/Editor/UnrealEd/Private/Cooker/LooseCookedPackageWriter.cpp
+++ b/Engine/Source/Editor/UnrealEd/Private/Cooker/LooseCookedPackageWriter.cpp
@@ -59,13 +59,15 @@ LLM_DEFINE_TAG(Cooker_PackageStoreManifest);
 
 FLooseCookedPackageWriter::FLooseCookedPackageWriter(const FString& InOutputPath,
 	const FString& InMetadataDirectoryPath, const ITargetPlatform* InTargetPlatform, FAsyncIODelete& InAsyncIODelete,
-	UE::Cook::FCookSandbox& InSandboxFile, FBeginCacheCallback&& InBeginCacheCallback)
+	UE::Cook::FCookSandbox& InSandboxFile, FBeginCacheCallback&& InBeginCacheCallback,
+	FRegisterDeterminismHelperCallback&& InRegisterDeterminismHelperCallback)
 	: OutputPath(InOutputPath)
 	, MetadataDirectoryPath(InMetadataDirectoryPath)
 	, TargetPlatform(*InTargetPlatform)
 	, SandboxFile(InSandboxFile)
 	, AsyncIODelete(InAsyncIODelete)
 	, BeginCacheCallback(MoveTemp(InBeginCacheCallback))
+	, RegisterDeterminismHelperCallback(MoveTemp(InRegisterDeterminismHelperCallback))
 {
 }
 
@@ -538,6 +540,15 @@ FDateTime FLooseCookedPackageWriter::GetPreviousCookTime() const
 	return IFileManager::Get().GetTimeStamp(*PreviousAssetRegistry);
 }
 
+void FLooseCookedPackageWriter::RegisterDeterminismHelper(UObject* SourceObject,
+	const TRefCountPtr<UE::Cook::IDeterminismHelper>& DeterminismHelper)
+{
+	if (RegisterDeterminismHelperCallback)
+	{
+		RegisterDeterminismHelperCallback(SourceObject, DeterminismHelper);
+	}
+}
+
 void FLooseCookedPackageWriter::Initialize(const FCookInfo& Info)
 {
 	bIterateSharedBuild = Info.bIterateSharedBuild;
diff --git a/Engine/Source/Editor/UnrealEd/Private/Cooker/LooseCookedPackageWriter.h b/Engine/Source/Editor/UnrealEd/Private/Cooker/LooseCookedPackageWriter.h
--- a/Engine/Source/Editor/UnrealEd/Private/Cooker/LooseCookedPackageWriter.h
+++ b/Engine/Source/Editor/UnrealEd/Private/Cooker/LooseCookedPackageWriter.h
@@ -39,12 +39,20 @@ public:
 
 	FLooseCookedPackageWriter(const FString& OutputPath, const FString& MetadataDirectoryPath,
 		const ITargetPlatform* TargetPlatform, FAsyncIODelete& InAsyncIODelete,
-		UE::Cook::FCookSandbox& InSandboxFile, FBeginCacheCallback&& InBeginCacheCallback);
+		UE::Cook::FCookSandbox& InSandboxFile, FBeginCacheCallback&& InBeginCacheCallback,
+		FRegisterDeterminismHelperCallback&& InRegisterDeterminismHelperCallback);
 	~FLooseCookedPackageWriter();
 
+	virtual FCapabilities GetCapabilities() const override
+	{
+		FCapabilities Result = Super::GetCapabilities();
+		Result.bDeterminismDebug = (bool)RegisterDeterminismHelperCallback;
+		return Result;
+	}
+
 	virtual FCookCapabilities GetCookCapabilities() const override
 	{
-		FCookCapabilities Result;
+		FCookCapabilities Result = Super::GetCookCapabilities();
 		Result.bDiffModeSupported = true;
 		return Result;
 	}
@@ -53,6 +61,8 @@ public:
 	virtual int64 GetExportsFooterSize() override;
 
 	virtual FDateTime GetPreviousCookTime() const override;
+	virtual void RegisterDeterminismHelper(UObject* SourceObject,
+		const TRefCountPtr<UE::Cook::IDeterminismHelper>& DeterminismHelper) override;
 	virtual void Initialize(const FCookInfo& Info) override;
 	virtual void BeginCook(const FCookInfo& Info) override;
 	virtual void EndCook(const FCookInfo& Info) override;
@@ -166,6 +176,7 @@ private:
 	UE::Cook::FCookSandbox& SandboxFile;
 	FAsyncIODelete& AsyncIODelete;
 	FBeginCacheCallback BeginCacheCallback;
+	FRegisterDeterminismHelperCallback RegisterDeterminismHelperCallback;
 	bool bIterateSharedBuild = false;
 	bool bProvidePerPackageResults = false;
 };
diff --git a/Engine/Source/Runtime/Core/Public/Serialization/PackageWriter.h b/Engine/Source/Runtime/Core/Public/Serialization/PackageWriter.h
--- a/Engine/Source/Runtime/Core/Public/Serialization/PackageWriter.h
+++ b/Engine/Source/Runtime/Core/Public/Serialization/PackageWriter.h
@@ -13,6 +13,7 @@
 #include "Misc/EnumClassFlags.h"
 #include "Misc/SecureHash.h"
 #include "Serialization/CompactBinary.h"
+#include "Templates/RefCounting.h"
 #include "Templates/UniquePtr.h"
 
 class FAssetRegistryState;
@@ -22,6 +23,7 @@ class FLargeMemoryWriter;
 class ICookedPackageWriter;
 class IPackageStoreWriter;
 class UObject;
+namespace UE::Cook { class IDeterminismHelper; }
 struct FPackageStoreEntryResource;
 struct FSavePackageArgs;
 struct FSavePackageResultStruct;
@@ -48,8 +50,17 @@ public:
 		 */
 		bool bDeclareRegionForEachAdditionalFile = false;
 
-		/** Applicable only to -diffonly saves; suppresses output and breakpoints for diffs in the header. */
+		/**
+		 * Applicable only to cook saves, and only to -diffonly saves; suppresses output and breakpoints for diffs
+		 * in the header.
+		 */
 		bool bIgnoreHeaderDiffs = false;
+
+		/**
+		 * Applicable only to cook saves: True if the SavePackage call should write extra debug data for debugging
+		 * cook determinism or incremental cook issues.
+		 */
+		bool bDeterminismDebug = false;
 	};
 
 	/** Return capabilities/settings this PackageWriter has/requires 
@@ -196,6 +207,15 @@ public:
 	/** Returns an archive to be used when serializing exports. */
 	COREUOBJECT_API virtual TUniquePtr<FLargeMemoryWriter> CreateLinkerExportsArchive(FName PackageName, UObject* Asset, uint16 MultiOutputIndex);
 
+	/**
+	 * Overridden by PackageWriters that handle bDeterminismDebug=true. A system will call this function to register
+	 * their callback class for adding determinism diagnostics for the given object to the package save.
+	 */
+	virtual void RegisterDeterminismHelper(UObject* SourceObject,
+		const TRefCountPtr<UE::Cook::IDeterminismHelper>& DeterminismHelper)
+	{
+	}
+
 	/** Report whether PreSave was already called by the PackageWriter before the current UPackage::Save call. */
 	virtual bool IsPreSaveCompleted() const
 	{
@@ -255,6 +275,9 @@ public:
 		 */
 		bool bOverridesPackageModificationStatus = false;
 
+		/** If true, this writer can save and load extra data for each package across cooks. */
+		bool bOplogAttachments = false;
+
 		/** What header format is produced as output by this writer. */
 		EPackageHeaderFormat HeaderFormat = EPackageHeaderFormat::PackageFileSummary;
 	};
@@ -376,6 +399,8 @@ public:
 	};
 	// Helper callback type for Writers that need to send the message on to UCookOnTheFlyServer
 	using FBeginCacheCallback = TUniqueFunction<EPackageWriterResult(FBeginCacheForCookedPlatformDataInfo& Info)>;
+	using FRegisterDeterminismHelperCallback = TUniqueFunction<void(UObject* SourceObject,
+		const TRefCountPtr<UE::Cook::IDeterminismHelper>& DeterminismHelper)>;
 
 	virtual EPackageWriterResult BeginCacheForCookedPlatformData(FBeginCacheForCookedPlatformDataInfo& Info)
 	{
diff --git a/Engine/Source/Runtime/CoreUObject/Private/UObject/SavePackage/SaveContext.h b/Engine/Source/Runtime/CoreUObject/Private/UObject/SavePackage/SaveContext.h
--- a/Engine/Source/Runtime/CoreUObject/Private/UObject/SavePackage/SaveContext.h
+++ b/Engine/Source/Runtime/CoreUObject/Private/UObject/SavePackage/SaveContext.h
@@ -531,6 +531,11 @@ public:
 		{
 			ObjectSaveContext.SaveOverrides = MoveTemp(*SaveArgs.InOutSaveOverrides);
 		}
+		ObjectSaveContext.PackageWriter = PackageWriter;
+		if (PackageWriter)
+		{
+			ObjectSaveContext.bDeterminismDebug = SaveArgs.SavePackageContext->PackageWriterCapabilities.bDeterminismDebug;
+		}
 
 		// Setup the harvesting flags and generate the context for harvesting the package
 		SetupHarvestingRealms();
diff --git a/Engine/Source/Runtime/CoreUObject/Private/UObject/SavePackage/SavePackageUtilities.cpp b/Engine/Source/Runtime/CoreUObject/Private/UObject/SavePackage/SavePackageUtilities.cpp
--- a/Engine/Source/Runtime/CoreUObject/Private/UObject/SavePackage/SavePackageUtilities.cpp
+++ b/Engine/Source/Runtime/CoreUObject/Private/UObject/SavePackage/SavePackageUtilities.cpp
@@ -684,6 +684,7 @@ void CallPreSaveRoot(UObject* Object, FObjectSaveContextData& ObjectSaveContext)
 	PRAGMA_ENABLE_DEPRECATION_WARNINGS;
 
 	ObjectSaveContext.bCleanupRequired = false;
+	ObjectSaveContext.Object = Object;
 	Object->PreSaveRoot(FObjectPreSaveRootContext(ObjectSaveContext));
 	ObjectSaveContext.bCleanupRequired |= bLegacyNeedsCleanup;
 }
@@ -695,6 +696,7 @@ void CallPostSaveRoot(UObject* Object, FObjectSaveContextData& ObjectSaveContext
 	Object->PostSaveRoot(bNeedsCleanup);
 	PRAGMA_ENABLE_DEPRECATION_WARNINGS;
 
+	ObjectSaveContext.Object = Object;
 	ObjectSaveContext.bCleanupRequired = bNeedsCleanup;
 	Object->PostSaveRoot(FObjectPostSaveRootContext(ObjectSaveContext));
 }
@@ -800,6 +802,21 @@ void FObjectPreSaveContext::HarvestCookRuntimeDependencies(UObject* HarvestRefer
 	}
 
 }
+
+bool FObjectPreSaveContext::IsDeterminismDebug()
+{
+	return Data.bDeterminismDebug;
+}
+
+void FObjectPreSaveContext::RegisterDeterminismHelper(
+	const TRefCountPtr<UE::Cook::IDeterminismHelper>& DeterminismHelper)
+{
+	if (Data.PackageWriter)
+	{
+		Data.PackageWriter->RegisterDeterminismHelper(Data.Object, DeterminismHelper);
+	}
+}
+
 #endif
 
 void FObjectSaveContextData::Set(UPackage* Package, const ITargetPlatform* InTargetPlatform, const TCHAR* InTargetFilename, uint32 InSaveFlags)
diff --git a/Engine/Source/Runtime/CoreUObject/Private/UObject/SavePackage2.cpp b/Engine/Source/Runtime/CoreUObject/Private/UObject/SavePackage2.cpp
--- a/Engine/Source/Runtime/CoreUObject/Private/UObject/SavePackage2.cpp
+++ b/Engine/Source/Runtime/CoreUObject/Private/UObject/SavePackage2.cpp
@@ -284,13 +284,15 @@ ESavePackageResult RoutePresave(FSaveContext& SaveContext)
 		// this is to prevent warning on objects which won't be harvested later since they are unreferenced
 		if (!SaveContext.IsUnsaveable(Object, false/*bEmitWarning*/))
 		{
+			FObjectSaveContextData& ObjectSaveContext = SaveContext.GetObjectSaveContext();
+			ObjectSaveContext.Object = Object;
 			if (SaveContext.IsCooking() && Object->HasAnyFlags(RF_ClassDefaultObject | RF_ArchetypeObject))
 			{
 				FArchiveObjectCrc32 CrcArchive;
 				CrcArchive.ArIsFilterEditorOnly = true;
 				FString PathNameBefore = Object->GetPathName();
 				int32 Before = CrcArchive.Crc32(Object);
-				UE::SavePackageUtilities::CallPreSave(Object, SaveContext.GetObjectSaveContext());
+				UE::SavePackageUtilities::CallPreSave(Object, ObjectSaveContext);
 				int32 After = CrcArchive.Crc32(Object);
 
 				if (Before != After)
@@ -324,7 +326,7 @@ ESavePackageResult RoutePresave(FSaveContext& SaveContext)
 			}
 			else
 			{
-				UE::SavePackageUtilities::CallPreSave(Object, SaveContext.GetObjectSaveContext());
+				UE::SavePackageUtilities::CallPreSave(Object, ObjectSaveContext);
 			}
 		}
 	}
@@ -2670,6 +2672,8 @@ void PostSavePackage(FSaveContext& SaveContext)
 		bContainsNoAsset &= bLinkerContainsNoAsset;
 
 		// Call the linker post save callbacks
+		FObjectSaveContextData& ObjectSaveContext = SaveContext.GetObjectSaveContext();
+		ObjectSaveContext.Object = nullptr;
 		Linker->OnPostSave(SaveContext.GetTargetPackagePath(), FObjectPostSaveContext(SaveContext.GetObjectSaveContext()));
 	}
 
@@ -2696,7 +2700,9 @@ void PostSavePackage(FSaveContext& SaveContext)
 	PRAGMA_DISABLE_DEPRECATION_WARNINGS;
 	UPackage::PackageSavedEvent.Broadcast(SaveContext.GetFilename(), Package);
 	PRAGMA_ENABLE_DEPRECATION_WARNINGS;
-	UPackage::PackageSavedWithContextEvent.Broadcast(SaveContext.GetFilename(), Package, FObjectPostSaveContext(SaveContext.GetObjectSaveContext()));
+	FObjectSaveContextData& ObjectSaveContext = SaveContext.GetObjectSaveContext();
+	ObjectSaveContext.Object = nullptr;
+	UPackage::PackageSavedWithContextEvent.Broadcast(SaveContext.GetFilename(), Package, FObjectPostSaveContext(ObjectSaveContext));
 
 	// update the internal package filename path if we're saving to a valid mounted path and we aren't currently cooking
 #if WITH_EDITOR
diff --git a/Engine/Source/Runtime/CoreUObject/Public/Cooker/CookDeterminismHelper.h b/Engine/Source/Runtime/CoreUObject/Public/Cooker/CookDeterminismHelper.h
new file mode 100644
--- /dev/null
+++ b/Engine/Source/Runtime/CoreUObject/Public/Cooker/CookDeterminismHelper.h
@@ -0,0 +1,149 @@
+// Copyright Epic Games, Inc. All Rights Reserved.
+
+#pragma once
+
+#if WITH_EDITOR
+
+#include "Containers/Map.h"
+#include "Containers/Set.h"
+#include "Containers/StringView.h"
+#include "Containers/UnrealString.h"
+#include "Containers/Utf8String.h"
+#include "Templates/RefCounting.h"
+#include "Serialization/CompactBinary.h"
+
+class ITargetPlatform;
+class UObject;
+
+namespace UE::Cook
+{
+
+struct IDeterminismModifiedPackageContext;
+
+/**
+ * Struct to hold input and receive output of IDeterminismHelper::ConstructDiagnostics. Virtual base class so it can
+ * be implemented in the UnrealEd project.
+ */
+struct IDeterminismConstructDiagnosticsContext
+{
+	virtual ~IDeterminismConstructDiagnosticsContext()
+	{
+	}
+
+	/**
+	 * Reports the TargetPlatform being cooked for which the diagnostics are created. In a multi-platform cook each
+	 * platform will have a separate ConstructDiagnostics call.
+	 */
+	virtual const ITargetPlatform* GetTargetPlatform() = 0;
+
+	/** Output function; call this to add a diagnostic with the given value. */
+	virtual void AddDiagnostic(FUtf8StringView DiagnosticName, const FCbField& Value) = 0;
+};
+
+/**
+ * Struct to hold input and receive output of IDeterminismHelper::OnExportModified. Virtual base class so it can
+ * be implemented in the UnrealEd project.
+ */
+struct IDeterminismModifiedExportContext
+{
+	virtual ~IDeterminismModifiedExportContext()
+	{
+	}
+
+	/**
+	 * True iff the export that registered the DeterminismHelper was found to have modifications in the bytes created
+	 * by its Serialize function.
+	 */
+	virtual bool IsModified() = 0;
+	/** True iff the export that registered the DeterminismHelper is the primary asset in the package. */
+	virtual bool IsPrimaryAsset() = 0;
+	/** Reports the TargetPlatform for which the package was found to be modified. */
+	virtual const ITargetPlatform* GetTargetPlatform() = 0;
+	/**
+	 * Returns a container of the diagnostics written for the export in the old version of the package from the
+	 * previous cook.
+	 */
+	virtual const TMap<FUtf8String, FCbField>& GetOldDiagnostics() = 0;
+	/**
+	 * Returns a container of the diagnostics written for the export in the current in-memory version of the package
+	 * from the current cook.
+	 */
+	virtual const TMap<FUtf8String, FCbField>& GetNewDiagnostics() = 0;
+	/** Return a reference to information about the package in general and all modified exports in the package. */
+	virtual IDeterminismModifiedPackageContext& GetPackageContext() = 0;
+	/** Create a string version of the values of the old and new diagnostics for the export. */
+	virtual FString GetCompareText() = 0;
+
+	/**
+	 * Output function; add the given text to the -diffonly or -iterativevalidate log output for the modified
+	 * package.
+	 */
+	virtual void AppendLog(FStringView LogText) = 0;
+	/** Output function; calls AppendLog(GetCompareText()). Activates only once per export if called multiple times. */
+	virtual void AppendDiagnostics() = 0;
+};
+
+/** Extended data for IDeterminismModifiedExportContext. Holds data about the entire package. */
+struct IDeterminismModifiedPackageContext
+{
+	virtual ~IDeterminismModifiedPackageContext()
+	{
+	}
+
+	/** Reports the TargetPlatform for which the package was found to be modified. */
+	virtual const ITargetPlatform* GetTargetPlatform() = 0;
+	/**
+	 * Return a list of all the exports in the package that were found to have modifications in the bytes created
+	 * by their Serialize functions.
+	 */
+	virtual const TSet<UObject*>& GetModifiedExports() = 0;
+	/**
+	 * Get the UObject that is the primary asset for the package, or nullptr in the few instances of packages that
+	 * lack a primary asset.
+	 */
+	virtual UObject* GetPrimaryAsset() = 0;
+	/**
+	 * Get the IDeterminismModifiedExportContext object for the given export. Creates it if it doesn't exist;
+	 * if it doesn't exist the IDeterminismModifiedExportContext will have IsModified=false and emtpy diagnostics.
+	 * Invalid to call with a Export pointer that is not in the package; in that case it will log an error and
+	 * return an empty IDeterminismModifiedExportContext pointing to the UPackage itself.
+	 */
+	virtual const IDeterminismModifiedExportContext& GetExportContext(UObject* Export) = 0;
+};
+
+/**
+ * Interface implemented by UObjects that want to store diagnostic data for comparison when
+ * a package is found to contain indeterminism by a -diffonly cook or is found to have a
+ * FalsePositiveIterativeSkip by a -iterativevalidate cook.
+ * 
+ * DeterminismHelper objects are registered via FObjectPreSaveContext.RegisterDeterminismHelper
+ * in the PreSave of a UObject.
+ */
+class IDeterminismHelper : public FRefCountBase
+{
+public:
+	/**
+	 * Override this function to add diagnostics to the package containing data about
+	 * the UObject that registered the DeterminismHelper.
+	 */
+	virtual void ConstructDiagnostics(IDeterminismConstructDiagnosticsContext& Context)
+	{
+	}
+
+	/**
+	 * Override this function to get a callback with the old and new values of the diagnostics
+	 * when a package is found to be modified by indeterminism or FalsePositiveIterativeSkip.
+	 * The default implementation just prints out the old and new values of each diagnostic.
+	 */
+	virtual void OnPackageModified(IDeterminismModifiedExportContext& ExportContext)
+	{
+		if (ExportContext.IsModified() || ExportContext.IsPrimaryAsset())
+		{
+			ExportContext.AppendDiagnostics();
+		}
+	}
+};
+
+} // namespace UE::Cook
+
+#endif // WITH_EDITOR
\ No newline at end of file
diff --git a/Engine/Source/Runtime/CoreUObject/Public/UObject/ObjectSaveContext.h b/Engine/Source/Runtime/CoreUObject/Public/UObject/ObjectSaveContext.h
--- a/Engine/Source/Runtime/CoreUObject/Public/UObject/ObjectSaveContext.h
+++ b/Engine/Source/Runtime/CoreUObject/Public/UObject/ObjectSaveContext.h
@@ -6,13 +6,16 @@
 #include "HAL/Platform.h"
 #include "UObject/CookEnums.h"
 #include "UObject/ObjectSaveOverride.h"
+#include "Templates/RefCounting.h"
 
 class FPackagePath;
+class IPackageWriter;
 class ITargetPlatform;
 class UPackage;
 struct FSoftObjectPath;
 #if WITH_EDITOR
 namespace UE::Cook { class FCookDependency; }
+namespace UE::Cook { class IDeterminismHelper; }
 #endif
 
 /** Data used to provide information about the save parameters during PreSave/PostSave. */
@@ -55,6 +58,12 @@ struct FObjectSaveContextData
 	/** The target platform of the save, if cooking. Null if not cooking. */
 	const ITargetPlatform* TargetPlatform = nullptr;
 
+	/** The PackageWriter passed to SavePackage, may be null. */
+	IPackageWriter* PackageWriter = nullptr;
+
+	/** The object the Save event is being called on, if known. */
+	UObject* Object = nullptr;
+
 	/** The save flags (ESaveFlags) of the save. */
 	uint32 SaveFlags = 0;
 
@@ -87,6 +96,12 @@ struct FObjectSaveContextData
 	/** Set to false if the save failed, before calling any PostSaves. */
 	bool bSaveSucceeded = true;
 
+	/**
+	 * Applicable only to cook saves: True if the SavePackage call should write extra debug data for debugging
+	 * cook determinism or incremental cook issues.
+	 */
+	bool bDeterminismDebug = false;
+
 	// Collection variables that are written but not read during the PreSave/PostSave functions
 #if WITH_EDITOR
 	TArray<UE::Cook::FCookDependency> CookBuildDependencies;
@@ -210,6 +225,17 @@ public:
 	COREUOBJECT_API void AddCookRuntimeDependency(FSoftObjectPath Dependency);
 	/** Serialize an object to find all packages that it references, and AddCookRuntimeDependency for each one. */
 	COREUOBJECT_API void HarvestCookRuntimeDependencies(UObject* HarvestReferencesFrom);
+
+	/**
+	 * Applicable only to cook saves: True if the SavePackage call should write extra debug data for debugging
+	 * cook determinism or incremental cook issues.
+	 */
+	COREUOBJECT_API bool IsDeterminismDebug();
+	/**
+	 * Ignored unless IsDeterminismDebug()=true. An object should call this function to register
+	 * their callback class for adding determinism diagnostics the package save.
+	 */
+	COREUOBJECT_API void RegisterDeterminismHelper(const TRefCountPtr<UE::Cook::IDeterminismHelper>& DeterminismHelper);
 #endif
 
 	/**
diff --git a/Engine/Source/Runtime/Engine/Private/Texture.cpp b/Engine/Source/Runtime/Engine/Private/Texture.cpp
--- a/Engine/Source/Runtime/Engine/Private/Texture.cpp
+++ b/Engine/Source/Runtime/Engine/Private/Texture.cpp
@@ -43,11 +43,14 @@
 #include "Async/ParallelFor.h"
 
 #if WITH_EDITOR
+#include "Cooker/CookDeterminismHelper.h"
 #include "DerivedDataBuildVersion.h"
 #include "Math/GuardedInt.h"
+#include "Misc/ScopeRWLock.h"
+#include "Serialization/CompactBinaryWriter.h"
 #include "TextureCompiler.h"
 #include "TextureBuildUtilities.h"
-#include "Misc/ScopeRWLock.h"
+#include "TextureDerivedDataBuildUtils.h"
 #endif
 
 #include UE_INLINE_GENERATED_CPP_BY_NAME(Texture)
@@ -1382,6 +1385,32 @@ void UTexture::PreSave(const class ITargetPlatform* TargetPlatform)
 	PRAGMA_ENABLE_DEPRECATION_WARNINGS;
 }
 
+#if WITH_EDITOR
+class FTextureDeterminismHelper : public UE::Cook::IDeterminismHelper
+{
+public:
+	FTextureDeterminismHelper(UTexture* InTexture)
+		: Texture(InTexture)
+	{
+	}
+
+	virtual void ConstructDiagnostics(UE::Cook::IDeterminismConstructDiagnosticsContext& Context)
+	{
+		FCbWriter Writer;
+		if (!UE::TextureBuildUtilities::TryWriteCookDeterminismDiagnostics(
+			Writer, Texture, Context.GetTargetPlatform()))
+		{
+			return;
+		}
+
+		Context.AddDiagnostic("UTexture", Writer.Save());
+	}
+
+private:
+	UTexture* Texture;
+};
+#endif
+
 void UTexture::PreSave(FObjectPreSaveContext ObjectSaveContext)
 {
 	PreSaveEvent.Broadcast(this);
@@ -1413,6 +1442,11 @@ void UTexture::PreSave(FObjectPreSaveContext ObjectSaveContext)
 		DeferCompression = false;
 		UpdateResource();
 	}
+
+	if (ObjectSaveContext.IsDeterminismDebug())
+	{
+		ObjectSaveContext.RegisterDeterminismHelper(new FTextureDeterminismHelper(this));
+	}
 #endif // #if WITH_EDITOR
 }
 
diff --git a/Engine/Source/Runtime/Engine/Private/TextureDerivedData.cpp b/Engine/Source/Runtime/Engine/Private/TextureDerivedData.cpp
--- a/Engine/Source/Runtime/Engine/Private/TextureDerivedData.cpp
+++ b/Engine/Source/Runtime/Engine/Private/TextureDerivedData.cpp
@@ -40,6 +40,7 @@
 #include "ProfilingDebugging/CookStats.h"
 #include "UObject/ArchiveCookContext.h"
 #include "VT/LightmapVirtualTexture.h"
+#include "Serialization/CompactBinaryWriter.h"
 #include "Serialization/MemoryWriter.h"
 #include "TextureBuildUtilities.h"
 #include "TextureCompiler.h"
@@ -4115,7 +4116,6 @@ void UTexture::CleanupCachedRunningPlatformData()
 	}
 }
 
-
 void UTexture::SerializeCookedPlatformData(FArchive& Ar, const bool bSerializeMipData)
 {
 	if (IsTemplate() )
@@ -4286,6 +4286,56 @@ void UTexture::SerializeCookedPlatformData(FArchive& Ar, const bool bSerializeMi
 	}
 }
 
+#if WITH_EDITOR
+namespace UE::TextureBuildUtilities
+{
+
+bool TryWriteCookDeterminismDiagnostics(FCbWriter& Writer, UTexture* Texture, const ITargetPlatform* TargetPlatform)
+{
+	if (!TargetPlatform->AllowAudioVisualData())
+	{
+		return false;
+	}
+	TMap<FString, FTexturePlatformData*>* CookedPlatformDataPtr = Texture->GetCookedPlatformData();
+	if (!CookedPlatformDataPtr)
+	{
+		return false;
+	}
+
+	ETextureEncodeSpeed EncodeSpeed = Texture->GetDesiredEncodeSpeed();
+	TArray<TArray<FTextureBuildSettings>> BuildSettingsToCacheFetchOrBuild;
+	if (EncodeSpeed == ETextureEncodeSpeed::FinalIfAvailable)
+	{
+		EncodeSpeed = ETextureEncodeSpeed::Fast;
+	}
+	FTextureBuildSettings BuildSettingsFetchOrBuild;
+	GetTextureBuildSettings(*Texture, TargetPlatform->GetTextureLODSettings(), *TargetPlatform, EncodeSpeed, BuildSettingsFetchOrBuild, nullptr);
+	GetBuildSettingsPerFormat(*Texture, BuildSettingsFetchOrBuild, nullptr, TargetPlatform, EncodeSpeed, BuildSettingsToCacheFetchOrBuild, nullptr);
+
+	if (BuildSettingsToCacheFetchOrBuild.IsEmpty())
+	{
+		return false;
+	}
+
+	Writer.BeginObject();
+	Writer.BeginArray("BuildSettings");
+	for (int32 SettingIndex = 0; SettingIndex < BuildSettingsToCacheFetchOrBuild.Num(); SettingIndex++)
+	{
+		// CookedPlatformData is keyed off of the fetchorbuild key.
+		FString DerivedDataKeyFetchOrBuild;
+		GetTextureDerivedDataKey(*Texture, BuildSettingsToCacheFetchOrBuild[SettingIndex].GetData(), DerivedDataKeyFetchOrBuild);
+		Writer.BeginObject();
+		Writer << "DerivedDataKey" << DerivedDataKeyFetchOrBuild;
+		Writer.EndObject();
+	}
+	Writer.EndArray();
+	Writer.EndObject();
+	return true;
+}
+
+}
+#endif
+
 int32 UTexture::GMinTextureResidentMipCount = NUM_INLINE_DERIVED_MIPS;
 
 void UTexture::SetMinTextureResidentMipCount(int32 InMinTextureResidentMipCount)
diff --git a/Engine/Source/Runtime/Engine/Private/TextureDerivedDataBuildUtils.h b/Engine/Source/Runtime/Engine/Private/TextureDerivedDataBuildUtils.h
--- a/Engine/Source/Runtime/Engine/Private/TextureDerivedDataBuildUtils.h
+++ b/Engine/Source/Runtime/Engine/Private/TextureDerivedDataBuildUtils.h
@@ -8,11 +8,22 @@
 #include "DerivedDataSharedStringFwd.h"
 
 class FCbObject;
+class FCbWriter;
 class FName;
+class ITargetPlatform;
 class UTexture;
 struct FTextureBuildSettings;
 
 UE::DerivedData::FUtf8SharedString FindTextureBuildFunction(FName TextureFormatName);
 FCbObject SaveTextureBuildSettings(const UTexture& Texture, const FTextureBuildSettings& BuildSettings, int32 LayerIndex, bool bUseCompositeTexture);
 
+#if WITH_EDITOR
+namespace UE::TextureBuildUtilities
+{
+
+bool TryWriteCookDeterminismDiagnostics(FCbWriter& Writer, UTexture* Texture, const ITargetPlatform* TargetPlatform);
+
+} // namespace UE::TextureBuildUtilities
+#endif
+
 #endif // WITH_EDITOR
