Cooker: Fix memory leak of requestcluster vertices. Their struct memory was being freed but their destructor was not being called, and this leaked data they allocated and stored in TUniquePtr.
#rnx

diff --git a/Engine/Source/Editor/UnrealEd/Private/Cooker/CookRequestCluster.cpp b/Engine/Source/Editor/UnrealEd/Private/Cooker/CookRequestCluster.cpp
--- a/Engine/Source/Editor/UnrealEd/Private/Cooker/CookRequestCluster.cpp
+++ b/Engine/Source/Editor/UnrealEd/Private/Cooker/CookRequestCluster.cpp
@@ -731,6 +731,16 @@ FRequestCluster::FGraphSearch::~FGraphSearch()
 		constexpr double WaitTime = 1.0;
 		WaitForAsyncQueue(WaitTime);
 	}
+
+	// Call the FVertexData destructors, but do not bother calling DeleteElement or Free on the VertexAllocator
+	// since we are destructing the VertexAllocator.
+	for (TPair<FName, FVertexData*>& VertexPair : this->Vertices)
+	{
+		FVertexData* VertexData = VertexPair.Value;
+		VertexData->~FVertexData();
+	}
+	// Empty frees the struct memory for each FVertexData we allocated, but it does not call the destructor.
+	VertexAllocator.Empty();
 }
 
 void FRequestCluster::FGraphSearch::OnNewReachablePlatforms(FPackageData* PackageData)
diff --git a/Engine/Source/Editor/UnrealEd/Private/Cooker/TypedBlockAllocator.h b/Engine/Source/Editor/UnrealEd/Private/Cooker/TypedBlockAllocator.h
--- a/Engine/Source/Editor/UnrealEd/Private/Cooker/TypedBlockAllocator.h
+++ b/Engine/Source/Editor/UnrealEd/Private/Cooker/TypedBlockAllocator.h
@@ -21,6 +21,9 @@
  * 
  * This is the base version and does not keep a free list. Freed elements will never be
  * reused and will only be freed by Empty() or destructor.
+ * In this base class, destructors for allocated elements are not automatically called;
+ * users of the allocator must call destructor manually or must call DeleteElement on each
+ * element they called NewElement on.
  */
 template <typename ElementType>
 class TTypedBlockAllocatorBase
@@ -102,6 +105,10 @@ protected:
 /**
  * A TTypedAllocatorBase with a freelist. Freed elements are destructed, but their memory
  * is reused and the constructor called again on it in future allocation.
+ *
+ * In this subclass, destructors for allocated elements are not automatically called;
+ * users of the allocator must call destructor manually or must call DeleteElement on each
+ * element they called NewElement on.
  */
 template <typename ElementType>
 class TTypedBlockAllocatorFreeList : public TTypedBlockAllocatorBase<ElementType>
@@ -131,6 +138,9 @@ protected:
  * A TTypedAllocatorBase with a resetlist. Freed elements are not destructed until empty is called. When NewElement is
  * called and the reset list is available, a previously-freed element is returned without having had the destructor and
  * a second call to the constructor called on it.
+ * 
+ * Destructors for elements added via NewElement are called automatically from Empty or from the allocator's
+ * destructor.
  */
 template <typename ElementType>
 class TTypedBlockAllocatorResetList : public TTypedBlockAllocatorBase<ElementType>
