Cooker: Fix performance problem when diagnosing garbage collection impact; array container was causing O(N*M) costs.
#rnx

diff --git a/Engine/Source/Editor/UnrealEd/Classes/CookOnTheSide/CookOnTheFlyServer.h b/Engine/Source/Editor/UnrealEd/Classes/CookOnTheSide/CookOnTheFlyServer.h
--- a/Engine/Source/Editor/UnrealEd/Classes/CookOnTheSide/CookOnTheFlyServer.h
+++ b/Engine/Source/Editor/UnrealEd/Classes/CookOnTheSide/CookOnTheFlyServer.h
@@ -1575,7 +1575,7 @@ private:
 	/** Used during garbagecolletion: a flat array of all the elements in UPackage::SoftGCPackageToObjectList arrayviews. */
 	TArray<UObject*> SoftGCPackageToObjectListBuffer;
 	/** Packages that were expected to be freed by the last Soft GC and we expect not to load again. */
-	TArray<FName> ExpectedFreedPackageNames;
+	TSet<FName> ExpectedFreedPackageNames;
 
 	UE::Cook::FPackageData* SavingPackageData = nullptr;
 	/** Helper struct for running cooking in diagnostic modes */
diff --git a/Engine/Source/Editor/UnrealEd/Private/CookOnTheFlyServer.cpp b/Engine/Source/Editor/UnrealEd/Private/CookOnTheFlyServer.cpp
--- a/Engine/Source/Editor/UnrealEd/Private/CookOnTheFlyServer.cpp
+++ b/Engine/Source/Editor/UnrealEd/Private/CookOnTheFlyServer.cpp
@@ -5392,7 +5392,7 @@ void UCookOnTheFlyServer::PreGarbageCollect()
 				}
 			}
 		}
-		ExpectedFreedPackageNames.Reset();
+		ExpectedFreedPackageNames.Empty(PackageTracker->NumLoadedPackages());
 		PackageTracker->ForEachLoadedPackage(
 			[this, &GCKeepPackagesSet](UPackage* Package)
 			{
@@ -5536,6 +5536,10 @@ void UCookOnTheFlyServer::EvaluateGarbageCollectionResults(bool bWasDueToOOM, bo
 {
 	using namespace UE::Cook;
 
+	ON_SCOPE_EXIT
+	{
+		ExpectedFreedPackageNames.Empty();
+	};
 	bWarnedExceededMaxMemoryWithinGCCooldown = false;
 	LastGCTime = FPlatformTime::Seconds();
 	bool bWasSoftGC = ResultFlags & COSR_RequiresGC_Soft_OOM;
diff --git a/Engine/Source/Editor/UnrealEd/Private/Cooker/PackageTracker.h b/Engine/Source/Editor/UnrealEd/Private/Cooker/PackageTracker.h
--- a/Engine/Source/Editor/UnrealEd/Private/Cooker/PackageTracker.h
+++ b/Engine/Source/Editor/UnrealEd/Private/Cooker/PackageTracker.h
@@ -213,7 +213,12 @@ public:
 			Function(Package);
 		}
 	}
-	void AddExpectedNeverLoadPackages(TArrayView<FName> PackageNames)
+	int32 NumLoadedPackages()
+	{
+		FReadScopeLock ScopeLock(Lock);
+		return LoadedPackages.Num();
+	}
+	void AddExpectedNeverLoadPackages(const TSet<FName>& PackageNames)
 	{
 		FWriteScopeLock ScopeLock(Lock);
 		ExpectedNeverLoadPackages.Append(PackageNames);
