IncrementalCook: IterativeValidate: suppress the redundant summary log lines from CookWorkers; we only need the summary on the director.
#rnx

diff --git a/Engine/Source/Editor/UnrealEd/Private/Cooker/IterativeValidatePackageWriter.cpp b/Engine/Source/Editor/UnrealEd/Private/Cooker/IterativeValidatePackageWriter.cpp
--- a/Engine/Source/Editor/UnrealEd/Private/Cooker/IterativeValidatePackageWriter.cpp
+++ b/Engine/Source/Editor/UnrealEd/Private/Cooker/IterativeValidatePackageWriter.cpp
@@ -525,7 +525,8 @@ void FIterativeValidatePackageWriter::Initialize(const FCookInfo& CookInfo)
 	case EPhase::AllInOnePhase:
 		if (CookInfo.bFullBuild)
 		{
-			UE_LOG(LogIterativeValidate, Display,
+			UE_CLOG(COTFS.GetCookMode() != ECookMode::CookWorker,
+				LogIterativeValidate, Display,
 				TEXT("The cook is running non-iteratively. All packages are reported \"modified\" and will be resaved as in a normal cook."));
 			bReadOnly = false;
 		}
@@ -537,7 +538,8 @@ void FIterativeValidatePackageWriter::Initialize(const FCookInfo& CookInfo)
 	case EPhase::Phase1:
 		if (CookInfo.bFullBuild)
 		{
-			UE_LOG(LogIterativeValidate, Display,
+			UE_CLOG(COTFS.GetCookMode() != ECookMode::CookWorker,
+				LogIterativeValidate, Display,
 				TEXT("The cook is running non-iteratively. All packages are reported \"modified\" and will be resaved during the final IterativeValidate phase."));
 		}
 		bReadOnly = false;
@@ -545,7 +547,8 @@ void FIterativeValidatePackageWriter::Initialize(const FCookInfo& CookInfo)
 	case EPhase::Phase2:
 		if (CookInfo.bFullBuild)
 		{
-			UE_LOG(LogIterativeValidate, Display,
+			UE_CLOG(COTFS.GetCookMode() != ECookMode::CookWorker, 
+				LogIterativeValidate, Display,
 				TEXT("The cook is running non-iteratively. Packages that were iteratively skipped and found valid will be resaved anyway."));
 		}
 		bReadOnly = false;
@@ -637,26 +640,31 @@ void FIterativeValidatePackageWriter::BeginCook(const FCookInfo& Info)
 	case EPhase::AllInOnePhase:
 		if (bReadOnly)
 		{
-			UE_LOG(LogIterativeValidate, Display,
+			UE_CLOG(COTFS.GetCookMode() != ECookMode::CookWorker,
+				LogIterativeValidate, Display,
 				TEXT("-IterativeValidateAllowWrite not present, read-only mode. Running -diffonly on all packages that were found to be iteratively unmodified."));
 		}
 		else
 		{
-			UE_LOG(LogIterativeValidate, Display,
+			UE_CLOG(COTFS.GetCookMode() != ECookMode::CookWorker,
+				LogIterativeValidate, Display,
 				TEXT("-IterativeValidateAllowWrite is present, writable mode. Resaving packages as in a normal cook, but also running -diffonly on all packages that were found to be iteratively unmodified."));
 		}
 		if (Info.bFullBuild)
 		{
-			UE_LOG(LogIterativeValidate, Error,
+			UE_CLOG(COTFS.GetCookMode() != ECookMode::CookWorker,
+				LogIterativeValidate, Error,
 				TEXT("IterativeValidate was bypassed on this run; it is a full cook and all packages are marked iteratively modified."));
 		}
 		break;
 	case EPhase::Phase1:
-		UE_LOG(LogIterativeValidate, Display,
+		UE_CLOG(COTFS.GetCookMode() != ECookMode::CookWorker,
+			LogIterativeValidate, Display,
 			TEXT("Phase1: running -diffonly and a resave on all packages discovered to be iteratively unmodified."));
 		if (Info.bFullBuild)
 		{
-			UE_LOG(LogIterativeValidate, Error,
+			UE_CLOG(COTFS.GetCookMode() != ECookMode::CookWorker,
+				LogIterativeValidate, Error,
 				TEXT("IterativeValidate was bypassed on this run; it is a full cook and all packages are marked iteratively modified."));
 		}
 		break;
@@ -664,11 +672,13 @@ void FIterativeValidatePackageWriter::BeginCook(const FCookInfo& Info)
 		{
 			Load();
 			FStatusCounts StatusCounts = CountPackagesByStatus();
-			UE_LOG(LogIterativeValidate, Display,
+			UE_CLOG(COTFS.GetCookMode() != ECookMode::CookWorker,
+				LogIterativeValidate, Display,
 				TEXT("Phase2: %d packages were found during Phase1 to be iteratively unmodified but had differences. "
 					"Running -diffonly on them again to check whether the differences are due to indeterminism or to IterativeFalsePositives."), 
 					StatusCounts[EPackageStatus::DeclaredUnmodified_FoundModified_IndeterminismOrFalsePositive]);
-			UE_LOG(LogIterativeValidate, Display,
+			UE_CLOG(COTFS.GetCookMode() != ECookMode::CookWorker,
+				LogIterativeValidate, Display,
 				TEXT("%d packages were found during Phase1 to be modified or new and will be resaved."),
 				StatusCounts[EPackageStatus::DeclaredModified_WillNotVerify]);
 			break;
@@ -691,7 +701,8 @@ void FIterativeValidatePackageWriter::EndCook(const FCookInfo& Info)
 	{
 		int32 DetectedUnmodified = StatusCounts[EPackageStatus::DeclaredUnmodified_ConfirmedUnmodified]
 			+ StatusCounts[EPackageStatus::DeclaredUnmodified_FoundModified_FalsePositive];
-		UE_LOG(LogIterativeValidate, Display,
+		UE_CLOG(COTFS.GetCookMode() != ECookMode::CookWorker,
+			LogIterativeValidate, Display,
 			TEXT("Modified: %d. DetectedUnmodified: %d. ValidatedUnmodified: %d. IterativeSkipFalsePositive: %d."),
 			StatusCounts[EPackageStatus::DeclaredModified_WillNotVerify], 
 			DetectedUnmodified,
@@ -702,11 +713,13 @@ void FIterativeValidatePackageWriter::EndCook(const FCookInfo& Info)
 			StatusCounts[EPackageStatus::DeclaredUnmodified_FoundModified_FalsePositive]);
 		if (StatusCounts[EPackageStatus::DeclaredUnmodified_FoundModified_FalsePositive] > 0)
 		{
-			UE_LOG(LogIterativeValidate, Error, TEXT("%s"), *Message);
+			UE_CLOG(COTFS.GetCookMode() != ECookMode::CookWorker,
+				LogIterativeValidate, Error, TEXT("%s"), *Message);
 		}
 		else
 		{
-			UE_LOG(LogIterativeValidate, Display, TEXT("%s"), *Message);
+			UE_CLOG(COTFS.GetCookMode() != ECookMode::CookWorker,
+				LogIterativeValidate, Display, TEXT("%s"), *Message);
 		}
 		break;
 	}
@@ -714,7 +727,8 @@ void FIterativeValidatePackageWriter::EndCook(const FCookInfo& Info)
 	{
 		int32 DetectedUnmodified = StatusCounts[EPackageStatus::DeclaredUnmodified_ConfirmedUnmodified]
 			+ StatusCounts[EPackageStatus::DeclaredUnmodified_FoundModified_IndeterminismOrFalsePositive];
-		UE_LOG(LogIterativeValidate, Display,
+		UE_CLOG(COTFS.GetCookMode() != ECookMode::CookWorker,
+			LogIterativeValidate, Display,
 			TEXT("Modified: %d. DetectedUnmodified: %d. ValidatedUnmodified: %d. IterativeSkipFalsePositiveOrIndeterminism: %d."),
 			StatusCounts[EPackageStatus::DeclaredModified_WillNotVerify],
 			DetectedUnmodified,
@@ -729,7 +743,8 @@ void FIterativeValidatePackageWriter::EndCook(const FCookInfo& Info)
 			+ StatusCounts[EPackageStatus::DeclaredUnmodified_FoundModified_IndeterminismOrFalsePositive]
 			+ StatusCounts[EPackageStatus::DeclaredUnmodified_FoundModified_Indeterminism]
 			+ StatusCounts[EPackageStatus::DeclaredUnmodified_FoundModified_FalsePositive];
-		UE_LOG(LogIterativeValidate, Display,
+		UE_CLOG(COTFS.GetCookMode() != ECookMode::CookWorker,
+			LogIterativeValidate, Display,
 			TEXT("Modified: %d. DetectedUnmodified: %d. ValidatedUnmodified: %d. Indeterminism: %d. IterativeSkipFalsePositive: %d."),
 			StatusCounts[EPackageStatus::DeclaredModified_WillNotVerify],
 			DetectedUnmodified,
@@ -742,11 +757,13 @@ void FIterativeValidatePackageWriter::EndCook(const FCookInfo& Info)
 			StatusCounts[EPackageStatus::DeclaredUnmodified_FoundModified_FalsePositive]);
 		if (StatusCounts[EPackageStatus::DeclaredUnmodified_FoundModified_FalsePositive] > 0)
 		{
-			UE_LOG(LogIterativeValidate, Error, TEXT("%s"), *Message);
+			UE_CLOG(COTFS.GetCookMode() != ECookMode::CookWorker,
+				LogIterativeValidate, Error, TEXT("%s"), *Message);
 		}
 		else
 		{
-			UE_LOG(LogIterativeValidate, Display, TEXT("%s"), *Message);
+			UE_CLOG(COTFS.GetCookMode() != ECookMode::CookWorker,
+				LogIterativeValidate, Display, TEXT("%s"), *Message);
 		}
 		break;
 	}
