[BaseIteratePackagesCommandlet]
 - Fix deadlock by moving complex code outside of the UObject hash table lock

diff --git a/Engine/Source/Editor/UnrealEd/Private/Commandlets/BaseIteratePackagesCommandlet.cpp b/Engine/Source/Editor/UnrealEd/Private/Commandlets/BaseIteratePackagesCommandlet.cpp
--- a/Engine/Source/Editor/UnrealEd/Private/Commandlets/BaseIteratePackagesCommandlet.cpp
+++ b/Engine/Source/Editor/UnrealEd/Private/Commandlets/BaseIteratePackagesCommandlet.cpp
@@ -705,7 +705,7 @@ void UBaseIteratePackagesCommandlet::LoadAndSaveOnePackage(const FString& Filena
 						check(Package);
 
 						TArray<UObject*> DependantObjects;
-						ForEachObjectWithPackage(Package, [this, &bSavePackage](UObject* Object)
+						ForEachObjectWithPackage(Package, [this, &bSavePackage, &DependantObjects](UObject* Object)
 							{
 								if (!IsValid(Object))
 								{
@@ -713,11 +713,16 @@ void UBaseIteratePackagesCommandlet::LoadAndSaveOnePackage(const FString& Filena
 								}
 								if (!Cast<UMetaData>(Object))
 								{
-									PerformWorldBuilderAdditionalOperations(Object, bSavePackage);
+									DependantObjects.Add(Object);
 								}
 								return true;
 							}, true);
 
+						for (UObject* Object : DependantObjects)
+						{
+							PerformWorldBuilderAdditionalOperations(Object, bSavePackage);
+						}
+
 						if (bSavePackage)
 						{
 							PackagesToSave.Add(Package);
