In UEditorProjectAppearanceSettings::PostInitProperties initialize missing display units.

In BaseEditor.ini set the default value for ImpulseUnits and PositionalImpulseUnits.

diff --git a/Engine/Config/BaseEditor.ini b/Engine/Config/BaseEditor.ini
--- a/Engine/Config/BaseEditor.ini
+++ b/Engine/Config/BaseEditor.ini
@@ -536,6 +536,8 @@ AccelerationUnits=EUnit::CentimetersPerSecondSquared
 TemperatureUnits=EUnit::Celsius
 ForceUnits=EUnit::Newtons
 TorqueUnits=EUnit::NewtonMeters
+ImpulseUnits=EUnit::NewtonSeconds
+PositionalImpulseUnits=EUnit::KilogramCentimeters
 
 [BlueprintSearchSettings]
 AsyncTaskBatchSize=2048
diff --git a/Engine/Source/Editor/UnrealEd/Private/Settings/EditorProjectSettings.cpp b/Engine/Source/Editor/UnrealEd/Private/Settings/EditorProjectSettings.cpp
--- a/Engine/Source/Editor/UnrealEd/Private/Settings/EditorProjectSettings.cpp
+++ b/Engine/Source/Editor/UnrealEd/Private/Settings/EditorProjectSettings.cpp
@@ -180,8 +180,13 @@ void UEditorProjectAppearanceSettings::PostInitProperties()
 	Settings.SetDisplayUnits(EUnitType::Time, TimeUnits);
 	Settings.SetDisplayUnits(EUnitType::Angle, AngleUnits);
 	Settings.SetDisplayUnits(EUnitType::Speed, SpeedUnits);
+	Settings.SetDisplayUnits(EUnitType::AngularSpeed, AngularSpeedUnits);
+	Settings.SetDisplayUnits(EUnitType::Acceleration, AccelerationUnits);
 	Settings.SetDisplayUnits(EUnitType::Temperature, TemperatureUnits);
 	Settings.SetDisplayUnits(EUnitType::Force, ForceUnits);
+	Settings.SetDisplayUnits(EUnitType::Torque, TorqueUnits);
+	Settings.SetDisplayUnits(EUnitType::Impulse, ImpulseUnits);
+	Settings.SetDisplayUnits(EUnitType::PositionalImpulse, PositionalImpulseUnits);
 
 	Settings.SetShouldDisplayUnits(bDisplayUnits);
 }
