GitHub 11914 : [MultiPlayerTesting][VRPreview]: All windows should have border
-All the window styles need to match for multi-vr-pie support, where the hmd can drive multiple pie windows or go between them at runtime.

diff --git a/Engine/Source/Editor/UnrealEd/Private/PlayLevel.cpp b/Engine/Source/Editor/UnrealEd/Private/PlayLevel.cpp
--- a/Engine/Source/Editor/UnrealEd/Private/PlayLevel.cpp
+++ b/Engine/Source/Editor/UnrealEd/Private/PlayLevel.cpp
@@ -3273,7 +3273,10 @@ TSharedRef<SPIEViewport> UEditorEngine::GeneratePIEViewportWindow(const FRequest
 	
 	// VR Preview overrides window location.
 	const bool bVRPreview = InWorldContext.bIsPrimaryPIEInstance && InSessionParams.SessionPreviewTypeOverride.Get(EPlaySessionPreviewType::NoPreview) == EPlaySessionPreviewType::VRPreview;
-	bool bUseOSWndBorder = bVRPreview;
+
+	// Because we could switch primary PIE on the fly, we should make all PIE windows with the same UI style
+	bool bUseOSWndBorder = InSessionParams.SessionPreviewTypeOverride.Get(EPlaySessionPreviewType::NoPreview) == EPlaySessionPreviewType::VRPreview;
+
 	if (bVRPreview)
 	{
 		bCenterNewWindowOverride = true;
