Remove usage of REN_ForceNoResetLoaders from the codebase since the flag has been deprecated. These removals all pertain to conditional use of the flag which usually ends up cleaning up the code a bit more than only removing the flag.

#rnx

diff --git a/Engine/Source/Editor/BlueprintGraph/Private/EdGraphSchema_K2.cpp b/Engine/Source/Editor/BlueprintGraph/Private/EdGraphSchema_K2.cpp
--- a/Engine/Source/Editor/BlueprintGraph/Private/EdGraphSchema_K2.cpp
+++ b/Engine/Source/Editor/BlueprintGraph/Private/EdGraphSchema_K2.cpp
@@ -6548,14 +6548,12 @@ UEdGraphNode* UEdGraphSchema_K2::CreateSubstituteNode(UEdGraphNode* Node, const
 			FBlueprintEditorUtils::GetFunctionNameList(Blueprint, ExistingNamesInUse);
 			FBlueprintEditorUtils::GetClassVariableList(Blueprint, ExistingNamesInUse);
 
-			const ERenameFlags RenameFlags = (Blueprint->bIsRegeneratingOnLoad ? REN_ForceNoResetLoaders : 0);
-
 			// Allow the old object name to be used in the graph
 			FName ObjName = EventNode->GetFName();
 			UObject* Found = FindObject<UObject>(EventNode->GetOuter(), *ObjName.ToString());
 			if(Found)
 			{
-				Found->Rename(NULL, NULL, REN_DontCreateRedirectors | RenameFlags | ((IsAsyncLoading() || Found->HasAnyFlags(RF_NeedLoad | RF_NeedPostLoad | RF_NeedPostLoadSubobjects)) ? REN_ForceNoResetLoaders : RF_NoFlags));
+				Found->Rename(NULL, NULL, REN_DontCreateRedirectors);
 			}
 
 			// Create a custom event node to replace the original event node imported from text
diff --git a/Engine/Source/Editor/BlueprintGraph/Private/K2Node_Composite.cpp b/Engine/Source/Editor/BlueprintGraph/Private/K2Node_Composite.cpp
--- a/Engine/Source/Editor/BlueprintGraph/Private/K2Node_Composite.cpp
+++ b/Engine/Source/Editor/BlueprintGraph/Private/K2Node_Composite.cpp
@@ -272,7 +272,7 @@ void UK2Node_Composite::RenameBoundGraphCloseToName(const FString& Name)
 		{
 			//Name is available
 			UBlueprint* BP = FBlueprintEditorUtils::FindBlueprintForGraphChecked(BoundGraph);
-			BoundGraph->Rename(*NewName, BoundGraph->GetOuter(), (BP->bIsRegeneratingOnLoad ? REN_ForceNoResetLoaders : 0) | REN_DontCreateRedirectors);
+			BoundGraph->Rename(*NewName, BoundGraph->GetOuter(), REN_DontCreateRedirectors);
 			bFoundName = true;
 		}
 	}
diff --git a/Engine/Source/Editor/BlueprintGraph/Private/K2Node_CustomEvent.cpp b/Engine/Source/Editor/BlueprintGraph/Private/K2Node_CustomEvent.cpp
--- a/Engine/Source/Editor/BlueprintGraph/Private/K2Node_CustomEvent.cpp
+++ b/Engine/Source/Editor/BlueprintGraph/Private/K2Node_CustomEvent.cpp
@@ -290,7 +290,7 @@ void UK2Node_CustomEvent::RenameCustomEventCloseToName(int32 StartIndex)
 		{
 			UBlueprint* Blueprint = GetBlueprint();
 			CustomFunctionName = FName(NewName.GetCharArray().GetData());
-			Rename(*NewName, GetOuter(), (Blueprint->bIsRegeneratingOnLoad ? REN_ForceNoResetLoaders : 0) | REN_DontCreateRedirectors);
+			Rename(*NewName, GetOuter(), REN_DontCreateRedirectors);
 			bFoundName = true;
 		}
 	}
diff --git a/Engine/Source/Editor/BlueprintGraph/Private/K2Node_Timeline.cpp b/Engine/Source/Editor/BlueprintGraph/Private/K2Node_Timeline.cpp
--- a/Engine/Source/Editor/BlueprintGraph/Private/K2Node_Timeline.cpp
+++ b/Engine/Source/Editor/BlueprintGraph/Private/K2Node_Timeline.cpp
@@ -209,7 +209,7 @@ void UK2Node_Timeline::DestroyNode()
 		FBlueprintEditorUtils::RemoveTimeline(Blueprint, Timeline, true);
 
 		// Move template object out of the way so that we can potentially create a timeline with the same name either through a paste or a new timeline action
-		Timeline->Rename(NULL, GetTransientPackage(), (Blueprint->bIsRegeneratingOnLoad ? REN_ForceNoResetLoaders : REN_None));
+		Timeline->Rename(NULL, GetTransientPackage(), REN_None);
 	}
 
 	Super::DestroyNode();
diff --git a/Engine/Source/Editor/KismetCompiler/Private/KismetCompiler.cpp b/Engine/Source/Editor/KismetCompiler/Private/KismetCompiler.cpp
--- a/Engine/Source/Editor/KismetCompiler/Private/KismetCompiler.cpp
+++ b/Engine/Source/Editor/KismetCompiler/Private/KismetCompiler.cpp
@@ -654,7 +654,7 @@ void FKismetCompilerContext::CleanAndSanitizeClass(UBlueprintGeneratedClass* Cla
 	SetNewClass( ClassToClean );
 	InOldCDO = ClassToClean->ClassDefaultObject; // we don't need to create the CDO at this point
 	
-	const ERenameFlags RenFlags = REN_DontCreateRedirectors |  ((bRecompilingOnLoad) ? REN_ForceNoResetLoaders : 0) | REN_NonTransactional | REN_DoNotDirty;
+	const ERenameFlags RenFlags = REN_DontCreateRedirectors | REN_NonTransactional | REN_DoNotDirty;
 
 	if( InOldCDO )
 	{
@@ -728,7 +728,7 @@ void FKismetCompilerContext::CleanAndSanitizeClass(UBlueprintGeneratedClass* Cla
 
         // Rename will remove the renamed object's linker when moving to a new package so invalidate the export beforehand
 		FLinkerLoad::InvalidateExport(CurrSubObj);
-		CurrSubObj->Rename(*NewSubobjectName.ToString(), TransientClass, RenFlags | REN_ForceNoResetLoaders);
+		CurrSubObj->Rename(*NewSubobjectName.ToString(), TransientClass, RenFlags);
 	}
 
 	// Purge the class to get it back to a "base" state
@@ -4052,7 +4052,7 @@ void FKismetCompilerContext::CreateAndProcessUbergraph()
 		{
 			if (OldEventGraph)
 			{
-				OldEventGraph->Rename(NULL, GetTransientPackage(), (Blueprint->bIsRegeneratingOnLoad) ? REN_ForceNoResetLoaders : 0);
+				OldEventGraph->Rename(NULL, GetTransientPackage());
 			}
 		}
 	}
@@ -4458,7 +4458,7 @@ void FKismetCompilerContext::ProcessOneFunctionGraph(UEdGraph* SourceGraph, bool
 		FString FunctionGraphName = SourceGraph->GetName() + TEXT("_MERGED");
 
 		ERenameFlags RenameFlags =
-			(REN_NonTransactional | REN_DoNotDirty | REN_DontCreateRedirectors | REN_ForceNoResetLoaders);
+			(REN_NonTransactional | REN_DoNotDirty | REN_DontCreateRedirectors);
 
 		if (UEdGraph* ExistingGraph = FindObject<UEdGraph>(Blueprint, *FunctionGraphName))
 		{
diff --git a/Engine/Source/Editor/KismetCompiler/Private/KismetCompilerMisc.cpp b/Engine/Source/Editor/KismetCompiler/Private/KismetCompilerMisc.cpp
--- a/Engine/Source/Editor/KismetCompiler/Private/KismetCompilerMisc.cpp
+++ b/Engine/Source/Editor/KismetCompiler/Private/KismetCompilerMisc.cpp
@@ -528,7 +528,7 @@ void FKismetCompilerUtilities::ConsignToOblivion(UClass* OldClass, bool bForceNo
 		}
 
 		const FString BaseName = FString::Printf(TEXT("DEADCLASS_%s_C_%d"), *OldClass->ClassGeneratedBy->GetName(), ConsignToOblivionCounter++);
-		OldClass->Rename(*BaseName, GetTransientPackage(), (REN_DontCreateRedirectors|REN_NonTransactional|(bForceNoResetLoaders ? REN_ForceNoResetLoaders : 0)));
+		OldClass->Rename(*BaseName, GetTransientPackage(), (REN_DontCreateRedirectors | REN_NonTransactional));
 
 		// Make sure MetaData doesn't have any entries to the class we just renamed out of package
 		OwnerOutermost->GetMetaData()->RemoveMetaDataOutsidePackage();
diff --git a/Engine/Source/Editor/UMGEditor/Private/WidgetBlueprintCompiler.cpp b/Engine/Source/Editor/UMGEditor/Private/WidgetBlueprintCompiler.cpp
--- a/Engine/Source/Editor/UMGEditor/Private/WidgetBlueprintCompiler.cpp
+++ b/Engine/Source/Editor/UMGEditor/Private/WidgetBlueprintCompiler.cpp
@@ -301,7 +301,7 @@ void FWidgetBlueprintCompilerContext::CleanAndSanitizeClass(UBlueprintGeneratedC
         // Rename will remove the renamed object's linker when moving to a new package so invalidate the export beforehand
 		FLinkerLoad::InvalidateExport(ObjectToRename);
 
-		const ERenameFlags RenFlags = REN_DontCreateRedirectors | (bRecompilingOnLoad ? REN_ForceNoResetLoaders : 0) | REN_NonTransactional | REN_DoNotDirty;
+		const ERenameFlags RenFlags = REN_DontCreateRedirectors | REN_NonTransactional | REN_DoNotDirty;
 
 		if (BaseName.IsNone())
 		{
@@ -779,10 +779,10 @@ void FWidgetBlueprintCompilerContext::FixAbandonedWidgetTree(UWidgetBlueprint* W
 			{
 				AbandonedWidgetTree->ClearFlags(RF_DefaultSubObject);
 				AbandonedWidgetTree->SetFlags(RF_Transient);
-				AbandonedWidgetTree->Rename(nullptr, GetTransientPackage(), REN_DontCreateRedirectors | REN_ForceNoResetLoaders | REN_NonTransactional | REN_DoNotDirty);
+				AbandonedWidgetTree->Rename(nullptr, GetTransientPackage(), REN_DontCreateRedirectors | REN_NonTransactional | REN_DoNotDirty);
 			}
 
-			WidgetTree->Rename(TEXT("WidgetTree"), nullptr, REN_DontCreateRedirectors | REN_ForceNoResetLoaders | REN_NonTransactional | REN_DoNotDirty);
+			WidgetTree->Rename(TEXT("WidgetTree"), nullptr, REN_DontCreateRedirectors | REN_NonTransactional | REN_DoNotDirty);
 			WidgetTree->SetFlags(RF_DefaultSubObject);
 		}
 	}
diff --git a/Engine/Source/Editor/UnrealEd/Private/EdGraphUtilities.cpp b/Engine/Source/Editor/UnrealEd/Private/EdGraphUtilities.cpp
--- a/Engine/Source/Editor/UnrealEd/Private/EdGraphUtilities.cpp
+++ b/Engine/Source/Editor/UnrealEd/Private/EdGraphUtilities.cpp
@@ -432,7 +432,7 @@ void FEdGraphUtilities::RenameGraphCloseToName(UEdGraph* Graph, const FString& B
 		if (Graph->Rename(*NewName, Graph->GetOuter(), REN_Test))
 		{
 			UBlueprint* BP = FBlueprintEditorUtils::FindBlueprintForGraphChecked(Graph);
-			Graph->Rename(*NewName, Graph->GetOuter(), (BP->bIsRegeneratingOnLoad ? REN_ForceNoResetLoaders : 0) | REN_DontCreateRedirectors);
+			Graph->Rename(*NewName, Graph->GetOuter(), REN_DontCreateRedirectors);
 			return;
 		}
 
@@ -446,7 +446,7 @@ void FEdGraphUtilities::RenameGraphToNameOrCloseToName(UEdGraph* Graph, const FS
 	if (Graph->Rename(*DesiredName, Graph->GetOuter(), REN_Test))
 	{
 		UBlueprint* BP = FBlueprintEditorUtils::FindBlueprintForGraphChecked(Graph);
-		Graph->Rename(*DesiredName, Graph->GetOuter(), (BP->bIsRegeneratingOnLoad ? REN_ForceNoResetLoaders : 0) | REN_DontCreateRedirectors);
+		Graph->Rename(*DesiredName, Graph->GetOuter(), REN_DontCreateRedirectors);
 	}
 	else
 	{
diff --git a/Engine/Source/Editor/UnrealEd/Private/Kismet2/BlueprintEditorUtils.cpp b/Engine/Source/Editor/UnrealEd/Private/Kismet2/BlueprintEditorUtils.cpp
--- a/Engine/Source/Editor/UnrealEd/Private/Kismet2/BlueprintEditorUtils.cpp
+++ b/Engine/Source/Editor/UnrealEd/Private/Kismet2/BlueprintEditorUtils.cpp
@@ -1237,7 +1237,7 @@ void FBlueprintEditorUtils::RemoveStaleFunctions(UBlueprintGeneratedClass* Class
 		OrphanedClass->ClassFlags |= CLASS_CompiledFromBlueprint;
 		OrphanedClass->ClassGeneratedBy = Class->ClassGeneratedBy;
 
-		const ERenameFlags RenFlags = REN_DontCreateRedirectors | (Blueprint->bIsRegeneratingOnLoad ? REN_ForceNoResetLoaders : 0) | REN_NonTransactional | REN_DoNotDirty;
+		const ERenameFlags RenFlags = REN_DontCreateRedirectors | REN_NonTransactional | REN_DoNotDirty;
 
 		while (Fn)
 		{
@@ -2168,7 +2168,7 @@ UEdGraph* FBlueprintEditorUtils::CreateNewGraph(UObject* ParentScope, const FNam
 				// Rename the old graph out of the way - this may confuse the user somewhat - and even
 				// break their logic. But name collisions are not avoidable e.g. someone can add
 				// a function to an interface that conflicts with something in a class hierarchy
-				ExistingObject->Rename(nullptr, ExistingObject->GetOuter(), REN_DoNotDirty | REN_ForceNoResetLoaders);
+				ExistingObject->Rename(nullptr, ExistingObject->GetOuter(), REN_DoNotDirty);
 			}
 			else if (ExistingObject->IsA<UObjectRedirector>())
 			{
@@ -2197,7 +2197,7 @@ UEdGraph* FBlueprintEditorUtils::CreateNewGraph(UObject* ParentScope, const FNam
 	// Now move to where we want it to. Workaround to ensure transaction buffer is correctly utilized
 	if (bRename)
 	{
-		NewGraph->Rename(*(GraphName.ToString()), ParentScope, REN_DoNotDirty | REN_ForceNoResetLoaders);
+		NewGraph->Rename(*(GraphName.ToString()), ParentScope, REN_DoNotDirty);
 	}
 	return NewGraph;
 }
@@ -2636,10 +2636,6 @@ void FBlueprintEditorUtils::RenameGraph(UEdGraph* Graph, const FString& NewNameS
 		};
 
 		ERenameFlags RenameFlagsToApply = REN_None;
-		if (Blueprint->bIsRegeneratingOnLoad)
-		{
-			RenameFlagsToApply |= REN_ForceNoResetLoaders;
-		}
 
 		// Macro library graphs are referenced indirectly and resolved at edit/compile time via GUID (see FGraphReference).
 		// However, they will be exported by name at save time, so renaming a macro library graph implies we should also
@@ -2665,7 +2661,7 @@ void FBlueprintEditorUtils::RenameGraph(UEdGraph* Graph, const FString& NewNameS
 			{
 				if (FunctionGraph->GetFName() == OldGraphName)
 				{
-					RenameGraphLambda(FunctionGraph, OldGraphName, NewGraphName, (InChildBP->bIsRegeneratingOnLoad ? REN_ForceNoResetLoaders : 0) | REN_DontCreateRedirectors);
+					RenameGraphLambda(FunctionGraph, OldGraphName, NewGraphName, REN_DontCreateRedirectors);
 				}
 			}
 
@@ -2745,7 +2741,7 @@ void FBlueprintEditorUtils::RenameGraphWithSuggestion(class UEdGraph* Graph, TSh
 	FString NewName = DesiredName;
 	NameValidator->FindValidString(NewName);
 	UBlueprint* BP = FBlueprintEditorUtils::FindBlueprintForGraphChecked(Graph);
-	Graph->Rename(*NewName, Graph->GetOuter(), (BP->bIsRegeneratingOnLoad ? REN_ForceNoResetLoaders : 0) | REN_DontCreateRedirectors);
+	Graph->Rename(*NewName, Graph->GetOuter(), REN_DontCreateRedirectors);
 }
 
 /** 
@@ -7378,7 +7374,7 @@ static void ConformInterfaceByName(UBlueprint* Blueprint, FBPInterfaceDescriptio
 					CurrentGraph->Rename(
 						nullptr,
 						CurrentGraph->GetOuter(),
-						(Blueprint->bIsRegeneratingOnLoad ? REN_ForceNoResetLoaders : 0) | REN_DoNotDirty | REN_DontCreateRedirectors);
+						REN_DoNotDirty | REN_DontCreateRedirectors);
 					// removing from root, standalone, and public is defensive to make sure it is not saved:
 					CurrentGraph->ClearFlags(RF_Standalone | RF_Public);
 					CurrentGraph->RemoveFromRoot();
@@ -7556,7 +7552,7 @@ void FBlueprintEditorUtils::UpdateOutOfDateCompositeWithOuter(UBlueprint* Bluepr
 				if (BoundGraph->GetOuter() != Node)
 				{
 					// change the outer of the BoundGraph to be the composite node instead of the OuterGraph
-					if (false == BoundGraph->Rename(*BoundGraph->GetName(), Node, ((BoundGraph->HasAnyFlags(RF_NeedLoad | RF_NeedPostLoad) ? REN_ForceNoResetLoaders : 0) | REN_DontCreateRedirectors)))
+					if (false == BoundGraph->Rename(*BoundGraph->GetName(), Node, REN_DontCreateRedirectors))
 					{
 						UE_LOG(LogBlueprintDebug, Log, TEXT("CompositeNode: On Blueprint '%s' could not fix Outer() for BoundGraph of composite node '%s'"), *Blueprint->GetPathName(), *Node->GetName());
 					}
@@ -8041,7 +8037,7 @@ bool FBlueprintEditorUtils::RenameTimeline(UBlueprint* Blueprint, const FName Ol
 			{
 				ExistingObject->Rename(*MakeUniqueObjectName(ExistingObject->GetOuter(), ExistingObject->GetClass(), ExistingObject->GetFName()).ToString());
 			}
-			Template->Rename(*NewTemplateName, Template->GetOuter(), (Blueprint->bIsRegeneratingOnLoad ? REN_ForceNoResetLoaders : REN_None));
+			Template->Rename(*NewTemplateName, Template->GetOuter(), REN_None);
 			Blueprint->Timelines.Add(Template);
 
 			// Validate child blueprints and adjust variable names to avoid a potential name collision
diff --git a/Engine/Source/Runtime/Engine/Private/EdGraph/EdGraph.cpp b/Engine/Source/Runtime/Engine/Private/EdGraph/EdGraph.cpp
--- a/Engine/Source/Runtime/Engine/Private/EdGraph/EdGraph.cpp
+++ b/Engine/Source/Runtime/Engine/Private/EdGraph/EdGraph.cpp
@@ -308,7 +308,7 @@ void UEdGraph::MoveNodesToAnotherGraph(UEdGraph* DestinationGraph, bool bIsLoadi
 		{
 			// Let the name be autogenerated, to automatically avoid naming conflicts
 			// Since this graph is always going to come from a cloned source graph, user readable names can come from the remap stored in a MessageLog.
-			Node->Rename(/*NewName=*/ nullptr, /*NewOuter=*/ DestinationGraph, REN_DontCreateRedirectors | (bIsLoading ? REN_ForceNoResetLoaders : 0) | (bInIsCompiling ? REN_NonTransactional : 0));
+			Node->Rename(/*NewName=*/ nullptr, /*NewOuter=*/ DestinationGraph, REN_DontCreateRedirectors | (bInIsCompiling ? REN_NonTransactional : 0));
 	
 			DestinationGraph->Nodes.Add(Node);
 		}
diff --git a/Engine/Source/Runtime/UMG/Private/WidgetBlueprintGeneratedClass.cpp b/Engine/Source/Runtime/UMG/Private/WidgetBlueprintGeneratedClass.cpp
--- a/Engine/Source/Runtime/UMG/Private/WidgetBlueprintGeneratedClass.cpp
+++ b/Engine/Source/Runtime/UMG/Private/WidgetBlueprintGeneratedClass.cpp
@@ -377,7 +377,7 @@ void UWidgetBlueprintGeneratedClass::PurgeClass(bool bRecompilingOnLoad)
 {
 	Super::PurgeClass(bRecompilingOnLoad);
 
-	const ERenameFlags RenFlags = REN_DontCreateRedirectors | ( ( bRecompilingOnLoad ) ? REN_ForceNoResetLoaders : 0 ) | REN_NonTransactional | REN_DoNotDirty;
+	const ERenameFlags RenFlags = REN_DontCreateRedirectors | REN_NonTransactional | REN_DoNotDirty;
 
 	// Remove the old widdget tree.
 	if ( WidgetTree )
