EditorAssetUtils: Fix package path being used where object path should be used

diff --git a/Engine/Source/Editor/UnrealEd/Private/Subsystems/EditorAssetSubsystem.cpp b/Engine/Source/Editor/UnrealEd/Private/Subsystems/EditorAssetSubsystem.cpp
--- a/Engine/Source/Editor/UnrealEd/Private/Subsystems/EditorAssetSubsystem.cpp
+++ b/Engine/Source/Editor/UnrealEd/Private/Subsystems/EditorAssetSubsystem.cpp
@@ -120,15 +120,12 @@ namespace UE::EditorAssetUtils
 			return MakeError(FString::Printf(TEXT("The object '%s' is not an asset."), *Object->GetName()));
 		}
 
-		UPackage* Package = Object->GetPackage();
+		FSoftObjectPath ObjectPath = FSoftObjectPath(Object);
 		FAssetRegistryModule& AssetRegistryModule = FModuleManager::LoadModuleChecked<FAssetRegistryModule>("AssetRegistry");
-PRAGMA_DISABLE_DEPRECATION_WARNINGS
-		// This should likely get GetAssetsByPackagename but the impact of changing it should be checked
-		FAssetData AssetData = AssetRegistryModule.Get().GetAssetByObjectPath(Package->GetName());
-PRAGMA_ENABLE_DEPRECATION_WARNINGS
+		FAssetData AssetData = AssetRegistryModule.Get().GetAssetByObjectPath(ObjectPath);
 		if (!AssetData.IsValid())
 		{
-			return MakeError(FString::Printf(TEXT("The AssetData '%s' could not be found in the Asset Registry."), *Package->GetName()));
+			return MakeError(FString::Printf(TEXT("The AssetData '%s' could not be found in the Asset Registry."), *Object->GetPathName()));
 		}
 
 		return MakeValue();
