Fix ODSC not reloading materials.
Even if we do release the materials and call GEngine->ForceGarbageCollection, we actually don't end up garbage collecting in ODSC server mode. Instead, COTF seems to provide a mechanism to call GC with UE::Cook::FTickStackData.
When GC materials, we lose the ability to merge materialshadermaps together. We need to keep the compiling shadermap alive, not the one sent back to ODSC since it's incomplete and thus not put in the DDC cache. This means that the next time the material is loaded, the compilation environment is already there, and we can merge the new shader permutations with the previous ones.
Destroying the material calls FMaterial::PrepareDestroy_GameThread -> FMaterialShaderMap::RemoveCompilingDependency which only checks if there are still dependencies left, and destroys it if there's no material left in CompilingMaterialDependencies.
This means that only keeping a TRefCountPtr on this materialshadermap is not enough to keep it around. Adding an additional refcount on the compiling shadermap does the trick to keep it alive. FODSCClientData is responsible for calling AddCompilingMaterialExternalDependency / RemoveCompilingMaterialExternalDependency
The CVars odsc.shadermaps.lifetime / odsc.shadermaps.numbeforegc allow to eventually get rid of materialshadermaps to keep memory stable

#rnx

diff --git a/Engine/Source/Editor/UnrealEd/Classes/CookOnTheSide/CookOnTheFlyServer.h b/Engine/Source/Editor/UnrealEd/Classes/CookOnTheSide/CookOnTheFlyServer.h
--- a/Engine/Source/Editor/UnrealEd/Classes/CookOnTheSide/CookOnTheFlyServer.h
+++ b/Engine/Source/Editor/UnrealEd/Classes/CookOnTheSide/CookOnTheFlyServer.h
@@ -565,7 +565,7 @@ private:
 	void PollGarbageCollection(UE::Cook::FTickStackData& StackData);
 	void PollQueuedCancel(UE::Cook::FTickStackData& StackData);
 	void WaitForAsync(UE::Cook::FTickStackData& StackData);
-	void TickRecompileShaderRequestsPrivate();
+	void TickRecompileShaderRequestsPrivate(UE::Cook::FTickStackData& StackData);
 
 public:
 
diff --git a/Engine/Source/Editor/UnrealEd/Private/CookOnTheFlyServer.cpp b/Engine/Source/Editor/UnrealEd/Private/CookOnTheFlyServer.cpp
--- a/Engine/Source/Editor/UnrealEd/Private/CookOnTheFlyServer.cpp
+++ b/Engine/Source/Editor/UnrealEd/Private/CookOnTheFlyServer.cpp
@@ -677,15 +677,22 @@ bool UCookOnTheFlyServer::StartCookOnTheFly(FCookOnTheFlyStartupOptions InCookOn
 					*Ar << RecompileData;
 				}
 
+				const void* ConnectionPtr = &Connection;
 				RecompileData.LoadedMaterialsToRecompile = &LoadedMaterialsToRecompile;
 
-				const void* ConnectionPtr = &Connection;
 				FEventRef RecompileCompletedEvent;
-				UE::Cook::FRecompileShaderCompletedCallback RecompileCompleted = [this, &RecompileCompletedEvent, ConnectionPtr, &LoadedMaterialsToRecompile]()
+				UE::Cook::FRecompileShaderCompletedCallback RecompileCompleted = [this, &RecompileCompletedEvent, ConnectionPtr, &LoadedMaterialsToRecompile, RecompileDataCommandType = RecompileData.CommandType]()
 				{
 					if (ODSCClientData)
 					{
-						ODSCClientData->KeepClientPersistentData(ConnectionPtr, LoadedMaterialsToRecompile);
+						if (RecompileDataCommandType == ODSCRecompileCommand::ResetMaterialCache)
+						{
+							ODSCClientData->FlushClientPersistentData(ConnectionPtr);
+						}
+						else
+						{
+							ODSCClientData->KeepClientPersistentData(ConnectionPtr, LoadedMaterialsToRecompile);
+						}
 					}
 					LoadedMaterialsToRecompile.Empty();
 
@@ -2287,7 +2294,7 @@ void UCookOnTheFlyServer::InitializePollables()
 	}
 	else
 	{
-		RecompileRequestsPollable = new FPollable(TEXT("RecompileShaderRequests"), FPollable::EManualTrigger(), [this](FTickStackData&) { TickRecompileShaderRequestsPrivate(); });
+		RecompileRequestsPollable = new FPollable(TEXT("RecompileShaderRequests"), FPollable::EManualTrigger(), [this](FTickStackData& TickStackData) { TickRecompileShaderRequestsPrivate(TickStackData); });
 		Pollables.Add(FPollableQueueKey(RecompileRequestsPollable));
 		Pollables.Emplace(new FPollable(TEXT("RequestManager"), 0.5f, 0.5f, [this](FTickStackData&) { TickRequestManager(); }));
 
@@ -5947,19 +5954,32 @@ void UCookOnTheFlyServer::TickRequestManager()
 	}
 }
 
-void UCookOnTheFlyServer::TickRecompileShaderRequestsPrivate()
+void UCookOnTheFlyServer::TickRecompileShaderRequestsPrivate(UE::Cook::FTickStackData& StackData)
 {
 	// try to pull off a request
 	UE::Cook::FRecompileShaderRequest RecompileShaderRequest;
+	bool bProcessedRequests = false;
 	if (PackageTracker->RecompileRequests.Dequeue(&RecompileShaderRequest))
 	{
-		RecompileShadersForRemote(RecompileShaderRequest.RecompileArguments, GetSandboxDirectory(RecompileShaderRequest.RecompileArguments.PlatformName));
+		if (RecompileShaderRequest.RecompileArguments.CommandType != ODSCRecompileCommand::ResetMaterialCache)
+		{
+			RecompileShadersForRemote(RecompileShaderRequest.RecompileArguments, GetSandboxDirectory(RecompileShaderRequest.RecompileArguments.PlatformName));
+		}
+
 		RecompileShaderRequest.CompletionCallback();
+		bProcessedRequests = true;
 	}
 	if (PackageTracker->RecompileRequests.HasItems())
 	{
 		RecompileRequestsPollable->Trigger(*this);
 	}
+
+	if (bProcessedRequests)
+	{
+		// Ask for GC to run again when we processed some shaders requests to ensure material get evicted and we don't keep their package open
+		StackData.ResultFlags |= COSR_RequiresGC | COSR_RequiresGC_IdleTimer | COSR_YieldTick;
+	}
+
 }
 
 class FDiffModeCookServerUtils
diff --git a/Engine/Source/Editor/UnrealEd/Private/Cooker/OnDemandShaderCompilation.cpp b/Engine/Source/Editor/UnrealEd/Private/Cooker/OnDemandShaderCompilation.cpp
--- a/Engine/Source/Editor/UnrealEd/Private/Cooker/OnDemandShaderCompilation.cpp
+++ b/Engine/Source/Editor/UnrealEd/Private/Cooker/OnDemandShaderCompilation.cpp
@@ -6,43 +6,92 @@
 #include "Async/Async.h"
 
 
+int32 GODSCShaderMapsLifetime = 25;
+static FAutoConsoleVariableRef CVarODSCShaderMapsLifetime(
+	TEXT("odsc.shadermaps.lifetime"),
+	GODSCShaderMapsLifetime,
+	TEXT("Controls how many shader recompiles can happen before deleting an unused shadermap. Higher values means more memory, but faster iteration time\n")
+	TEXT("-1 means we never delete shader maps\n"),
+	ECVF_Default);
+
+int32 GODSCNumShaderMapsBeforeGC = 5000;
+static FAutoConsoleVariableRef CVarODSCNumShaderMapsBeforeGC(
+	TEXT("odsc.shadermaps.numbeforegc"),
+	GODSCNumShaderMapsBeforeGC,
+	TEXT("Controls how many shader maps we keep in memory before we start deleting them. Higher values means more memory, but faster iteration time\n")
+	TEXT("-1 means we never delete shader maps\n"),
+	ECVF_Default);
+
+
 namespace UE::Cook
 {
 
 void FODSCClientData::OnClientConnected(const void* ConnectionPtr)
 {
-	FScopeLock PollablesScopeLock(&ODSCClientPersistentDataMapLock);
-	ODSCClientPersistentDataMap.FindOrAdd(ConnectionPtr);
 }
 
 void FODSCClientData::OnClientDisconnected(const void* ConnectionPtr)
 {
-	FScopeLock PollablesScopeLock(&ODSCClientPersistentDataMapLock);
+}
 
-	FODSCClientPersistentData* ODSCClientPersistentData = ODSCClientPersistentDataMap.Find(ConnectionPtr);
-	if (ODSCClientPersistentData)
+void FODSCClientData::PurgeMaterialShaderMaps(int32 Lifetime, int32 NumMapsToDelete, FODSCClientPersistentData::Value& MaterialShaderMapsKeptAlive)
+{
+	// Don't start counting shader lifetime until we go over the limit of shadermaps we want to keep in memory
+	if (NumMapsToDelete <= 0)
 	{
-		// We can delete TStrongObjectPtr only on the game thread
-		TSet<TStrongObjectPtr<UMaterialInterface>>* MaterialsToRemove = new TSet<TStrongObjectPtr<UMaterialInterface>>;
-		*MaterialsToRemove = MoveTemp(ODSCClientPersistentData->MaterialsKeptAlive);
-		AsyncTask(ENamedThreads::GameThread, [MaterialsToRemove]()
-				  {
-					  delete MaterialsToRemove;
-				  });
-
-		ODSCClientPersistentDataMap.Remove(ConnectionPtr);
+		return;
+	}
+
+	for (FODSCClientPersistentData::Value::TIterator Iter = MaterialShaderMapsKeptAlive.CreateIterator(); Iter; ++Iter)
+	{
+		int32& ShaderMapLifetime = Iter.Value();
+
+		++ShaderMapLifetime;
+		FMaterialShaderMap* MaterialShaderMap = Iter.Key();
+		if ((Lifetime >= 0) && (ShaderMapLifetime > Lifetime) && (NumMapsToDelete > 0))
+		{
+			MaterialShaderMap->RemoveCompilingMaterialExternalDependency();
+			Iter.RemoveCurrent();
+			--NumMapsToDelete;
+		}
 	}
 }
 
+void FODSCClientData::FlushClientPersistentData(const void* ConnectionPtr)
+{
+	FScopeLock PollablesScopeLock(&ODSCClientPersistentDataLock);
+	int32 NumMapsToDelete = (GODSCNumShaderMapsBeforeGC >= 0) ? ODSCClientPersistentData.MaterialShaderMapsKeptAlive.Num() - GODSCNumShaderMapsBeforeGC : 0;
+	PurgeMaterialShaderMaps(GODSCShaderMapsLifetime, NumMapsToDelete, ODSCClientPersistentData.MaterialShaderMapsKeptAlive);
+}
+
 void FODSCClientData::KeepClientPersistentData(const void* ConnectionPtr, const TArray<TStrongObjectPtr<UMaterialInterface>>& LoadedMaterialsToRecompile)
 {
-	FScopeLock PollablesScopeLock(&ODSCClientPersistentDataMapLock);
-	FODSCClientPersistentData* ODSCClientPersistentData = ODSCClientPersistentDataMap.Find(ConnectionPtr);
-	if (ODSCClientPersistentData)
+	FScopeLock PollablesScopeLock(&ODSCClientPersistentDataLock);
+	for (const TStrongObjectPtr<UMaterialInterface>& MaterialInterface : LoadedMaterialsToRecompile)
 	{
-		for (const TStrongObjectPtr<UMaterialInterface>& MaterialInterface : LoadedMaterialsToRecompile)
+		for (int32 QualityLevelIndex = 0; QualityLevelIndex <= EMaterialQualityLevel::Num; ++QualityLevelIndex)
 		{
-			ODSCClientPersistentData->MaterialsKeptAlive.Add(MaterialInterface);
+			for (int32 FeatureLevelIndex = 0; FeatureLevelIndex <= ERHIFeatureLevel::Num; ++FeatureLevelIndex)
+			{
+				const FMaterialResource* MaterialResource = MaterialInterface->GetMaterialResource((ERHIFeatureLevel::Type)FeatureLevelIndex, (EMaterialQualityLevel::Type)QualityLevelIndex);
+				if (!MaterialResource)
+				{
+					continue;
+				}
+				FMaterialShaderMap* CompilingShaderMap = FMaterialShaderMap::FindCompilingShaderMap(MaterialResource->GetGameThreadCompilingShaderMapId());
+				if (!CompilingShaderMap)
+				{
+					continue;
+				}
+
+				int32& Lifetime = ODSCClientPersistentData.MaterialShaderMapsKeptAlive.FindOrAdd(CompilingShaderMap, -1);
+				// On first insertion, we add the external dependency, and set to 0 such that the call to PurgeMaterialShaderMaps will always work with positive value
+				if (Lifetime == -1)
+				{
+					CompilingShaderMap->AddCompilingMaterialExternalDependency();
+				}
+				Lifetime = 0;
+			}
 		}
 	}
 }
diff --git a/Engine/Source/Editor/UnrealEd/Private/Cooker/OnDemandShaderCompilation.h b/Engine/Source/Editor/UnrealEd/Private/Cooker/OnDemandShaderCompilation.h
--- a/Engine/Source/Editor/UnrealEd/Private/Cooker/OnDemandShaderCompilation.h
+++ b/Engine/Source/Editor/UnrealEd/Private/Cooker/OnDemandShaderCompilation.h
@@ -16,15 +16,20 @@ public:
 	void OnClientConnected(const void* ConnectionPtr);
 	void OnClientDisconnected(const void* ConnectionPtr); 
 	void KeepClientPersistentData(const void* ConnectionPtr, const TArray<TStrongObjectPtr<UMaterialInterface>>& LoadedMaterialsToRecompile);
+	void FlushClientPersistentData(const void* ConnectionPtr);
 
 private:
+
 	struct FODSCClientPersistentData
 	{
-		TSet<TStrongObjectPtr<UMaterialInterface> > MaterialsKeptAlive;
+		typedef TMap<TRefCountPtr<FMaterialShaderMap>, int32> Value;
+		Value MaterialShaderMapsKeptAlive;
 	};
 
-	TMap<const void*, FODSCClientPersistentData> ODSCClientPersistentDataMap;
-	FCriticalSection ODSCClientPersistentDataMapLock;
+	static void PurgeMaterialShaderMaps(int32 Lifetime, int32 NumMapsToDelete, FODSCClientPersistentData::Value& MaterialShaderMapsKeptAlive);
+
+	FODSCClientPersistentData ODSCClientPersistentData;
+	FCriticalSection ODSCClientPersistentDataLock;
 };
 
 }
\ No newline at end of file
diff --git a/Engine/Source/Runtime/Engine/Private/Materials/MaterialShader.cpp b/Engine/Source/Runtime/Engine/Private/Materials/MaterialShader.cpp
--- a/Engine/Source/Runtime/Engine/Private/Materials/MaterialShader.cpp
+++ b/Engine/Source/Runtime/Engine/Private/Materials/MaterialShader.cpp
@@ -2202,6 +2202,7 @@ void FMaterialShaderMap::ReleaseCompilingId()
 	{
 		FWriteScopeLock Locker(GCompilingShaderMapLock);
 		check(CompilingMaterialDependencies.Num() == 0);
+		check(CompilingMaterialNumExternalDependencies == 0);
 		verify(GetCompilingShaderMapLookup().Remove(CompilingId) == 1);
 		CompilingId = 0;
 	}
@@ -2222,7 +2223,12 @@ void FMaterialShaderMap::RemoveCompilingDependency(FMaterial* Material)
 	const int32 NumRemoved = CompilingMaterialDependencies.RemoveSingle(Material);
 	check(NumRemoved == 1);
 	checkSlow(!CompilingMaterialDependencies.Contains(Material));
-	if (CompilingMaterialDependencies.Num() == 0)
+	CheckReleaseCompilingId();
+}
+
+void FMaterialShaderMap::CheckReleaseCompilingId()
+{
+	if (CompilingMaterialDependencies.Num() == 0 && CompilingMaterialNumExternalDependencies==0)
 	{
 		const TArray<int32> CompilingIdsToCancel = { (int32)CompilingId };
 		GShaderCompilingManager->CancelCompilation(GetFriendlyName(), CompilingIdsToCancel);
diff --git a/Engine/Source/Runtime/Engine/Private/ODSC/ODSCManager.cpp b/Engine/Source/Runtime/Engine/Private/ODSC/ODSCManager.cpp
--- a/Engine/Source/Runtime/Engine/Private/ODSC/ODSCManager.cpp
+++ b/Engine/Source/Runtime/Engine/Private/ODSC/ODSCManager.cpp
@@ -161,6 +161,8 @@ void FODSCManager::AddThreadedRequest(
 	{
 		ClearErrorMessage();
 
+			Thread->AddRequest(TArray<FString>(), FString(), ShaderPlatform, FeatureLevel, QualityLevel, ODSCRecompileCommand::ResetMaterialCache);
+
 		if ((RecompileCommandType == ODSCRecompileCommand::Material || RecompileCommandType == ODSCRecompileCommand::Changed)
 			&& (CVarODSCRecompileMode.GetValueOnAnyThread() > 0))
 		{
diff --git a/Engine/Source/Runtime/Engine/Private/ShaderCompiler/ShaderCompiler.cpp b/Engine/Source/Runtime/Engine/Private/ShaderCompiler/ShaderCompiler.cpp
--- a/Engine/Source/Runtime/Engine/Private/ShaderCompiler/ShaderCompiler.cpp
+++ b/Engine/Source/Runtime/Engine/Private/ShaderCompiler/ShaderCompiler.cpp
@@ -6634,7 +6634,7 @@ void FShaderCompilingManager::ProcessCompiledShaderMaps(
 				}
 				else if (bSuccess)
 				{
-					bool bIsComplete = CompilingShaderMap->IsComplete(Material, true);
+					bool bIsComplete = CompilingShaderMap->IsComplete(Material, true) && (CompilingShaderMap->CompilingMaterialNumExternalDependencies == 0);
 
 					// If running a cook, only process complete shader maps, as there's no rendering of partially complete shader maps to worry about.
 					if (bIsComplete || IsRunningCookCommandlet() == false || bAllowForIncompleteShaderMaps)
@@ -10202,6 +10202,8 @@ extern ENGINE_API const TCHAR* ODSCCmdEnumToString(ODSCRecompileCommand Cmd)
 		return TEXT("Material");
 	case ODSCRecompileCommand::SingleShader:
 		return TEXT("SingleShader");
+	case ODSCRecompileCommand::ResetMaterialCache:
+			return TEXT("ResetMaterialCache");
 	}
 	ensure(false);
 	return TEXT("Unknown");
diff --git a/Engine/Source/Runtime/Engine/Public/MaterialShared.h b/Engine/Source/Runtime/Engine/Public/MaterialShared.h
--- a/Engine/Source/Runtime/Engine/Public/MaterialShared.h
+++ b/Engine/Source/Runtime/Engine/Public/MaterialShared.h
@@ -105,6 +105,11 @@ namespace UE::MIR
 	class FEmitter;
 }
 
+namespace UE::Cook
+{
+	class FODSCClientData;
+}
+
 template <class ElementType> class TLinkedList;
 
 #define ME_CAPTION_HEIGHT		18
@@ -1404,7 +1409,7 @@ public:
 	static TRefCountPtr<FMaterialShaderMap> FindId(const FMaterialShaderMapId& ShaderMapId, EShaderPlatform Platform);
 
 #if WITH_EDITOR
-	static FMaterialShaderMap* FindCompilingShaderMap(uint32 CompilingId);
+	ENGINE_API static FMaterialShaderMap* FindCompilingShaderMap(uint32 CompilingId);
 
 	/** Gets outdated types from all loaded material shader maps */
 	static void GetAllOutdatedTypes(TArray<const FShaderType*>& OutdatedShaderTypes, TArray<const FShaderPipelineType*>& OutdatedShaderPipelineTypes, TArray<const FVertexFactoryType*>& OutdatedFactoryTypes);
@@ -1675,6 +1680,12 @@ public:
 #endif
 
 protected:
+#if WITH_EDITOR
+	inline void AddCompilingMaterialExternalDependency() { check(IsInGameThread()); ++CompilingMaterialNumExternalDependencies; }
+	inline void RemoveCompilingMaterialExternalDependency() { check(IsInGameThread()); --CompilingMaterialNumExternalDependencies; CheckReleaseCompilingId();  }
+	ENGINE_API void CheckReleaseCompilingId();
+#endif
+
 	void PostFinalizeContent() override;
 
 private:
@@ -1701,6 +1712,7 @@ private:
 	TRefCountPtr<FMaterialShaderMap> FinalizedClone;
 	TRefCountPtr<FSharedShaderCompilerEnvironment> PendingCompilerEnvironment;
 	TArray<TRefCountPtr<FMaterial>> CompilingMaterialDependencies;
+	int32 CompilingMaterialNumExternalDependencies = 0;
 #endif // WITH_EDITOR
 
 	FUniformBufferLayoutRHIRef UniformBufferLayout;
@@ -1751,6 +1763,7 @@ private:
 
 	friend ENGINE_API void DumpMaterialStats( EShaderPlatform Platform );
 	friend class FShaderCompilingManager;
+	friend class UE::Cook::FODSCClientData;
 };
 
 
@@ -2810,6 +2823,7 @@ private:
 	friend class FShaderCompilingManager;
 	friend class FHLSLMaterialTranslator;
 	friend class FMaterialHLSLErrorHandler;
+	friend class UE::Cook::FODSCClientData;
 };
 
 
diff --git a/Engine/Source/Runtime/Engine/Public/ShaderCompiler.h b/Engine/Source/Runtime/Engine/Public/ShaderCompiler.h
--- a/Engine/Source/Runtime/Engine/Public/ShaderCompiler.h
+++ b/Engine/Source/Runtime/Engine/Public/ShaderCompiler.h
@@ -1210,7 +1210,8 @@ enum class ODSCRecompileCommand
 	Changed,
 	Global,
 	Material,
-	SingleShader
+	SingleShader,
+	ResetMaterialCache
 };
 
 extern ENGINE_API const TCHAR* ODSCCmdEnumToString(ODSCRecompileCommand Cmd);
