Update comments on K2_AddLevelToWorld functions and disallow abstract subclasses of ULevelStreaming to prevent users from accidentally selecting the base class of 'Level Streaming", which will always return a null value.

#rnx

diff --git a/Engine/Source/Editor/UnrealEd/Public/EditorLevelUtils.h b/Engine/Source/Editor/UnrealEd/Public/EditorLevelUtils.h
--- a/Engine/Source/Editor/UnrealEd/Public/EditorLevelUtils.h
+++ b/Engine/Source/Editor/UnrealEd/Public/EditorLevelUtils.h
@@ -41,7 +41,7 @@ public:
 	 * @return	Returns the newly created level, or NULL on failure
 	 */
 	UFUNCTION(BlueprintCallable, Category = "Editor Scripting | Level Creation")
-	static UNREALED_API ULevelStreaming* CreateNewStreamingLevel(TSubclassOf<ULevelStreaming> LevelStreamingClass, const FString& NewLevelPath = TEXT(""), bool bMoveSelectedActorsIntoNewLevel = false);
+	static UNREALED_API ULevelStreaming* CreateNewStreamingLevel(UPARAM(meta = (AllowAbstract = "false")) TSubclassOf<ULevelStreaming> LevelStreamingClass, const FString& NewLevelPath = TEXT(""), bool bMoveSelectedActorsIntoNewLevel = false);
 
 	/**
 	 * Makes the specified streaming level the current level for editing.
@@ -231,6 +231,8 @@ public:
 
 	/**
 	 * Adds the named level package to the world.  Does nothing if the level already exists in the world.
+	 * 
+	 * Levels are not saved when added to the world. They can be saved with the "Save Map" function
 	 *
 	 * @param	InWorld				World in which to add the level.
 	 * @param	LevelPackageName	The package name ("e.g /Game/MyLevel") of the level package to add.
@@ -239,7 +241,7 @@ public:
 	 * @return								The new level, or NULL if the level couldn't added.
 	 */
 	UFUNCTION(BlueprintCallable, Category = "Editor Scripting | Level Creation", meta=(DisplayName="Add Level to World", ScriptName="AddLevelToWorld"))
-	static ULevelStreaming* K2_AddLevelToWorld(UWorld* World, const FString& LevelPackageName, TSubclassOf<ULevelStreaming> LevelStreamingClass)
+	static ULevelStreaming* K2_AddLevelToWorld(UWorld* World, const FString& LevelPackageName, UPARAM(meta = (AllowAbstract = "false")) TSubclassOf<ULevelStreaming> LevelStreamingClass)
 	{
 		return AddLevelToWorld(World, *LevelPackageName, LevelStreamingClass, FTransform::Identity);
 	}
@@ -247,6 +249,8 @@ public:
 	/**
 	 * Removes given level from the world. Note, this will only work for sub-levels in the main level.
 	 *
+	 * Levels are not saved when added to the world. They can be saved with the "Save Map" function
+	 *
 	 * @param	InLevel				    Level asset to remove from the world.
 	 * @param	bClearSelection			If true, it will clear the editor selection.
 	 * @param	bResetTransactionBuffer	If true, it will reset the transaction buffer (i.e. clear undo history)
@@ -270,7 +274,7 @@ public:
 	 * @return								The new level, or NULL if the level couldn't added.
 	 */
 	UFUNCTION(BlueprintCallable, Category = "Editor Scripting | Level Creation", meta=(DisplayName="Add Level to World with Transform", ScriptName="AddLevelToWorldWithTransform"))
-	static ULevelStreaming* K2_AddLevelToWorldWithTransform(UWorld* World, const FString& LevelPackageName, TSubclassOf<ULevelStreaming> LevelStreamingClass, const FTransform& LevelTransform)
+	static ULevelStreaming* K2_AddLevelToWorldWithTransform(UWorld* World, const FString& LevelPackageName, UPARAM(meta = (AllowAbstract = "false")) TSubclassOf<ULevelStreaming> LevelStreamingClass, const FTransform& LevelTransform)
 	{
 		return AddLevelToWorld(World, *LevelPackageName, LevelStreamingClass, LevelTransform);
 	}
