[Selection] Hook up browse content browse to assets for ISMs

diff --git a/Engine/Source/Editor/UnrealEd/Private/EditorEngine.cpp b/Engine/Source/Editor/UnrealEd/Private/EditorEngine.cpp
--- a/Engine/Source/Editor/UnrealEd/Private/EditorEngine.cpp
+++ b/Engine/Source/Editor/UnrealEd/Private/EditorEngine.cpp
@@ -163,6 +163,9 @@
 #include "EngineModule.h"
 
 #include "EditorWorldExtension.h"
+#include "Elements/Interfaces/TypedElementAssetDataInterface.h"
+#include "Elements/Interfaces/TypedElementObjectInterface.h"
+#include "Elements/Framework/TypedElementRegistry.h"
 
 #if PLATFORM_WINDOWS
 	#include "Windows/WindowsHWrapper.h"
@@ -3207,72 +3210,41 @@ void UEditorEngine::GetAssetsToSyncToContentBrowser(TArray<FAssetData>& Assets,
 	// Otherwise, assemble a list of resources from selected actors.
 	if (!bFoundSurfaceMaterial)
 	{
-		for (FSelectionIterator It(GetSelectedActorIterator()); It; ++It)
+		if (UTypedElementSelectionSet* SelectionSet = GEditor->GetSelectedActors()->GetElementSelectionSet())
 		{
-			AActor* Actor = static_cast<AActor*>(*It);
-			checkSlow(Actor->IsA(AActor::StaticClass()));
-
-			bool bFoundOverride = false;
-			if (bAllowBrowseToAssetOverride)
+			SelectionSet->ForEachSelectedElementHandle([&Assets, bAllowBrowseToAssetOverride](const FTypedElementHandle& SelectedHandle)
 			{
-				// If BrowseToAssetOverride is set, then use the asset it points to instead of the selected asset
-				const FString& BrowseToAssetOverride = Actor->GetBrowseToAssetOverride();
-				if (!BrowseToAssetOverride.IsEmpty())
+				if (bAllowBrowseToAssetOverride)
 				{
-					if (IAssetRegistry* AssetRegistry = IAssetRegistry::Get())
+					if (TTypedElement<ITypedElementObjectInterface> ObjectInterface = UTypedElementRegistry::GetInstance()->GetElement<ITypedElementObjectInterface>(SelectedHandle))
 					{
-						TArray<FAssetData> FoundAssets;
-						if (AssetRegistry->GetAssetsByPackageName(*BrowseToAssetOverride, FoundAssets) && FoundAssets.Num() > 0)
+						if (AActor* Actor = ObjectInterface.GetObjectAs<AActor>())
 						{
-							Assets.Add(FoundAssets[0]);
-							bFoundOverride = true;
+							// If BrowseToAssetOverride is set, then use the asset it points to instead of the selected asset
+							const FString& BrowseToAssetOverride = Actor->GetBrowseToAssetOverride();
+							if (!BrowseToAssetOverride.IsEmpty())
+							{
+								if (IAssetRegistry* AssetRegistry = IAssetRegistry::Get())
+								{
+									TArray<FAssetData> FoundAssets;
+									if (AssetRegistry->GetAssetsByPackageName(*BrowseToAssetOverride, FoundAssets) && FoundAssets.Num() > 0)
+									{
+										Assets.Add(FoundAssets[0]);
+										return true;
+									}
+								}
+							}
 						}
 					}
 				}
-			}
 
-			if (!bFoundOverride)
-			{
-				// If the actor is an instance of a blueprint, just add the blueprint.
-				UBlueprint* GeneratingBP = Cast<UBlueprint>(It->GetClass()->ClassGeneratedBy);
-				if (GeneratingBP != NULL)
+				if (TTypedElement<ITypedElementAssetDataInterface> AssetDataInterface = UTypedElementRegistry::GetInstance()->GetElement<ITypedElementAssetDataInterface>(SelectedHandle))
 				{
-					Assets.Add(FAssetData(GeneratingBP));
+					Assets.Append(AssetDataInterface.GetAllReferencedAssetDatas());
 				}
-				// Cooked editor sometimes only contains UBlueprintGeneratedClass with no UBlueprint
-				else if (UBlueprintGeneratedClass* BlueprintGeneratedClass = Cast<UBlueprintGeneratedClass>(It->GetClass()))
-				{
-					Assets.Add(FAssetData(BlueprintGeneratedClass));
-				}
-				// Otherwise, add the results of the GetReferencedContentObjects call
-				else
-				{
-					TArray<UObject*> Objects;
-					Actor->GetReferencedContentObjects(Objects);
-					for (UObject* Object : Objects)
-					{
-						Assets.Add(FAssetData(Object));
-					}
 
-					TArray<FSoftObjectPath> SoftObjects;
-					Actor->GetSoftReferencedContentObjects(SoftObjects);
-
-					if (SoftObjects.Num())
-					{
-						IAssetRegistry& AssetRegistry = IAssetRegistry::GetChecked();
-
-						for (const FSoftObjectPath& SoftObject : SoftObjects)
-						{
-							FAssetData AssetData = AssetRegistry.GetAssetByObjectPath(SoftObject);
-
-							if (AssetData.IsValid())
-							{
-								Assets.Add(AssetData);
-							}
-						}
-					}
-				}
-			}
+				return true;
+			});
 		}
 	}
 }
diff --git a/Engine/Source/Editor/UnrealEd/Private/Elements/Actor/ActorElementEditorAssetDataInterface.cpp b/Engine/Source/Editor/UnrealEd/Private/Elements/Actor/ActorElementEditorAssetDataInterface.cpp
--- a/Engine/Source/Editor/UnrealEd/Private/Elements/Actor/ActorElementEditorAssetDataInterface.cpp
+++ b/Engine/Source/Editor/UnrealEd/Private/Elements/Actor/ActorElementEditorAssetDataInterface.cpp
@@ -3,9 +3,9 @@
 #include "Elements/Actor/ActorElementEditorAssetDataInterface.h"
 
 #include "AssetRegistry/AssetData.h"
+#include "AssetRegistry/IAssetRegistry.h"
 #include "Elements/Actor/ActorElementData.h"
 #include "GameFramework/Actor.h"
-#include "HAL/PlatformCrt.h"
 
 TArray<FAssetData> UActorElementEditorAssetDataInterface::GetAllReferencedAssetDatas(const FTypedElementHandle& InElementHandle)
 {
@@ -24,6 +24,23 @@ TArray<FAssetData> UActorElementEditorAssetDataInterface::GetAllReferencedAssetD
 			}
 		}
 
+		TArray<FSoftObjectPath> SoftObjects;
+		RawActorPtr->GetSoftReferencedContentObjects(SoftObjects);
+		if (SoftObjects.Num())
+		{
+			IAssetRegistry& AssetRegistry = IAssetRegistry::GetChecked();
+
+			for (const FSoftObjectPath& SoftObject : SoftObjects)
+			{
+				FAssetData AssetData = AssetRegistry.GetAssetByObjectPath(SoftObject);
+
+				if (AssetData.IsValid())
+				{
+					AssetDatas.Add(AssetData);
+				}
+			}
+		}
+
 		AssetDatas.Emplace(FAssetData(RawActorPtr));
 	}
 
