[FBX Material Exporter]  Adding support for MaterialBaking and Interchange material matching.

#rnx

diff --git a/Engine/Source/Editor/UnrealEd/Classes/Exporters/FbxExportOption.h b/Engine/Source/Editor/UnrealEd/Classes/Exporters/FbxExportOption.h
--- a/Engine/Source/Editor/UnrealEd/Classes/Exporters/FbxExportOption.h
+++ b/Engine/Source/Editor/UnrealEd/Classes/Exporters/FbxExportOption.h
@@ -36,6 +36,30 @@ enum class EMovieSceneBakeType : uint8
 	BakeAll = BakeChannels | BakeTransforms,
 };
 
+UENUM(BlueprintType)
+enum class EFbxMaterialBakeMode : uint8
+{
+	/** Never bake material inputs. */
+	Disabled,
+	/** Only use a simple quad if a material input needs to be baked out. */
+	Simple,
+	/** Allow usage of the mesh data if a material input needs to be baked out with vertex data. */
+	UseMeshData,
+};
+
+USTRUCT(Blueprintable)
+struct FFbxMaterialBakeSize
+{
+	GENERATED_BODY()
+
+	UPROPERTY(EditAnywhere, BlueprintReadWrite, Category = "")
+	FIntPoint Size = FIntPoint(1024, 1024);
+
+	/** If enabled, bake size is based on the largest texture used in the material input's expression graph. If none found, bake size will fall back to the explicit dimensions. */
+	UPROPERTY(EditAnywhere, BlueprintReadWrite, Category = "")
+	bool bAutoDetect = true;
+};
+
 UCLASS(config = EditorPerProjectUserSettings, MinimalAPI, BlueprintType)
 class UFbxExportOption : public UObject
 {
@@ -98,6 +122,14 @@ public:
 	UPROPERTY(EditAnywhere, BlueprintReadWrite, AdvancedDisplay, config, Category = Animation)
 	EMovieSceneBakeType BakeActorAnimation;
 
+	/** Bake mode determining if and how a material input is baked out to a texture. Baking is only used for non-trivial material inputs (i.e. not simple texture or constant expressions). */
+	UPROPERTY(EditAnywhere, BlueprintReadWrite, Config, Category = Material)
+	EFbxMaterialBakeMode BakeMaterialInputs;
+
+	/** Default size of the baked out texture (containing the material input).*/
+	UPROPERTY(EditAnywhere, BlueprintReadWrite, Config, Category = Material, Meta = (EditCondition = "BakeMaterialInputs != EFbxMaterialBakeMode::Disabled"))
+	FFbxMaterialBakeSize DefaultMaterialBakeSize;
+
 	/* Set all the FProperty to the CDO value */
 	void ResetToDefault();
 
diff --git a/Engine/Source/Editor/UnrealEd/Private/Fbx/FbxMainExport.cpp b/Engine/Source/Editor/UnrealEd/Private/Fbx/FbxMainExport.cpp
--- a/Engine/Source/Editor/UnrealEd/Private/Fbx/FbxMainExport.cpp
+++ b/Engine/Source/Editor/UnrealEd/Private/Fbx/FbxMainExport.cpp
@@ -118,6 +118,11 @@
 #include "Interfaces/IMainFrameModule.h"
 #include "UObject/MetaData.h"
 
+#include "FbxMaterialExportUtilities.h"
+
+#include "Exporters/Exporter.h"
+#include "MaterialPropertyEx.h"
+
 namespace UnFbx
 {
 
@@ -485,7 +490,7 @@ void FFbxExporter::ExportLevelMesh(ULevel* InLevel, bool bExportLevelGeometry, T
 			Scene->GetRootNode()->AddChild(Node);
 
 			// Export the mesh for the world
-			ExportModel(InLevel->Model, Node, "Level Mesh");
+			ExportModel(InLevel->Model, Node, "Level Mesh", FFbxMaterialBakingMeshData(InLevel->Model));
 		}
 	}
 
@@ -725,7 +730,7 @@ void FFbxExporter::ExportBrush(ABrush* Actor, UModel* InModel, bool bConvertToSt
 		Scene->GetRootNode()->AddChild(FbxActor);
  
 		// Export the mesh information
-		ExportModel(Model, FbxActor, TCHAR_TO_UTF8(*Actor->GetName()));
+		ExportModel(Model, FbxActor, TCHAR_TO_UTF8(*Actor->GetName()), FFbxMaterialBakingMeshData(InModel, Actor));
 	}
 	else
 	{
@@ -743,7 +748,7 @@ void FFbxExporter::ExportBrush(ABrush* Actor, UModel* InModel, bool bConvertToSt
 	}
 }
 
-void FFbxExporter::ExportModel(UModel* Model, FbxNode* Node, const char* Name)
+void FFbxExporter::ExportModel(UModel* Model, FbxNode* Node, const char* Name, const FFbxMaterialBakingMeshData& MaterialBakingMeshData)
 {
 	//int32 VertexCount = Model->VertexBuffer.Vertices.Num();
 	int32 MaterialCount = Model->MaterialIndexBuffers.Num();
@@ -842,7 +847,7 @@ void FFbxExporter::ExportModel(UModel* Model, FbxNode* Node, const char* Name)
 		FbxSurfaceMaterial* FbxMaterial;
 		if (MaterialInterface != NULL && MaterialInterface->GetMaterial() != NULL)
 		{
-			FbxMaterial = ExportMaterial(MaterialInterface);
+			FbxMaterial = ExportMaterial(MaterialInterface, MaterialBakingMeshData.GetUModelStaticMeshMaterialIndex(MaterialInterface), MaterialBakingMeshData);
 		}
 		else
 		{
@@ -963,7 +968,7 @@ void FFbxExporter::ExportStaticMesh(AActor* Actor, UStaticMeshComponent* StaticM
 
 			const int32 LightmapUVChannel = -1;
 			const TArray<FStaticMaterial>* MaterialOrderOverride = nullptr;
-			ExportStaticMeshToFbx(StaticMesh, CurrentLodIndex, *FbxMeshName, FbxActorLOD, LightmapUVChannel, ColorBuffer, MaterialOrderOverride, &ToRawPtrTArrayUnsafe(StaticMeshComponent->OverrideMaterials));
+			ExportStaticMeshToFbx(StaticMesh, CurrentLodIndex, *FbxMeshName, FbxActorLOD, FFbxMaterialBakingMeshData(StaticMesh, StaticMeshComponent, CurrentLodIndex), LightmapUVChannel, ColorBuffer, MaterialOrderOverride, &ToRawPtrTArrayUnsafe(StaticMeshComponent->OverrideMaterials));
 		}
 	}
 	else
@@ -976,7 +981,7 @@ void FFbxExporter::ExportStaticMesh(AActor* Actor, UStaticMeshComponent* StaticM
 		FbxNode* FbxActor = ExportActor(Actor, false, NodeNameAdapter);
 		const int32 LightmapUVChannel = -1;
 		const TArray<FStaticMaterial>* MaterialOrderOverride = nullptr;
-		ExportStaticMeshToFbx(StaticMesh, LODIndex, *FbxMeshName, FbxActor, LightmapUVChannel, ColorBuffer, MaterialOrderOverride, &ToRawPtrTArrayUnsafe(StaticMeshComponent->OverrideMaterials));
+		ExportStaticMeshToFbx(StaticMesh, LODIndex, *FbxMeshName, FbxActor, FFbxMaterialBakingMeshData(StaticMesh, StaticMeshComponent, LODIndex), LightmapUVChannel, ColorBuffer, MaterialOrderOverride, &ToRawPtrTArrayUnsafe(StaticMeshComponent->OverrideMaterials));
 	}
 }
 
@@ -1205,12 +1210,12 @@ void FFbxExporter::ExportStaticMesh( UStaticMesh* StaticMesh, const TArray<FStat
 				double LodScreenSize = (double)(10.0f / StaticMesh->GetRenderData()->ScreenSize[CurrentLodIndex].Default);
 				FbxLodGroupAttribute->AddThreshold(LodScreenSize);
 			}
-			ExportStaticMeshToFbx(StaticMesh, CurrentLodIndex, *MeshName, FbxActorLOD, -1, nullptr, MaterialOrder);
+			ExportStaticMeshToFbx(StaticMesh, CurrentLodIndex, *MeshName, FbxActorLOD, FFbxMaterialBakingMeshData(StaticMesh, nullptr, CurrentLodIndex), -1, nullptr, MaterialOrder);
 		}
 	}
 	else
 	{
-		ExportStaticMeshToFbx(StaticMesh, 0, *MeshName, MeshNode, -1, NULL, MaterialOrder);
+		ExportStaticMeshToFbx(StaticMesh, 0, *MeshName, MeshNode, FFbxMaterialBakingMeshData(StaticMesh), -1, NULL, MaterialOrder);
 	}
 }
 
@@ -1222,7 +1227,7 @@ void FFbxExporter::ExportStaticMeshLightMap( UStaticMesh* StaticMesh, int32 LODI
 	StaticMesh->GetName(MeshName);
 	FbxNode* MeshNode = FbxNode::Create(Scene, TCHAR_TO_UTF8(*MeshName));
 	Scene->GetRootNode()->AddChild(MeshNode);
-	ExportStaticMeshToFbx(StaticMesh, LODIndex, *MeshName, MeshNode, UVChannel);
+	ExportStaticMeshToFbx(StaticMesh, LODIndex, *MeshName, MeshNode, FFbxMaterialBakingMeshData(StaticMesh, nullptr, LODIndex), UVChannel);
 }
 
 void FFbxExporter::ExportSkeletalMesh( USkeletalMesh* SkeletalMesh )
@@ -1382,14 +1387,17 @@ bool FFbxExporter::FillFbxTextureProperty(const char *PropertyName, const FExpre
 /**
 * Exports the profile_COMMON information for a material.
 */
-FbxSurfaceMaterial* FFbxExporter::ExportMaterial(UMaterialInterface* MaterialInterface)
+FbxSurfaceMaterial* FFbxExporter::ExportMaterial(UMaterialInterface* MaterialInterface, const int32& MaterialIndex, const FFbxMaterialBakingMeshData& BakingMeshData)
 {
+	FString ExportFolderPath = FPaths::GetPath(UExporter::CurrentFilename);
+
 	if (Scene == nullptr || MaterialInterface == nullptr || MaterialInterface->GetMaterial() == nullptr) return nullptr;
 	
 	// Verify that this material has not already been exported:
-	if (FbxMaterials.Find(MaterialInterface))
+	TMap<int32, FbxSurfaceMaterial*>* MaterialIndexToFbxMaterials = FbxMaterials.Find(MaterialInterface);
+	if (MaterialIndexToFbxMaterials && MaterialIndexToFbxMaterials->Find(MaterialIndex))
 	{
-		return *FbxMaterials.Find(MaterialInterface);
+		return *MaterialIndexToFbxMaterials->Find(MaterialIndex);
 	}
 
 	// Create the Fbx material
@@ -1402,33 +1410,62 @@ FbxSurfaceMaterial* FFbxExporter::ExportMaterial(UMaterialInterface* MaterialInt
 		return nullptr;
 	}
 
+	if (!MaterialIndexToFbxMaterials)
+	{
+		MaterialIndexToFbxMaterials = &FbxMaterials.Add(MaterialInterface);
+	}
+
 	UMaterialEditorOnlyData* MaterialEditorOnly = Material->GetEditorOnlyData();
 	
 	// Set the shading model
-	if (Material->GetShadingModels().HasOnlyShadingModel(MSM_DefaultLit))
+	bool bUseLambert = false;
+	bool bInterchangeOriginated = FFbxMaterialExportUtilities::GetInterchangeShadingModel(MaterialInterface, bUseLambert);
+
+	FString FbxMaterialName = MaterialInterface->GetName();
+	if (MaterialIndexToFbxMaterials != nullptr)
 	{
-		FbxMaterial = FbxSurfacePhong::Create(Scene, TCHAR_TO_UTF8(*MaterialInterface->GetName()));
+		FbxMaterialName += FString(TEXT("_")) + FString::FromInt(MaterialIndex);
+	}
+	
+	if ((bInterchangeOriginated && !bUseLambert) || (!bInterchangeOriginated && Material->GetShadingModels().HasOnlyShadingModel(MSM_DefaultLit)))
+	{
+		FbxMaterial = FbxSurfacePhong::Create(Scene, TCHAR_TO_UTF8(*FbxMaterialName));
 		//((FbxSurfacePhong*)FbxMaterial)->Specular.Set(Material->Specular));
 		//((FbxSurfacePhong*)FbxMaterial)->Shininess.Set(Material->SpecularPower.Constant);
 	}
 	else // if (Material->ShadingModel == MSM_Unlit)
 	{
-		FbxMaterial = FbxSurfaceLambert::Create(Scene, TCHAR_TO_UTF8(*MaterialInterface->GetName()));
+		FbxMaterial = FbxSurfaceLambert::Create(Scene, TCHAR_TO_UTF8(*FbxMaterialName));
 	}
 
+	if (bInterchangeOriginated)
+	{
+		FFbxMaterialExportUtilities::ProcessInterchangeMaterials(MaterialInterface, Scene, FbxMaterial);
+
+		MaterialIndexToFbxMaterials->Add(MaterialIndex, FbxMaterial);
+
+		return FbxMaterial;
+	}
 
 	//Get the base material connected expression parameter name for all the supported fbx material exported properties
 	//We only support material input where the connected expression is a parameter of type (constant, scalar, vector, texture, TODO virtual texture)
 
-	FName BaseColorParamName = (!MaterialEditorOnly->BaseColor.UseConstant && MaterialEditorOnly->BaseColor.Expression) ? MaterialEditorOnly->BaseColor.Expression->GetParameterName() : NAME_None;
+	bool bBaseColorNonConstAndNonDefault = (!MaterialEditorOnly->BaseColor.UseConstant && MaterialEditorOnly->BaseColor.Expression);
+	bool bEmissiveNonConstAndNonDefault = (!MaterialEditorOnly->EmissiveColor.UseConstant && MaterialEditorOnly->EmissiveColor.Expression);
+	bool bNormalNonConstAndNonDefault = (MaterialEditorOnly->Normal.Expression != nullptr);
+	bool bOpacityNonConstAndNonDefault = (!MaterialEditorOnly->Opacity.UseConstant && MaterialEditorOnly->Opacity.Expression);
+	bool bOpacityMaskNonConstAndNonDefault = (!MaterialEditorOnly->OpacityMask.UseConstant && MaterialEditorOnly->OpacityMask.Expression);
+
+	FName BaseColorParamName = bBaseColorNonConstAndNonDefault ? MaterialEditorOnly->BaseColor.Expression->GetParameterName() : NAME_None;
+	FName EmissiveParamName = bEmissiveNonConstAndNonDefault ? MaterialEditorOnly->EmissiveColor.Expression->GetParameterName() : NAME_None;
+	FName NormalParamName = bNormalNonConstAndNonDefault ? MaterialEditorOnly->Normal.Expression->GetParameterName() : NAME_None;
+	FName OpacityParamName = bOpacityNonConstAndNonDefault ? MaterialEditorOnly->Opacity.Expression->GetParameterName() : NAME_None;
+	FName OpacityMaskParamName = bOpacityMaskNonConstAndNonDefault ? MaterialEditorOnly->OpacityMask.Expression->GetParameterName() : NAME_None;
+	
 	bool BaseColorParamSet = false;
-	FName EmissiveParamName = (!MaterialEditorOnly->EmissiveColor.UseConstant && MaterialEditorOnly->EmissiveColor.Expression) ? MaterialEditorOnly->EmissiveColor.Expression->GetParameterName() : NAME_None;
 	bool EmissiveParamSet = false;
-	FName NormalParamName = MaterialEditorOnly->Normal.Expression ? MaterialEditorOnly->Normal.Expression->GetParameterName() : NAME_None;
 	bool NormalParamSet = false;
-	FName OpacityParamName = (!MaterialEditorOnly->Opacity.UseConstant && MaterialEditorOnly->Opacity.Expression) ? MaterialEditorOnly->Opacity.Expression->GetParameterName() : NAME_None;
 	bool OpacityParamSet = false;
-	FName OpacityMaskParamName = (!MaterialEditorOnly->OpacityMask.UseConstant && MaterialEditorOnly->OpacityMask.Expression) ? MaterialEditorOnly->OpacityMask.Expression->GetParameterName() : NAME_None;
 	bool OpacityMaskParamSet = false;
 
 	UMaterialInstance* MaterialInstance = Cast<UMaterialInstance>(MaterialInterface);
@@ -1471,22 +1508,22 @@ FbxSurfaceMaterial* FFbxExporter::ExportMaterial(UMaterialInterface* MaterialInt
 				return false;
 			};
 
-			if (BaseColorParamName != NAME_None && ParameterName == BaseColorParamName)
+			if (!BaseColorParamSet && BaseColorParamName != NAME_None && ParameterName == BaseColorParamName)
 			{
 				BaseColorParamSet = SetTextureProperty(FbxSurfaceMaterial::sDiffuse, Scene);
 			}
-			if (EmissiveParamName != NAME_None && ParameterName == EmissiveParamName)
+			if (!EmissiveParamSet && EmissiveParamName != NAME_None && ParameterName == EmissiveParamName)
 			{
 				EmissiveParamSet = SetTextureProperty(FbxSurfaceMaterial::sEmissive, Scene);
 			}
 			
 			if (BlendMode == BLEND_Translucent)
 			{
-				if (OpacityParamName != NAME_None && ParameterName == OpacityParamName)
+				if (!OpacityParamSet && OpacityParamName != NAME_None && ParameterName == OpacityParamName)
 				{
 					OpacityParamSet = SetTextureProperty(FbxSurfaceMaterial::sTransparentColor, Scene);
 				}
-				if (OpacityMaskParamName != NAME_None && ParameterName == OpacityMaskParamName)
+				if (!OpacityMaskParamSet && OpacityMaskParamName != NAME_None && ParameterName == OpacityMaskParamName)
 				{
 					OpacityMaskParamSet = SetTextureProperty(FbxSurfaceMaterial::sTransparencyFactor, Scene);
 				}
@@ -1494,7 +1531,7 @@ FbxSurfaceMaterial* FFbxExporter::ExportMaterial(UMaterialInterface* MaterialInt
 			else
 			{
 				//There is no normal input in Blend translucent mode
-				if (NormalParamName != NAME_None && ParameterName == NormalParamName)
+				if (!NormalParamSet && NormalParamName != NAME_None && ParameterName == NormalParamName)
 				{
 					NormalParamSet = SetTextureProperty(FbxSurfaceMaterial::sNormalMap, Scene);
 				}
@@ -1589,13 +1626,26 @@ FbxSurfaceMaterial* FFbxExporter::ExportMaterial(UMaterialInterface* MaterialInt
 		((FbxSurfaceLambert*)FbxMaterial)->TransparencyFactor.Set(MaterialEditorOnly->Opacity.Constant);
 	}
 
+	UFbxExportOption* FbxExportOptions = GetExportOptions();
+
 	// Fill in the profile_COMMON effect with the material information.
 	//Fill the texture or constant
 	if(!BaseColorParamSet)
 	{
 		if (!FillFbxTextureProperty(FbxSurfaceMaterial::sDiffuse, MaterialEditorOnly->BaseColor, FbxMaterial))
 		{
-			((FbxSurfaceLambert*)FbxMaterial)->Diffuse.Set(SetMaterialComponent(MaterialEditorOnly->BaseColor, true));
+			if (bBaseColorNonConstAndNonDefault)
+			{
+				FFbxMaterialExportUtilities::BakeMaterialProperty(FbxExportOptions,
+					Scene, FbxMaterial, FbxSurfaceMaterial::sDiffuse,
+					MP_BaseColor, MaterialInterface, MaterialIndex,
+					BakingMeshData, 
+					ExportFolderPath);
+			}
+			else
+			{
+				((FbxSurfaceLambert*)FbxMaterial)->Diffuse.Set(SetMaterialComponent(MaterialEditorOnly->BaseColor, true));
+			}
 		}
 	}
 
@@ -1603,7 +1653,18 @@ FbxSurfaceMaterial* FFbxExporter::ExportMaterial(UMaterialInterface* MaterialInt
 	{
 		if (!FillFbxTextureProperty(FbxSurfaceMaterial::sEmissive, MaterialEditorOnly->EmissiveColor, FbxMaterial))
 		{
-			((FbxSurfaceLambert*)FbxMaterial)->Emissive.Set(SetMaterialComponent(MaterialEditorOnly->EmissiveColor, true));
+			if (bEmissiveNonConstAndNonDefault)
+			{
+				FFbxMaterialExportUtilities::BakeMaterialProperty(FbxExportOptions,
+					Scene, FbxMaterial, FbxSurfaceMaterial::sEmissive,
+					MP_EmissiveColor, MaterialInterface, MaterialIndex,
+					BakingMeshData,
+					ExportFolderPath);
+			}
+			else
+			{
+				((FbxSurfaceLambert*)FbxMaterial)->Emissive.Set(SetMaterialComponent(MaterialEditorOnly->EmissiveColor, true));
+			}
 		}
 	}
 
@@ -1616,8 +1677,19 @@ FbxSurfaceMaterial* FFbxExporter::ExportMaterial(UMaterialInterface* MaterialInt
 		{
 			if (!FillFbxTextureProperty(FbxSurfaceMaterial::sTransparentColor, MaterialEditorOnly->Opacity, FbxMaterial))
 			{
-				FbxDouble3 OpacityValue((FbxDouble)(MaterialEditorOnly->Opacity.Constant), (FbxDouble)(MaterialEditorOnly->Opacity.Constant), (FbxDouble)(MaterialEditorOnly->Opacity.Constant));
-				((FbxSurfaceLambert*)FbxMaterial)->TransparentColor.Set(OpacityValue);
+				if (bOpacityNonConstAndNonDefault)
+				{
+					FFbxMaterialExportUtilities::BakeMaterialProperty(FbxExportOptions,
+						Scene, FbxMaterial, FbxSurfaceMaterial::sTransparentColor,
+						MP_Opacity, MaterialInterface, MaterialIndex,
+						BakingMeshData,
+						ExportFolderPath);
+				}
+				else
+				{
+					FbxDouble3 OpacityValue((FbxDouble)(MaterialEditorOnly->Opacity.Constant), (FbxDouble)(MaterialEditorOnly->Opacity.Constant), (FbxDouble)(MaterialEditorOnly->Opacity.Constant));
+					((FbxSurfaceLambert*)FbxMaterial)->TransparentColor.Set(OpacityValue);
+				}
 			}
 		}
 
@@ -1625,7 +1697,18 @@ FbxSurfaceMaterial* FFbxExporter::ExportMaterial(UMaterialInterface* MaterialInt
 		{
 			if (!FillFbxTextureProperty(FbxSurfaceMaterial::sTransparencyFactor, MaterialEditorOnly->OpacityMask, FbxMaterial))
 			{
-				((FbxSurfaceLambert*)FbxMaterial)->TransparencyFactor.Set(MaterialEditorOnly->OpacityMask.Constant);
+				if (bOpacityMaskNonConstAndNonDefault)
+				{
+					FFbxMaterialExportUtilities::BakeMaterialProperty(FbxExportOptions,
+						Scene, FbxMaterial, FbxSurfaceMaterial::sTransparencyFactor,
+						MP_OpacityMask, MaterialInterface, MaterialIndex,
+						BakingMeshData,
+						ExportFolderPath);
+				}
+				else
+				{
+					((FbxSurfaceLambert*)FbxMaterial)->TransparencyFactor.Set(MaterialEditorOnly->OpacityMask.Constant);
+				}
 			}
 		}
 	}
@@ -1635,12 +1718,22 @@ FbxSurfaceMaterial* FFbxExporter::ExportMaterial(UMaterialInterface* MaterialInt
 		if (!NormalParamSet)
 		{
 			//Set the Normal map only if there is a texture sampler
-			FillFbxTextureProperty(FbxSurfaceMaterial::sNormalMap, MaterialEditorOnly->Normal, FbxMaterial);
+			if (!FillFbxTextureProperty(FbxSurfaceMaterial::sNormalMap, MaterialEditorOnly->Normal, FbxMaterial))
+			{
+				if (bNormalNonConstAndNonDefault)
+				{
+					FFbxMaterialExportUtilities::BakeMaterialProperty(FbxExportOptions,
+						Scene, FbxMaterial, FbxSurfaceMaterial::sNormalMap,
+						MP_Normal, MaterialInterface, MaterialIndex,
+						BakingMeshData,
+						ExportFolderPath);
+				}
+			}
 		}
 	}
 
-	FbxMaterials.Add(MaterialInterface, FbxMaterial);
-	
+	MaterialIndexToFbxMaterials->Add(MaterialIndex, FbxMaterial);
+
 	return FbxMaterial;
 }
 
@@ -2718,11 +2811,13 @@ FbxNode* FFbxExporter::ExportActor(AActor* Actor, bool bExportComponents, INodeN
 				{
 					if (USplineMeshComponent* SplineMeshComp = Cast<USplineMeshComponent>(StaticMeshComp))
 					{
-						ExportSplineMeshToFbx(SplineMeshComp, *SplineMeshComp->GetName(), ExportNode);
+						//TODO: Validate Spline material baking
+						//(Spline mesh staticmesh does not take into account the Spline transforms, might not be accurate/enough to use Spline->StaticMesh for material baking.)
+						ExportSplineMeshToFbx(SplineMeshComp, *SplineMeshComp->GetName(), ExportNode, FFbxMaterialBakingMeshData(SplineMeshComp->GetStaticMesh(), SplineMeshComp, (SplineMeshComp->ForcedLodModel > 0 ? SplineMeshComp->ForcedLodModel - 1 : /* auto-select*/ 0)));
 					}
 					else if (UInstancedStaticMeshComponent* InstancedMeshComp = Cast<UInstancedStaticMeshComponent>(StaticMeshComp))
 					{
-						ExportInstancedMeshToFbx(InstancedMeshComp, *InstancedMeshComp->GetName(), ExportNode);
+						ExportInstancedMeshToFbx(InstancedMeshComp, *InstancedMeshComp->GetName(), ExportNode, FFbxMaterialBakingMeshData(InstancedMeshComp->GetStaticMesh(), InstancedMeshComp, (InstancedMeshComp->ForcedLodModel > 0 ? InstancedMeshComp->ForcedLodModel - 1 : /* auto-select*/ 0)));
 					}
 					else
 					{
@@ -2730,7 +2825,7 @@ FbxNode* FFbxExporter::ExportActor(AActor* Actor, bool bExportComponents, INodeN
 						const int32 LightmapUVChannel = -1;
 						const TArray<FStaticMaterial>* MaterialOrderOverride = nullptr;
 						const FColorVertexBuffer* ColorBuffer = nullptr;
-						ExportStaticMeshToFbx(StaticMeshComp->GetStaticMesh(), LODIndex, *StaticMeshComp->GetName(), ExportNode, LightmapUVChannel, ColorBuffer, MaterialOrderOverride, &ToRawPtrTArrayUnsafe(StaticMeshComp->OverrideMaterials));
+						ExportStaticMeshToFbx(StaticMeshComp->GetStaticMesh(), LODIndex, *StaticMeshComp->GetName(), ExportNode, FFbxMaterialBakingMeshData(StaticMeshComp->GetStaticMesh(), StaticMeshComp, LODIndex), LightmapUVChannel, ColorBuffer, MaterialOrderOverride, &ToRawPtrTArrayUnsafe(StaticMeshComp->OverrideMaterials));
 					}
 				}
 				else if (SkelMeshComp && SkelMeshComp->GetSkeletalMeshAsset())
@@ -4601,7 +4696,8 @@ bool FFbxExporter::ExportStaticMeshFromMeshDescription(FbxMesh* Mesh
 	, FbxNode* FbxActor
 	, int32 LightmapUVChannel
 	, const TArray<FStaticMaterial>* MaterialOrderOverride
-	, const TArray<UMaterialInterface*>* OverrideMaterials)
+	, const TArray<UMaterialInterface*>* OverrideMaterials
+	, const FFbxMaterialBakingMeshData& MaterialBakingMeshData)
 {
 
 	if (MeshDescription->IsEmpty() || MeshDescription->Vertices().Num() == 0)
@@ -4766,26 +4862,27 @@ bool FFbxExporter::ExportStaticMeshFromMeshDescription(FbxMesh* Mesh
 		{
 			UMaterialInterface* Material = nullptr;
 
+			FName CurrentMaterialSlotName = PolygonGroupMaterialSlotNames[PolygonGroupID];
+			int32 MaterialIndex = StaticMesh->GetMaterialIndexFromImportedMaterialSlotName(CurrentMaterialSlotName);
+			if (MaterialIndex == INDEX_NONE)
+			{
+				MaterialIndex = PolygonGroupID.GetValue();
+			}
+			if (!StaticMesh->GetStaticMaterials().IsValidIndex(MaterialIndex))
+			{
+				MaterialIndex = 0;
+			}
+
 			if (OverrideMaterials && OverrideMaterials->IsValidIndex(PolygonGroupID.GetValue()))
 			{
 				Material = (*OverrideMaterials)[PolygonGroupID.GetValue()];
 			}
 			else
 			{
-				FName CurrentMaterialSlotName = PolygonGroupMaterialSlotNames[PolygonGroupID];
-				int32 MaterialIndex = StaticMesh->GetMaterialIndexFromImportedMaterialSlotName(CurrentMaterialSlotName);
-				if(MaterialIndex == INDEX_NONE)
-				{
-					MaterialIndex = PolygonGroupID.GetValue();
-				}
-				if(!StaticMesh->GetStaticMaterials().IsValidIndex(MaterialIndex))
-				{
-					MaterialIndex = 0;
-				}
 				Material = StaticMesh->GetMaterial(MaterialIndex);
 			}
 
-			FbxSurfaceMaterial* FbxMaterial = Material ? ExportMaterial(Material) : nullptr;
+			FbxSurfaceMaterial* FbxMaterial = Material ? ExportMaterial(Material, MaterialIndex, MaterialBakingMeshData) : nullptr;
 			if (!FbxMaterial)
 			{
 				FbxMaterial = CreateDefaultMaterial();
@@ -4906,7 +5003,8 @@ bool FFbxExporter::ExportStaticMeshFromRenderData(FbxMesh* Mesh
 	, int32 LightmapUVChannel
 	, const FColorVertexBuffer* ColorBuffer
 	, const TArray<FStaticMaterial>* MaterialOrderOverride
-	, const TArray<UMaterialInterface*>* OverrideMaterials)
+	, const TArray<UMaterialInterface*>* OverrideMaterials
+	, const FFbxMaterialBakingMeshData& MaterialBakingMeshData)
 {
 	// Verify the integrity of the static mesh.
 	if (RenderMesh.VertexBuffers.StaticMeshVertexBuffer.GetNumVertices() == 0)
@@ -5100,7 +5198,7 @@ bool FFbxExporter::ExportStaticMeshFromRenderData(FbxMesh* Mesh
 			Material = StaticMesh->GetMaterial(Polygons.MaterialIndex);
 		}
 
-		FbxSurfaceMaterial* FbxMaterial = Material ? ExportMaterial(Material) : NULL;
+		FbxSurfaceMaterial* FbxMaterial = Material ? ExportMaterial(Material, Polygons.MaterialIndex, MaterialBakingMeshData) : NULL;
 		if (!FbxMaterial)
 		{
 			FbxMaterial = CreateDefaultMaterial();
@@ -5207,7 +5305,7 @@ bool FFbxExporter::ExportStaticMeshFromRenderData(FbxMesh* Mesh
  * @param ColorBuffer	Vertex color overrides to export
  * @param MaterialOrderOverride	Optional ordering of materials to set up correct material ID's across multiple meshes being export such as BSP surfaces which share common materials. Should be used sparingly
  */
-FbxNode* FFbxExporter::ExportStaticMeshToFbx(const UStaticMesh* StaticMesh, int32 ExportLOD, const TCHAR* MeshName, FbxNode* FbxActor, int32 LightmapUVChannel /*= -1*/, const FColorVertexBuffer* ColorBuffer /*= NULL*/, const TArray<FStaticMaterial>* MaterialOrderOverride /*= NULL*/, const TArray<UMaterialInterface*>* OverrideMaterials /*= NULL*/)
+FbxNode* FFbxExporter::ExportStaticMeshToFbx(const UStaticMesh* StaticMesh, int32 ExportLOD, const TCHAR* MeshName, FbxNode* FbxActor, const FFbxMaterialBakingMeshData& MaterialBakingMeshData, int32 LightmapUVChannel /*= -1*/, const FColorVertexBuffer* ColorBuffer /*= NULL*/, const TArray<FStaticMaterial>* MaterialOrderOverride /*= NULL*/, const TArray<UMaterialInterface*>* OverrideMaterials /*= NULL*/)
 {
 	FbxMesh* Mesh = nullptr;
 	if ((ExportLOD == 0 || ExportLOD == -1) && LightmapUVChannel == -1 && ColorBuffer == nullptr && MaterialOrderOverride == nullptr)
@@ -5230,7 +5328,7 @@ FbxNode* FFbxExporter::ExportStaticMeshToFbx(const UStaticMesh* StaticMesh, int3
 			if (bUseNaniteData)
 			{
 				//Export the nanite mesh description
-				if (!ExportStaticMeshFromMeshDescription(Mesh, StaticMesh, StaticMesh->GetHiResMeshDescription(), FbxActor, LightmapUVChannel, MaterialOrderOverride, OverrideMaterials))
+				if (!ExportStaticMeshFromMeshDescription(Mesh, StaticMesh, StaticMesh->GetHiResMeshDescription(), FbxActor, LightmapUVChannel, MaterialOrderOverride, OverrideMaterials, MaterialBakingMeshData))
 				{
 					return nullptr;
 				}
@@ -5239,7 +5337,7 @@ FbxNode* FFbxExporter::ExportStaticMeshToFbx(const UStaticMesh* StaticMesh, int3
 			{
 				ensure(bUseLodData);
 				//Export the lod mesh description
-				if(!ExportStaticMeshFromMeshDescription(Mesh, StaticMesh, StaticMesh->GetMeshDescription(LodIndex), FbxActor, LightmapUVChannel, MaterialOrderOverride, OverrideMaterials))
+				if(!ExportStaticMeshFromMeshDescription(Mesh, StaticMesh, StaticMesh->GetMeshDescription(LodIndex), FbxActor, LightmapUVChannel, MaterialOrderOverride, OverrideMaterials, MaterialBakingMeshData))
 				{
 					return nullptr;
 				}
@@ -5249,7 +5347,7 @@ FbxNode* FFbxExporter::ExportStaticMeshToFbx(const UStaticMesh* StaticMesh, int3
 		{
 			//Export the render data
 			const FStaticMeshLODResources& RenderMesh = StaticMesh->GetLODForExport(LodIndex);
-			if (!ExportStaticMeshFromRenderData(Mesh, StaticMesh, RenderMesh, FbxActor, LightmapUVChannel, ColorBuffer, MaterialOrderOverride, OverrideMaterials))
+			if (!ExportStaticMeshFromRenderData(Mesh, StaticMesh, RenderMesh, FbxActor, LightmapUVChannel, ColorBuffer, MaterialOrderOverride, OverrideMaterials, MaterialBakingMeshData))
 			{
 				return nullptr;
 			}
@@ -5282,7 +5380,7 @@ FbxNode* FFbxExporter::ExportStaticMeshToFbx(const UStaticMesh* StaticMesh, int3
 				Material = StaticMesh->GetMaterial(Polygons.MaterialIndex);
 			}
 
-			FbxSurfaceMaterial* FbxMaterial = Material ? ExportMaterial(Material) : NULL;
+			FbxSurfaceMaterial* FbxMaterial = Material ? ExportMaterial(Material, Polygons.MaterialIndex, MaterialBakingMeshData) : NULL;
 			if (!FbxMaterial)
 			{
 				FbxMaterial = CreateDefaultMaterial();
@@ -5304,7 +5402,7 @@ FbxNode* FFbxExporter::ExportStaticMeshToFbx(const UStaticMesh* StaticMesh, int3
 	return FbxActor;
 }
 
-void FFbxExporter::ExportSplineMeshToFbx(const USplineMeshComponent* SplineMeshComp, const TCHAR* MeshName, FbxNode* FbxActor)
+void FFbxExporter::ExportSplineMeshToFbx(const USplineMeshComponent* SplineMeshComp, const TCHAR* MeshName, FbxNode* FbxActor, const FFbxMaterialBakingMeshData& MaterialBakingMeshData)
 {
 	const UStaticMesh* StaticMesh = SplineMeshComp->GetStaticMesh();
 	check(StaticMesh);
@@ -5468,7 +5566,7 @@ void FFbxExporter::ExportSplineMeshToFbx(const USplineMeshComponent* SplineMeshC
 		FIndexArrayView RawIndices = RenderMesh.IndexBuffer.GetArrayView();
 		UMaterialInterface* Material = StaticMesh->GetMaterial(Polygons.MaterialIndex);
 
-		FbxSurfaceMaterial* FbxMaterial = Material ? ExportMaterial(Material) : NULL;
+		FbxSurfaceMaterial* FbxMaterial = Material ? ExportMaterial(Material, Polygons.MaterialIndex, MaterialBakingMeshData) : NULL;
 		if (!FbxMaterial)
 		{
 			FbxMaterial = CreateDefaultMaterial();
@@ -5543,7 +5641,7 @@ void FFbxExporter::ExportSplineMeshToFbx(const USplineMeshComponent* SplineMeshC
 	FbxActor->SetNodeAttribute(Mesh);
 }
 
-void FFbxExporter::ExportInstancedMeshToFbx(const UInstancedStaticMeshComponent* InstancedMeshComp, const TCHAR* MeshName, FbxNode* FbxActor)
+void FFbxExporter::ExportInstancedMeshToFbx(const UInstancedStaticMeshComponent* InstancedMeshComp, const TCHAR* MeshName, FbxNode* FbxActor, const FFbxMaterialBakingMeshData& MaterialBakingMeshData)
 {
 	const UStaticMesh* StaticMesh = InstancedMeshComp->GetStaticMesh();
 	check(StaticMesh);
@@ -5565,7 +5663,7 @@ void FFbxExporter::ExportInstancedMeshToFbx(const UInstancedStaticMeshComponent*
 			const int32 LightmapUVChannel = -1;
 			const TArray<FStaticMaterial>* MaterialOrderOverride = nullptr;
 			const FColorVertexBuffer* ColorBuffer = nullptr;
-			ExportStaticMeshToFbx(StaticMesh, LODIndex, *FString::Printf(TEXT("%d"), InstanceIndex), InstNode, LightmapUVChannel, ColorBuffer, MaterialOrderOverride, &ToRawPtrTArrayUnsafe(InstancedMeshComp->OverrideMaterials));
+			ExportStaticMeshToFbx(StaticMesh, LODIndex, *FString::Printf(TEXT("%d"), InstanceIndex), InstNode, MaterialBakingMeshData, LightmapUVChannel, ColorBuffer, MaterialOrderOverride, &ToRawPtrTArrayUnsafe(InstancedMeshComp->OverrideMaterials));
 			FbxActor->AddChild(InstNode);
 		}
 	}
@@ -5743,13 +5841,13 @@ void FFbxExporter::ExportLandscapeToFbx(ALandscapeProxy* Landscape, const TCHAR*
 	Layer0->SetMaterials(LayerElementMaterials);
 
 	UMaterialInterface* Material = Landscape->GetLandscapeMaterial();
-	FbxSurfaceMaterial* FbxMaterial = Material ? ExportMaterial(Material) : NULL;
+	FbxSurfaceMaterial* FbxMaterial = Material ? ExportMaterial(Material, 0, FFbxMaterialBakingMeshData()) : NULL;
 	if (!FbxMaterial)
 	{
 		FbxMaterial = CreateDefaultMaterial();
 	}
-	const int32 MaterialIndex = FbxActor->AddMaterial(FbxMaterial);
-	LayerElementMaterials->GetIndexArray().Add(MaterialIndex);
+	const int32 FbxMaterialIndex = FbxActor->AddMaterial(FbxMaterial);
+	LayerElementMaterials->GetIndexArray().Add(FbxMaterialIndex);
 
 	const int32 VisThreshold = 170;
 	// Copy over the index buffer into the FBX polygons set.
diff --git a/Engine/Source/Editor/UnrealEd/Private/Fbx/FbxMaterialExportUtilities.cpp b/Engine/Source/Editor/UnrealEd/Private/Fbx/FbxMaterialExportUtilities.cpp
new file mode 100644
--- /dev/null
+++ b/Engine/Source/Editor/UnrealEd/Private/Fbx/FbxMaterialExportUtilities.cpp
@@ -0,0 +1,1049 @@
+// Copyright Epic Games, Inc. All Rights Reserved.
+
+#include "FbxMaterialExportUtilities.h"
+
+#include "Exporters/FbxExportOption.h"
+
+#include "Editor.h"
+#include "IMeshMergeUtilities.h"
+#include "MeshMergeModule.h"
+#include "Modules/ModuleManager.h"
+#include "IImageWrapperModule.h"
+#include "IImageWrapper.h"
+#include "IMaterialBakingModule.h"
+#include "Misc/FileHelper.h"
+#include "EditorFramework/AssetImportData.h"
+
+#include "Rendering/SkeletalMeshRenderData.h"
+#include "Components/BrushComponent.h"
+#include "Components/SkeletalMeshComponent.h"
+#include "Components/StaticMeshComponent.h"
+#include "StaticMeshAttributes.h"
+#include "StaticMeshComponentLODInfo.h"
+
+#include "Engine/MapBuildDataRegistry.h"
+#include "Engine/TextureRenderTarget2D.h"
+#include "Engine/Brush.h"
+#include "MaterialPropertyEx.h"
+#include "MaterialBakingStructures.h"
+#include "Materials/Material.h"
+#include "Materials/MaterialInstance.h"
+#include "Materials/MaterialExpressionCustomOutput.h"
+#include "Materials/MaterialExpressionMaterialFunctionCall.h"
+#include "Materials/MaterialExpressionTextureSample.h"
+#include "Materials/MaterialExpressionTextureSampleParameter2D.h"
+#include "Materials/MaterialExpressionTextureObject.h"
+#include "Materials/MaterialExpressionTextureObjectParameter.h"
+#include "Materials/MaterialExpressionConstant.h"
+#include "Materials/MaterialExpressionConstant2Vector.h"
+#include "Materials/MaterialExpressionConstant3Vector.h"
+#include "Materials/MaterialExpressionConstant4Vector.h"
+#include "Materials/MaterialExpressionLinearInterpolate.h"
+#include "Materials/MaterialExpressionScalarParameter.h"
+#include "TextureCompiler.h"
+
+#include "InterchangeAssetImportData.h"
+
+THIRD_PARTY_INCLUDES_START
+#include <fbxsdk.h>
+THIRD_PARTY_INCLUDES_END
+
+namespace UnFbx
+{
+	namespace Private
+	{
+		UMaterialExpressionCustomOutput* GetCustomOutputByName(const UMaterialInterface* Material, const FString& FunctionName)
+		{
+			for (const TObjectPtr<UMaterialExpression>& Expression : Material->GetMaterial()->GetExpressions())
+			{
+				UMaterialExpressionCustomOutput* CustomOutput = Cast<UMaterialExpressionCustomOutput>(Expression);
+				if (CustomOutput != nullptr && CustomOutput->GetFunctionName() == FunctionName)
+				{
+					return CustomOutput;
+				}
+			}
+
+			return nullptr;
+		}
+
+		const FExpressionInput* GetInputForProperty(const UMaterialInterface* Material, const FMaterialPropertyEx& Property)
+		{
+			if (Property.IsCustomOutput())
+			{
+				const FString FunctionName = Property.CustomOutput.ToString();
+				const UMaterialExpressionCustomOutput* CustomOutput = GetCustomOutputByName(Material, FunctionName);
+				if (CustomOutput == nullptr)
+				{
+					return nullptr;
+				}
+
+				// Assume custom outputs always have a single input (which is true for all supported custom outputs)
+				return const_cast<UMaterialExpressionCustomOutput*>(CustomOutput)->GetInput(0);
+			}
+
+			UMaterial* UnderlyingMaterial = const_cast<UMaterial*>(Material->GetMaterial());
+			return UnderlyingMaterial->GetExpressionInputForProperty(Property.Type);
+		}
+
+		template <typename ExpressionType>
+		void GetAllInputExpressionsOfType(const UMaterialInterface* Material, const FMaterialPropertyEx& Property, TArray<ExpressionType*>& OutExpressions)
+		{
+			const FExpressionInput* Input = GetInputForProperty(Material, Property);
+			if (Input == nullptr)
+			{
+				return;
+			}
+
+			UMaterialExpression* InputExpression = Input->Expression;
+			if (InputExpression == nullptr)
+			{
+				return;
+			}
+
+			TArray<UMaterialExpression*> AllInputExpressions;
+			InputExpression->GetAllInputExpressions(AllInputExpressions);
+
+			for (UMaterialExpression* Expression : AllInputExpressions)
+			{
+				if (ExpressionType* ExpressionOfType = Cast<ExpressionType>(Expression))
+				{
+					OutExpressions.Add(ExpressionOfType);
+				}
+
+				if (UMaterialFunctionInterface* MaterialFunction = UMaterial::GetExpressionFunctionPointer(Expression))
+				{
+					MaterialFunction->GetAllExpressionsOfType<ExpressionType>(OutExpressions);
+				}
+				else if (TOptional<UMaterial::FLayersInterfaces> LayersInterfaces = UMaterial::GetExpressionLayers(Expression))
+				{
+					for (UMaterialFunctionInterface* Layer : LayersInterfaces->Layers)
+					{
+						if (Layer != nullptr)
+						{
+							Layer->GetAllExpressionsOfType<ExpressionType>(OutExpressions);
+						}
+					}
+
+					for (UMaterialFunctionInterface* Blend : LayersInterfaces->Blends)
+					{
+						if (Blend != nullptr)
+						{
+							Blend->GetAllExpressionsOfType<ExpressionType>(OutExpressions);
+						}
+					}
+				}
+			}
+		}
+		UTexture* GetTextureFromSample(const UMaterialInterface* Material, const UMaterialExpressionTextureSample* SampleExpression)
+		{
+			if (const UMaterialExpressionTextureSampleParameter2D* SampleParameter = ExactCast<UMaterialExpressionTextureSampleParameter2D>(SampleExpression))
+			{
+				UTexture* ParameterValue = SampleParameter->Texture;
+
+				if (!Material->GetTextureParameterValue(SampleParameter->GetParameterName(), ParameterValue))
+				{
+					return nullptr;
+				}
+
+				return ParameterValue;
+			}
+
+			if (const UMaterialExpressionTextureSample* Sample = ExactCast<UMaterialExpressionTextureSample>(SampleExpression))
+			{
+				UMaterialExpression* ObjectExpression = Sample->TextureObject.Expression;
+				if (ObjectExpression == nullptr)
+				{
+					return Sample->Texture;
+				}
+
+				if (const UMaterialExpressionTextureObjectParameter* ObjectParameter = ExactCast<UMaterialExpressionTextureObjectParameter>(ObjectExpression))
+				{
+					UTexture* ParameterValue = ObjectParameter->Texture;
+
+					if (!Material->GetTextureParameterValue(ObjectParameter->GetParameterName(), ParameterValue))
+					{
+						return nullptr;
+					}
+
+					return ParameterValue;
+				}
+
+				if (const UMaterialExpressionTextureObject* Object = ExactCast<UMaterialExpressionTextureObject>(ObjectExpression))
+				{
+					return Object->Texture;
+				}
+
+				return nullptr;
+			}
+
+			if (const UMaterialExpressionTextureObjectParameter* ObjectParameter = ExactCast<UMaterialExpressionTextureObjectParameter>(SampleExpression))
+			{
+				UTexture* ParameterValue = ObjectParameter->Texture;
+
+				if (!Material->GetTextureParameterValue(ObjectParameter->GetParameterName(), ParameterValue))
+				{
+					return nullptr;
+				}
+
+				return ParameterValue;
+			}
+
+			if (const UMaterialExpressionTextureObject* Object = ExactCast<UMaterialExpressionTextureObject>(SampleExpression))
+			{
+				return Object->Texture;
+			}
+
+
+			return nullptr;
+		}
+		FIntPoint TryGetMaxTextureSize(const UMaterialInterface* Material, const FMaterialPropertyEx& Property, const FIntPoint& DefaultMaxSize)
+		{
+			TArray<UMaterialExpressionTextureSample*> TextureSamples;
+			GetAllInputExpressionsOfType(Material, Property, TextureSamples);
+
+			if (TextureSamples.Num() == 0)
+			{
+				return DefaultMaxSize;
+			}
+
+			FIntPoint MaxSize = { 0, 0 };
+
+			for (const UMaterialExpressionTextureSample* TextureSample : TextureSamples)
+			{
+				UTexture* Texture = GetTextureFromSample(Material, TextureSample);
+				if (Texture == nullptr ||
+					!(Texture->IsA<UTexture2D>() || Texture->IsA<UTextureRenderTarget2D>()))
+				{
+					continue;
+				}
+
+
+				//Load Texture
+				{
+#if WITH_EDITOR
+					FTextureCompilingManager::Get().FinishCompilation({ Texture });
+#endif
+					Texture->SetForceMipLevelsToBeResident(30.0f);
+					Texture->WaitForStreaming();
+				}
+
+				auto GetMipBias = [](const UTexture* Texture)
+					{
+						if (const UTexture2D* Texture2D = Cast<UTexture2D>(Texture))
+						{
+							return Texture2D->GetNumMips() - Texture2D->GetNumMipsAllowed(true);
+						}
+
+						return Texture->GetCachedLODBias();
+					};
+
+				auto GetInGameSize = [&GetMipBias](const UTexture* Texture) -> FIntPoint
+					{
+						const int32 Width = FMath::CeilToInt(Texture->GetSurfaceWidth());
+						const int32 Height = FMath::CeilToInt(Texture->GetSurfaceHeight());
+
+						const int32 MipBias = GetMipBias(Texture);
+
+						const int32 InGameWidth = FMath::Max(Width >> MipBias, 1);
+						const int32 InGameHeight = FMath::Max(Height >> MipBias, 1);
+
+						return { InGameWidth, InGameHeight };
+					};
+
+				FIntPoint TextureSize = GetInGameSize(Texture);
+
+				MaxSize = MaxSize.ComponentMax(TextureSize);
+			}
+
+			return (MaxSize.X == 0 || MaxSize.Y == 0) ? DefaultMaxSize : MaxSize;
+		}
+
+		void AnalyzeMaterialProperty(const UMaterialInterface* InMaterial, const FMaterialPropertyEx& InProperty, FMaterialAnalysisResult& OutAnalysis)
+		{
+			if (GetInputForProperty(InMaterial, InProperty) == nullptr)
+			{
+				OutAnalysis = FMaterialAnalysisResult();
+				return;
+			}
+
+			UMaterial* BaseMaterial = const_cast<UMaterial*>(InMaterial->GetMaterial());
+			bool bRequiresPrimitiveData = false;
+
+			if (BaseMaterial)
+			{
+				BaseMaterial->AnalyzeMaterialPropertyEx(InProperty.Type, OutAnalysis);
+			}
+
+			// Also make sure the analysis takes into account primitive data
+			OutAnalysis.bRequiresVertexData |= bRequiresPrimitiveData;
+		}
+
+		bool NeedsMeshDataForProperty(const UMaterialInterface* Material, const FMaterialPropertyEx& Property)
+		{
+			if (Material != nullptr)
+			{
+				FMaterialAnalysisResult Analysis;
+
+				AnalyzeMaterialProperty(Material, Property, Analysis);
+
+				return Analysis.bRequiresVertexData;
+			}
+
+			return false;
+		}
+
+		void GetAllTextureCoordinateIndices(const UMaterialInterface* Material, const FMaterialPropertyEx& Property, TArray<int32>& OutTexCoords)
+		{
+			FMaterialAnalysisResult Analysis;
+			AnalyzeMaterialProperty(Material, Property, Analysis);
+
+			const TBitArray<>& TexCoords = Analysis.TextureCoordinates;
+			for (int32 Index = 0; Index < TexCoords.Num(); Index++)
+			{
+				if (TexCoords[Index])
+				{
+					OutTexCoords.Add(Index);
+				}
+			}
+		}
+
+		void TransformColorSpace(TArray<FColor>& Pixels, bool bFromSRGB, bool bToSRGB)
+		{
+			if (bFromSRGB == bToSRGB)
+			{
+				return;
+			}
+
+			if (bToSRGB)
+			{
+				for (FColor& Pixel : Pixels)
+				{
+					Pixel = Pixel.ReinterpretAsLinear().ToFColor(true);
+				}
+			}
+			else
+			{
+				for (FColor& Pixel : Pixels)
+				{
+					Pixel = FLinearColor(Pixel).ToFColor(false);
+				}
+			}
+		}
+
+		TArray<int32> GetSectionIndices(const UStaticMesh* StaticMesh, int32 LODIndex, const int32& MaterialIndex)
+		{
+			if (StaticMesh == nullptr)
+			{
+				return {};
+			}
+
+			const FStaticMeshLODResources& RenderData = StaticMesh->GetLODForExport(LODIndex);
+			const FStaticMeshSectionArray& Sections = RenderData.Sections;
+
+			TArray<int32> SectionIndices;
+			SectionIndices.Reserve(Sections.Num());
+
+			for (int32 SectionIndex = 0; SectionIndex < Sections.Num(); ++SectionIndex)
+			{
+				if (Sections[SectionIndex].MaterialIndex == MaterialIndex)
+				{
+					SectionIndices.Add(SectionIndex);
+				}
+			}
+
+			return SectionIndices;
+		}
+
+		TArray<int32> GetSectionIndices(const USkeletalMesh* SkeletalMesh, int32 LODIndex, const int32& MaterialIndex)
+		{
+			if (SkeletalMesh == nullptr)
+			{
+				return {};
+			}
+
+			const FSkeletalMeshLODRenderData& RenderData = SkeletalMesh->GetResourceForRendering()->LODRenderData[LODIndex];			
+			const TArray<FSkelMeshRenderSection>& Sections = RenderData.RenderSections;
+
+			TArray<int32> SectionIndices;
+			SectionIndices.Reserve(Sections.Num());
+
+			for (int32 SectionIndex = 0; SectionIndex < Sections.Num(); ++SectionIndex)
+			{
+				if (Sections[SectionIndex].MaterialIndex == MaterialIndex)
+				{
+					SectionIndices.Add(SectionIndex);
+				}
+			}
+
+			return SectionIndices;
+		}
+
+		namespace InterchangeMaterialProcessHelpers
+		{
+			FString LambertSurfaceMaterialPath = TEXT("/Interchange/Materials/LambertSurfaceMaterial.LambertSurfaceMaterial");
+			FString PhongSurfaceMaterialPath = TEXT("/Interchange/Materials/PhongSurfaceMaterial.PhongSurfaceMaterial");
+			FString MF_PhongToMetalRoughnessPathName = TEXT("/Interchange/Functions/MF_PhongToMetalRoughness.MF_PhongToMetalRoughness");
+
+			bool HandleMaterialProperty(UMaterialInterface* MaterialInterface, fbxsdk::FbxScene* Scene, fbxsdk::FbxSurfaceMaterial* FbxMaterial,
+				const FString& ParameterPropertyName, const FMaterialPropertyEx& Property, 
+				const char* FbxPropertyName, const char* FbxFactorPropertyName,
+				bool bPrimaryPropertyIsMap = true)
+				{
+					fbxsdk::FbxProperty FbxProperty = FbxMaterial->FindProperty(FbxPropertyName);
+					fbxsdk::FbxProperty FbxFactorProperty = FbxMaterial->FindProperty(FbxFactorPropertyName);
+					if (!FbxProperty.IsValid())
+					{
+						return false;
+					}
+
+					bool bPropertyProcessed = false;
+
+					//If Map is present(non-default) then Factor should be acquired from the Weight (and Color is presumed black and does not need to be set on the fbx material (it cannot be set))
+					//If no Map is present then the factor cannot be deduced, and just the Color should be set(If R = G = B then it should be a float otherwise RGB)
+
+					//Use Map
+					UTexture* Texture = nullptr;
+					if (bPrimaryPropertyIsMap && MaterialInterface->GetTextureParameterValue(*(ParameterPropertyName + TEXT("Map")), Texture, true))
+					{
+						FString TextureSourceFullPath = Texture->AssetImportData->GetFirstFilename();
+						//Create a fbx property
+						FbxFileTexture* lTexture = FbxFileTexture::Create(Scene, "EnvSamplerTex");
+						lTexture->SetFileName(TCHAR_TO_UTF8(*TextureSourceFullPath));
+						lTexture->SetTextureUse(FbxTexture::eStandard);
+						lTexture->SetMappingType(FbxTexture::eUV);
+						lTexture->ConnectDstProperty(FbxProperty);
+						
+						float Weight;
+						if (MaterialInterface->GetScalarParameterValue(*(ParameterPropertyName + TEXT("MapWeight")), Weight, true))
+						{
+							if (FbxFactorProperty.IsValid()) FbxFactorProperty.Set(Weight);
+						}
+
+						bPropertyProcessed = true;
+					}
+					else
+					{
+						FLinearColor Color;
+						float Value;
+						if (MaterialInterface->GetVectorParameterValue(*ParameterPropertyName, Color, false))
+						{
+							FLinearColor DefaultColor;
+							MaterialInterface->GetVectorParameterDefaultValue(*ParameterPropertyName, DefaultColor);
+							
+							if (Color != DefaultColor)
+							{
+								if (Color.R == Color.G
+									&& Color.G == Color.B)
+								{
+									if (FbxProperty.IsValid()) FbxProperty.Set(Color.R);
+								}
+								else
+								{
+									FbxDouble3 FbxColor(Color.R, Color.G, Color.B);
+									if (FbxProperty.IsValid()) FbxProperty.Set(FbxColor);
+								}
+							}
+
+							bPropertyProcessed = true;
+						}
+						else if (MaterialInterface->GetScalarParameterValue(*ParameterPropertyName, Value, false))
+						{
+							float DefaultValue;
+							MaterialInterface->GetScalarParameterDefaultValue(*ParameterPropertyName, DefaultValue);
+
+							if (Value != DefaultValue
+								&& FbxProperty.IsValid())
+							{
+								FbxProperty.Set(Value);
+							}
+
+							bPropertyProcessed = true;
+						}
+					}
+					
+					return bPropertyProcessed;
+				};
+
+			UMaterialExpressionMaterialFunctionCall* GetInterchangeMFPongToMetalRoughness(UMaterialInterface* MaterialInterface)
+			{
+				if (!MaterialInterface || !MaterialInterface->GetMaterial())
+				{
+					return nullptr;
+				}
+
+				//Acquire the Material's expressions and check if any matches the MF_PhongToMetalRoughness MaterialFunctionCall:
+				TConstArrayView<TObjectPtr<UMaterialExpression>> MaterialExpressions = MaterialInterface->GetMaterial()->GetExpressions();
+
+				UMaterialExpressionMaterialFunctionCall* ImportedMaterialFunction = nullptr;
+
+				for (const TObjectPtr<UMaterialExpression>& MaterialExpression : MaterialExpressions)
+				{
+					if (UMaterialExpressionMaterialFunctionCall* FunctionCallExpression = Cast<UMaterialExpressionMaterialFunctionCall>(MaterialExpression))
+					{
+						if (FunctionCallExpression->MaterialFunction)
+						{
+							FString MF_PathName = FunctionCallExpression->MaterialFunction->GetPathName();
+
+							if (MF_PhongToMetalRoughnessPathName == MF_PathName)
+							{
+								ImportedMaterialFunction = FunctionCallExpression;
+								break;
+							}
+						}
+					}
+				}
+
+				return ImportedMaterialFunction;
+			}
+
+			void HandleMaterialExpression(UMaterialInterface* MaterialInterface, fbxsdk::FbxScene* Scene, fbxsdk::FbxSurfaceMaterial* FbxMaterial,
+				const char* FbxPropertyName, const char* FbxFactorPropertyName,
+				UMaterialExpression* MaterialExpression)
+			{
+				fbxsdk::FbxProperty FbxProperty = FbxMaterial->FindProperty(FbxPropertyName);
+				fbxsdk::FbxProperty FbxFactorProperty = FbxMaterial->FindProperty(FbxFactorPropertyName);
+				if (!FbxProperty.IsValid())
+				{
+					return;
+				}
+
+				//Check for Lerp
+				// Interchange import creates a Lepr Setup with:
+				// A => Color
+				// B => Map
+				// Alpha => Weight (Factor)
+				//If Map is present(non-default) then Factor should be acquired from the Weight (and Color is presumed black and does not need to be set on the fbx material (it cannot be set))
+				//If no Map is present then the factor cannot be deduced, and just the Color should be set(If R = G = B then it should be a float otherwise RGB)
+				if (UMaterialExpressionLinearInterpolate* LerpExpr = Cast< UMaterialExpressionLinearInterpolate>(MaterialExpression))
+				{
+					if (UMaterialExpressionTextureBase* TextureBaseExpr = Cast<UMaterialExpressionTextureBase>(LerpExpr->B.Expression))
+					{
+						FString TextureSourceFullPath = TextureBaseExpr->Texture->AssetImportData->GetFirstFilename();
+						//Create a fbx property
+						FbxFileTexture* lTexture = FbxFileTexture::Create(Scene, "EnvSamplerTex");
+						lTexture->SetFileName(TCHAR_TO_UTF8(*TextureSourceFullPath));
+						lTexture->SetTextureUse(FbxTexture::eStandard);
+						lTexture->SetMappingType(FbxTexture::eUV);
+						lTexture->ConnectDstProperty(FbxProperty);
+
+						if (UMaterialExpressionConstant* ConstExpr1 = Cast<UMaterialExpressionConstant>(LerpExpr->Alpha.Expression))
+						{
+							if (FbxFactorProperty.IsValid()) FbxFactorProperty.Set(ConstExpr1->R);
+						}
+						else if (UMaterialExpressionScalarParameter* ScalarParamExpr = Cast<UMaterialExpressionScalarParameter>(LerpExpr->Alpha.Expression))
+						{
+							const FHashedMaterialParameterInfo ParameterInfo(ScalarParamExpr->GetParameterName());
+							float Value;
+							if (MaterialInterface->GetScalarParameterValue(ParameterInfo, Value))
+							{
+								if (FbxFactorProperty.IsValid()) FbxFactorProperty.Set(Value);
+							}
+						}
+					}
+				}
+				else
+				{
+					//If Texture/Constant is directly connected
+					//Then Texture and Constant is targetting the FbxProperty.
+					if (UMaterialExpressionTextureBase* TextureBaseExpr = Cast<UMaterialExpressionTextureBase>(MaterialExpression))
+					{
+						FString TextureSourceFullPath = TextureBaseExpr->Texture->AssetImportData->GetFirstFilename();
+						//Create a fbx property
+						FbxFileTexture* lTexture = FbxFileTexture::Create(Scene, "EnvSamplerTex");
+						lTexture->SetFileName(TCHAR_TO_UTF8(*TextureSourceFullPath));
+						lTexture->SetTextureUse(FbxTexture::eStandard);
+						lTexture->SetMappingType(FbxTexture::eUV);
+						lTexture->ConnectDstProperty(FbxProperty);
+					}
+					else if (UMaterialExpressionConstant* ConstExpr1 = Cast<UMaterialExpressionConstant>(MaterialExpression))
+					{
+						FbxDouble3 FbxColor(ConstExpr1->R, ConstExpr1->R, ConstExpr1->R);
+						FbxProperty.Set(FbxColor);
+					}
+					if (UMaterialExpressionConstant3Vector* ConstExpr3 = Cast<UMaterialExpressionConstant3Vector>(MaterialExpression))
+					{
+						FbxDouble3 FbxColor(ConstExpr3->Constant.R, ConstExpr3->Constant.G, ConstExpr3->Constant.B);
+						FbxProperty.Set(FbxColor);
+					}
+					if (UMaterialExpressionConstant4Vector* ConstExpr4 = Cast<UMaterialExpressionConstant4Vector>(MaterialExpression))
+					{
+						FbxDouble3 FbxColor(ConstExpr4->Constant.R, ConstExpr4->Constant.G, ConstExpr4->Constant.B);
+						FbxProperty.Set(FbxColor);
+					}
+				}
+			}
+
+			void HandleExpressionInput(UMaterialInterface* MaterialInterface, fbxsdk::FbxScene* Scene, fbxsdk::FbxSurfaceMaterial* FbxMaterial,
+				const char* FbxPropertyName, const char* FbxFactorPropertyName,
+				FExpressionInput* ExpressionInput)
+			{
+				if (ExpressionInput && ExpressionInput->Expression && ExpressionInput->IsConnected())
+				{
+					HandleMaterialExpression(MaterialInterface, Scene, FbxMaterial,
+						FbxPropertyName, FbxFactorPropertyName,
+						ExpressionInput->GetTracedInput().Expression);
+				}
+			}
+		}
+	}
+
+	FFbxMaterialBakingMeshData::FFbxMaterialBakingMeshData()
+	{
+	}
+
+	FFbxMaterialBakingMeshData::FFbxMaterialBakingMeshData(UModel* Model, ABrush* Actor, int32 InLODIndex)
+		: bHasMeshData(true)
+		, LODIndex(InLODIndex)
+	{
+		FMeshDescription Mesh;
+		FStaticMeshAttributes MeshAttributes(Description);
+		MeshAttributes.Register();
+		TArray<FStaticMaterial>	Materials;
+		GetBrushMesh(Actor, Model, Mesh, Materials);
+
+		StaticMesh = CreateStaticMesh(Mesh, Materials, GetTransientPackage(), Actor->GetFName());
+
+		const IMeshMergeUtilities& MeshMergeUtilities = FModuleManager::Get().LoadModuleChecked<IMeshMergeModule>("MeshMergeUtilities").GetUtilities();
+		FStaticMeshAttributes(Description).Register();
+
+		if (StaticMeshComponent != nullptr)
+		{
+			//MeshMergeUtilities.RetrieveMeshDescription (for StaticMeshComponent) uses bApplyComponentTransform==true, 
+			// however MeshData->Description is only used for DegenerateVertices/Triangles and UV checks, and material baking.
+			MeshMergeUtilities.RetrieveMeshDescription(StaticMeshComponent, LODIndex, Description, true);
+
+			constexpr int32 LightMapLODIndex = 0; // TODO: why is this zero?
+			if (StaticMeshComponent->LODData.IsValidIndex(LightMapLODIndex))
+			{
+				const FStaticMeshComponentLODInfo& LODData = StaticMeshComponent->LODData[LightMapLODIndex];
+				const FMeshMapBuildData* BuildData = StaticMeshComponent->GetMeshMapBuildData(LODData);
+				if (BuildData != nullptr)
+				{
+					LightMap = BuildData->LightMap;
+					LightMapResourceCluster = BuildData->ResourceCluster;
+				}
+			}
+		}
+		else
+		{
+			MeshMergeUtilities.RetrieveMeshDescription(StaticMesh, LODIndex, Description);
+		}
+
+		LightMapTexCoord = StaticMesh->GetLightMapCoordinateIndex();
+		const int32 NumTexCoords = StaticMesh->GetLODForExport(LODIndex).VertexBuffers.StaticMeshVertexBuffer.GetNumTexCoords();
+		BakeUsingTexCoord = FMath::Min(LightMapTexCoord, NumTexCoords - 1);
+	}
+
+	FFbxMaterialBakingMeshData::FFbxMaterialBakingMeshData(const UStaticMesh* InStaticMesh, const UStaticMeshComponent* InStaticMeshComponent, int32 InLODIndex)
+		: bHasMeshData(true)
+		, StaticMeshComponent(InStaticMeshComponent)
+		, StaticMesh(InStaticMesh)
+		, LODIndex(InLODIndex)
+	{
+		const IMeshMergeUtilities& MeshMergeUtilities = FModuleManager::Get().LoadModuleChecked<IMeshMergeModule>("MeshMergeUtilities").GetUtilities();
+		FStaticMeshAttributes(Description).Register();
+
+		if (StaticMeshComponent != nullptr)
+		{
+			//MeshMergeUtilities.RetrieveMeshDescription (for StaticMeshComponent) uses bApplyComponentTransform==true, 
+			// however MeshData->Description is only used for DegenerateVertices/Triangles and UV checks, and material baking.
+			MeshMergeUtilities.RetrieveMeshDescription(StaticMeshComponent, LODIndex, Description, true);
+
+			constexpr int32 LightMapLODIndex = 0; // TODO: why is this zero?
+			if (StaticMeshComponent->LODData.IsValidIndex(LightMapLODIndex))
+			{
+				const FStaticMeshComponentLODInfo& LODData = StaticMeshComponent->LODData[LightMapLODIndex];
+				const FMeshMapBuildData* BuildData = StaticMeshComponent->GetMeshMapBuildData(LODData);
+				if (BuildData != nullptr)
+				{
+					LightMap = BuildData->LightMap;
+					LightMapResourceCluster = BuildData->ResourceCluster;
+				}
+			}
+		}
+		else
+		{
+			MeshMergeUtilities.RetrieveMeshDescription(StaticMesh, LODIndex, Description);
+		}
+
+		LightMapTexCoord = StaticMesh->GetLightMapCoordinateIndex();
+		const int32 NumTexCoords = StaticMesh->GetLODForExport(LODIndex).VertexBuffers.StaticMeshVertexBuffer.GetNumTexCoords();
+		BakeUsingTexCoord = FMath::Min(LightMapTexCoord, NumTexCoords - 1);
+	}
+
+	FFbxMaterialBakingMeshData::FFbxMaterialBakingMeshData(const USkeletalMesh* InSkeletalMesh, const USkeletalMeshComponent* SkeletalMeshComponent, int32 InLODIndex)
+		: bHasMeshData(true)
+		, SkeletalMesh(InSkeletalMesh)
+		, LODIndex(InLODIndex)
+	{
+		const IMeshMergeUtilities& MeshMergeUtilities = FModuleManager::Get().LoadModuleChecked<IMeshMergeModule>("MeshMergeUtilities").GetUtilities();
+		FStaticMeshAttributes(Description).Register();
+
+		if (SkeletalMeshComponent != nullptr)
+		{
+			MeshMergeUtilities.RetrieveMeshDescription(SkeletalMeshComponent, LODIndex, Description, true);
+		}
+		else
+		{
+			// NOTE: this is a workaround for the fact that there's no overload for FMeshMergeHelpers::RetrieveMesh
+			// that accepts a USkeletalMesh, only a USkeletalMeshComponent.
+			// Writing a custom utility function that would work on a "standalone" skeletal mesh is problematic
+			// since we would need to implement an equivalent of USkinnedMeshComponent::GetCPUSkinnedVertices too.
+
+			if (UWorld* World = GEditor->GetEditorWorldContext().World())
+			{
+				FActorSpawnParameters SpawnParams;
+				SpawnParams.SpawnCollisionHandlingOverride = ESpawnActorCollisionHandlingMethod::AlwaysSpawn;
+				SpawnParams.ObjectFlags |= RF_Transient;
+				SpawnParams.bAllowDuringConstructionScript = true;
+
+				if (AActor* TempActor = World->SpawnActor<AActor>(SpawnParams))
+				{
+					USkeletalMeshComponent* TempComponent = NewObject<USkeletalMeshComponent>(TempActor, TEXT(""), RF_Transient);
+					TempComponent->RegisterComponent();
+					TempComponent->SetSkeletalMesh(const_cast<USkeletalMesh*>(SkeletalMesh));
+
+					MeshMergeUtilities.RetrieveMeshDescription(TempComponent, LODIndex, Description, true);
+
+					World->DestroyActor(TempActor, false, false);
+				}
+			}
+		}
+
+		// TODO: don't assume last UV channel is non-overlapping
+		const int32 NumTexCoords = SkeletalMesh->GetResourceForRendering()->LODRenderData[LODIndex].StaticVertexBuffers.StaticMeshVertexBuffer.GetNumTexCoords();
+		BakeUsingTexCoord = NumTexCoords - 1;
+	}
+
+	TArray<int32> FFbxMaterialBakingMeshData::GetSectionIndices(const int32& MaterialIndex) const
+	{
+		if (StaticMesh)
+		{
+			return Private::GetSectionIndices(StaticMesh, LODIndex, MaterialIndex);
+		}
+
+		if (SkeletalMesh)
+		{
+			return Private::GetSectionIndices(SkeletalMesh, LODIndex, MaterialIndex);
+		}
+
+		return TArray<int32>();
+	}
+
+	int32 FFbxMaterialBakingMeshData::GetUModelStaticMeshMaterialIndex(const UMaterialInterface* MaterialInterface) const
+	{
+		//UModel to StaticMesh generation usaes the MaterialInterface's FName to set the Slot naming:
+		return StaticMesh->GetMaterialIndexFromImportedMaterialSlotName(MaterialInterface->GetFName());
+	}
+
+	namespace FFbxMaterialExportUtilities
+	{
+		void BakeMaterialProperty(const UFbxExportOption* FbxExportOptions,
+			fbxsdk::FbxScene* Scene, fbxsdk::FbxSurfaceMaterial* FbxMaterial, const char* FbxPropertyName,
+			const FMaterialPropertyEx& Property, const UMaterialInterface* Material, const int32& MaterialIndex,
+			const FFbxMaterialBakingMeshData& MeshData,
+			const FString& ExportFolderPath)
+		{
+			using namespace Private;
+
+			const FIntPoint DefaultBakeSize = FIntPoint(512, 512);
+			const FBox2f TexCoordBounds = FBox2f{ { 0.0f, 0.0f }, { 1.0f, 1.0f } };
+			const bool bFillAlpha = true;
+
+			if (FbxExportOptions->BakeMaterialInputs == EFbxMaterialBakeMode::Disabled)
+			{
+				return;
+			}
+
+			fbxsdk::FbxProperty FbxColorProperty = FbxMaterial->FindProperty(FbxPropertyName);
+			if (!FbxColorProperty.IsValid())
+			{
+				return;
+			}
+
+			FIntPoint BakeSize = FbxExportOptions->DefaultMaterialBakeSize.bAutoDetect ? TryGetMaxTextureSize(Material, Property, DefaultBakeSize) : FbxExportOptions->DefaultMaterialBakeSize.Size;
+
+			if (BakeSize == FIntPoint(0, 0))
+			{
+				return;
+			}
+
+			bool bNeedsMeshData = FbxExportOptions->BakeMaterialInputs == EFbxMaterialBakeMode::UseMeshData && MeshData.bHasMeshData && NeedsMeshDataForProperty(Material, Property);
+
+			TArray<int32> TexCoords;
+			GetAllTextureCoordinateIndices(Material, Property, TexCoords);
+
+			int32 TexCoordIndex = bNeedsMeshData ? MeshData.BakeUsingTexCoord : (TexCoords.Num() > 0 ? TexCoords[0] : 0);
+
+			FMeshData MeshSet;
+			MeshSet.TextureCoordinateBox = FBox2D(TexCoordBounds);
+			MeshSet.TextureCoordinateIndex = TexCoordIndex;
+			MeshSet.MaterialIndices = MeshData.GetSectionIndices(MaterialIndex); // NOTE: MaterialIndices is actually section indices
+			if (bNeedsMeshData)
+			{
+				MeshSet.MeshDescription = const_cast<FMeshDescription*>(&MeshData.Description);
+				MeshSet.LightMap = MeshData.LightMap;
+				MeshSet.LightMapIndex = MeshData.LightMapTexCoord;
+				MeshSet.LightmapResourceCluster = MeshData.LightMapResourceCluster;
+				MeshSet.PrimitiveData = MeshData.StaticMeshComponent ? FPrimitiveData(MeshData.StaticMeshComponent) : (MeshData.StaticMesh ? FPrimitiveData(MeshData.StaticMesh) : FPrimitiveData(MeshData.SkeletalMesh));
+			}
+
+			FMaterialDataEx MatSet;
+			MatSet.Material = const_cast<UMaterialInterface*>(Material);
+			MatSet.PropertySizes.Add(Property, BakeSize);
+			MatSet.bTangentSpaceNormal = true;
+
+			TArray<FMeshData*> MeshSettings;
+			TArray<FMaterialDataEx*> MatSettings;
+			MeshSettings.Add(&MeshSet);
+			MatSettings.Add(&MatSet);
+
+			TArray<FBakeOutputEx> BakeOutputs;
+			IMaterialBakingModule& Module = FModuleManager::Get().LoadModuleChecked<IMaterialBakingModule>("MaterialBaking");
+
+			Module.SetLinearBake(true);
+			Module.BakeMaterials(MatSettings, MeshSettings, BakeOutputs);
+			const bool bIsLinearBake = Module.IsLinearBake(Property);
+			Module.SetLinearBake(false);
+
+			FBakeOutputEx& BakeOutput = BakeOutputs[0];
+
+			TArray<FColor> BakedPixels = MoveTemp(BakeOutput.PropertyData.FindChecked(Property));
+			const FIntPoint BakedSize = BakeOutput.PropertySizes.FindChecked(Property);
+			const float EmissiveScale = BakeOutput.EmissiveScale;
+
+			if (bFillAlpha)
+			{
+				// NOTE: alpha is 0 by default after baking a property, but we prefer 255 (1.0).
+				// It makes it easier to view the exported textures.
+				for (FColor& Pixel : BakedPixels)
+				{
+					Pixel.A = 255;
+				}
+			}
+
+			if (Property == MP_EmissiveColor)
+			{
+				bool bFromSRGB = !bIsLinearBake;
+				bool bToSRGB = true;
+				TransformColorSpace(BakedPixels, bFromSRGB, true);
+			}
+
+			//Save Out:
+			const void* InRawData = BakedPixels.GetData();
+			const int64 ByteLength = BakedSize.X * BakedSize.Y * sizeof(FColor);
+			int64 InRawSize = ByteLength;
+			int32 InWidth = BakedSize.X;
+			int32 InHeight = BakedSize.Y;
+			ERGBFormat InRGBFormat = ERGBFormat::BGRA;
+			int32 InBitDepth = 8;
+			EImageFormat InCompressionFormat = EImageFormat::PNG;
+			int32 InCompressionQuality = 0;
+			TArray64<uint8> OutCompressedData;
+
+			IImageWrapperModule& ImageWrapperModule = FModuleManager::Get().LoadModuleChecked<IImageWrapperModule>(TEXT("ImageWrapper"));
+			const TSharedPtr<IImageWrapper> ImageWrapper = ImageWrapperModule.CreateImageWrapper(InCompressionFormat);
+
+			if (ImageWrapper.IsValid() && ImageWrapper->SetRaw(InRawData, InRawSize, InWidth, InHeight, InRGBFormat, InBitDepth))
+			{
+				const TArray64<uint8>& CompressedData = ImageWrapper->GetCompressed(InCompressionQuality);
+
+				FString FileName = Material->GetName() + TEXT("_") + Property.ToString() + TEXT("_") + FString::FromInt(MaterialIndex) + TEXT(".png");
+				FString FilePath = ExportFolderPath + TEXT("/") + FileName;
+				if (FFileHelper::SaveArrayToFile(CompressedData, *FilePath))
+				{
+					FbxFileTexture* lTexture = FbxFileTexture::Create(Scene, "EnvSamplerTex");
+					lTexture->SetFileName(TCHAR_TO_UTF8(*FilePath)); //SetRelativeFileName does not seem to work?
+					lTexture->SetTextureUse(FbxTexture::eStandard);
+					lTexture->SetMappingType(FbxTexture::eUV);
+					lTexture->ConnectDstProperty(FbxColorProperty);
+				}
+			}
+		}
+
+		bool GetInterchangeShadingModel(const UMaterialInterface* MaterialInterface, bool& bLambert)
+		{
+			using namespace Private::InterchangeMaterialProcessHelpers;
+
+			FString PathName = MaterialInterface->GetMaterial()->GetPathName();
+			if (MaterialInterface->IsA(UMaterialInstance::StaticClass()))
+			{
+				if (PathName == LambertSurfaceMaterialPath
+					|| PathName == PhongSurfaceMaterialPath)
+				{
+					bLambert = PathName == LambertSurfaceMaterialPath;
+
+					return true;
+				}
+			}
+			else
+			{
+				//Acquire the Material's expressions and check if any matches the MF_PhongToMetalRoughness MaterialFunctionCall:
+				TConstArrayView<TObjectPtr<UMaterialExpression>> MaterialExpressions = MaterialInterface->GetMaterial()->GetExpressions();
+
+				bool MFPhongUsed = false;
+				for (const TObjectPtr<UMaterialExpression>& MaterialExpression : MaterialExpressions)
+				{
+					if (UMaterialExpressionMaterialFunctionCall* FunctionCallExpression = Cast<UMaterialExpressionMaterialFunctionCall>(MaterialExpression))
+					{
+						if (FunctionCallExpression->MaterialFunction)
+						{
+							FString MF_PathName = FunctionCallExpression->MaterialFunction->GetPathName();
+
+							if (MF_PhongToMetalRoughnessPathName == MF_PathName)
+							{
+								MFPhongUsed = true;
+								break;
+							}
+						}
+					}
+				}
+
+				bLambert = !MFPhongUsed;
+
+				return MFPhongUsed || (MaterialInterface->AssetImportData.IsA(UInterchangeAssetImportData::StaticClass()));
+			}
+
+			return false;
+		}
+
+		void ProcessInterchangeMaterials(UMaterialInterface* MaterialInterface, fbxsdk::FbxScene* Scene, fbxsdk::FbxSurfaceMaterial* FbxMaterial)
+		{
+			using namespace Private::InterchangeMaterialProcessHelpers;
+
+			FString PathName = MaterialInterface->GetMaterial()->GetPathName();
+
+			if (MaterialInterface->IsA(UMaterialInstance::StaticClass()))
+			{
+				if (PathName == LambertSurfaceMaterialPath
+					|| PathName == PhongSurfaceMaterialPath)
+				{
+
+					bool bSpecularToBeBaked = false; //also indicates to bake the BaseColor and Metallic
+					if (PathName == PhongSurfaceMaterialPath)
+					{
+						//SpecularColor
+						//TEXT("SpecularColor");
+						//TEXT("SpecularColorMap");
+						//TEXT("SpecularColorMapWeight");
+						if (!HandleMaterialProperty(MaterialInterface, Scene, FbxMaterial, TEXT("SpecularColor"), MP_Specular, FbxSurfaceMaterial::sSpecular, FbxSurfaceMaterial::sSpecularFactor))
+						{
+							bSpecularToBeBaked = true;
+						}
+
+						//Roughness
+						//TEXT("Shininess");
+						//TEXT("ShininessMap");
+						//TEXT("ShininessMapWeight");
+						HandleMaterialProperty(MaterialInterface, Scene, FbxMaterial, TEXT("Shininess"), MP_Roughness, FbxSurfaceMaterial::sShininess, "NoFactorForThisProperty", false);
+
+					}
+
+					//Diffuse
+					//TEXT("DiffuseColor");
+					//TEXT("DiffuseColorMap");
+					//TEXT("DiffuseColorMapWeight");
+					if (!bSpecularToBeBaked) HandleMaterialProperty(MaterialInterface, Scene, FbxMaterial, TEXT("DiffuseColor"), MP_BaseColor, FbxSurfaceMaterial::sDiffuse, FbxSurfaceMaterial::sDiffuseFactor);
+
+					//Emissive
+					//TEXT("EmissiveColorMap");
+					//TEXT("EmissiveColorMapWeight");
+					//TEXT("EmissiveColor");
+					HandleMaterialProperty(MaterialInterface, Scene, FbxMaterial, TEXT("EmissiveColor"), MP_EmissiveColor, FbxSurfaceMaterial::sEmissive, FbxSurfaceMaterial::sEmissiveFactor);
+
+					//Normal
+					//TEXT("NormalMap");
+					//TEXT("NormalMapWeight");
+					HandleMaterialProperty(MaterialInterface, Scene, FbxMaterial, TEXT("Normal"), MP_Normal, FbxSurfaceMaterial::sNormalMap, "NoFactorForThisProperty");
+
+					//AmbientOcclusion is NOT AmbientColor, FbxSurfaceMaterail supports AmbientColor!
+					// Meaning Interchange's Lambert and Phong Surface Materials have no support for AmbientColor
+					//TEXT("AmbientOcclusionMap");
+					//TEXT("AmbientOcclusionMapWeight");
+					//HandleProperty(TEXT("AmbientOcclusion"), MP_AmbientOcclusion, FbxSurfaceMaterial::sAmbient, FbxSurfaceMaterial::sAmbientFactor);
+				}
+			}
+			else //!MaterialInstance (==Material)
+			{
+				//Acquire the Material's expressions and check if any matches the MF_PhongToMetalRoughness MaterialFunctionCall:
+				UMaterialExpressionMaterialFunctionCall* ImportedMF_PhongToMetalRoughness = GetInterchangeMFPongToMetalRoughness(MaterialInterface);
+
+				if (ImportedMF_PhongToMetalRoughness)
+				{
+					TArrayView<FExpressionInput*> InputsView = ImportedMF_PhongToMetalRoughness->GetInputsView();
+
+					for (FExpressionInput* Input : InputsView)
+					{
+						if (!Input || !Input->Expression || !Input->IsConnected())
+						{
+							continue;
+						}
+
+						FString Name = Input->InputName.ToString();
+						UMaterialExpression* InputExpression = Input->GetTracedInput().Expression;
+
+						if (Name == TEXT("AmbientColor"))
+						{
+							//Always set the ambient to zero since we dont have ambient in unreal we want to avoid default value in DCCs
+							((FbxSurfaceLambert*)FbxMaterial)->Ambient.Set(FbxDouble3(0.0, 0.0, 0.0));
+							//Overwrite with Interchange values:
+							HandleMaterialExpression(MaterialInterface, Scene, FbxMaterial, FbxSurfaceMaterial::sAmbient, FbxSurfaceMaterial::sAmbientFactor, InputExpression);
+						}
+						else if (Name == TEXT("DiffuseColor"))
+						{
+							HandleMaterialExpression(MaterialInterface, Scene, FbxMaterial, FbxSurfaceMaterial::sDiffuse, FbxSurfaceMaterial::sDiffuseFactor, InputExpression);
+						}
+						else if (Name == TEXT("SpecularColor"))
+						{
+							HandleMaterialExpression(MaterialInterface, Scene, FbxMaterial, FbxSurfaceMaterial::sSpecular, FbxSurfaceMaterial::sSpecularFactor, InputExpression);
+						}
+						else if (Name == TEXT("Shininess"))
+						{
+							HandleMaterialExpression(MaterialInterface, Scene, FbxMaterial, FbxSurfaceMaterial::sShininess, "NoFactorForThisProperty", InputExpression);
+						}
+					}
+				}
+				
+
+				if (MaterialInterface->AssetImportData.IsA(UInterchangeAssetImportData::StaticClass()))
+				{
+					UMaterial* Material = MaterialInterface->GetMaterial();
+
+					if (ImportedMF_PhongToMetalRoughness == nullptr)
+					{
+						//AmbientColor
+						// Only supported via MF_PhongToMetalRoughness
+						
+						//Diffuse
+						HandleExpressionInput(MaterialInterface, Scene, FbxMaterial, FbxSurfaceMaterial::sDiffuse, FbxSurfaceMaterial::sDiffuseFactor, Material->GetExpressionInputForProperty(MP_BaseColor));
+
+						//Specular
+						HandleExpressionInput(MaterialInterface, Scene, FbxMaterial, FbxSurfaceMaterial::sSpecular, FbxSurfaceMaterial::sSpecularFactor, Material->GetExpressionInputForProperty(MP_Specular));
+						
+						//Shininess
+						// Only supported via MF_PhongToMetalRoughness
+					}
+
+					//Normal
+					HandleExpressionInput(MaterialInterface, Scene, FbxMaterial, FbxSurfaceMaterial::sNormalMap, "NoFactorForThisProperty", Material->GetExpressionInputForProperty(MP_Normal));
+
+					//Emissive
+					HandleExpressionInput(MaterialInterface, Scene, FbxMaterial, FbxSurfaceMaterial::sEmissive, FbxSurfaceMaterial::sEmissiveFactor, Material->GetExpressionInputForProperty(MP_EmissiveColor));
+
+
+					UMaterialInstance* MaterialInstance = Cast<UMaterialInstance>(MaterialInterface);
+					EBlendMode BlendMode = MaterialInstance != nullptr && MaterialInstance->BasePropertyOverrides.bOverride_BlendMode ? MaterialInstance->BasePropertyOverrides.BlendMode : Material->BlendMode;
+
+					if (BlendMode == BLEND_Translucent)
+					{
+						//Opacity
+						HandleExpressionInput(MaterialInterface, Scene, FbxMaterial, FbxSurfaceMaterial::sTransparentColor, "NoFactorForThisProperty", Material->GetExpressionInputForProperty(MP_Opacity));
+
+						//OpacityMask
+						HandleExpressionInput(MaterialInterface, Scene, FbxMaterial, FbxSurfaceMaterial::sTransparencyFactor, "NoFactorForThisProperty", Material->GetExpressionInputForProperty(MP_OpacityMask));
+					}
+				}
+			}
+		}
+	}
+}
\ No newline at end of file
diff --git a/Engine/Source/Editor/UnrealEd/Private/Fbx/FbxMaterialExportUtilities.h b/Engine/Source/Editor/UnrealEd/Private/Fbx/FbxMaterialExportUtilities.h
new file mode 100644
--- /dev/null
+++ b/Engine/Source/Editor/UnrealEd/Private/Fbx/FbxMaterialExportUtilities.h
@@ -0,0 +1,97 @@
+// Copyright Epic Games, Inc. All Rights Reserved.
+
+#pragma once
+
+#include "MeshDescription.h"
+#include "SceneTypes.h"
+
+
+class UModel;
+class ABrush;
+class UStaticMesh;
+class UStaticMeshComponent;
+class USkeletalMesh;
+class USkeletalMeshComponent;
+class FLightMap;
+class FLightmapResourceCluster;
+class UFbxExportOption;
+struct FMaterialPropertyEx;
+class UMaterialInterface;
+
+namespace fbxsdk
+{
+	class FbxScene;
+	class FbxSurfaceMaterial;
+}
+
+
+namespace UnFbx
+{
+	/*
+	* Stores the necessary mesh information for Material Baking that requires vertex data
+	*/
+	struct FFbxMaterialBakingMeshData
+	{
+		FFbxMaterialBakingMeshData(); //Used for ULandscapes
+		FFbxMaterialBakingMeshData(UModel* Model, ABrush* Actor = nullptr, int32 LODIndex = 0);
+		FFbxMaterialBakingMeshData(const UStaticMesh* StaticMesh, const UStaticMeshComponent* StaticMeshComponent = nullptr, int32 LODIndex = 0);
+		FFbxMaterialBakingMeshData(const USkeletalMesh* SkeletalMesh, const USkeletalMeshComponent* SkeletalMeshComponent = nullptr, int32 LODIndex = 0);
+
+		bool bHasMeshData = false;
+		FMeshDescription Description;
+
+		FLightMapRef LightMap = nullptr;
+		const FLightmapResourceCluster* LightMapResourceCluster = nullptr;
+		int32 LightMapTexCoord = 0;
+
+		int32 BakeUsingTexCoord = 0;
+
+		//For FPrimitiveData:
+		const UStaticMeshComponent* StaticMeshComponent = nullptr;
+		const UStaticMesh* StaticMesh = nullptr;
+		const USkeletalMesh* SkeletalMesh = nullptr;
+
+		int32 LODIndex = 0;
+
+		TArray<int32> GetSectionIndices(const int32& MaterialIndex) const;
+
+		int32 GetUModelStaticMeshMaterialIndex(const UMaterialInterface* MaterialInterface) const; //Used for UModel
+	};
+
+	namespace FFbxMaterialExportUtilities
+	{
+		/*
+		* Checks Material if it was Interchange imported, and if so, was it Lambert or Phong
+		* Material presumed InterchangeImported when 1 of the following conditions are met:
+		*		- MI -> PathName == Interchange.LambertSurfaceMaterialPath
+		*		- MI -> PathName == Interchange.PhongSurfaceMaterialPath
+		*		- M -> Has MF with path pointing to Interchange.MF_PhongToMetalRoughness
+		*		- M -> AssetImportData == UInterchangeAssetImportData
+		*/
+		bool GetInterchangeShadingModel(const UMaterialInterface* MaterialInterface, bool& bLambert);
+
+
+		/*
+		* Bakes the Material Property. Saves it to png at the given folder path. Sets the Absolute path for the FbxMaterial for the given FbxPropertyName
+		*/
+		void BakeMaterialProperty(const UFbxExportOption* FbxExportOptions,
+			fbxsdk::FbxScene* Scene, fbxsdk::FbxSurfaceMaterial* FbxMaterial, const char* FbxPropertyName,
+			const FMaterialPropertyEx& Property, const UMaterialInterface* MaterialInterface, const int32& MaterialIndex,
+			const FFbxMaterialBakingMeshData& MeshData,
+			const FString& ExportFolderPath);
+
+		/*
+		* Checks and processes for Interchange Lambert/Phong Surface Materials and PhongToMaterialRoughness MaterialFunctions.
+		* Returns Processed Fbx Property Names
+		* Material presumed InterchangeImported when 1 of the following conditions are met:
+		*		- MI -> PathName == Interchange.LambertSurfaceMaterialPath
+		*		- MI -> PathName == Interchange.PhongSurfaceMaterialPath
+		*		- M -> Has MF with path pointing to Interchange.MF_PhongToMetalRoughness
+		*		- M -> AssetImportData == UInterchangeAssetImportData
+		* Note: Possible deviation on Roundtripping:
+		*				-ReflectionFactor is not used on import and is not stored so cannot be retrived on export.
+		*				-AmbientColor is only supported for MF_PhongToMetalRoughness (Interchange's Phong and Lambert Surface Materials do not support it.)
+		*/
+		void ProcessInterchangeMaterials(UMaterialInterface* MaterialInterface, fbxsdk::FbxScene* Scene, fbxsdk::FbxSurfaceMaterial* FbxMaterial);
+	}
+}
\ No newline at end of file
diff --git a/Engine/Source/Editor/UnrealEd/Private/Fbx/FbxSkeletalMeshExport.cpp b/Engine/Source/Editor/UnrealEd/Private/Fbx/FbxSkeletalMeshExport.cpp
--- a/Engine/Source/Editor/UnrealEd/Private/Fbx/FbxSkeletalMeshExport.cpp
+++ b/Engine/Source/Editor/UnrealEd/Private/Fbx/FbxSkeletalMeshExport.cpp
@@ -18,6 +18,8 @@
 #include "Exporters/FbxExportOption.h"
 #include "UObject/MetaData.h"
 
+#include "FbxMaterialExportUtilities.h"
+
 DEFINE_LOG_CATEGORY_STATIC(LogFbxSkeletalMeshExport, Log, All);
 
 namespace UnFbx
@@ -347,6 +349,9 @@ FbxNode* FFbxExporter::CreateMesh(const USkeletalMesh* SkelMesh, const TCHAR* Me
 	const TArray<FSkeletalMaterial>& SkelMeshMaterials = SkelMesh->GetMaterials();
 	int32 MaterialCount = SkelMeshMaterials.Num();
 
+	//TODO: move the init of MaterialBakingMeshData up a callstack step:
+	FFbxMaterialBakingMeshData MaterialBakingMeshData(SkelMesh, nullptr, 0);
+
 	for(int32 MaterialIndex = 0; MaterialIndex < MaterialCount; ++MaterialIndex)
 	{
 		UMaterialInterface* MatInterface = nullptr;
@@ -365,15 +370,20 @@ FbxNode* FFbxExporter::CreateMesh(const USkeletalMesh* SkelMesh, const TCHAR* Me
 		{
 			if (MatInterface)
 			{
-				FbxMaterial = ExportMaterial(MatInterface);
+				FbxMaterial = ExportMaterial(MatInterface, MaterialIndex, MaterialBakingMeshData);
 			}
 		}
 		else if(MatInterface)
 		{
-			if ( FbxSurfaceMaterial** FbxMaterialPtr = FbxMaterials.Find( MatInterface ) )
+			TMap<int32, FbxSurfaceMaterial*>* MaterialIndexToFbxMaterials = FbxMaterials.Find(MatInterface);
+			if (MaterialIndexToFbxMaterials && MaterialIndexToFbxMaterials->Find(MaterialIndex))
 			{
-				FbxMaterial = *FbxMaterialPtr;
+				if (FbxSurfaceMaterial** FbxMaterialPtr = MaterialIndexToFbxMaterials->Find(MaterialIndex))
+				{
+					FbxMaterial = *FbxMaterialPtr;
+				}
 			}
+			
 		}
 
 		if(!FbxMaterial)
diff --git a/Engine/Source/Editor/UnrealEd/Private/FbxExporter.h b/Engine/Source/Editor/UnrealEd/Private/FbxExporter.h
--- a/Engine/Source/Editor/UnrealEd/Private/FbxExporter.h
+++ b/Engine/Source/Editor/UnrealEd/Private/FbxExporter.h
@@ -15,6 +15,9 @@
 #include "UObject/GCObject.h"
 #include "Animation/AnimTypes.h"
 
+#include "SceneTypes.h"
+#include "LightMap.h"
+
 class ABrush;
 class ACameraActor;
 class ALandscapeProxy;
@@ -50,6 +53,8 @@ struct FMovieSceneSequenceTransform;
 
 namespace UnFbx
 {
+	struct FFbxMaterialBakingMeshData;
+
 	/** Adapter interface which allows ExportAnimTrack to act on sequencer without a tight coupling. */
 	class IAnimTrackAdapter
 	{
@@ -286,7 +291,7 @@ private:
 	TMap<FString,int32> FbxNodeNameToIndexMap;
 	TMap<const AActor*, FbxNode*> FbxActors;
 	TMap<const USkeletalMeshComponent*, FbxNode*> FbxSkeletonRoots;
-	TMap<const UMaterialInterface*, FbxSurfaceMaterial*> FbxMaterials;
+	TMap<const UMaterialInterface*, TMap<int32, FbxSurfaceMaterial*>> FbxMaterials;
 	TMap<const UStaticMesh*, FbxMesh*> FbxMeshes;
 	TMap<const UStaticMesh*, FbxMesh*> FbxCollisionMeshes;
 
@@ -304,7 +309,7 @@ private:
 	*/
 	UNREALED_API void ExportAnimTrack( IAnimTrackAdapter& AnimTrackAdapter, AActor* Actor, USkeletalMeshComponent* SkeletalMeshComponent, double SamplingRate );
 
-	UNREALED_API void ExportModel(UModel* Model, FbxNode* Node, const char* Name);
+	UNREALED_API void ExportModel(UModel* Model, FbxNode* Node, const char* Name, const FFbxMaterialBakingMeshData& MaterialBakingMeshData);
 
 	UNREALED_API FbxNode* ExportCollisionMesh(const UStaticMesh* StaticMesh, const TCHAR* MeshName, FbxNode* ParentActor);
 
@@ -319,7 +324,7 @@ private:
 	 * @param MaterialOrderOverride	Optional ordering of materials to set up correct material ID's across multiple meshes being export such as BSP surfaces which share common materials. Should be used sparingly
 	 * @param OverrideMaterials	Optional array of materials to be used instead of the static mesh materials. Used for material overrides in static mesh components.
 	 */
-	UNREALED_API FbxNode* ExportStaticMeshToFbx(const UStaticMesh* StaticMesh, int32 ExportLOD, const TCHAR* MeshName, FbxNode* FbxActor, int32 LightmapUVChannel = -1, const FColorVertexBuffer* ColorBuffer = nullptr, const TArray<FStaticMaterial>* MaterialOrderOverride = nullptr, const TArray<UMaterialInterface*>* OverrideMaterials = nullptr);
+	UNREALED_API FbxNode* ExportStaticMeshToFbx(const UStaticMesh* StaticMesh, int32 ExportLOD, const TCHAR* MeshName, FbxNode* FbxActor, const FFbxMaterialBakingMeshData& MaterialBakingMeshData, int32 LightmapUVChannel = -1, const FColorVertexBuffer* ColorBuffer = nullptr, const TArray<FStaticMaterial>* MaterialOrderOverride = nullptr, const TArray<UMaterialInterface*>* OverrideMaterials = nullptr);
 
 	UNREALED_API bool ExportStaticMeshFromMeshDescription(FbxMesh* Mesh
 		, const UStaticMesh* StaticMesh
@@ -327,7 +332,8 @@ private:
 		, FbxNode* FbxActor
 		, int32 LightmapUVChannel
 		, const TArray<FStaticMaterial>* MaterialOrderOverride
-		, const TArray<UMaterialInterface*>* OverrideMaterials);
+		, const TArray<UMaterialInterface*>* OverrideMaterials
+		, const FFbxMaterialBakingMeshData& MaterialBakingMeshData);
 
 	UNREALED_API bool ExportStaticMeshFromRenderData(FbxMesh* Mesh
 		, const UStaticMesh* StaticMesh
@@ -336,7 +342,8 @@ private:
 		, int32 LightmapUVChannel
 		, const FColorVertexBuffer* ColorBuffer
 		, const TArray<FStaticMaterial>* MaterialOrderOverride
-		, const TArray<UMaterialInterface*>* OverrideMaterials);
+		, const TArray<UMaterialInterface*>* OverrideMaterials
+		, const FFbxMaterialBakingMeshData& MaterialBakingMeshData);
 
 	/**
 	 * Exports a spline mesh
@@ -344,7 +351,7 @@ private:
 	 * @param MeshName		The name of the mesh for the FBX file
 	 * @param FbxActor		The fbx node representing the mesh
 	 */
-	UNREALED_API void ExportSplineMeshToFbx(const USplineMeshComponent* SplineMeshComp, const TCHAR* MeshName, FbxNode* FbxActor);
+	UNREALED_API void ExportSplineMeshToFbx(const USplineMeshComponent* SplineMeshComp, const TCHAR* MeshName, FbxNode* FbxActor, const FFbxMaterialBakingMeshData& MaterialBakingMeshData);
 
 	/**
 	 * Exports an instanced mesh
@@ -352,7 +359,7 @@ private:
 	 * @param MeshName		The name of the mesh for the FBX file
 	 * @param FbxActor		The fbx node representing the mesh
 	 */
-	UNREALED_API void ExportInstancedMeshToFbx(const UInstancedStaticMeshComponent* InstancedMeshComp, const TCHAR* MeshName, FbxNode* FbxActor);
+	UNREALED_API void ExportInstancedMeshToFbx(const UInstancedStaticMeshComponent* InstancedMeshComp, const TCHAR* MeshName, FbxNode* FbxActor, const FFbxMaterialBakingMeshData& MaterialBakingMeshData);
 
 	/**
 	* Exports a landscape
@@ -529,7 +536,7 @@ private:
 	/**
 	 * Exports the profile_COMMON information for a material.
 	 */
-	UNREALED_API FbxSurfaceMaterial* ExportMaterial(UMaterialInterface* Material);
+	UNREALED_API FbxSurfaceMaterial* ExportMaterial(UMaterialInterface* Material, const int32& MaterialIndex, const FFbxMaterialBakingMeshData& BakingMeshData);
 	
 	UNREALED_API FbxSurfaceMaterial* CreateDefaultMaterial();
 	
