[FBX Exporter] Material Export helper Private namespace change

#rnx

diff --git a/Engine/Source/Editor/UnrealEd/Private/Fbx/FbxMaterialExportUtilities.cpp b/Engine/Source/Editor/UnrealEd/Private/Fbx/FbxMaterialExportUtilities.cpp
--- a/Engine/Source/Editor/UnrealEd/Private/Fbx/FbxMaterialExportUtilities.cpp
+++ b/Engine/Source/Editor/UnrealEd/Private/Fbx/FbxMaterialExportUtilities.cpp
@@ -50,7 +50,7 @@ THIRD_PARTY_INCLUDES_END
 
 namespace UnFbx
 {
-	namespace Private
+	namespace FbxMaterialExportHelpers
 	{
 		UMaterialExpressionCustomOutput* GetCustomOutputByName(const UMaterialInterface* Material, const FString& FunctionName)
 		{
@@ -719,12 +719,12 @@ namespace UnFbx
 	{
 		if (StaticMesh)
 		{
-			return Private::GetSectionIndices(StaticMesh, LODIndex, MaterialIndex);
+			return FbxMaterialExportHelpers::GetSectionIndices(StaticMesh, LODIndex, MaterialIndex);
 		}
 
 		if (SkeletalMesh)
 		{
-			return Private::GetSectionIndices(SkeletalMesh, LODIndex, MaterialIndex);
+			return FbxMaterialExportHelpers::GetSectionIndices(SkeletalMesh, LODIndex, MaterialIndex);
 		}
 
 		return TArray<int32>();
@@ -744,7 +744,7 @@ namespace UnFbx
 			const FFbxMaterialBakingMeshData& MeshData,
 			const FString& ExportFolderPath)
 		{
-			using namespace Private;
+			using namespace FbxMaterialExportHelpers;
 
 			const FIntPoint DefaultBakeSize = FIntPoint(512, 512);
 			const FBox2f TexCoordBounds = FBox2f{ { 0.0f, 0.0f }, { 1.0f, 1.0f } };
@@ -863,7 +863,7 @@ namespace UnFbx
 
 		bool GetInterchangeShadingModel(const UMaterialInterface* MaterialInterface, bool& bLambert)
 		{
-			using namespace Private::InterchangeMaterialProcessHelpers;
+			using namespace FbxMaterialExportHelpers::InterchangeMaterialProcessHelpers;
 
 			FString PathName = MaterialInterface->GetMaterial()->GetPathName();
 			if (MaterialInterface->IsA(UMaterialInstance::StaticClass()))
@@ -909,7 +909,7 @@ namespace UnFbx
 
 		void ProcessInterchangeMaterials(UMaterialInterface* MaterialInterface, fbxsdk::FbxScene* Scene, fbxsdk::FbxSurfaceMaterial* FbxMaterial)
 		{
-			using namespace Private::InterchangeMaterialProcessHelpers;
+			using namespace FbxMaterialExportHelpers::InterchangeMaterialProcessHelpers;
 
 			FString PathName = MaterialInterface->GetMaterial()->GetPathName();
 
