#UE Provide the asset reference filter information about which referencers are editor only

diff --git a/Engine/Plugins/Editor/AssetReferenceRestrictions/Source/AssetReferenceRestrictions/Private/AssetReferencingPolicySubsystem.cpp b/Engine/Plugins/Editor/AssetReferenceRestrictions/Source/AssetReferenceRestrictions/Private/AssetReferencingPolicySubsystem.cpp
--- a/Engine/Plugins/Editor/AssetReferenceRestrictions/Source/AssetReferenceRestrictions/Private/AssetReferencingPolicySubsystem.cpp
+++ b/Engine/Plugins/Editor/AssetReferenceRestrictions/Source/AssetReferenceRestrictions/Private/AssetReferencingPolicySubsystem.cpp
@@ -101,7 +101,7 @@ TValueOrError<void, TArray<FAssetReferenceError>> UAssetReferencingPolicySubsyst
 	if (AllDependencyAssets.Num() > 0)
 	{
 		FAssetReferenceFilterContext AssetReferenceFilterContext;
-		AssetReferenceFilterContext.ReferencingAssets = { InAssetData };
+		AssetReferenceFilterContext.AddReferencingAsset(InAssetData);
 		TSharedPtr<IAssetReferenceFilter> AssetReferenceFilter = GEditor ? GEditor->MakeAssetReferenceFilter(AssetReferenceFilterContext) : nullptr;
 		if (ensure(AssetReferenceFilter.IsValid()))
 		{
diff --git a/Engine/Plugins/Editor/AssetReferenceRestrictions/Source/AssetReferenceRestrictions/Private/DomainAssetReferenceFilter.cpp b/Engine/Plugins/Editor/AssetReferenceRestrictions/Source/AssetReferenceRestrictions/Private/DomainAssetReferenceFilter.cpp
--- a/Engine/Plugins/Editor/AssetReferenceRestrictions/Source/AssetReferenceRestrictions/Private/DomainAssetReferenceFilter.cpp
+++ b/Engine/Plugins/Editor/AssetReferenceRestrictions/Source/AssetReferenceRestrictions/Private/DomainAssetReferenceFilter.cpp
@@ -16,7 +16,7 @@ FDomainAssetReferenceFilter::FDomainAssetReferenceFilter(const FAssetReferenceFi
 	: IAssetReferenceFilter()
 	, DomainDB(InDomainDB)
 {
-	OriginalReferencingAssets = Context.ReferencingAssets;
+	OriginalReferencingAssets = Context.GetReferencingAssets();
 
 	DetermineReferencingDomain();
 
@@ -42,7 +42,10 @@ void FDomainAssetReferenceFilter::DetermineReferencingDomain()
 	TArray<FAssetData> PendingAssets;
 	TSet<FAssetData> AssetsConsidered;
 
-	PendingAssets.Append(OriginalReferencingAssets);
+	for (const FAssetReferenceFilterReferencerInfo& AssetReferencer : OriginalReferencingAssets)
+	{
+		PendingAssets.Add(AssetReferencer.Data);
+	}
 
 	while (!PendingAssets.IsEmpty())
 	{
diff --git a/Engine/Plugins/Editor/AssetReferenceRestrictions/Source/AssetReferenceRestrictions/Public/DomainAssetReferenceFilter.h b/Engine/Plugins/Editor/AssetReferenceRestrictions/Source/AssetReferenceRestrictions/Public/DomainAssetReferenceFilter.h
--- a/Engine/Plugins/Editor/AssetReferenceRestrictions/Source/AssetReferenceRestrictions/Public/DomainAssetReferenceFilter.h
+++ b/Engine/Plugins/Editor/AssetReferenceRestrictions/Source/AssetReferenceRestrictions/Public/DomainAssetReferenceFilter.h
@@ -3,10 +3,9 @@
 #pragma once
 
 #include "AssetRegistry/AssetData.h"
+#include "Editor/EditorEngine.h"
 #include "Editor/AssetReferenceFilter.h"
 
-struct FAssetReferenceFilterContext;
-
 struct FAssetData;
 struct FDomainDatabase;
 struct FDomainData;
@@ -36,7 +35,7 @@ private:
 
 	TSharedPtr<FDomainDatabase> DomainDB;
 
-	TArray<FAssetData> OriginalReferencingAssets;
+	TArray<FAssetReferenceFilterReferencerInfo> OriginalReferencingAssets;
 	TSet<TSharedPtr<FDomainData>> ReferencingDomains;
 
 	FText Failure_CouldNotDetermineDomain;
diff --git a/Engine/Plugins/Experimental/MeshModelingToolsetExp/Source/ModelingEditorUI/Private/ModelingWidgets/SToolInputAssetPicker.cpp b/Engine/Plugins/Experimental/MeshModelingToolsetExp/Source/ModelingEditorUI/Private/ModelingWidgets/SToolInputAssetPicker.cpp
--- a/Engine/Plugins/Experimental/MeshModelingToolsetExp/Source/ModelingEditorUI/Private/ModelingWidgets/SToolInputAssetPicker.cpp
+++ b/Engine/Plugins/Experimental/MeshModelingToolsetExp/Source/ModelingEditorUI/Private/ModelingWidgets/SToolInputAssetPicker.cpp
@@ -130,16 +130,9 @@ void SToolInputAssetPicker::Construct( const FArguments& InArgs )
 	FOnShouldFilterAsset ShouldFilterAssetDelegate;
 	{
 		FAssetReferenceFilterContext AssetReferenceFilterContext;
-		AssetReferenceFilterContext.ReferencingAssets = InArgs._AssetPickerConfig.AdditionalReferencingAssets;
-		if (InArgs._AssetPickerConfig.PropertyHandle.IsValid())
-		{
-			TArray<UObject*> ReferencingObjects;
-			InArgs._AssetPickerConfig.PropertyHandle->GetOuterObjects(ReferencingObjects);
-			for (UObject* ReferencingObject : ReferencingObjects)
-			{
-				AssetReferenceFilterContext.ReferencingAssets.Add(FAssetData(ReferencingObject));
-			}
-		}
+		AssetReferenceFilterContext.AddReferencingAssets(InArgs._AssetPickerConfig.AdditionalReferencingAssets);
+		AssetReferenceFilterContext.AddReferencingAssetsFromPropertyHandle(InArgs._AssetPickerConfig.PropertyHandle);
+		
 		TSharedPtr<IAssetReferenceFilter> AssetReferenceFilter = GEditor ? GEditor->MakeAssetReferenceFilter(AssetReferenceFilterContext) : nullptr;
 		if (AssetReferenceFilter.IsValid())
 		{
diff --git a/Engine/Plugins/Experimental/StructUtils/Source/StructUtilsEditor/Private/SInstancedStructPicker.cpp b/Engine/Plugins/Experimental/StructUtils/Source/StructUtilsEditor/Private/SInstancedStructPicker.cpp
--- a/Engine/Plugins/Experimental/StructUtils/Source/StructUtilsEditor/Private/SInstancedStructPicker.cpp
+++ b/Engine/Plugins/Experimental/StructUtils/Source/StructUtilsEditor/Private/SInstancedStructPicker.cpp
@@ -224,7 +224,7 @@ TSharedRef<SWidget> SInstancedStructPicker::GenerateStructPicker()
 		StructProperty->GetOuterPackages(OuterPackages);
 		for (UPackage* OuterPackage : OuterPackages)
 		{
-			AssetReferenceFilterContext.ReferencingAssets.Add(FAssetData(OuterPackage));
+			AssetReferenceFilterContext.AddReferencingAsset(FAssetData(OuterPackage));
 		}
 
 		StructFilter->AssetReferenceFilter = GEditor->MakeAssetReferenceFilter(AssetReferenceFilterContext);
diff --git a/Engine/Source/Editor/ClassViewer/Private/ClassViewerFilter.cpp b/Engine/Source/Editor/ClassViewer/Private/ClassViewerFilter.cpp
--- a/Engine/Source/Editor/ClassViewer/Private/ClassViewerFilter.cpp
+++ b/Engine/Source/Editor/ClassViewer/Private/ClassViewerFilter.cpp
@@ -406,16 +406,8 @@ FClassViewerFilter::FClassViewerFilter(const FClassViewerInitializationOptions&
 	if (GEditor)
 	{
 		FAssetReferenceFilterContext AssetReferenceFilterContext;
-		AssetReferenceFilterContext.ReferencingAssets = InInitOptions.AdditionalReferencingAssets;
-		if (InInitOptions.PropertyHandle.IsValid())
-		{
-			TArray<UObject*> ReferencingObjects;
-			InInitOptions.PropertyHandle->GetOuterObjects(ReferencingObjects);
-			for (UObject* ReferencingObject : ReferencingObjects)
-			{
-				AssetReferenceFilterContext.ReferencingAssets.Add(FAssetData(ReferencingObject));
-			}
-		}
+		AssetReferenceFilterContext.AddReferencingAssets(InInitOptions.AdditionalReferencingAssets);
+		AssetReferenceFilterContext.AddReferencingAssetsFromPropertyHandle(InInitOptions.PropertyHandle);
 		AssetReferenceFilter = GEditor->MakeAssetReferenceFilter(AssetReferenceFilterContext);
 	}
 }
diff --git a/Engine/Source/Editor/ContentBrowser/Private/SAssetPicker.cpp b/Engine/Source/Editor/ContentBrowser/Private/SAssetPicker.cpp
--- a/Engine/Source/Editor/ContentBrowser/Private/SAssetPicker.cpp
+++ b/Engine/Source/Editor/ContentBrowser/Private/SAssetPicker.cpp
@@ -286,16 +286,9 @@ void SAssetPicker::Construct( const FArguments& InArgs )
 	FOnShouldFilterAsset ShouldFilterAssetDelegate;
 	{
 		FAssetReferenceFilterContext AssetReferenceFilterContext;
-		AssetReferenceFilterContext.ReferencingAssets = InArgs._AssetPickerConfig.AdditionalReferencingAssets;
-		if (InArgs._AssetPickerConfig.PropertyHandle.IsValid())
-		{
-			TArray<UObject*> ReferencingObjects;
-			InArgs._AssetPickerConfig.PropertyHandle->GetOuterObjects(ReferencingObjects);
-			for (UObject* ReferencingObject : ReferencingObjects)
-			{
-				AssetReferenceFilterContext.ReferencingAssets.Add(FAssetData(ReferencingObject));
-			}
-		}
+		AssetReferenceFilterContext.AddReferencingAssets(InArgs._AssetPickerConfig.AdditionalReferencingAssets);
+		AssetReferenceFilterContext.AddReferencingAssetsFromPropertyHandle(InArgs._AssetPickerConfig.PropertyHandle);
+
 		TSharedPtr<IAssetReferenceFilter> AssetReferenceFilter = GEditor ? GEditor->MakeAssetReferenceFilter(AssetReferenceFilterContext) : nullptr;
 		if (AssetReferenceFilter.IsValid())
 		{
diff --git a/Engine/Source/Editor/GraphEditor/Private/SGraphPanel.cpp b/Engine/Source/Editor/GraphEditor/Private/SGraphPanel.cpp
--- a/Engine/Source/Editor/GraphEditor/Private/SGraphPanel.cpp
+++ b/Engine/Source/Editor/GraphEditor/Private/SGraphPanel.cpp
@@ -1365,7 +1365,7 @@ bool SGraphPanel::PassesAssetReferenceFilter(const TArray<FAssetData>& Reference
 		UObject* GraphOuter = GraphObj ? GraphObj->GetOuter() : nullptr;
 		if (GraphOuter)
 		{
-			AssetReferenceFilterContext.ReferencingAssets.Add(FAssetData(GraphOuter));
+			AssetReferenceFilterContext.AddReferencingAsset(FAssetData(GraphOuter));
 		}
 		TSharedPtr<IAssetReferenceFilter> AssetReferenceFilter = GEditor->MakeAssetReferenceFilter(AssetReferenceFilterContext);
 		if (AssetReferenceFilter.IsValid())
diff --git a/Engine/Source/Editor/Kismet/Private/BlueprintActionMenuUtils.cpp b/Engine/Source/Editor/Kismet/Private/BlueprintActionMenuUtils.cpp
--- a/Engine/Source/Editor/Kismet/Private/BlueprintActionMenuUtils.cpp
+++ b/Engine/Source/Editor/Kismet/Private/BlueprintActionMenuUtils.cpp
@@ -538,7 +538,7 @@ void FBlueprintActionMenuUtils::MakeContextMenu(FBlueprintActionContext const& C
 	FAssetReferenceFilterContext AssetReferenceFilterContext;
 	for (UBlueprint* Blueprint : Context.Blueprints)
 	{
-		AssetReferenceFilterContext.ReferencingAssets.Add(FAssetData(Blueprint));
+		AssetReferenceFilterContext.AddReferencingAsset(FAssetData(Blueprint));
 	}
 
 	MainMenuFilter.AssetReferenceFilter = GEditor->MakeAssetReferenceFilter(AssetReferenceFilterContext);
diff --git a/Engine/Source/Editor/Kismet/Private/BlueprintEditor.cpp b/Engine/Source/Editor/Kismet/Private/BlueprintEditor.cpp
--- a/Engine/Source/Editor/Kismet/Private/BlueprintEditor.cpp
+++ b/Engine/Source/Editor/Kismet/Private/BlueprintEditor.cpp
@@ -406,11 +406,9 @@ namespace BlueprintEditorImpl
 		FPermissionsPinTypeSelectorFilter(TConstArrayView<UBlueprint*> InBlueprints)
 		{
 			FAssetReferenceFilterContext Context;
-			Context.ReferencingAssets.Reserve(InBlueprints.Num());
-
 			for (UBlueprint* Blueprint : InBlueprints)
 			{
-				Context.ReferencingAssets.Add(FAssetData(Blueprint));
+				Context.AddReferencingAsset(FAssetData(Blueprint));
 			}
 
 			AssetReferenceFilter = GEditor->MakeAssetReferenceFilter(Context);
diff --git a/Engine/Source/Editor/MovieSceneTools/Private/MovieSceneToolHelpers.cpp b/Engine/Source/Editor/MovieSceneTools/Private/MovieSceneToolHelpers.cpp
--- a/Engine/Source/Editor/MovieSceneTools/Private/MovieSceneToolHelpers.cpp
+++ b/Engine/Source/Editor/MovieSceneTools/Private/MovieSceneToolHelpers.cpp
@@ -4586,7 +4586,7 @@ bool MovieSceneToolHelpers::IsValidAsset(UMovieSceneSequence* Sequence, const FA
 	}
 
 	FAssetReferenceFilterContext AssetReferenceFilterContext;
-	AssetReferenceFilterContext.ReferencingAssets.Add(FAssetData(Sequence));
+	AssetReferenceFilterContext.AddReferencingAsset(FAssetData(Sequence));
 
 	TSharedPtr<IAssetReferenceFilter> AssetReferenceFilter = GEditor->MakeAssetReferenceFilter(AssetReferenceFilterContext);
 
diff --git a/Engine/Source/Editor/PropertyEditor/Private/UserInterface/PropertyEditor/SPropertyEditorAsset.cpp b/Engine/Source/Editor/PropertyEditor/Private/UserInterface/PropertyEditor/SPropertyEditorAsset.cpp
--- a/Engine/Source/Editor/PropertyEditor/Private/UserInterface/PropertyEditor/SPropertyEditorAsset.cpp
+++ b/Engine/Source/Editor/PropertyEditor/Private/UserInterface/PropertyEditor/SPropertyEditorAsset.cpp
@@ -1644,30 +1644,10 @@ bool SPropertyEditorAsset::CanSetBasedOnAssetReferenceFilter( const FAssetData&
 {
 	if (GEditor && InAssetData.IsValid())
 	{
-		TSharedPtr<IPropertyHandle> PropertyHandleToUse = GetMostSpecificPropertyHandle();
 		FAssetReferenceFilterContext AssetReferenceFilterContext;
-		if (PropertyHandleToUse.IsValid())
-		{
-			TArray<UObject*> ReferencingObjects;
-			PropertyHandleToUse->GetOuterObjects(ReferencingObjects);
-			for (UObject* ReferencingObject : ReferencingObjects)
-			{
-				AssetReferenceFilterContext.ReferencingAssets.Add(FAssetData(ReferencingObject));
-			}
-		}
+		AssetReferenceFilterContext.AddReferencingAssets(OwnerAssetDataArray);
+		AssetReferenceFilterContext.AddReferencingAssetsFromPropertyHandle(GetMostSpecificPropertyHandle());
 		
-		if(OwnerAssetDataArray.Num() > 0)
-		{
-			for (const FAssetData& AssetData : OwnerAssetDataArray)
-			{
-				if (AssetData.IsValid())
-				{
-					//Use add unique in case the PropertyHandle as already add the referencing asset
-					AssetReferenceFilterContext.ReferencingAssets.AddUnique(AssetData);
-				}
-			}
-		}
-
 		TSharedPtr<IAssetReferenceFilter> AssetReferenceFilter = GEditor->MakeAssetReferenceFilter(AssetReferenceFilterContext);
 		if (AssetReferenceFilter.IsValid() && !AssetReferenceFilter->PassesFilter(InAssetData, OutOptionalFailureReason))
 		{
diff --git a/Engine/Source/Editor/UnrealEd/Classes/Editor/EditorEngine.h b/Engine/Source/Editor/UnrealEd/Classes/Editor/EditorEngine.h
--- a/Engine/Source/Editor/UnrealEd/Classes/Editor/EditorEngine.h
+++ b/Engine/Source/Editor/UnrealEd/Classes/Editor/EditorEngine.h
@@ -276,10 +276,46 @@ struct FPreviewPlatformInfo
 	UNREALED_API ERHIFeatureLevel::Type GetEffectivePreviewFeatureLevel() const;
 };
 
+enum class EAssetReferenceFilterProperties : uint8
+{
+	None = 0,
+	EditorOnly = 1
+};
+ENUM_CLASS_FLAGS(EAssetReferenceFilterProperties);
+
+struct FAssetReferenceFilterReferencerInfo
+{
+	FAssetData Data;
+	EAssetReferenceFilterProperties Properties;
+
+	FAssetReferenceFilterReferencerInfo() : Properties(EAssetReferenceFilterProperties::None) {}
+	FAssetReferenceFilterReferencerInfo(const FAssetData& InData) : Data(InData), Properties(EAssetReferenceFilterProperties::None) {}
+	FAssetReferenceFilterReferencerInfo(const FAssetData& InData, EAssetReferenceFilterProperties InProperties) : Data(InData), Properties(InProperties) {}
+	bool operator==(const FAssetReferenceFilterReferencerInfo& Other) const
+	{
+		return Data == Other.Data && Properties == Other.Properties;
+	}
+
+	bool operator!=(const FAssetReferenceFilterReferencerInfo& Other) const
+	{
+		return !(*this == Other);
+	}
+};
+
 /** Struct used in filtering allowed references between assets. Passes context about the referencers to game-level filters */
 struct FAssetReferenceFilterContext
 {
+	UNREALED_API void AddReferencingAsset(const FAssetData& InReferencingAsset, EAssetReferenceFilterProperties InProperties = EAssetReferenceFilterProperties::None);
+	UNREALED_API void AddReferencingAssets(const TArray<FAssetData>& InReferencingAssets, EAssetReferenceFilterProperties InProperties = EAssetReferenceFilterProperties::None);
+	UNREALED_API void AddReferencingAssetsFromPropertyHandle(const TSharedPtr<class IPropertyHandle>& PropertyHandle);
+
+	const TArray<FAssetReferenceFilterReferencerInfo>& GetReferencingAssets() const { return ReferencingAssetInfo; }
+
+	UE_DEPRECATED(5.5, "ReferencingAssets is deprecated. Use the AddReferencingAsset functions to add referencers and GetReferencingAssets to get them.")
 	TArray<FAssetData> ReferencingAssets;
+
+private:
+	TArray<FAssetReferenceFilterReferencerInfo> ReferencingAssetInfo;
 };
 
 /**
diff --git a/Engine/Source/Editor/UnrealEd/Private/EditorEngine.cpp b/Engine/Source/Editor/UnrealEd/Private/EditorEngine.cpp
--- a/Engine/Source/Editor/UnrealEd/Private/EditorEngine.cpp
+++ b/Engine/Source/Editor/UnrealEd/Private/EditorEngine.cpp
@@ -447,6 +447,68 @@ ERHIFeatureLevel::Type FPreviewPlatformInfo::GetEffectivePreviewFeatureLevel() c
 	return bPreviewFeatureLevelActive ? PreviewFeatureLevel : GMaxRHIFeatureLevel;
 }
 
+void FAssetReferenceFilterContext::AddReferencingAsset(const FAssetData& InReferencingAsset, EAssetReferenceFilterProperties InProperties)
+{
+	if (InReferencingAsset.IsValid())
+	{
+		ReferencingAssetInfo.Emplace(InReferencingAsset, InProperties);
+
+		PRAGMA_DISABLE_DEPRECATION_WARNINGS
+		ReferencingAssets.Add(InReferencingAsset);
+		PRAGMA_ENABLE_DEPRECATION_WARNINGS
+	}
+}
+
+void FAssetReferenceFilterContext::AddReferencingAssets(const TArray<FAssetData>& InReferencingAssets, EAssetReferenceFilterProperties InProperties)
+{
+	for (const FAssetData& ReferencingAsset : InReferencingAssets)
+	{
+		AddReferencingAsset(ReferencingAsset);
+	}
+}
+
+void FAssetReferenceFilterContext::AddReferencingAssetsFromPropertyHandle(const TSharedPtr<IPropertyHandle>& PropertyHandle)
+{
+	if (!PropertyHandle.IsValid())
+	{
+		return;
+	}
+
+	EAssetReferenceFilterProperties ReferencerProperties = EAssetReferenceFilterProperties::None;
+
+	// Determine if this is an editor only property by following the parent property chain and checking if any are editor only
+	{
+		TSharedPtr<IPropertyHandle> CurrentPropertyHandle = PropertyHandle;
+		while (CurrentPropertyHandle.IsValid())
+		{
+			const FProperty* Property = CurrentPropertyHandle->GetProperty();
+			if (Property != nullptr && Property->HasAnyPropertyFlags(CPF_EditorOnly))
+			{
+				ReferencerProperties |= EAssetReferenceFilterProperties::EditorOnly;
+				break;
+			}
+
+			CurrentPropertyHandle = CurrentPropertyHandle->GetParentHandle();
+		}
+	}
+
+	TArray<UObject*> ReferencingObjects;
+	PropertyHandle->GetOuterObjects(ReferencingObjects);
+
+	for (UObject* ReferencingObject : ReferencingObjects)
+	{
+		if (ReferencingObject)
+		{
+			FAssetData ReferencingObjectData(ReferencingObject);
+			ReferencingAssetInfo.Emplace(ReferencingObjectData, ReferencerProperties);
+			
+			PRAGMA_DISABLE_DEPRECATION_WARNINGS
+			ReferencingAssets.Add(ReferencingObjectData);
+			PRAGMA_ENABLE_DEPRECATION_WARNINGS
+		}
+	}
+}
+
 //////////////////////////////////////////////////////////////////////////
 // UEditorEngine
 
diff --git a/Engine/Source/Editor/UnrealEd/Private/LevelEditorDragDropHandler.cpp b/Engine/Source/Editor/UnrealEd/Private/LevelEditorDragDropHandler.cpp
--- a/Engine/Source/Editor/UnrealEd/Private/LevelEditorDragDropHandler.cpp
+++ b/Engine/Source/Editor/UnrealEd/Private/LevelEditorDragDropHandler.cpp
@@ -28,7 +28,7 @@ bool ULevelEditorDragDropHandler::PassesFilter(UWorld* World, const FAssetData&
 		auto AssetPassesFilter = [](const FAssetData& InAssetData, const UObject* InReferencingAsset, FText* OutFailureReason = nullptr)
 		{
 			FAssetReferenceFilterContext AssetReferenceFilterContext;
-			AssetReferenceFilterContext.ReferencingAssets.Add(FAssetData(InReferencingAsset));
+			AssetReferenceFilterContext.AddReferencingAsset(FAssetData(InReferencingAsset));
 
 			TSharedPtr<IAssetReferenceFilter> AssetReferenceFilter = GEditor->MakeAssetReferenceFilter(AssetReferenceFilterContext);
 			return AssetReferenceFilter.IsValid() ? AssetReferenceFilter->PassesFilter(InAssetData, OutFailureReason) : true;
diff --git a/Engine/Source/Runtime/Engine/Private/ActorEditor.cpp b/Engine/Source/Runtime/Engine/Private/ActorEditor.cpp
--- a/Engine/Source/Runtime/Engine/Private/ActorEditor.cpp
+++ b/Engine/Source/Runtime/Engine/Private/ActorEditor.cpp
@@ -1923,7 +1923,7 @@ bool AActor::CanAddDataLayer(const UDataLayerInstance* InDataLayerInstance, FTex
 	auto PassesAssetReferenceFiltering = [](const UObject * InReferencingObject, const UDataLayerAsset * InDataLayerAsset, FText* OutReason)
 	{
 		FAssetReferenceFilterContext AssetReferenceFilterContext;
-		AssetReferenceFilterContext.ReferencingAssets.Add(FAssetData(InReferencingObject));
+		AssetReferenceFilterContext.AddReferencingAsset(FAssetData(InReferencingObject));
 		TSharedPtr<IAssetReferenceFilter> AssetReferenceFilter = GEditor->MakeAssetReferenceFilter(AssetReferenceFilterContext);
 		return AssetReferenceFilter.IsValid() ? AssetReferenceFilter->PassesFilter(FAssetData(InDataLayerAsset), OutReason) : true;
 	};
diff --git a/Engine/Source/Runtime/Engine/Private/WorldPartition/DataLayer/ExternalDataLayerHelper.cpp b/Engine/Source/Runtime/Engine/Private/WorldPartition/DataLayer/ExternalDataLayerHelper.cpp
--- a/Engine/Source/Runtime/Engine/Private/WorldPartition/DataLayer/ExternalDataLayerHelper.cpp
+++ b/Engine/Source/Runtime/Engine/Private/WorldPartition/DataLayer/ExternalDataLayerHelper.cpp
@@ -109,9 +109,8 @@ namespace UE::Private::ExternalDataLayerHelper
 		uint32 ErrorCount = 0;
 		if (InReferencedAssets.Num() > 0)
 		{
-			FAssetData ReferencingAssetData = FAssetData(InAsset);
 			FAssetReferenceFilterContext AssetReferenceFilterContext;
-			AssetReferenceFilterContext.ReferencingAssets = { ReferencingAssetData };
+			AssetReferenceFilterContext.AddReferencingAsset(FAssetData(InAsset));
 			TSharedPtr<IAssetReferenceFilter> AssetReferenceFilter = GEditor ? GEditor->MakeAssetReferenceFilter(AssetReferenceFilterContext) : nullptr;
 			if (ensure(AssetReferenceFilter.IsValid()))
 			{
diff --git a/Engine/Source/Runtime/Engine/Private/WorldPartition/DataLayer/WorldDataLayers.cpp b/Engine/Source/Runtime/Engine/Private/WorldPartition/DataLayer/WorldDataLayers.cpp
--- a/Engine/Source/Runtime/Engine/Private/WorldPartition/DataLayer/WorldDataLayers.cpp
+++ b/Engine/Source/Runtime/Engine/Private/WorldPartition/DataLayer/WorldDataLayers.cpp
@@ -1300,7 +1300,7 @@ bool AWorldDataLayers::CanReferenceDataLayerAsset(const UDataLayerAsset* InDataL
 	auto PassesAssetReferenceFiltering = [](const UObject* InReferencingObject, const UDataLayerAsset* InDataLayerAsset, FText* OutReason)
 	{
 		FAssetReferenceFilterContext AssetReferenceFilterContext;
-		AssetReferenceFilterContext.ReferencingAssets.Add(FAssetData(InReferencingObject));
+		AssetReferenceFilterContext.AddReferencingAsset(FAssetData(InReferencingObject));
 		TSharedPtr<IAssetReferenceFilter> AssetReferenceFilter = GEditor ? GEditor->MakeAssetReferenceFilter(AssetReferenceFilterContext) : nullptr;
 		return AssetReferenceFilter.IsValid() ? AssetReferenceFilter->PassesFilter(FAssetData(InDataLayerAsset), OutReason) : true;
 	};
