[AssetEditorToolkit] Save and SaveAs command entry visibility contingent on the built-in UIAction, rather than IsActuallyAnAsset. The callback maps to overridable functions (IsSaveAssetVisible, IsSaveAssetAsVisible)

diff --git a/Engine/Source/Editor/UnrealEd/Private/Toolkits/AssetEditorToolkit.cpp b/Engine/Source/Editor/UnrealEd/Private/Toolkits/AssetEditorToolkit.cpp
--- a/Engine/Source/Editor/UnrealEd/Private/Toolkits/AssetEditorToolkit.cpp
+++ b/Engine/Source/Editor/UnrealEd/Private/Toolkits/AssetEditorToolkit.cpp
@@ -252,16 +252,20 @@ void FAssetEditorToolkit::InitAssetEditor( const EToolkitMode::Type Mode, const
 
 	check( ToolkitHost.IsValid() );
 	ToolkitManager.RegisterNewToolkit( SharedThis( this ) );
-	
+
 	ToolkitCommands->MapAction(
 		FAssetEditorCommonCommands::Get().SaveAsset,
 		FExecuteAction::CreateSP( this, &FAssetEditorToolkit::SaveAsset_Execute ),
-		FCanExecuteAction::CreateSP( this, &FAssetEditorToolkit::CanSaveAsset_Internal ));
+		FCanExecuteAction::CreateSP( this, &FAssetEditorToolkit::CanSaveAsset_Internal ),
+		FIsActionChecked(),
+		FIsActionButtonVisible::CreateSP( this, &FAssetEditorToolkit::IsSaveAssetVisible ));
 
 	ToolkitCommands->MapAction(
 		FAssetEditorCommonCommands::Get().SaveAssetAs,
 		FExecuteAction::CreateSP( this, &FAssetEditorToolkit::SaveAssetAs_Execute ),
-		FCanExecuteAction::CreateSP( this, &FAssetEditorToolkit::CanSaveAssetAs_Internal ));
+		FCanExecuteAction::CreateSP( this, &FAssetEditorToolkit::CanSaveAssetAs_Internal ),
+		FIsActionChecked(),
+		FIsActionButtonVisible::CreateSP( this, &FAssetEditorToolkit::IsSaveAssetAsVisible ));
 
 	ToolkitCommands->MapAction(
 		FGlobalEditorCommonCommands::Get().FindInContentBrowser,
@@ -639,6 +643,11 @@ bool FAssetEditorToolkit::CanSaveAsset_Internal() const
 	return CanSaveAsset();
 }
 
+bool FAssetEditorToolkit::IsSaveAssetVisible() const
+{
+	return true;
+}
+
 void FAssetEditorToolkit::SaveAsset_Execute()
 {
 	if (EditingObjects.Num() == 0)
@@ -688,6 +697,11 @@ bool FAssetEditorToolkit::CanSaveAssetAs_Internal() const
 	return CanSaveAssetAs();
 }
 
+bool FAssetEditorToolkit::IsSaveAssetAsVisible() const
+{
+	return IsActuallyAnAsset();
+}
+
 void FAssetEditorToolkit::SaveAssetAs_Execute()
 {
 	if (EditingObjects.Num() == 0)
@@ -1079,10 +1093,7 @@ void FAssetEditorToolkit::FillDefaultFileMenuCommands(FToolMenuSection& InSectio
 	if (UAssetEditorToolkitMenuContext* Context = InSection.FindContext<UAssetEditorToolkitMenuContext>())
 	{
 		InSection.AddMenuEntry(FAssetEditorCommonCommands::Get().SaveAsset, TAttribute<FText>(), TAttribute<FText>(), FSlateIcon(FAppStyle::GetAppStyleSetName(), "AssetEditor.SaveAsset")).InsertPosition = InsertPosition;
-		if( IsActuallyAnAsset() )
-		{
-			InSection.AddMenuEntry(FAssetEditorCommonCommands::Get().SaveAssetAs, TAttribute<FText>(), TAttribute<FText>(), FSlateIcon(FAppStyle::GetAppStyleSetName(), "AssetEditor.SaveAssetAs")).InsertPosition = InsertPosition;
-		}
+		InSection.AddMenuEntry(FAssetEditorCommonCommands::Get().SaveAssetAs, TAttribute<FText>(), TAttribute<FText>(), FSlateIcon(FAppStyle::GetAppStyleSetName(), "AssetEditor.SaveAssetAs")).InsertPosition = InsertPosition;
 	}
 
 	if( IsWorldCentricAssetEditor() )
diff --git a/Engine/Source/Editor/UnrealEd/Public/Toolkits/AssetEditorToolkit.h b/Engine/Source/Editor/UnrealEd/Public/Toolkits/AssetEditorToolkit.h
--- a/Engine/Source/Editor/UnrealEd/Public/Toolkits/AssetEditorToolkit.h
+++ b/Engine/Source/Editor/UnrealEd/Public/Toolkits/AssetEditorToolkit.h
@@ -347,6 +347,9 @@ protected:
 	/** Internal function to check if the asset can be saved that calls CanSaveAsset */
 	bool CanSaveAsset_Internal() const;
 
+	/** Called to test if "Save" should be visible for this asset */
+	UNREALED_API virtual bool IsSaveAssetVisible() const;
+
 	/** Called when "Save" is clicked for this asset */
 	UNREALED_API virtual void SaveAsset_Execute();
 
@@ -359,6 +362,9 @@ protected:
 	/** Internal function to check if the asset can be saved that calls CanSaveAssetAs */
 	bool CanSaveAssetAs_Internal() const;
 
+	/** Called to test if "Save As" should be visible for this asset */
+	UNREALED_API virtual bool IsSaveAssetAsVisible() const;
+
 	/** Called when "Save As" is clicked for this asset */
 	UNREALED_API virtual void SaveAssetAs_Execute();
 
