-Use FMaterial::DeferredDelete to destroy the FMaterialExportDataEntry lightmass material proxies
-This gracefully handles the case where the FMaterial still is compiling

diff --git a/Engine/Source/Editor/UnrealEd/Private/Lightmass/LightmassRender.cpp b/Engine/Source/Editor/UnrealEd/Private/Lightmass/LightmassRender.cpp
--- a/Engine/Source/Editor/UnrealEd/Private/Lightmass/LightmassRender.cpp
+++ b/Engine/Source/Editor/UnrealEd/Private/Lightmass/LightmassRender.cpp
@@ -1020,15 +1020,26 @@ FMaterialExportDataEntry::~FMaterialExportDataEntry()
 	FLightmassMaterialProxy* LocalEmissiveMaterialProxy = EmissiveMaterialProxy;
 	FLightmassMaterialProxy* LocalOpacityMaterialProxy = OpacityMaterialProxy;
 	FLightmassMaterialProxy* LocalNormalMaterialProxy = NormalMaterialProxy;
-	ENQUEUE_RENDER_COMMAND(FMaterialExportDataEntryDelete)(
-		[LocalDiffuseMaterialProxy, LocalEmissiveMaterialProxy, LocalOpacityMaterialProxy, LocalNormalMaterialProxy](FRHICommandListImmediate& RHICmdList)
-		{
-			delete LocalDiffuseMaterialProxy;
-			delete LocalEmissiveMaterialProxy;
-			delete LocalOpacityMaterialProxy;
-			delete LocalNormalMaterialProxy;
-		}
-	);
+
+	if (LocalDiffuseMaterialProxy )
+	{
+		FMaterial::DeferredDelete(LocalDiffuseMaterialProxy);		
+	}
+
+	if (LocalEmissiveMaterialProxy)
+	{
+		FMaterial::DeferredDelete(LocalEmissiveMaterialProxy);
+	}
+
+	if (LocalOpacityMaterialProxy)
+	{
+		FMaterial::DeferredDelete(LocalOpacityMaterialProxy);
+	}
+
+	if (LocalNormalMaterialProxy)
+	{	
+		FMaterial::DeferredDelete(LocalNormalMaterialProxy);
+	}
 }
 
 //
