Disable localization chunking when cooking DLC

Running this code as-is causes a name conflict in the DLC output that can override the base localization data when the DLC is mounted. This can result in most of the base localization data appearing to be missing, as the DLC may have been cooked against a much smaller set of packages in each chunk.

While this could be supported by adding a suffix to the chunked data within the DLC, plugin localization is a better way to achieve this, and so localization chunking is now disabled when cooking a DLC plugin.

diff --git a/Engine/Source/Editor/UnrealEd/Private/CookOnTheFlyServer.cpp b/Engine/Source/Editor/UnrealEd/Private/CookOnTheFlyServer.cpp
--- a/Engine/Source/Editor/UnrealEd/Private/CookOnTheFlyServer.cpp
+++ b/Engine/Source/Editor/UnrealEd/Private/CookOnTheFlyServer.cpp
@@ -12971,6 +12971,13 @@ void UCookOnTheFlyServer::RegisterLocalizationChunkDataGenerator()
 {
 	check(!IsCookWorkerMode());
 
+	// Localization chunking is disabled when cooking DLC as it produces output that can override the base localization data
+	// Localization chunking is disabled when we're not cooking for any languages, as there would be no output generated
+	if (IsCookingDLC() || CookByTheBookOptions->AllCulturesToCook.IsEmpty())
+	{
+		return;
+	}
+
 	// Get the list of localization targets to chunk, and remove any targets that we've been asked not to stage
 	const UProjectPackagingSettings* const PackagingSettings = GetDefault<UProjectPackagingSettings>();
 	TArray<FString> LocalizationTargetsToChunk = PackagingSettings->LocalizationTargetsToChunk;
@@ -12986,16 +12993,19 @@ void UCookOnTheFlyServer::RegisterLocalizationChunkDataGenerator()
 		}
 	}
 
-	if (LocalizationTargetsToChunk.Num() > 0 && CookByTheBookOptions->AllCulturesToCook.Num() > 0)
+	// Localization chunking is disabled when there are no localization targets to chunk
+	if (LocalizationTargetsToChunk.IsEmpty())
 	{
-		for (const ITargetPlatform* TargetPlatform : PlatformManager->GetSessionPlatforms())
-		{
-			FAssetRegistryGenerator& RegistryGenerator = *(PlatformManager->GetPlatformData(TargetPlatform)->RegistryGenerator);
-			TSharedRef<FLocalizationChunkDataGenerator> LocalizationGenerator =
-				MakeShared<FLocalizationChunkDataGenerator>(RegistryGenerator.GetPakchunkIndex(PackagingSettings->LocalizationTargetCatchAllChunkId),
-					LocalizationTargetsToChunk, CookByTheBookOptions->AllCulturesToCook);
-			RegistryGenerator.RegisterChunkDataGenerator(MoveTemp(LocalizationGenerator));
-		}
+		return;
+	}
+
+	for (const ITargetPlatform* TargetPlatform : PlatformManager->GetSessionPlatforms())
+	{
+		FAssetRegistryGenerator& RegistryGenerator = *(PlatformManager->GetPlatformData(TargetPlatform)->RegistryGenerator);
+		TSharedRef<FLocalizationChunkDataGenerator> LocalizationGenerator =
+			MakeShared<FLocalizationChunkDataGenerator>(RegistryGenerator.GetPakchunkIndex(PackagingSettings->LocalizationTargetCatchAllChunkId),
+				LocalizationTargetsToChunk, CookByTheBookOptions->AllCulturesToCook);
+		RegistryGenerator.RegisterChunkDataGenerator(MoveTemp(LocalizationGenerator));
 	}
 }
 
