[BaseIteratePackagesCommandlet]
Fix two deadlocks by moving complex code outside of the UObject hash table lock

diff --git a/Engine/Source/Editor/UnrealEd/Private/Commandlets/BaseIteratePackagesCommandlet.cpp b/Engine/Source/Editor/UnrealEd/Private/Commandlets/BaseIteratePackagesCommandlet.cpp
--- a/Engine/Source/Editor/UnrealEd/Private/Commandlets/BaseIteratePackagesCommandlet.cpp
+++ b/Engine/Source/Editor/UnrealEd/Private/Commandlets/BaseIteratePackagesCommandlet.cpp
@@ -659,10 +659,12 @@ void UBaseIteratePackagesCommandlet::LoadAndSaveOnePackage(const FString& Filena
 
 
 				// process all the objects loaded with the base package.
-				ForEachObjectWithOuter(Package, [this, &bSavePackage](UObject* Object)
-					{
-						PerformAdditionalOperations(Object, bSavePackage);
-					});
+				TArray<UObject*> PackageObjects;
+				GetObjectsWithOuter(Package, PackageObjects);
+				for (UObject* Object : PackageObjects)
+				{
+					PerformAdditionalOperations(Object, bSavePackage);
+				}
 
 				FScopedEditorWorld ScopeEditorWorld(World, WorldInitialisationValues);
 
@@ -739,14 +741,19 @@ void UBaseIteratePackagesCommandlet::LoadAndSaveOnePackage(const FString& Filena
 				VerboseMessage(TEXT("Post PerformAdditionalOperations"));
 
 				// Check for any special per object operations
-				ForEachObjectWithOuter(Package, [this, &bSavePackage](UObject* Object)
+				TArray<UObject*> PackageObjects;
+				ForEachObjectWithOuter(Package, [this, &PackageObjects](UObject* Object)
 					{
 						if (!IsValid(Object))
 						{
 							return;
 						}
-						PerformAdditionalOperations(Object, bSavePackage);
+						PackageObjects.Add(Object);
 					});
+				for(UObject* Object : PackageObjects)
+				{
+					PerformAdditionalOperations(Object, bSavePackage);
+				}
 			}
 
 			PostPerformAdditionalOperations(Package);
