[UserDefinedStructEditorData] fixed GetFriendlyNameForProperty implementation

diff --git a/Engine/Source/Editor/UnrealEd/Private/UserDefinedStructEditorData.cpp b/Engine/Source/Editor/UnrealEd/Private/UserDefinedStructEditorData.cpp
--- a/Engine/Source/Editor/UnrealEd/Private/UserDefinedStructEditorData.cpp
+++ b/Engine/Source/Editor/UnrealEd/Private/UserDefinedStructEditorData.cpp
@@ -321,7 +321,7 @@ FProperty* UUserDefinedStructEditorData::FindProperty(const UUserDefinedStruct*
 
 FString UUserDefinedStructEditorData::GetFriendlyNameForProperty(const UUserDefinedStruct* Struct, const FProperty* Property) const
 {
-	return Super::GetFriendlyNameForProperty(Struct, Property);
+	return FStructureEditorUtils::GetVariableFriendlyNameForProperty(Struct, Property);
 }
 
 FString UUserDefinedStructEditorData::GetTooltip() const
