[Selection] Fix an issue with ISM browse to asset pulling in too much data

diff --git a/Engine/Source/Editor/UnrealEd/Private/Elements/Actor/ActorElementEditorAssetDataInterface.cpp b/Engine/Source/Editor/UnrealEd/Private/Elements/Actor/ActorElementEditorAssetDataInterface.cpp
--- a/Engine/Source/Editor/UnrealEd/Private/Elements/Actor/ActorElementEditorAssetDataInterface.cpp
+++ b/Engine/Source/Editor/UnrealEd/Private/Elements/Actor/ActorElementEditorAssetDataInterface.cpp
@@ -41,7 +41,10 @@ TArray<FAssetData> UActorElementEditorAssetDataInterface::GetAllReferencedAssetD
 			}
 		}
 
-		AssetDatas.Emplace(FAssetData(RawActorPtr));
+		if (RawActorPtr->IsAsset())
+		{
+			AssetDatas.Emplace(FAssetData(RawActorPtr));
+		}
 	}
 
 	return AssetDatas;
