Fixing a bug where we could try to reinstance CDOs asynchronously as part of live coding. One required step in this process involves collecting CDO references which involves serializes objects via a specialized archive. Serialization of some types causes us to try to finish pending compilations of those instances and that cannot always be done off the game thread. This change ensures we finish all asset compilation before trying to gather the CDO references.

diff --git a/Engine/Source/Editor/UnrealEd/Private/Kismet2/ReloadUtilities.cpp b/Engine/Source/Editor/UnrealEd/Private/Kismet2/ReloadUtilities.cpp
--- a/Engine/Source/Editor/UnrealEd/Private/Kismet2/ReloadUtilities.cpp
+++ b/Engine/Source/Editor/UnrealEd/Private/Kismet2/ReloadUtilities.cpp
@@ -5,6 +5,7 @@
 =============================================================================*/
 
 #include "Kismet2/ReloadUtilities.h"
+#include "AssetCompilingManager.h"
 #include "Async/AsyncWork.h"
 #include "Engine/Engine.h"
 #include "Engine/EngineTypes.h"
@@ -1176,6 +1177,7 @@ const UObject* FReload::GetReinstancedCDO(const UObject* CDO)
 
 void FReload::Finalize(bool bRunGC)
 {
+	FAssetCompilingManager::Get().FinishAllCompilation();
 
 	// Make sure new classes have the token stream assembled
 	UClass::AssembleReferenceTokenStreams();
