Added boolean guard to prevent recursive calls to RebuildAlteredBSP
This is to fix the infinite loop when the user is in Brush Editing Mode

diff --git a/Engine/Source/Editor/UnrealEd/Classes/Editor/EditorEngine.h b/Engine/Source/Editor/UnrealEd/Classes/Editor/EditorEngine.h
--- a/Engine/Source/Editor/UnrealEd/Classes/Editor/EditorEngine.h
+++ b/Engine/Source/Editor/UnrealEd/Classes/Editor/EditorEngine.h
@@ -596,6 +596,10 @@ public:
 	/** Whether or not the editor is currently compiling */
 	bool bIsCompiling;
 
+protected:
+	/** Used to prevent reentrant calls to RebuildAlteredBSP(). */
+	bool bIsRebuildingAlteredBSP;
+
 private:
 
 	/** Manager that holds all extensions paired with a world */
diff --git a/Engine/Source/Editor/UnrealEd/Private/EditorEngine.cpp b/Engine/Source/Editor/UnrealEd/Private/EditorEngine.cpp
--- a/Engine/Source/Editor/UnrealEd/Private/EditorEngine.cpp
+++ b/Engine/Source/Editor/UnrealEd/Private/EditorEngine.cpp
@@ -550,6 +550,7 @@ UEditorEngine::UEditorEngine(const FObjectInitializer& ObjectInitializer)
 
 	DetailMode = DM_MAX;
 	CurrentPlayWorldDestination = -1;
+	bIsRebuildingAlteredBSP = false;
 	bDisableDeltaModification = false;
 	bAllowMultiplePIEWorlds = true;
 	bIsEndingPlay = false;
diff --git a/Engine/Source/Editor/UnrealEd/Private/EditorServer.cpp b/Engine/Source/Editor/UnrealEd/Private/EditorServer.cpp
--- a/Engine/Source/Editor/UnrealEd/Private/EditorServer.cpp
+++ b/Engine/Source/Editor/UnrealEd/Private/EditorServer.cpp
@@ -1789,6 +1789,13 @@ void UEditorEngine::RebuildModelFromBrushes(TArray<ABrush*> &BrushesToBuild, UMo
 
 void UEditorEngine::RebuildAlteredBSP()
 {
+	if (bIsRebuildingAlteredBSP)
+	{
+		return;
+	}
+
+	TGuardValue<bool> GuardIsRebuildingAlteredBSP(bIsRebuildingAlteredBSP, true);
+
 	if( !GIsTransacting )
 	{
 		// Early out if BSP auto-updating is disabled
