[Selection] Fix an issue with ISM browse to asset pulling in too much data

#rnx

diff --git a/Engine/Source/Editor/UnrealEd/Private/Elements/Actor/ActorElementEditorAssetDataInterface.cpp b/Engine/Source/Editor/UnrealEd/Private/Elements/Actor/ActorElementEditorAssetDataInterface.cpp
--- a/Engine/Source/Editor/UnrealEd/Private/Elements/Actor/ActorElementEditorAssetDataInterface.cpp
+++ b/Engine/Source/Editor/UnrealEd/Private/Elements/Actor/ActorElementEditorAssetDataInterface.cpp
@@ -40,11 +40,6 @@ TArray<FAssetData> UActorElementEditorAssetDataInterface::GetAllReferencedAssetD
 				}
 			}
 		}
-
-		if (RawActorPtr->IsAsset())
-		{
-			AssetDatas.Emplace(FAssetData(RawActorPtr));
-		}
 	}
 
 	return AssetDatas;
