[Actor Actions Usability] Rework the message in the Delete Actor with References dialog

New Dialog to display a list of references for the actor to be deleted, and a customized message based on the type of the references.

A custom widget is used instead of the previous message dialog, so that references can be listed with aligned layout and scrollbars for long names/lists.

diff --git a/Engine/Source/Editor/UnrealEd/Private/EditorActor.cpp b/Engine/Source/Editor/UnrealEd/Private/EditorActor.cpp
--- a/Engine/Source/Editor/UnrealEd/Private/EditorActor.cpp
+++ b/Engine/Source/Editor/UnrealEd/Private/EditorActor.cpp
@@ -76,6 +76,7 @@
 #include "Misc/ScopedSlowTask.h"
 #include "Subsystems/AssetEditorSubsystem.h"
 #include "SceneInterface.h"
+#include "SDeleteReferencedActorDialog.h"
 
 #define LOCTEXT_NAMESPACE "UnrealEd.EditorActor"
 
@@ -928,6 +929,7 @@ bool UUnrealEdEngine::DeleteActors(const TArray<AActor*>& InActorsToDelete, UWor
 	bool bRequestedDeleteAllByLevel = false;
 	bool bRequestedDeleteAllByActor = false;
 	bool bRequestedDeleteAllBySoftReference = false;
+	bool bRequestedDeleteAllByGroup = false;
 	int32 DeleteCount = 0;
 
 	TUniquePtr<FTypedElementList::FLegacySyncScopedBatch> LegacySyncBatch = MakeUnique<FTypedElementList::FLegacySyncScopedBatch>(*InSelectionSet->GetElementList());
@@ -957,6 +959,7 @@ bool UUnrealEdEngine::DeleteActors(const TArray<AActor*>& InActorsToDelete, UWor
 		bool bReferencedByLevelScript = bWarnAboutReferences && (nullptr != LSB && ReferencedToActorsFromLevelScriptArray.Num() > 0);
 		bool bReferencedByActor = false;
 		bool bReferencedByLODActor = false;
+		bool bReferencedByGroupActor = false;
 		bool bReferencedBySoftReference = false;
 		TArray<UObject*>* SoftReferencingObjects = nullptr;
 
@@ -1005,6 +1008,10 @@ bool UUnrealEdEngine::DeleteActors(const TArray<AActor*>& InActorsToDelete, UWor
 				{
 					bReferencedByLODActor = true;
 				}
+				else if (Cast<AGroupActor>(ReferencingActor))
+				{
+					bReferencedByGroupActor = true;
+				}
 				else
 				{
 					// If the referencing actor is a child actor that is referencing us, do not treat it
@@ -1025,98 +1032,108 @@ bool UUnrealEdEngine::DeleteActors(const TArray<AActor*>& InActorsToDelete, UWor
 		}
 
 		// We have references from one or more sources, prompt the user for feedback.
-		if (bReferencedByLevelScript || bReferencedByActor || bReferencedBySoftReference || bReferencedByLODActor)
+		if (bReferencedByLevelScript || bReferencedByActor || bReferencedBySoftReference || bReferencedByLODActor || bReferencedByGroupActor)
 		{
 			if ((bReferencedByLevelScript && !bRequestedDeleteAllByLevel) ||
 				(bReferencedByActor && !bRequestedDeleteAllByActor) ||
-				(bReferencedBySoftReference && !bRequestedDeleteAllBySoftReference))
+				(bReferencedBySoftReference && !bRequestedDeleteAllBySoftReference) ||
+				(bReferencedByGroupActor && !bRequestedDeleteAllByGroup))
 			{
 				FText ConfirmDelete;
 
-				FString LevelScriptReferenceString;
-
-				for (UK2Node* Node : ReferencedToActorsFromLevelScriptArray)
-				{
-					LevelScriptReferenceString += Node->GetFindReferenceSearchString(EGetFindReferenceSearchStringFlags::None);
-
-					if (bReferencedByLevelScript && bReferencedByActor)
-					{
-						LevelScriptReferenceString += TEXT(" (Level Blueprint)");
-					}
-
-					LevelScriptReferenceString += TEXT("\n");
-				}
-
-				LevelScriptReferenceString.TrimEndInline();
-
-				FString ActorReferenceString;
+				// We separate groups and other actors, so we can visually display them separately in the list
+				TArray<TSharedPtr<FText>> GroupActorReferencers;
+				TArray<TSharedPtr<FText>> OtherReferencers;
 
 				if (ReferencingActors != nullptr)
 				{
 					for (AActor* ReferencingActor : *ReferencingActors)
 					{
-						ActorReferenceString += ReferencingActor->GetActorLabel();
-
-						if (bReferencedByLevelScript && bReferencedByActor)
+						if (ReferencingActor->IsA<ALevelScriptActor>())
 						{
-							ActorReferenceString += TEXT(" (Other Actor)");
+							continue;
 						}
 
-						ActorReferenceString += TEXT("\n");
+						const FString ActorLabel = FString::Printf(TEXT("%s"), *ReferencingActor->GetActorLabel());
+
+						if (ReferencingActor->IsA<AGroupActor>())
+						{
+							GroupActorReferencers.Add(MakeShared<FText>(FText::FromString(ActorLabel)));
+						}
+						else if (!ReferencingActor->IsA<ALevelScriptActor>())
+						{
+							OtherReferencers.Add(MakeShared<FText>(FText::FromString(ActorLabel)));
+						}
 					}
 				}
 
+				OtherReferencers.Append(GroupActorReferencers);
+
 				if (bReferencedBySoftReference)
 				{
 					for (UObject* ReferencingObject : *SoftReferencingObjects)
 					{
+						FString ActorLabel;
+
 						if (AActor* ReferencingActor = Cast<AActor>(ReferencingObject))
 						{
-							ActorReferenceString += FString::Printf(TEXT("(Soft) Actor %s in %s\n"), *ReferencingActor->GetActorLabel(), *FPackageName::GetLongPackageAssetName(ReferencingActor->GetOutermost()->GetName()));
+							ActorLabel = FString::Printf(TEXT("%s in %s (soft reference)"), *ReferencingActor->GetActorLabel(), *FPackageName::GetLongPackageAssetName(ReferencingActor->GetOutermost()->GetName()));
 						}
 						else
 						{
-							ActorReferenceString += FString::Printf(TEXT("(Soft) Object %s\n"), *ReferencingObject->GetPathName());
+							ActorLabel = FString::Printf(TEXT("%s (soft reference)"), *ReferencingObject->GetPathName());
 						}
+
+						OtherReferencers.Add(MakeShared<FText>(FText::FromString(ActorLabel)));
 					}
 				}
 
-				ActorReferenceString.TrimEndInline();
+				EDeletedActorReferenceTypes DeletedActorReferenceTypes = EDeletedActorReferenceTypes::None;
 
-				if (bReferencedByLevelScript && (bReferencedByActor || bReferencedBySoftReference))
+				if (bReferencedByGroupActor)
 				{
-					ConfirmDelete = FText::Format(LOCTEXT("ConfirmDeleteActorReferenceByScriptAndActor", "Actor {0} is referenced by the level blueprint and other Actors/Objects.\nDo you really want to delete it? This will break references.\n\nReference List:\n\n{1}\n{2}"),
-						FText::FromString(Actor->GetActorLabel()), FText::FromString(LevelScriptReferenceString), FText::FromString(ActorReferenceString));
+					EnumAddFlags(DeletedActorReferenceTypes, EDeletedActorReferenceTypes::Group);
 				}
-				else if (bReferencedByLevelScript)
+
+				if (bReferencedByActor || bReferencedBySoftReference)
 				{
-					ConfirmDelete = FText::Format(LOCTEXT("ConfirmDeleteActorReferencedByScript", "Actor {0} is referenced by the level blueprint.\nDo you really want to delete it? This will break references.\n\nReference List:\n\n{1}"),
-						FText::FromString(Actor->GetActorLabel()), FText::FromString(LevelScriptReferenceString));
+					EnumAddFlags(DeletedActorReferenceTypes, EDeletedActorReferenceTypes::ActorOrAsset);
 				}
-				else
+
+				if (bReferencedByLevelScript)
 				{
-					ConfirmDelete = FText::Format(LOCTEXT("ConfirmDeleteActorReferencedByActor", "Actor {0} is referenced by other Actors/Objects.\nDo you really want to delete it? This will break references.\n\nReference List:\n\n{1}"),
-						FText::FromString(Actor->GetActorLabel()), FText::FromString(ActorReferenceString));
+					EnumAddFlags(DeletedActorReferenceTypes, EDeletedActorReferenceTypes::LevelBlueprint);
 				}
 
 				const double DialogStartSeconds = FPlatformTime::Seconds();
 
-				const EAppMsgType::Type MessageType = ActorsToDelete.Num() > 1 ? EAppMsgType::YesNoYesAllNoAll : EAppMsgType::YesNo;
-				int32 Result = FMessageDialog::Open(MessageType, ConfirmDelete);
+				const bool bCanShowApplyToAll = ActorsToDelete.Num() > 1;
+
+				TSharedRef<SDeleteReferencedActorDialog> ConfirmDialog = SNew(SDeleteReferencedActorDialog)
+					.Referencers(OtherReferencers)
+					.ShowApplyToAll(bCanShowApplyToAll)
+					.ActorToDeleteLabel(Actor->GetActorLabel())
+					.ReferenceTypes(DeletedActorReferenceTypes);
+
+				uint32 Result = ConfirmDialog->ShowModal();
 
 				DialogWaitingSeconds += FPlatformTime::Seconds() - DialogStartSeconds;
 
-				if (Result == EAppReturnType::YesAll)
-				{
-					bRequestedDeleteAllByLevel |= bReferencedByLevelScript;
-					bRequestedDeleteAllByActor |= bReferencedByActor;
-					bRequestedDeleteAllBySoftReference |= bReferencedBySoftReference;
-				}
-				else if (Result == EAppReturnType::NoAll)
+				if (ConfirmDialog->GetApplyToAll())
 				{
-					break;
+					if (Result == 0)
+					{
+						bRequestedDeleteAllByLevel |= bReferencedByLevelScript;
+						bRequestedDeleteAllByActor |= bReferencedByActor;
+						bRequestedDeleteAllBySoftReference |= bReferencedBySoftReference;
+						bRequestedDeleteAllByGroup |= bReferencedByGroupActor;
+					}
+					else
+					{
+						break;
+					}
 				}
-				else if (Result == EAppReturnType::No || Result == EAppReturnType::Cancel)
+				else if (Result == 1)
 				{
 					continue;
 				}
@@ -1127,7 +1144,7 @@ bool UUnrealEdEngine::DeleteActors(const TArray<AActor*>& InActorsToDelete, UWor
 				FBlueprintEditorUtils::ModifyActorReferencedGraphNodes(LSB, Actor);
 			}
 
-			if (bReferencedByActor || bReferencedByLODActor)
+			if (bReferencedByActor || bReferencedByLODActor || bReferencedByGroupActor)
 			{
 				check(ReferencingActors != nullptr);
 				for (AActor* ReferencingActor : *ReferencingActors)
diff --git a/Engine/Source/Editor/UnrealEd/Private/SDeleteReferencedActorDialog.cpp b/Engine/Source/Editor/UnrealEd/Private/SDeleteReferencedActorDialog.cpp
new file mode 100644
--- /dev/null
+++ b/Engine/Source/Editor/UnrealEd/Private/SDeleteReferencedActorDialog.cpp
@@ -0,0 +1,279 @@
+// Copyright Epic Games, Inc. All Rights Reserved.
+
+#include "SDeleteReferencedActorDialog.h"
+#include "HAL/PlatformApplicationMisc.h"
+#include "SlateOptMacros.h"
+#include "Widgets/Input/SButton.h"
+#include "Widgets/Input/SCheckBox.h"
+#include "Widgets/Layout/SBox.h"
+#include "Widgets/Layout/SScrollBox.h"
+#include "Widgets/Views/SListView.h"
+#include "Widgets/Views/STableRow.h"
+#include "Widgets/Views/STableViewBase.h"
+
+#define LOCTEXT_NAMESPACE "DeleteReferencedActorDialog"
+
+BEGIN_SLATE_FUNCTION_BUILD_OPTIMIZATION
+
+void SDeleteReferencedActorDialog::Construct(const FArguments& InArgs)
+{
+	ActorLabel = InArgs._ActorToDeleteLabel;
+	ReferenceTypes = InArgs._ReferenceTypes;
+	bShowApplyToAll = InArgs._ShowApplyToAll;
+	ActorReferencers = InArgs._Referencers;
+
+	constexpr uint32 Width  = 600;
+	constexpr uint32 MaxReferencersListHeight = 200;
+
+	CreateMessage();
+
+	const TSharedRef<SScrollBar> HorizontalScrollBar = SNew(SScrollBar)
+		.Orientation(Orient_Horizontal);
+
+	const TSharedRef<SScrollBar> VerticalScrollBar = SNew(SScrollBar)
+	.Orientation(Orient_Vertical);
+
+	SCustomDialog::Construct(SCustomDialog::FArguments()
+		.RootPadding(16.f)
+		.HAlignContent(HAlign_Fill)
+		.VAlignContent(VAlign_Fill)
+		.IconDesiredSizeOverride(FVector2D(24.f, 24.f))
+		.HAlignIcon(HAlign_Left)
+		.VAlignIcon(VAlign_Top)
+		.ContentAreaPadding(FMargin(8.f, 0.f, 0.f, 0.f))
+		.ButtonAreaPadding(FMargin(0.f, 32.f, 0.f, 0.f))
+		.Icon(FAppStyle::Get().GetBrush("Icons.WarningWithColor.Large"))
+		.Title(LOCTEXT("ConfirmDeleteActorMessageTitle", "Delete Referenced Actor"))
+		.Buttons({
+			SCustomDialog::FButton(LOCTEXT("ConfirmDeleteActorMessage_ButtonDelete", "Delete")).SetPrimary(true)
+			, SCustomDialog::FButton(LOCTEXT("ConfirmDeleteActorMessage_ButtonCancel", "Cancel")),
+			})
+		.Content()
+		[
+			SNew(SBox)
+			.MinDesiredWidth(Width)
+			[
+				SNew(SVerticalBox)
+				+ SVerticalBox::Slot()
+				.Padding(FMargin(0.0f, 0.0f, 0.0f, 4.0f))
+				.FillHeight(1.0f)
+				[
+					SNew(STextBlock)
+					.Text(Message)
+				]
+				+ SVerticalBox::Slot()
+				.Padding(FMargin(0.0f, 4.0f))
+				.HAlign(HAlign_Fill)
+				.VAlign(VAlign_Fill)
+				.AutoHeight()
+				[
+					SNew(SBox)
+					.MaxDesiredWidth(Width)
+					.MaxDesiredHeight(MaxReferencersListHeight)
+					.Visibility(this, &SDeleteReferencedActorDialog::GetReferencersListVisibility)
+					[
+						SNew(SVerticalBox)
+						+ SVerticalBox::Slot()
+						[
+							SNew(SHorizontalBox)
+							+ SHorizontalBox::Slot()
+							[
+								SNew(SScrollBox)
+								.ExternalScrollbar(HorizontalScrollBar)
+								.Orientation(Orient_Horizontal)
+								+ SScrollBox::Slot()
+								.FillSize(1.0)
+								[
+									SNew(SScrollBox)
+									.ExternalScrollbar(VerticalScrollBar)
+									.Orientation(Orient_Vertical)
+									+ SScrollBox::Slot()
+									.HAlign(HAlign_Fill)
+									.FillSize(1.0)
+									[
+										SNew(SListView<TSharedPtr<FText>>)
+										.ListItemsSource(&ActorReferencers)
+										.SelectionMode(ESelectionMode::None)
+										.OnGenerateRow(this, &SDeleteReferencedActorDialog::OnGenerateRow)
+									]
+								]
+							]
+							+ SHorizontalBox::Slot()
+							.AutoWidth()
+							[
+								VerticalScrollBar
+							]
+						]
+						+ SVerticalBox::Slot()
+						.AutoHeight()
+						[
+							HorizontalScrollBar
+						]
+					]
+				]
+				+ SVerticalBox::Slot()
+				.Padding(FMargin(0.0f, 8.0f, 0.0f, 2.0f))
+				.FillHeight(1.0f)
+				.HAlign(HAlign_Fill)
+				.AutoHeight()
+				[
+					SNew(SHorizontalBox)
+					+ SHorizontalBox::Slot()
+					.HAlign(HAlign_Left)
+					[
+						SNew(STextBlock)
+						.Text(FText::FromString(TEXT("Delete it anyway?")))
+					]
+				]
+			]
+		]
+		.BeforeButtons()
+		[
+			SNew(SHorizontalBox)
+			+ SHorizontalBox::Slot()
+			.AutoWidth()
+			.HAlign(HAlign_Left)
+			.VAlign(VAlign_Center)
+			[
+				SNew(SButton)
+				.ButtonStyle(FAppStyle::Get(), "SimpleButton")
+				.OnClicked(this, &SDeleteReferencedActorDialog::OnCopyMessageClicked)
+				.ToolTipText(NSLOCTEXT("SChoiceDialog", "CopyMessageTooltip", "Copy the text in this message to the clipboard (CTRL+C)"))
+				.ContentPadding(2.0f)
+				.Content()
+				[
+					SNew(SImage)
+					.Image(FAppStyle::Get().GetBrush("Icons.Clipboard"))
+					.ColorAndOpacity(FSlateColor::UseForeground())
+				]
+			]
+			+ SHorizontalBox::Slot()
+			.FillWidth(1.0f)
+            .HAlign(HAlign_Left)
+            .VAlign(VAlign_Center)
+            .Padding(FMargin(16.f, 0.f, 0.f, 0.f))
+			[
+				SNew(SCheckBox)
+				.IsChecked(ECheckBoxState::Unchecked)
+				.OnCheckStateChanged(this, &SDeleteReferencedActorDialog::OnApplyToAllCheckboxStateChanged)
+				.Visibility(this, &SDeleteReferencedActorDialog::GetApplyToAllCheckboxVisibility)
+				[
+					SNew(STextBlock)
+					.WrapTextAt(615.0f)
+					.Text(LOCTEXT("ApplyToAllLabel", "Apply to All"))
+				]
+			]
+		]
+	);
+}
+
+END_SLATE_FUNCTION_BUILD_OPTIMIZATION
+
+void SDeleteReferencedActorDialog::CreateMessage()
+{
+	Message = FText();
+
+	if (EnumHasAllFlags(ReferenceTypes, EDeletedActorReferenceTypes::All))
+	{
+		Message = FText::Format(LOCTEXT("ConfirmDeleteActorReferenceByScriptActorAndGroup", "Actor '{0}' is referenced by the Level Blueprint by the following other actors or assets and group:")
+				, FText::FromString(ActorLabel));
+	}
+	else if (EnumHasAllFlags(ReferenceTypes, EDeletedActorReferenceTypes::LevelAndActorOrAsset))
+	{
+		Message = FText::Format(LOCTEXT("ConfirmDeleteActorReferenceByScriptAndActor", "Actor '{0}' is referenced by the Level Blueprint and by the following other actors or assets:")
+				, FText::FromString(ActorLabel));
+	}
+	else if (EnumHasAllFlags(ReferenceTypes, EDeletedActorReferenceTypes::GroupAndActorOrAsset))
+    	{
+    		Message = FText::Format(LOCTEXT("ConfirmDeleteActorReferenceByGroupAndActor", "Actor '{0}' is referenced by the following other actors or assets and group:")
+    				, FText::FromString(ActorLabel));
+    	}
+	else if (EnumHasAllFlags(ReferenceTypes, EDeletedActorReferenceTypes::LevelAndGroup))
+	{
+		Message = FText::Format(LOCTEXT("ConfirmDeleteActorReferenceByGroup", "Actor '{0}' is referenced by the Level Blueprint and by the following group:")
+				, FText::FromString(ActorLabel));
+	}
+	else if (EnumHasAllFlags(ReferenceTypes, EDeletedActorReferenceTypes::ActorOrAsset))
+	{
+		Message = FText::Format(LOCTEXT("ConfirmDeleteActorReferenceByActor", "Actor '{0}' is referenced by the following other actors or assets:")
+					, FText::FromString(ActorLabel));
+	}
+	else if (EnumHasAllFlags(ReferenceTypes, EDeletedActorReferenceTypes::Group))
+	{
+		Message = FText::Format(LOCTEXT("ConfirmDeleteActorReferencedByGroup", "Actor '{0}' is in the following group:")
+				, FText::FromString(ActorLabel));
+	}
+	else if (EnumHasAllFlags(ReferenceTypes, EDeletedActorReferenceTypes::LevelBlueprint))
+	{
+		Message = FText::Format(LOCTEXT("ConfirmDeleteActorReferenceByScript", "Actor '{0}' is referenced by the Level Blueprint.")
+				, FText::FromString(ActorLabel));
+	}
+}
+
+void SDeleteReferencedActorDialog::OnApplyToAllCheckboxStateChanged(ECheckBoxState InCheckBoxState)
+{
+	bApplyToAll = InCheckBoxState == ECheckBoxState::Checked;
+}
+
+EVisibility SDeleteReferencedActorDialog::GetApplyToAllCheckboxVisibility() const
+{
+	return bShowApplyToAll ? EVisibility::Visible : EVisibility::Hidden; 
+}
+
+EVisibility SDeleteReferencedActorDialog::GetReferencersListVisibility() const
+{
+	return ReferenceTypes == EDeletedActorReferenceTypes::LevelBlueprint ? EVisibility::Collapsed : EVisibility::Visible;
+}
+
+FReply SDeleteReferencedActorDialog::OnKeyDown(const FGeometry& MyGeometry, const FKeyEvent& InKeyEvent)
+{
+	if (InKeyEvent.GetKey() == EKeys::C && InKeyEvent.IsControlDown())
+	{
+		CopyMessageToClipboard();
+		return FReply::Handled();
+	}
+
+	return FReply::Unhandled();
+}
+
+void SDeleteReferencedActorDialog::CopyMessageToClipboard()
+{
+	FString ClipboardMessage = Message.ToString() + LINE_TERMINATOR;
+	for (const TSharedPtr<FText>& Referencer : ActorReferencers)
+	{
+		if (Referencer)
+		{
+			ClipboardMessage += + TEXT("\t") + Referencer->ToString() + LINE_TERMINATOR;
+		}
+	}
+
+	ClipboardMessage += LINE_TERMINATOR + LOCTEXT("ConfirmDeleteMessageDeleteAnyway", "Delete it anyway?").ToString();
+
+	FPlatformApplicationMisc::ClipboardCopy(*ClipboardMessage);
+}
+
+FReply SDeleteReferencedActorDialog::OnCopyMessageClicked()
+{
+	CopyMessageToClipboard();
+	return FReply::Handled();
+}
+
+TSharedRef<ITableRow> SDeleteReferencedActorDialog::OnGenerateRow(TSharedPtr<FText> InText, const TSharedRef<STableViewBase>& InTableView)
+{
+	return SNew(STableRow<TSharedPtr<FText>>, InTableView)
+		.Padding(FMargin(4.0f))
+		.Content()
+		[
+			SNew(SBox)
+			.VAlign(VAlign_Center)
+			.MaxDesiredHeight(16.0f)
+			.Padding(1.0f)
+			[
+				SNew(STextBlock)
+				.Justification(ETextJustify::Left)
+				.Text(*InText.Get())
+			]
+		];
+}
+
+#undef LOCTEXT_NAMESPACE
diff --git a/Engine/Source/Editor/UnrealEd/Private/SDeleteReferencedActorDialog.h b/Engine/Source/Editor/UnrealEd/Private/SDeleteReferencedActorDialog.h
new file mode 100644
--- /dev/null
+++ b/Engine/Source/Editor/UnrealEd/Private/SDeleteReferencedActorDialog.h
@@ -0,0 +1,66 @@
+// Copyright Epic Games, Inc. All Rights Reserved.
+
+#pragma once
+
+#include "Dialog/SMessageDialog.h"
+#include "Widgets/SCompoundWidget.h"
+
+UENUM()
+enum class EDeletedActorReferenceTypes : uint8
+{
+	None = 0,
+	ActorOrAsset = 1 << 0,
+	Group = 1 << 1,
+	LevelBlueprint = 1 << 2,
+
+	All = ActorOrAsset | Group | LevelBlueprint,
+	LevelAndActorOrAsset = ActorOrAsset | LevelBlueprint,
+	GroupAndActorOrAsset = ActorOrAsset | Group,
+	LevelAndGroup = Group | LevelBlueprint
+};
+ENUM_CLASS_FLAGS(EDeletedActorReferenceTypes)
+
+class ITableRow;
+class STableViewBase;
+
+/**
+ * A dialog to be shown when deleting scene actors which are referenced by one or multiple entities (groups, actors, level BP, etc)
+ */
+class SDeleteReferencedActorDialog : public SCustomDialog
+{
+public:
+	SLATE_BEGIN_ARGS(SDeleteReferencedActorDialog)
+		: _ShowApplyToAll(false)
+		, _ReferenceTypes(EDeletedActorReferenceTypes::None)
+		{
+		}
+
+		SLATE_ARGUMENT(bool, ShowApplyToAll)
+		SLATE_ARGUMENT(FString, ActorToDeleteLabel)
+		SLATE_ARGUMENT(EDeletedActorReferenceTypes, ReferenceTypes)
+		SLATE_ARGUMENT(TArray<TSharedPtr<FText>>, Referencers)
+
+	SLATE_END_ARGS()
+
+	void Construct(const FArguments& InArgs);
+
+	virtual FReply OnKeyDown(const FGeometry& InGeometry, const FKeyEvent& InKeyEvent) override;
+
+	bool GetApplyToAll() const { return bApplyToAll; }
+
+private:
+	void CreateMessage();
+	void OnApplyToAllCheckboxStateChanged(ECheckBoxState InCheckBoxState);
+	EVisibility GetApplyToAllCheckboxVisibility() const;
+	EVisibility GetReferencersListVisibility() const;
+	void CopyMessageToClipboard();
+	FReply OnCopyMessageClicked();
+	TSharedRef<ITableRow> OnGenerateRow(TSharedPtr<FText> InText, const TSharedRef<STableViewBase>& InTableView);
+
+	EDeletedActorReferenceTypes ReferenceTypes = EDeletedActorReferenceTypes::None;
+	TArray<TSharedPtr<FText>> ActorReferencers;
+	bool bShowApplyToAll = false;
+	bool bApplyToAll = false;
+	FString ActorLabel;
+	FText Message;
+};
