BaseToolkit: Make RebuildModeToolBar and ShouldShowModeToolbar methods virtual, to allow derived toolkits more flexibility over constructing and managing the mode toolbar.

diff --git a/Engine/Source/Editor/UnrealEd/Public/Toolkits/BaseToolkit.h b/Engine/Source/Editor/UnrealEd/Public/Toolkits/BaseToolkit.h
--- a/Engine/Source/Editor/UnrealEd/Public/Toolkits/BaseToolkit.h
+++ b/Engine/Source/Editor/UnrealEd/Public/Toolkits/BaseToolkit.h
@@ -189,7 +189,7 @@ protected:
 	/**
 	 * Whether or not the mode toolbar should be shown.  If any active modes generated a toolbar this method will return true
 	 */
-	UNREALED_API bool ShouldShowModeToolbar() const;
+	UNREALED_API virtual bool ShouldShowModeToolbar() const;
 	UNREALED_API TSharedRef<SDockTab> CreatePrimaryModePanel(const FSpawnTabArgs& Args);
 	UNREALED_API void UpdatePrimaryModePanel();
 	UNREALED_API EVisibility GetInlineContentHolderVisibility() const;
@@ -206,7 +206,7 @@ protected:
 	UNREALED_API virtual TSharedRef<SWidget> CreatePaletteWidget(TSharedPtr<FUICommandList> InCommandList, FName InToolbarCustomizationName, FName InPaletteName);
 
 	UNREALED_API void SpawnOrUpdateModeToolbar();
-	UNREALED_API void RebuildModeToolBar();
+	UNREALED_API virtual void RebuildModeToolBar();
 
 
 protected:
