MPCook: Add new substates for Save state - SaveStalledRetracted and SaveStalledAssignedToWorker - so that generated packages retracted from a worker can stay in the save state and be ready to be picked back up if reassigned back to the worker.
The new states are not used yet; this change just makes them available; a followup change will handle moving packages into the new states.
#rnx

diff --git a/Engine/Source/Editor/UnrealEd/Private/CookOnTheFlyServer.cpp b/Engine/Source/Editor/UnrealEd/Private/CookOnTheFlyServer.cpp
--- a/Engine/Source/Editor/UnrealEd/Private/CookOnTheFlyServer.cpp
+++ b/Engine/Source/Editor/UnrealEd/Private/CookOnTheFlyServer.cpp
@@ -2359,7 +2359,7 @@ UCookOnTheFlyServer::ECookAction UCookOnTheFlyServer::DecideNextCookAction(UE::C
 	UE::Cook::FPackageDataMonitor& Monitor = PackageDatas->GetMonitor();
 	if (Monitor.GetNumUrgent() > 0)
 	{
-		if (!bSaveBusy && Monitor.GetNumUrgent(UE::Cook::EPackageState::Save) > 0)
+		if (!bSaveBusy && Monitor.GetNumUrgent(UE::Cook::EPackageState::SaveActive) > 0)
 		{
 			SetIdleStatus(StackData, EIdleStatus::Active);
 			return ECookAction::Save;
@@ -2379,7 +2379,7 @@ UCookOnTheFlyServer::ECookAction UCookOnTheFlyServer::DecideNextCookAction(UE::C
 			SetIdleStatus(StackData, EIdleStatus::Active);
 			return ECookAction::Request;
 		}
-		else if (Monitor.GetNumUrgent(UE::Cook::EPackageState::Save) > 0)
+		else if (Monitor.GetNumUrgent(UE::Cook::EPackageState::SaveActive) > 0)
 		{
 			SetIdleStatus(StackData, EIdleStatus::Active);
 			return ECookAction::Save;
@@ -2401,6 +2401,8 @@ UCookOnTheFlyServer::ECookAction UCookOnTheFlyServer::DecideNextCookAction(UE::C
 		}
 		else
 		{
+			// UE::Cook::EPackageState::SaveStalledAssignedToWorker should not be possible for urgent packages
+			// UE::Cook::EPackageState::SaveStalledRetracted should not be possible for urgent packages
 			checkf(false, TEXT("Urgent request is in state not yet handled by DecideNextCookAction"));
 		}
 	}
@@ -2727,7 +2729,10 @@ void UCookOnTheFlyServer::AssignRequests(TArrayView<UE::Cook::FPackageData*> Req
 			}
 			else
 			{
-				PackageData->SendToState(EPackageState::AssignedToWorker, ESendFlags::QueueAdd, EStateChangeReason::Requested);
+				EPackageState NewState = PackageData->IsInStateProperty(EPackageStateProperty::Saving)
+					? EPackageState::SaveStalledAssignedToWorker
+					: EPackageState::AssignedToWorker;
+				PackageData->SendToState(NewState, ESendFlags::QueueAdd, EStateChangeReason::Requested);
 				PackageData->SetWorkerAssignment(Assignment);
 			}
 		}
@@ -3014,7 +3019,7 @@ void UCookOnTheFlyServer::LoadPackageInQueue(UE::Cook::FPackageData& PackageData
 
 	PostLoadPackageFixup(PackageData, LoadedPackage);
 	PackageData.SetPackage(LoadedPackage);
-	PackageData.SendToState(EPackageState::Save, ESendFlags::QueueAdd, EStateChangeReason::Loaded);
+	PackageData.SendToState(EPackageState::SaveActive, ESendFlags::QueueAdd, EStateChangeReason::Loaded);
 	++OutNumPushed;
 }
 
@@ -3826,7 +3831,7 @@ UE::Cook::EPollStatus UCookOnTheFlyServer::PrepareSaveInternal(UE::Cook::FPackag
 #endif
 	UPackage* Package = PackageData.GetPackage();
 	check(Package && Package->IsFullyLoaded());
-	check(PackageData.GetState() == EPackageState::Save);
+	check(PackageData.GetState() == EPackageState::SaveActive);
 	TRefCountPtr<FGenerationHelper> GenerationHelper;
 
 	if (!PackageData.GetCookedPlatformDataCalled())
@@ -4705,7 +4710,11 @@ void UCookOnTheFlyServer::PumpSaves(UE::Cook::FTickStackData& StackData, uint32
 		{
 			// Timeouts can occur because of new objects created during the save, so we need to update our object cache,
 			// so we call ReleaseCookedPlatformData and ClearObjectCache to clear it and recache on next attempt.
-			ReleaseCookedPlatformData(PackageData, EStateChangeReason::RecreateObjectCache, EPackageState::Save);
+			check(PackageData.GetState() == EPackageState::SaveActive);
+			// TODO: ReleaseCookedPlatformData is not valid for resetting the objectcache for a generator or generated
+			// package; we need to add a function to handle it on the GenerationHelper
+			check(!PackageData.GetGenerationHelper() && !PackageData.GetParentGenerationHelper());
+			ReleaseCookedPlatformData(PackageData, EStateChangeReason::RecreateObjectCache, EPackageState::SaveActive);
 			PackageData.ClearObjectCache();
 			if (PackageData.GetIsUrgent())
 			{
@@ -5249,11 +5258,11 @@ void UCookOnTheFlyServer::PreGarbageCollect()
 	TArray<UE::Cook::FPackageData*> GCKeepPackageDatas;
 
 #if COOK_CHECKSLOW_PACKAGEDATA
-	// Verify that only packages in the save state have pointers to objects
-	for (const FPackageData* PackageData : *PackageDatas.Get())
-	{
-		check(PackageData->GetState() == EPackageState::Save || !PackageData->HasReferencedObjects());
-	}
+	// Verify that only packages in the saving states have pointers to objects
+	PackageDatas->LockAndEnumeratePackageDatas([](const FPackageData* PackageData)
+		{
+			check(PackageData->IsInStateProperty(EPackageStateProperty::Saving) || !PackageData->HasReferencedObjects());
+		});
 #endif
 	if (SavingPackageData)
 	{
@@ -5275,7 +5284,7 @@ void UCookOnTheFlyServer::PreGarbageCollect()
 				bool bShouldDemote;
 				GenerationHelper->PreGarbageCollect(GenerationHelper, *PackageData, GCKeepObjects, GCKeepPackages,
 					GCKeepPackageDatas, bShouldDemote);
-				if (bShouldDemote && PackageData->GetState() >= EPackageState::Save)
+				if (bShouldDemote && PackageData->IsInStateProperty(EPackageStateProperty::Saving))
 				{
 					// Demote any Generated/Generator packages we called PreSave on so they call their PostSave before the GC
 					// or prevent them from being garbage collected if the splitter wants to keep them referenced
@@ -5283,7 +5292,7 @@ void UCookOnTheFlyServer::PreGarbageCollect()
 						EPackageState::Request);
 				}
 			}
-			if (PackageData->GetIsCookLast() && PackageData->GetState() >= EPackageState::Save)
+			if (PackageData->GetIsCookLast() && PackageData->IsInStateProperty(EPackageStateProperty::Saving))
 			{
 				GCKeepPackages.Add(PackageData->GetPackage());
 				GCKeepPackageDatas.Add(PackageData);
@@ -5366,6 +5375,10 @@ void UCookOnTheFlyServer::PreGarbageCollect()
 		{
 			AddPackageName(PackageData->GetPackageName());
 		}
+		for (FPackageData* PackageData : PackageDatas->GetSaveStalledSet())
+		{
+			AddPackageName(PackageData->GetPackageName());
+		}
 
 		TArray<FName> Dependencies;
 		while (!Queue.IsEmpty())
@@ -5463,7 +5476,8 @@ void UCookOnTheFlyServer::PostGarbageCollect()
 	// force delete them despite our reference, and the package is then in an unknown state. If that happens we
 	// demote the package back to request and start its load and save over.
 	TArray<FPackageData*> Demotes;
-	for (FPackageData* PackageData : PackageDatas->GetSaveQueue())
+	auto UpdateSavingPackageAfterGarbageCollect =
+		[&Demotes, &SaveQueueObjectsThatStillExist](FPackageData* PackageData)
 	{
 		bool bOutDemote;
 		PackageData->UpdateSaveAfterGarbageCollect(bOutDemote);
@@ -5483,16 +5497,38 @@ void UCookOnTheFlyServer::PostGarbageCollect()
 				}
 			}
 		}
+	};
+	for (FPackageData* PackageData : PackageDatas->GetSaveQueue())
+	{
+		UpdateSavingPackageAfterGarbageCollect(PackageData);
+	}
+	for (FPackageData* PackageData : PackageDatas->GetSaveStalledSet())
+	{
+		UpdateSavingPackageAfterGarbageCollect(PackageData);
 	}
 	for (FPackageData* PackageData : Demotes)
 	{
-		PackageData->SendToState(EPackageState::Request, ESendFlags::QueueRemove, EStateChangeReason::GarbageCollected);
-		if (PackageData->GetIsCookLast())
+		switch (PackageData->GetState())
 		{
-			// CookLast packages in SaveState have had their urgency removed. Add it back if we need to demote them.
-			PackageData->AddUrgency(true /* bValue */, false /* bAllowUpdateState */);
+		case EPackageState::SaveActive:
+			PackageData->SendToState(EPackageState::Request, ESendFlags::QueueRemove, EStateChangeReason::GarbageCollected);
+			if (PackageData->GetIsCookLast())
+			{
+				// CookLast packages in SaveState have had their urgency removed. Add it back if we need to demote them.
+				PackageData->AddUrgency(true /* bValue */, false /* bAllowUpdateState */);
+			}
+			PackageDatas->GetRequestQueue().AddRequest(PackageData, /* bForceUrgent */ true);
+			break;
+		case EPackageState::SaveStalledAssignedToWorker:
+			PackageData->SendToState(EPackageState::AssignedToWorker, ESendFlags::QueueAddAndRemove, EStateChangeReason::GarbageCollected);
+			break;
+		case EPackageState::SaveStalledRetracted:
+			DemoteToIdle(*PackageData, ESendFlags::QueueAddAndRemove, ESuppressCookReason::RetractedByCookDirector);
+			break;
+		default:
+			checkf(false, TEXT("State %s not handled in a demoted package."), LexToString(PackageData->GetState()));
+			break;
 		}
-		PackageDatas->GetRequestQueue().AddRequest(PackageData, /* bForceUrgent */ true);
 	}
 
 	// Mark that any objects in PendingCookedPlatformDatas should be kept in CachedCookedPlatformData records
@@ -10282,6 +10318,7 @@ void UCookOnTheFlyServer::CookByTheBookFinishedInternal()
 	check(PackageDatas->GetLoadPrepareQueue().IsEmpty());
 	check(PackageDatas->GetLoadReadyQueue().IsEmpty());
 	check(PackageDatas->GetSaveQueue().IsEmpty());
+	check(PackageDatas->GetSaveStalledSet().IsEmpty());
 
 	TArray<FPackageData*> DanglingGenerationHelpers;
 	PackageDatas->LockAndEnumeratePackageDatas([&DanglingGenerationHelpers](FPackageData* PackageData)
@@ -10862,6 +10899,11 @@ void UCookOnTheFlyServer::CancelAllQueues()
 		DemoteToIdle(*PackageData, ESendFlags::QueueAdd, ESuppressCookReason::CookCanceled);
 	}
 	PackageDatas->GetAssignedToWorkerSet().Empty();
+	for (FPackageData* PackageData : PackageDatas->GetSaveStalledSet())
+	{
+		DemoteToIdle(*PackageData, ESendFlags::QueueAdd, ESuppressCookReason::CookCanceled);
+	}
+	PackageDatas->GetSaveStalledSet().Empty();
 	FRequestQueue& RequestQueue = PackageDatas->GetRequestQueue();
 	RequestQueue.GetDiscoveryQueue().Empty();
 	FPackageDataSet& RestartedRequests = RequestQueue.GetRestartedRequests();
diff --git a/Engine/Source/Editor/UnrealEd/Private/Cooker/CookDirector.cpp b/Engine/Source/Editor/UnrealEd/Private/Cooker/CookDirector.cpp
--- a/Engine/Source/Editor/UnrealEd/Private/Cooker/CookDirector.cpp
+++ b/Engine/Source/Editor/UnrealEd/Private/Cooker/CookDirector.cpp
@@ -248,11 +248,14 @@ FCookDirector::~FCookDirector()
 	}
 	for (FPackageData* PackageData : AbortedAssignments)
 	{
-		// Packages that were assigned to workers should be in the AssignedToWorker state
+		// Packages that were assigned to workers should be in an AssignedToWorker state and
+		// therefore should be InProgress.
 		check(PackageData->IsInProgress());
 		PackageData->SetWorkerAssignment(FWorkerId::Invalid(), ESendFlags::QueueNone);
-		PackageData->SendToState(UE::Cook::EPackageState::Request, ESendFlags::QueueAddAndRemove,
-			EStateChangeReason::CookerShutdown);
+		EPackageState NewState = PackageData->IsInStateProperty(EPackageStateProperty::Saving)
+			? EPackageState::SaveActive
+			: EPackageState::Request;
+		PackageData->SendToState(NewState, ESendFlags::QueueAddAndRemove, EStateChangeReason::CookerShutdown);
 	}
 	RemoteWorkers.Empty();
 	RemoteWorkerProfileDatas.Empty();
@@ -720,8 +723,28 @@ void FCookDirector::TickFromSchedulerThread()
 		RetractionHandler->TickFromSchedulerThread(bAllWorkersConnected, bAnyIdle, BusiestNumAssignments);
 	}
 
-	bool bIsStalled = bLocalWorkerIdle && !COTFS.PackageDatas->GetAssignedToWorkerSet().IsEmpty()
-		&& WorkersWithMessage.IsEmpty();
+	bool bIsStalled = bLocalWorkerIdle && WorkersWithMessage.IsEmpty();
+	if (bIsStalled)
+	{
+		// We are only stalled if we have no local work to do and there is a remote worker with assigned work. If no
+		// worker has assigned work then we are done with the cook rather than stalled.
+		bool bRemoteWorkerHasWork = !COTFS.PackageDatas->GetAssignedToWorkerSet().IsEmpty();
+		if (!bRemoteWorkerHasWork)
+		{
+			// Do the slow check to see whether we have a package in the SaveStalledAssignedToWorker state only if
+			// we have to
+			for (FPackageData* PackageData : COTFS.PackageDatas->GetSaveStalledSet())
+			{
+				if (PackageData->GetState() == EPackageState::SaveStalledAssignedToWorker)
+				{
+					bRemoteWorkerHasWork = true;
+					break;
+				}
+			}
+		}
+		bIsStalled = bIsStalled && bRemoteWorkerHasWork;
+	}
+
 	{
 		bReceivingMessages = true;
 		for (TRefCountPtr<FCookWorkerServer>& Worker : WorkersWithMessage)
@@ -1635,11 +1658,13 @@ void FCookDirector::ReassignAbortedPackages(TArray<FPackageData*>& PackagesToRea
 {
 	for (FPackageData* PackageData : PackagesToReassign)
 	{
-		// Packages that were assigned to a worker should be in the AssignedToWorker state
+		// Packages that were assigned to a worker should be in the AssignedToWorker state and therefore in progress.
 		check(PackageData->IsInProgress());
 		PackageData->SetWorkerAssignment(FWorkerId::Invalid());
-		PackageData->SendToState(UE::Cook::EPackageState::Request, ESendFlags::QueueAddAndRemove,
-			EStateChangeReason::ReassignAbortedPackages);
+		EPackageState NewState = PackageData->IsInStateProperty(EPackageStateProperty::Saving)
+			? EPackageState::SaveActive
+			: EPackageState::Request;
+		PackageData->SendToState(NewState, ESendFlags::QueueAddAndRemove, EStateChangeReason::ReassignAbortedPackages);
 	}
 	PackagesToReassign.Empty();
 }
@@ -2127,8 +2152,10 @@ FCookDirector::FRetractionHandler::ReassignPackages(const FWorkerId& FromWorker,
 		}
 		else
 		{
-			PackageData->SendToState(EPackageState::AssignedToWorker, ESendFlags::QueueAdd,
-				EStateChangeReason::Retraction);
+			EPackageState NewState = PackageData->IsInStateProperty(EPackageStateProperty::Saving)
+				? EPackageState::SaveStalledAssignedToWorker
+				: EPackageState::AssignedToWorker;
+			PackageData->SendToState(NewState, ESendFlags::QueueAdd, EStateChangeReason::Retraction);
 			PackageData->SetWorkerAssignment(Assignment);
 		}
 	}
diff --git a/Engine/Source/Editor/UnrealEd/Private/Cooker/CookGenerationHelper.cpp b/Engine/Source/Editor/UnrealEd/Private/Cooker/CookGenerationHelper.cpp
--- a/Engine/Source/Editor/UnrealEd/Private/Cooker/CookGenerationHelper.cpp
+++ b/Engine/Source/Editor/UnrealEd/Private/Cooker/CookGenerationHelper.cpp
@@ -1542,7 +1542,7 @@ void FGenerationHelper::PostGarbageCollect(const TRefCountPtr<FGenerationHelper>
 	bCurrentGCHasKeptGeneratorKeepPackages = false;
 
 	FPackageData& Owner = GetOwner();
-	if (Owner.GetState() == EPackageState::Save)
+	if (Owner.IsInStateProperty(EPackageStateProperty::Saving))
 	{
 		// UCookOnTheFlyServer::PreGarbageCollect adds references for the Generator package and all its public
 		// objects, so it should still be loaded
diff --git a/Engine/Source/Editor/UnrealEd/Private/Cooker/CookPackageData.cpp b/Engine/Source/Editor/UnrealEd/Private/Cooker/CookPackageData.cpp
--- a/Engine/Source/Editor/UnrealEd/Private/Cooker/CookPackageData.cpp
+++ b/Engine/Source/Editor/UnrealEd/Private/Cooker/CookPackageData.cpp
@@ -591,7 +591,7 @@ struct FStateProperties
 			Properties = EPackageStateProperty::InProgress;
 			break;
 		case EPackageState::AssignedToWorker:
-			Properties = EPackageStateProperty::InProgress;
+			Properties = EPackageStateProperty::InProgress | EPackageStateProperty::AssignedToWorkerProperty;
 			break;
 		case EPackageState::LoadPrepare:
 			Properties = EPackageStateProperty::InProgress | EPackageStateProperty::Loading;
@@ -599,9 +599,15 @@ struct FStateProperties
 		case EPackageState::LoadReady:
 			Properties = EPackageStateProperty::InProgress | EPackageStateProperty::Loading;
 			break;
-		// TODO_SaveQueue: When we add state PrepareForSave, it will also have bHasPackage = true, 
-		case EPackageState::Save:
-			Properties = EPackageStateProperty::InProgress | EPackageStateProperty::HasPackage;
+		case EPackageState::SaveActive:
+			Properties = EPackageStateProperty::InProgress | EPackageStateProperty::Saving;
+			break;
+		case EPackageState::SaveStalledRetracted:
+			Properties = EPackageStateProperty::InProgress | EPackageStateProperty::Saving;
+			break;
+		case EPackageState::SaveStalledAssignedToWorker:
+			Properties = EPackageStateProperty::InProgress | EPackageStateProperty::Saving
+				| EPackageStateProperty::AssignedToWorkerProperty;
 			break;
 		default:
 			check(false);
@@ -647,12 +653,26 @@ void FPackageData::SendToState(EPackageState NextState, ESendFlags SendFlags, ES
 		}
 		OnExitLoadReady();
 		break;
-	case EPackageState::Save:
+	case EPackageState::SaveActive:
 		if (!!(SendFlags & ESendFlags::QueueRemove))
 		{
 			ensure(PackageDatas.GetSaveQueue().Remove(this) == 1);
 		}
-		OnExitSave(ReleaseSaveReason, NextState);
+		OnExitSaveActive();
+		break;
+	case EPackageState::SaveStalledRetracted:
+		if (!!(SendFlags & ESendFlags::QueueRemove))
+		{
+			ensure(PackageDatas.GetSaveStalledSet().Remove(this) == 1);
+		}
+		OnExitSaveStalledRetracted();
+		break;
+	case EPackageState::SaveStalledAssignedToWorker:
+		if (!!(SendFlags & ESendFlags::QueueRemove))
+		{
+			ensure(PackageDatas.GetSaveStalledSet().Remove(this) == 1);
+		}
+		OnExitSaveStalledAssignedToWorker();
 		break;
 	default:
 		check(false);
@@ -679,8 +699,11 @@ void FPackageData::SendToState(EPackageState NextState, ESendFlags SendFlags, ES
 			case EPackageStateProperty::Loading:
 				OnExitLoading();
 				break;
-			case EPackageStateProperty::HasPackage:
-				OnExitHasPackage();
+			case EPackageStateProperty::Saving:
+				OnExitSaving(ReleaseSaveReason, NextState);
+				break;
+			case EPackageStateProperty::AssignedToWorkerProperty:
+				OnExitAssignedToWorkerProperty();
 				break;
 			default:
 				check(false);
@@ -703,8 +726,11 @@ void FPackageData::SendToState(EPackageState NextState, ESendFlags SendFlags, ES
 			case EPackageStateProperty::Loading:
 				OnEnterLoading();
 				break;
-			case EPackageStateProperty::HasPackage:
-				OnEnterHasPackage();
+			case EPackageStateProperty::Saving:
+				OnEnterSaving();
+				break;
+			case EPackageStateProperty::AssignedToWorkerProperty:
+				OnEnterAssignedToWorkerProperty();
 				break;
 			default:
 				check(false);
@@ -762,8 +788,8 @@ void FPackageData::SendToState(EPackageState NextState, ESendFlags SendFlags, ES
 			}
 		}
 		break;
-	case EPackageState::Save:
-		OnEnterSave();
+	case EPackageState::SaveActive:
+		OnEnterSaveActive();
 		if (((SendFlags & ESendFlags::QueueAdd) != ESendFlags::QueueNone))
 		{
 			if (GetIsUrgent())
@@ -776,6 +802,20 @@ void FPackageData::SendToState(EPackageState NextState, ESendFlags SendFlags, ES
 			}
 		}
 		break;
+	case EPackageState::SaveStalledRetracted:
+		OnEnterSaveStalledRetracted();
+		if (((SendFlags & ESendFlags::QueueAdd) != ESendFlags::QueueNone))
+		{
+			PackageDatas.GetSaveStalledSet().Add(this);
+		}
+		break;
+	case EPackageState::SaveStalledAssignedToWorker:
+		OnEnterSaveStalledAssignedToWorker();
+		if (((SendFlags & ESendFlags::QueueAdd) != ESendFlags::QueueNone))
+		{
+			PackageDatas.GetSaveStalledSet().Add(this);
+		}
+		break;
 	default:
 		check(false);
 		break;
@@ -802,11 +842,17 @@ void FPackageData::CheckInContainer() const
 	case EPackageState::LoadReady:
 		check(Algo::Find(PackageDatas.GetLoadReadyQueue(), this) != nullptr);
 		break;
-	case EPackageState::Save:
+	case EPackageState::SaveActive:
 		// The save queue is huge and often pushed at end. Check last element first and then scan.
 		check(PackageDatas.GetSaveQueue().Num() && (PackageDatas.GetSaveQueue().Last() == this
 			|| Algo::Find(PackageDatas.GetSaveQueue(), this)));
 		break;
+	case EPackageState::SaveStalledRetracted:
+		check(PackageDatas.GetSaveStalledSet().Contains(this));
+		break;
+	case EPackageState::SaveStalledAssignedToWorker:
+		check(PackageDatas.GetSaveStalledSet().Contains(this));
+		break;
 	default:
 		check(false);
 		break;
@@ -869,9 +915,9 @@ void FPackageData::SetWorkerAssignment(FWorkerId InWorkerAssignment, ESendFlags
 	{
 		if (InWorkerAssignment.IsValid())
 		{
-			checkf(GetState() == EPackageState::AssignedToWorker,
-				TEXT("Package %s is being assigned to worker %d while in a state other than AssignedToWorker. This is invalid."),
-				*GetPackageName().ToString(), GetWorkerAssignment().GetRemoteIndex());
+			checkf(IsInStateProperty(EPackageStateProperty::AssignedToWorkerProperty),
+				TEXT("Package %s is being assigned to worker %d while in state %s, which is not an AssignedToWorker state. This is invalid."),
+				*GetPackageName().ToString(), GetWorkerAssignment().GetRemoteIndex(), LexToString(GetState()));
 		}
 		WorkerAssignment = InWorkerAssignment;
 	}
@@ -879,7 +925,6 @@ void FPackageData::SetWorkerAssignment(FWorkerId InWorkerAssignment, ESendFlags
 
 void FPackageData::OnExitAssignedToWorker()
 {
-	SetWorkerAssignment(FWorkerId::Invalid());
 }
 
 void FPackageData::OnEnterLoadPrepare()
@@ -898,21 +943,28 @@ void FPackageData::OnExitLoadReady()
 {
 }
 
-void FPackageData::OnEnterSave()
+void FPackageData::OnEnterSaveActive()
 {
-	check(GetPackage() != nullptr && GetPackage()->IsFullyLoaded());
+}
 
-	check(!HasPrepareSaveFailed());
-	CheckObjectCacheEmpty();
-	CheckCookedPlatformDataEmpty();
+void FPackageData::OnExitSaveActive()
+{
 }
 
-void FPackageData::OnExitSave(EStateChangeReason ReleaseSaveReason, EPackageState NewState)
+void FPackageData::OnEnterSaveStalledRetracted()
+{
+}
+
+void FPackageData::OnExitSaveStalledRetracted()
+{
+}
+
+void FPackageData::OnEnterSaveStalledAssignedToWorker()
+{
+}
+
+void FPackageData::OnExitSaveStalledAssignedToWorker()
 {
-	PackageDatas.GetCookOnTheFlyServer().ReleaseCookedPlatformData(*this, ReleaseSaveReason, NewState);
-	ClearObjectCache();
-	SetHasPrepareSaveFailed(false);
-	SetIsPrepareSaveRequiresGC(false);
 }
 
 void FPackageData::OnEnterInProgress()
@@ -947,12 +999,21 @@ void FPackageData::OnExitLoading()
 	ClearPreload();
 }
 
-void FPackageData::OnEnterHasPackage()
+void FPackageData::OnEnterSaving()
 {
+	check(GetPackage() != nullptr && GetPackage()->IsFullyLoaded());
+
+	check(!HasPrepareSaveFailed());
+	CheckObjectCacheEmpty();
+	CheckCookedPlatformDataEmpty();
 }
 
-void FPackageData::OnExitHasPackage()
+void FPackageData::OnExitSaving(EStateChangeReason ReleaseSaveReason, EPackageState NewState)
 {
+	PackageDatas.GetCookOnTheFlyServer().ReleaseCookedPlatformData(*this, ReleaseSaveReason, NewState);
+	ClearObjectCache();
+	SetHasPrepareSaveFailed(false);
+	SetIsPrepareSaveRequiresGC(false);
 	SetPackage(nullptr);
 }
 
@@ -964,6 +1025,15 @@ void FPackageData::OnPackageDataFirstMarkedReachable(FInstigator&& InInstigator)
 	PackageDatas.UpdateThreadsafePackageData(*this);
 }
 
+void FPackageData::OnEnterAssignedToWorkerProperty()
+{
+}
+
+void FPackageData::OnExitAssignedToWorkerProperty()
+{
+	SetWorkerAssignment(FWorkerId::Invalid());
+}
+
 void FPackageData::SetState(EPackageState NextState)
 {
 	State = static_cast<uint32>(NextState);
@@ -1457,7 +1527,7 @@ void FPackageData::RemapTargetPlatforms(const TMap<ITargetPlatform*, ITargetPlat
 void FPackageData::UpdateSaveAfterGarbageCollect(bool& bOutDemote)
 {
 	bOutDemote = false;
-	if (GetState() != EPackageState::Save)
+	if (!IsInStateProperty(EPackageStateProperty::Saving))
 	{
 		return;
 	}
@@ -1989,8 +2059,10 @@ void FPackageDataMonitor::OnStateChanged(FPackageData& PackageData, EPackageStat
 		TrackCookLastRequests(OldState, -1);
 		TrackCookLastRequests(NewState, 1);
 	}
-	bool bOldStateAssignedToLocal = OldState != EPackageState::Idle && OldState != EPackageState::AssignedToWorker;
-	bool bNewStateAssignedToLocal = NewState != EPackageState::Idle && NewState != EPackageState::AssignedToWorker;
+	bool bOldStateAssignedToLocal = OldState != EPackageState::Idle &&
+		!EnumHasAnyFlags(FStateProperties(OldState).Properties, EPackageStateProperty::AssignedToWorkerProperty);
+	bool bNewStateAssignedToLocal = NewState != EPackageState::Idle &&
+		!EnumHasAnyFlags(FStateProperties(NewState).Properties, EPackageStateProperty::AssignedToWorkerProperty);
 	if (bOldStateAssignedToLocal != bNewStateAssignedToLocal)
 	{
 		++(bNewStateAssignedToLocal ? MPCookAssignedFenceMarker : MPCookRetiredFenceMarker);
@@ -2083,26 +2155,6 @@ void FPackageDatas::AddReferencedObjects(FReferenceCollector& Collector)
 	return CookOnTheFlyServer.CookerAddReferencedObjects(Collector);
 }
 
-FPackageDataMonitor& FPackageDatas::GetMonitor()
-{
-	return Monitor;
-}
-
-UCookOnTheFlyServer& FPackageDatas::GetCookOnTheFlyServer()
-{
-	return CookOnTheFlyServer;
-}
-
-FRequestQueue& FPackageDatas::GetRequestQueue()
-{
-	return RequestQueue;
-}
-
-FPackageDataQueue& FPackageDatas::GetSaveQueue()
-{
-	return SaveQueue;
-}
-
 FPackageData& FPackageDatas::FindOrAddPackageData(const FName& PackageName, const FName& NormalizedFileName)
 {
 	{
@@ -2547,6 +2599,8 @@ void FPackageDatas::Clear()
 	PendingCookedPlatformDataLists.Empty(); // These destructors will read/write PackageDatas
 	RequestQueue.Empty();
 	SaveQueue.Empty();
+	AssignedToWorkerSet.Empty();
+	SaveStalledSet.Empty();
 	PackageNameToPackageData.Empty();
 	FileNameToPackageData.Empty();
 	CachedCookedPlatformDataObjects.Empty();
diff --git a/Engine/Source/Editor/UnrealEd/Private/Cooker/CookPackageData.h b/Engine/Source/Editor/UnrealEd/Private/Cooker/CookPackageData.h
--- a/Engine/Source/Editor/UnrealEd/Private/Cooker/CookPackageData.h
+++ b/Engine/Source/Editor/UnrealEd/Private/Cooker/CookPackageData.h
@@ -763,15 +763,21 @@ private:
 	void OnExitLoadPrepare();
 	void OnEnterLoadReady();
 	void OnExitLoadReady();
-	void OnEnterSave();
-	void OnExitSave(EStateChangeReason ReleaseSaveReason, EPackageState NewState);
+	void OnEnterSaveActive();
+	void OnExitSaveActive();
+	void OnEnterSaveStalledRetracted();
+	void OnExitSaveStalledRetracted();
+	void OnEnterSaveStalledAssignedToWorker();
+	void OnExitSaveStalledAssignedToWorker();
 	/* Entry/Exit gates for Properties shared between multiple states */
 	void OnExitInProgress();
 	void OnEnterInProgress();
 	void OnExitLoading();
 	void OnEnterLoading();
-	void OnExitHasPackage();
-	void OnEnterHasPackage();
+	void OnExitSaving(EStateChangeReason ReleaseSaveReason, EPackageState NewState);
+	void OnEnterSaving();
+	void OnExitAssignedToWorkerProperty();
+	void OnEnterAssignedToWorkerProperty();
 
 	void OnPackageDataFirstMarkedReachable(FInstigator&& InInstigator);
 
@@ -1118,24 +1124,27 @@ public:
 	FRequestQueue& GetRequestQueue();
 
 	/** Return the Set that holds unordered all PackageDatas that are in the AssignedToWorker state. */
-	TFastPointerSet<FPackageData*>& GetAssignedToWorkerSet() { return AssignedToWorkerSet; }
+	TFastPointerSet<FPackageData*>& GetAssignedToWorkerSet();
 
 	/**
 	 * Return the LoadPrepareQueue used by the CookOnTheFlyServer. The LoadPrepareQueue is the dependency-ordered
 	 * list of FPackageData that need to be preloaded before they can be loaded.
 	 */
-	FLoadPrepareQueue& GetLoadPrepareQueue() { return LoadPrepareQueue; }
+	FLoadPrepareQueue& GetLoadPrepareQueue();
 	/**
 	 * Return the LoadReadyQueue used by the CookOnTheFlyServer. The LoadReadyQueue is the dependency-ordered list
 	 * of PackageData that need to be loaded.
 	 */
-	FPackageDataQueue& GetLoadReadyQueue() { return LoadReadyQueue; }
+	FPackageDataQueue& GetLoadReadyQueue();
 	/**
 	 * Return the SaveQueue used by the CookOnTheFlyServer. The SaveQueue is the performance-sorted list of
 	 * PackageData that have been loaded and need to start or are only part way through saving.
 	 */
 	FPackageDataQueue& GetSaveQueue();
 
+	/** Return the Set that holds unordered all PackageDatas that are in one of the SaveStalled states. */
+	TFastPointerSet<FPackageData*>& GetSaveStalledSet();
+
 	/**
 	Return the PackageData for the given PackageName and FileName; no validation is done on the names.
 	 * Creates the PackageData if it does not already exist.
@@ -1352,6 +1361,7 @@ private:
 	int32 PendingCookedPlatformDataNum = 0;
 	FRequestQueue RequestQueue;
 	TFastPointerSet<FPackageData*> AssignedToWorkerSet;
+	TFastPointerSet<FPackageData*> SaveStalledSet;
 	FLoadPrepareQueue LoadPrepareQueue;
 	FPackageDataQueue LoadReadyQueue;
 	FPackageDataQueue SaveQueue;
@@ -1391,6 +1401,12 @@ struct FPoppedPackageDataScope
 #endif
 };
 
+
+///////////////////////////////////////////////////////
+// Inline implementations
+///////////////////////////////////////////////////////
+
+
 template<typename CallbackType>
 inline void FPackageDatas::UpdateThreadsafePackageData(FName PackageName, CallbackType&& Callback)
 {
@@ -1412,6 +1428,46 @@ inline TOptional<FThreadsafePackageData> FPackageDatas::FindThreadsafePackageDat
 	return Value ? TOptional<FThreadsafePackageData>(*Value) : TOptional<FThreadsafePackageData>();
 }
 
+inline FPackageDataMonitor& FPackageDatas::GetMonitor()
+{
+	return Monitor;
+}
+
+inline UCookOnTheFlyServer& FPackageDatas::GetCookOnTheFlyServer()
+{
+	return CookOnTheFlyServer;
+}
+
+inline FRequestQueue& FPackageDatas::GetRequestQueue()
+{
+	return RequestQueue;
+}
+
+inline TFastPointerSet<FPackageData*>& FPackageDatas::GetAssignedToWorkerSet()
+{
+	return AssignedToWorkerSet;
+}
+
+inline FLoadPrepareQueue& FPackageDatas::GetLoadPrepareQueue()
+{
+	return LoadPrepareQueue;
+}
+
+inline FPackageDataQueue& FPackageDatas::GetLoadReadyQueue()
+{
+	return LoadReadyQueue;
+}
+
+inline TFastPointerSet<FPackageData*>& FPackageDatas::GetSaveStalledSet()
+{
+	return SaveStalledSet;
+}
+
+inline FPackageDataQueue& FPackageDatas::GetSaveQueue()
+{
+	return SaveQueue;
+}
+
 template <typename ArrayType>
 inline void FPackageData::GetPlatformsNeedingCooking(ArrayType& OutPlatforms) const
 {
diff --git a/Engine/Source/Editor/UnrealEd/Private/Cooker/CookTypes.cpp b/Engine/Source/Editor/UnrealEd/Private/Cooker/CookTypes.cpp
--- a/Engine/Source/Editor/UnrealEd/Private/Cooker/CookTypes.cpp
+++ b/Engine/Source/Editor/UnrealEd/Private/Cooker/CookTypes.cpp
@@ -77,7 +77,7 @@ const TCHAR* LexToString(ESuppressCookReason Reason)
 
 const TCHAR* LexToString(UE::Cook::EPackageState Reason)
 {
-	static_assert(static_cast<int>(EPackageState::Count) == 6);
+	static_assert(static_cast<int>(EPackageState::Count) == 8);
 	switch (Reason)
 	{
 	case EPackageState::Idle: return TEXT("Idle");
@@ -85,7 +85,9 @@ const TCHAR* LexToString(UE::Cook::EPackageState Reason)
 	case EPackageState::AssignedToWorker: return TEXT("AssignedToWorker");
 	case EPackageState::LoadPrepare: return TEXT("LoadPrepare");
 	case EPackageState::LoadReady: return TEXT("LoadReady");
-	case EPackageState::Save: return TEXT("Save");
+	case EPackageState::SaveActive: return TEXT("SaveActive");
+	case EPackageState::SaveStalledRetracted: return TEXT("SaveStalledRetracted");
+	case EPackageState::SaveStalledAssignedToWorker: return TEXT("SaveStalledAssignedToWorker");
 	default: return TEXT("Invalid");
 	}
 }
diff --git a/Engine/Source/Editor/UnrealEd/Private/Cooker/CookTypes.h b/Engine/Source/Editor/UnrealEd/Private/Cooker/CookTypes.h
--- a/Engine/Source/Editor/UnrealEd/Private/Cooker/CookTypes.h
+++ b/Engine/Source/Editor/UnrealEd/Private/Cooker/CookTypes.h
@@ -197,14 +197,27 @@ namespace UE::Cook
 		 */
 		LoadReady,
 		/** The Package is in the SaveQueue; it has been fully loaded and some target data may have been calculated. */
-		Save,
+		SaveActive,
+		/**
+		 * The Package is in the SaveStalled Set. It might have Saving data, but it has been retracted by the CookDirector
+		 * and has not yet completed save elsewhere. It will stay in this stalled state until the CookDirector reassigns
+		 * it back to this worker or reports that its save was completed elsewhere.
+		 */
+		SaveStalledRetracted,
+		/**
+		 * The Package is in the SaveStalled Set. We are on the CookDirector and the package was previously assigned
+		 * locally for saving on the Director, but we retracted it from saving locally and assigned it to a remote
+		 * CookWorker. It will stay in this stalled state until COTFS.Director reassigns it back for local saving or a
+		 * worker reports that it finished saving.
+		 */
+		SaveStalledAssignedToWorker,
 
 		Min = Idle,
-		Max = Save,
+		Max = SaveStalledAssignedToWorker,
 		/** Number of values in this enum, not a valid value for any EPackageState variable. */
 		Count = Max + 1,
 		/** Number of bits required to store a valid EPackageState */
-		BitCount = 3,
+		BitCount = FPlatformMath::ConstExprCeilLogTwo(Count),
 	};
 	const TCHAR* LexToString(UE::Cook::EPackageState State);
 
@@ -216,13 +229,15 @@ namespace UE::Cook
 		/** The package is in one of the loading states and has preload data. */
 		Loading		= 0x2,
 		/**
-		 * The package has progressed past the loading state, and the UPackage pointer is available on the
-		 * FPackageData.
+		 * The package is in one of the saving states and has access to saving-only data. The UPackage pointer on
+		 * the FPackageData is non-null.
 		 */
-		HasPackage	= 0x4,
+		Saving		= 0x4,
+		/** The package is assigned to a remote worker, and here on the director it is in a stalled state. */
+		AssignedToWorkerProperty = 0x8,
 
 		Min = InProgress,
-		Max = HasPackage
+		Max = AssignedToWorkerProperty,
 	};
 	ENUM_CLASS_FLAGS(EPackageStateProperty);
 
