FBX Exporter: Fix issue when exporting material where the material name get a index append at the end of the name. The index should be append only if we have several sections using the same material.

#rnx

diff --git a/Engine/Source/Editor/UnrealEd/Private/Fbx/FbxMainExport.cpp b/Engine/Source/Editor/UnrealEd/Private/Fbx/FbxMainExport.cpp
--- a/Engine/Source/Editor/UnrealEd/Private/Fbx/FbxMainExport.cpp
+++ b/Engine/Source/Editor/UnrealEd/Private/Fbx/FbxMainExport.cpp
@@ -1422,7 +1422,7 @@ FbxSurfaceMaterial* FFbxExporter::ExportMaterial(UMaterialInterface* MaterialInt
 	bool bInterchangeOriginated = FFbxMaterialExportUtilities::GetInterchangeShadingModel(MaterialInterface, bUseLambert);
 
 	FString FbxMaterialName = MaterialInterface->GetName();
-	if (MaterialIndexToFbxMaterials != nullptr)
+	if (MaterialIndexToFbxMaterials && MaterialIndexToFbxMaterials->Num() > 0)
 	{
 		FbxMaterialName += FString(TEXT("_")) + FString::FromInt(MaterialIndex);
 	}
