CookPackageSplitter: use the _Generated_ constants from FPackageName now that we have moved them into FPackageName.
#rnx

diff --git a/Engine/Source/Editor/UnrealEd/Private/Cooker/CookPackageSplitter.cpp b/Engine/Source/Editor/UnrealEd/Private/Cooker/CookPackageSplitter.cpp
--- a/Engine/Source/Editor/UnrealEd/Private/Cooker/CookPackageSplitter.cpp
+++ b/Engine/Source/Editor/UnrealEd/Private/Cooker/CookPackageSplitter.cpp
@@ -3,34 +3,18 @@
 #include "CookPackageSplitter.h"
 
 #if WITH_EDITOR
+#include "Misc/PackageName.h"
 #include "Misc/Paths.h"
 #include "Misc/PathViews.h"
 
 const TCHAR* ICookPackageSplitter::GetGeneratedPackageSubPath()
 {
-	return TEXT("_Generated_");
+	return FPackageName::GetGeneratedPackageSubPath();
 }
 
 bool ICookPackageSplitter::IsUnderGeneratedPackageSubPath(FStringView FileOrLongPackagePath)
 {
-	FStringView GeneratedSubDir(GetGeneratedPackageSubPath());
-	int32 Index = FileOrLongPackagePath.Find(GeneratedSubDir, ESearchCase::IgnoreCase);
-	if (Index < 0)
-	{
-		return false;
-	}
-	if (Index == 0 || !FPathViews::IsSeparator(FileOrLongPackagePath[Index - 1]))
-	{
-		// "_Generated_/..." or ".../Prefix_Generated_/...", not the /_Generated_/ subfolder we're looking for.
-		return false;
-	}
-	if (Index + GeneratedSubDir.Len() < FileOrLongPackagePath.Len()
-		&& !FPathViews::IsSeparator(FileOrLongPackagePath[Index + GeneratedSubDir.Len()]))
-	{
-		// ".../_Generated_Suffix/...", not the /_Generated_/ subfolder we're looking for.
-		return false;
-	}
-	return true;
+	return FPackageName::IsUnderGeneratedPackageSubPath(FileOrLongPackagePath);
 }
 
 FString ICookPackageSplitter::ConstructGeneratedPackageName(FName OwnerPackageName, FStringView RelPath,
