Fixed a bug where Asset Guidelines failed to update ini files when the Enable Missing button was clicked.

The previous code was writing the files in an unsupported way. This change uses a method adapted from the Project Settings dialog.

diff --git a/Engine/Source/Editor/UnrealEd/Classes/Editor/AssetGuideline.h b/Engine/Source/Editor/UnrealEd/Classes/Editor/AssetGuideline.h
--- a/Engine/Source/Editor/UnrealEd/Classes/Editor/AssetGuideline.h
+++ b/Engine/Source/Editor/UnrealEd/Classes/Editor/AssetGuideline.h
@@ -26,6 +26,14 @@ struct FIniStringValue
 	/** From .ini, relative to {PROJECT}. Ex: /Config/DefaultEngine.ini */
 	UPROPERTY(EditAnywhere, Category = General)
 	FString Filename;
+
+	/** 
+	 * Ini branch that Filename belongs to. Ex: if Filename is "/Config/DefaultEngine.ini", Branch should be "Engine".
+	 *
+	 * If this isn't set, the system will attempt to detect it based on Filename.
+	 */
+	UPROPERTY(EditAnywhere, Category = General)
+	FString Branch;
 };
 
 /**
diff --git a/Engine/Source/Editor/UnrealEd/Private/AssetGuideline.cpp b/Engine/Source/Editor/UnrealEd/Private/AssetGuideline.cpp
--- a/Engine/Source/Editor/UnrealEd/Private/AssetGuideline.cpp
+++ b/Engine/Source/Editor/UnrealEd/Private/AssetGuideline.cpp
@@ -18,9 +18,12 @@
 #include "Application/SlateApplicationBase.h"
 
 #include "Misc/ConfigCacheIni.h"
+#include "Misc/ConfigContext.h"
 #include "Engine/EngineTypes.h"
 #include "Framework/Notifications/NotificationManager.h"
 #include "Framework/Docking/TabManager.h"
+#include "Internationalization/Regex.h"
+#include "Logging/StructuredLog.h"
 
 #include "Widgets/Notifications/SNotificationList.h"
 #include "Widgets/Notifications/INotificationWidget.h"
@@ -35,7 +38,7 @@
 
 #define LOCTEXT_NAMESPACE "AssetGuideine"
 
-DEFINE_LOG_CATEGORY_STATIC(LogAssetGuideline, Warning, All);
+DEFINE_LOG_CATEGORY_STATIC(LogAssetGuideline, Log, All);
 
 bool UAssetGuideline::bAssetGuidelinesEnabled = true;
 
@@ -114,10 +117,10 @@ void UAssetGuideline::PostLoad()
 			FText AssetName = FText::AsCultureInvariant(GetPackage() ? GetPackage()->GetFName().ToString() : GetFName().ToString());
 
 			FText MissingPlugins = NeededPlugins.IsEmpty() ? FText::GetEmpty() : FText::Format(LOCTEXT("MissingPlugins", "Needed Plugins:\n{0}"), FText::AsCultureInvariant(NeededPlugins));
-			FText PluginWarning = NeededPlugins.IsEmpty() ? FText::GetEmpty() : FText::Format(LOCTEXT("PluginWarning", "Asset '{0}' needs the plugins listed above. Releated assets may not display properly.\nAttemping to save this asset or related assets may result in irreverisble modification due to missing plugins."), AssetName);
+			FText PluginWarning = NeededPlugins.IsEmpty() ? FText::GetEmpty() : FText::Format(LOCTEXT("PluginWarning", "Asset '{0}' needs the plugins listed above. Related assets may not display properly.\nAttemping to save this asset or related assets may result in irreversible modification due to missing plugins."), AssetName);
 
 			FText MissingProjectSettings = NeededProjectSettings.IsEmpty() ? FText::GetEmpty() : FText::Format(LOCTEXT("MissingProjectSettings", "Needed project settings: \n{0}"), FText::AsCultureInvariant(NeededProjectSettings));
-			FText ProjectSettingWarning = NeededProjectSettings.IsEmpty() ? FText::GetEmpty() : FText::Format(LOCTEXT("ProjectSettingWarning", "Asset '{0}' needs the project settings listed above. Releated assets may not display properly."), AssetName);
+			FText ProjectSettingWarning = NeededProjectSettings.IsEmpty() ? FText::GetEmpty() : FText::Format(LOCTEXT("ProjectSettingWarning", "Asset '{0}' needs the project settings listed above. Related assets may not display properly."), AssetName);
 
 			FFormatNamedArguments SubTextArgs;
 			SubTextArgs.Add("PluginSubText", NeededPlugins.IsEmpty() ? FText::GetEmpty() : FText::Format(FText::AsCultureInvariant("{0}{1}\n"), MissingPlugins, PluginWarning));
@@ -279,36 +282,73 @@ void UAssetGuideline::EnableMissingGuidelines(TArray<FString> IncorrectPlugins,
 		}
 
 		TSet<FString> ConfigFilesToFlush;
-		for (const FIniStringValue& IncorrectProjectSetting : IncorrectProjectSettings)
+		if (bSuccess)
 		{
-			// Only fails if file DNE
-			FString FilenamePath = FConfigCacheIni::NormalizeConfigIniPath(FPaths::ProjectDir() + IncorrectProjectSetting.Filename);
-			if (bSuccess && GConfig->Find(FilenamePath))
+			for (const FIniStringValue& IncorrectProjectSetting : IncorrectProjectSettings)
 			{
+				FString Branch = IncorrectProjectSetting.Branch;
+				if (Branch.Len() == 0)
+				{
+					// Try to detect the branch from the filename
+					const FRegexPattern DefaultPathPattern(TEXT("^/Config/Default(\\w+)\\.ini$"), ERegexPatternFlags::CaseInsensitive);
+					FRegexMatcher Matcher(DefaultPathPattern, IncorrectProjectSetting.Filename);
+					if (Matcher.FindNext())
+					{
+						Branch = Matcher.GetCaptureGroup(1);
+						UE_LOGFMT(LogAssetGuideline, Log, "Detected ini branch name {Branch} from filename {Filename}", Branch, IncorrectProjectSetting.Filename);
+					}
+					else
+					{
+						UE_LOGFMT(LogAssetGuideline, Error, "Failed to detect ini branch name from filename {Filename}. Please specify the ini branch explicitly in the asset guideline.", IncorrectProjectSetting.Filename);
+						
+						bSuccess = false;
+						break;
+					}
+				}
+
+				const FString FilenamePath = FConfigCacheIni::NormalizeConfigIniPath(FPaths::ProjectDir() + IncorrectProjectSetting.Filename);
+
+				// Prompt the user to make the file writeable if necessary
 				FGameProjectGenerationModule::Get().TryMakeProjectFileWriteable(FilenamePath);
+				if (FPlatformFileManager::Get().GetPlatformFile().IsReadOnly(*FilenamePath))
+				{
+					UE_LOGFMT(LogAssetGuideline, Error, "Failed to make {Filename} writeable", IncorrectProjectSetting.Filename);
+
+					bSuccess = false;
+					break;
+				}
+
+				// Flush pending writes from any other code systems, so that they're not lost when we reload this branch
+				GConfig->Flush(false, Branch);
+
+				// Create a sandbox FConfigCache
+				FConfigCacheIni Config(EConfigCacheType::Temporary);
 
-				if (!FPlatformFileManager::Get().GetPlatformFile().IsReadOnly(*FilenamePath))
+				// Add an empty file to the config so it doesn't read in the original file (see FConfigCacheIni.Find())
+				FConfigFile& NewFile = Config.Add(FilenamePath, FConfigFile());
+
+				// Set the key in the temp file and write it out to the actual file
+				NewFile.SetString(*IncorrectProjectSetting.Section, *IncorrectProjectSetting.Key, *IncorrectProjectSetting.Value);
+				if (!NewFile.UpdateSinglePropertyInSection(*FilenamePath, *IncorrectProjectSetting.Key, *IncorrectProjectSetting.Section))
 				{
-					GConfig->SetString(*IncorrectProjectSetting.Section, *IncorrectProjectSetting.Key, *IncorrectProjectSetting.Value, FilenamePath);
-					ConfigFilesToFlush.Add(MoveTemp(FilenamePath));
+					UE_LOGFMT(LogAssetGuideline, Error, "Failed to write ini file {Filename}", FilenamePath);
+
+					bSuccess = false;
+					break;
 				}
-				else
+
+				// Reload the branch from disk to synchronise the in-memory value with the on-disk value
+				FConfigContext Context = FConfigContext::ForceReloadIntoGConfig();
+				// No need to write the combined ini out, as all the necessary parts are already synced
+				Context.bWriteDestIni = false;
+				if (!Context.Load(*Branch))
 				{
+					UE_LOGFMT(LogAssetGuideline, Error, "Failed to reload ini branch {Branch}", Branch);
+
 					bSuccess = false;
 					break;
 				}
 			}
-			else
-			{
-				bSuccess = false;
-				break;
-			}
-		}
-
-		for (const FString& ConfigFileToFlush : ConfigFilesToFlush)
-		{
-			constexpr bool bRemoveFromCache = false;
-			GConfig->Flush(bRemoveFromCache, ConfigFileToFlush);
 		}
 
 		if (bSuccess)
