Marquee selection: ensure that the frustum volume is not degenerate when creating it
    - if the frustum's horizontal/vertical plane is flat, frustum selection will fail. to avoid this, we extend the frustum if necessary (FDragTool_ActorFrustumSelect::EndDrag also extends the box rectangle to avoid this kind of problem).

#rnx

diff --git a/Engine/Plugins/Animation/ControlRig/Source/ControlRigEditor/Private/EditMode/ControlRigEditMode.cpp b/Engine/Plugins/Animation/ControlRig/Source/ControlRigEditor/Private/EditMode/ControlRigEditMode.cpp
--- a/Engine/Plugins/Animation/ControlRig/Source/ControlRigEditor/Private/EditMode/ControlRigEditMode.cpp
+++ b/Engine/Plugins/Animation/ControlRig/Source/ControlRigEditor/Private/EditMode/ControlRigEditMode.cpp
@@ -1808,19 +1808,13 @@ bool FControlRigEditMode::FrustumSelect(const FConvexVolume& InFrustum, FEditorV
 	}
 	
 	//need to check for a zero frustum since ComponentIsTouchingSelectionFrustum will return true, selecting everything, when this is the case
-	const bool bMalformedFrustum = (InFrustum.Planes[0].IsNearlyZero() && InFrustum.Planes[2].IsNearlyZero()) || (InFrustum.Planes[3].IsNearlyZero() &&
-		InFrustum.Planes[4].IsNearlyZero());
+	// cf. FDragTool_ActorFrustumSelect::CalculateFrustum 
+	const bool bAreTopBottomMalformed = InFrustum.Planes[0].IsNearlyZero() && InFrustum.Planes[2].IsNearlyZero();
+	const bool bAreRightLeftMalformed = InFrustum.Planes[1].IsNearlyZero() && InFrustum.Planes[3].IsNearlyZero();
+	const bool bMalformedFrustum = bAreTopBottomMalformed || bAreRightLeftMalformed;
 	if (bMalformedFrustum || InViewportClient->IsInGameView() == true || Settings->bHideControlShapes)
 	{
-		if (Settings->bOnlySelectRigControls)
-		{
-			return true;
-		}
-		else
-		{
-			return false;
-		}
-		return FEdMode::FrustumSelect(InFrustum, InViewportClient, InSelect);
+		return Settings->bOnlySelectRigControls;
 	}
 
 	FScopedTransaction ScopedTransaction(LOCTEXT("SelectControlTransaction", "Select Control"), !AreEditingControlRigDirectly() && !GIsTransacting);
diff --git a/Engine/Source/Editor/UnrealEd/Private/DragTool_FrustumSelect.cpp b/Engine/Source/Editor/UnrealEd/Private/DragTool_FrustumSelect.cpp
--- a/Engine/Source/Editor/UnrealEd/Private/DragTool_FrustumSelect.cpp
+++ b/Engine/Source/Editor/UnrealEd/Private/DragTool_FrustumSelect.cpp
@@ -278,22 +278,29 @@ bool FDragTool_ActorFrustumSelect::IntersectsFrustum( const UModel& InModel, int
 	return bIntersects && (!bUseStrictSelection || (bUseStrictSelection && bFullyContained));
 }
 
-void FDragTool_ActorFrustumSelect::CalculateFrustum( FSceneView* View, FConvexVolume& OutFrustum, bool bUseBoxFrustum )
+void FDragTool_ActorFrustumSelect::CalculateFrustum( const FSceneView* InView, FConvexVolume& OutFrustum, bool bUseBoxFrustum ) const
 {
 	if( bUseBoxFrustum )
 	{
 		FVector CamPoint = EditorViewportClient->GetViewLocation();
 		FVector BoxPoint1, BoxPoint2, BoxPoint3, BoxPoint4;
 		FVector WorldDir1, WorldDir2, WorldDir3, WorldDir4;
-		// Deproject the four corners of the selection box
-		FVector2D Point1(FMath::Min(Start.X, End.X), FMath::Min(Start.Y, End.Y)); // Upper Left Corner
-		FVector2D Point2(FMath::Max(Start.X, End.X), FMath::Min(Start.Y, End.Y)); // Upper Right Corner
-		FVector2D Point3(FMath::Max(Start.X, End.X), FMath::Max(Start.Y, End.Y)); // Lower Right Corner
-		FVector2D Point4(FMath::Min(Start.X, End.X), FMath::Max(Start.Y, End.Y)); // Lower Left Corner
-		View->DeprojectFVector2D(Point1, BoxPoint1, WorldDir1);
-		View->DeprojectFVector2D(Point2, BoxPoint2, WorldDir2);
-		View->DeprojectFVector2D(Point3, BoxPoint3, WorldDir3);
-		View->DeprojectFVector2D(Point4, BoxPoint4, WorldDir4);
+		// extend the 2D box of 1 pixel if needed to avoid degenerated volume 
+		const double DX = Start.X == End.X ? 0.5 : 0.0;
+		const double Left = FMath::Min(Start.X, End.X) - DX;
+		const double Right = FMath::Max(Start.X, End.X) + DX;
+		const double DY = Start.Y == End.Y ? 0.5 : 0.0;
+		const double Bottom = FMath::Min(Start.Y, End.Y) - DY;
+		const double Top = FMath::Max(Start.Y, End.Y) + DY;
+		// Deproject the four corners of the selection box		
+		FVector2D Point1(Left, Bottom); // Upper Left Corner
+		FVector2D Point2(Right, Bottom); // Upper Right Corner
+		FVector2D Point3(Right, Top); // Lower Right Corner
+		FVector2D Point4(Left, Top); // Lower Left Corner
+		InView->DeprojectFVector2D(Point1, BoxPoint1, WorldDir1);
+		InView->DeprojectFVector2D(Point2, BoxPoint2, WorldDir2);
+		InView->DeprojectFVector2D(Point3, BoxPoint3, WorldDir3);
+		InView->DeprojectFVector2D(Point4, BoxPoint4, WorldDir4);
 		// Use the camera position and the selection box to create the bounding planes
 		FPlane TopPlane(BoxPoint1, BoxPoint2, CamPoint); // Top Plane
 		FPlane RightPlane(BoxPoint2, BoxPoint3, CamPoint); // Right Plane
@@ -309,11 +316,11 @@ void FDragTool_ActorFrustumSelect::CalculateFrustum( FSceneView* View, FConvexVo
 		OutFrustum.Planes.Add(RightPlane);
 		OutFrustum.Planes.Add(BottomPlane);
 		OutFrustum.Planes.Add(LeftPlane);
-		if ( View->ViewMatrices.GetViewProjectionMatrix().GetFrustumNearPlane(NearPlane) )
+		if ( InView->ViewMatrices.GetViewProjectionMatrix().GetFrustumNearPlane(NearPlane) )
 		{
 			OutFrustum.Planes.Add(NearPlane);
 		}
-		if ( View->ViewMatrices.GetViewProjectionMatrix().GetFrustumFarPlane(FarPlane) )
+		if ( InView->ViewMatrices.GetViewProjectionMatrix().GetFrustumFarPlane(FarPlane) )
 		{
 			OutFrustum.Planes.Add(FarPlane);
 		}
@@ -321,7 +328,7 @@ void FDragTool_ActorFrustumSelect::CalculateFrustum( FSceneView* View, FConvexVo
 	}
 	else
 	{
-		OutFrustum = View->ViewFrustum;
+		OutFrustum = InView->ViewFrustum;
 		OutFrustum.Init();
 	}
 }
diff --git a/Engine/Source/Editor/UnrealEd/Public/DragTool_FrustumSelect.h b/Engine/Source/Editor/UnrealEd/Public/DragTool_FrustumSelect.h
--- a/Engine/Source/Editor/UnrealEd/Public/DragTool_FrustumSelect.h
+++ b/Engine/Source/Editor/UnrealEd/Public/DragTool_FrustumSelect.h
@@ -53,11 +53,12 @@ public:
 private:
 	/** 
 	 * Calculates a frustum to check actors against 
-	 * 
+	 *
+	 * @param InView			Information about scene projection
 	 * @param OutFrustum		The created frustum
 	 * @param bUseBoxFrustum	If true a frustum out of the current dragged box will be created.  false will use the view frustum.
 	 */
-	UNREALED_API void CalculateFrustum( FSceneView* View, FConvexVolume& OutFrustum, bool bUseBoxFrustum );
+	UNREALED_API void CalculateFrustum( const FSceneView* InView, FConvexVolume& OutFrustum, bool bUseBoxFrustum ) const;
 
 	/** 
 	 * Returns true if the provided BSP node intersects with the provided frustum 
