MPCook: CookPackageSplitter: When retracting generated packages that have started saving, stall then in a SaveStalled state rather than demoting them, so that we do not lose data that we will need in them if they are assigned back to the worker. This fixes the error "PackageSplitter specified generated package is map in GetGenerateList results, but then in PreSaveGeneratedPackage created it as uasset", which happens when we garbage collect one of WorldPartition's generated packages and need to recreate it but cant do so correctly without garbage collecting the entire generator package.
#rnx

diff --git a/Engine/Source/Editor/UnrealEd/Private/Cooker/CookDirector.cpp b/Engine/Source/Editor/UnrealEd/Private/Cooker/CookDirector.cpp
--- a/Engine/Source/Editor/UnrealEd/Private/Cooker/CookDirector.cpp
+++ b/Engine/Source/Editor/UnrealEd/Private/Cooker/CookDirector.cpp
@@ -735,7 +735,7 @@ void FCookDirector::TickFromSchedulerThread()
 			// we have to
 			for (FPackageData* PackageData : COTFS.PackageDatas->GetSaveStalledSet())
 			{
-				if (PackageData->GetState() == EPackageState::SaveStalledAssignedToWorker)
+				if (PackageData->IsInStateProperty(EPackageStateProperty::AssignedToWorkerProperty))
 				{
 					bRemoteWorkerHasWork = true;
 					break;
@@ -2152,10 +2152,28 @@ FCookDirector::FRetractionHandler::ReassignPackages(const FWorkerId& FromWorker,
 		}
 		else
 		{
-			EPackageState NewState = PackageData->IsInStateProperty(EPackageStateProperty::Saving)
-				? EPackageState::SaveStalledAssignedToWorker
-				: EPackageState::AssignedToWorker;
-			PackageData->SendToState(NewState, ESendFlags::QueueAdd, EStateChangeReason::Retraction);
+			TRefCountPtr<FGenerationHelper> GenerationHelper = PackageData->GetGenerationHelper();
+			if (!GenerationHelper)
+			{
+				GenerationHelper = PackageData->GetParentGenerationHelper();
+			}
+			bool bShouldStall = false;
+			if (GenerationHelper)
+			{
+				bShouldStall = GenerationHelper->ShouldRetractionStallRatherThanDemote(*PackageData);
+			}
+			if (bShouldStall)
+			{
+				PackageData->Stall(EPackageState::SaveStalledAssignedToWorker, ESendFlags::QueueAddAndRemove);
+				// ShouldRetractionStallRatherThanDemote should not return true unless the Stall will succeed
+				check(PackageData->GetState() == EPackageState::SaveStalledAssignedToWorker);
+				UE_LOG(LogCook, Display, TEXT("Retracting generated package %s; it will remain in memory on the director's worker until the generator finishes saving."),
+					*WriteToString<256>(PackageData->GetPackageName()));
+			}
+			else
+			{
+				PackageData->SendToState(EPackageState::AssignedToWorker, ESendFlags::QueueAdd, EStateChangeReason::Retraction);
+			}
 			PackageData->SetWorkerAssignment(Assignment);
 		}
 	}
diff --git a/Engine/Source/Editor/UnrealEd/Private/Cooker/CookGenerationHelper.cpp b/Engine/Source/Editor/UnrealEd/Private/Cooker/CookGenerationHelper.cpp
--- a/Engine/Source/Editor/UnrealEd/Private/Cooker/CookGenerationHelper.cpp
+++ b/Engine/Source/Editor/UnrealEd/Private/Cooker/CookGenerationHelper.cpp
@@ -115,6 +115,12 @@ void FGenerationHelper::Uninitialize()
 	{
 		return;
 	}
+
+	// Demote stalled packages; we will be garbage collecting so they no longer need to be preserved.
+	// And we need to demote them so that they drop their references to the generation helper and allow it to be
+	// deleted if no longer referenced
+	DemoteStalledPackages(OwnerInfo.PackageData->GetPackageDatas().GetCookOnTheFlyServer());
+
 	NotifyCompletion(ICookPackageSplitter::ETeardown::Complete);
 	check(!CookPackageSplitterInstance);
 
@@ -186,6 +192,28 @@ void FGenerationHelper::OnAllSavesCompleted(UCookOnTheFlyServer& COTFS)
 	// these references
 	ClearKeepForCompletedAllSavesMessage();
 	ClearKeepForAllSavedOrGC();
+
+	// Demote stalled packages; we will no longer need to come back to them
+	DemoteStalledPackages(COTFS);
+}
+
+void FGenerationHelper::DemoteStalledPackages(UCookOnTheFlyServer& COTFS)
+{
+	// For any packages that we stalled because they were retracted and assigned to another worker,
+	// demote them now. But don't demote non-stalled packages, because doing so could demote the final
+	// package that we just saved locally and still needs to finish its work in PumpSaves.
+	auto ConditionalDemote = [&COTFS](FCookGenerationInfo& Info)
+		{
+			if (Info.PackageData->IsStalled())
+			{
+				COTFS.DemoteToIdle(*Info.PackageData, ESendFlags::QueueAddAndRemove, ESuppressCookReason::RetractedByCookDirector);
+			}
+		};
+	ConditionalDemote(OwnerInfo);
+	for (FCookGenerationInfo& Info : PackagesToGenerate)
+	{
+		ConditionalDemote(Info);
+	}
 }
 
 void FGenerationHelper::DiagnoseWhyNotShutdown()
@@ -1268,6 +1296,23 @@ void FGenerationHelper::ResetSaveState(FCookGenerationInfo& Info, UPackage* Pack
 	}
 }
 
+bool FGenerationHelper::ShouldRetractionStallRatherThanDemote(FPackageData& PackageData)
+{
+	FCookGenerationInfo* Info = FindInfo(PackageData);
+	if (Info)
+	{
+		if (PackageData.IsInStateProperty(EPackageStateProperty::Saving))
+		{
+			if (Info->GetSaveState() > FCookGenerationInfo::ESaveState::StartPopulate)
+			{
+				return true;
+			}
+		}
+	}
+	return false;
+}
+
+
 void FGenerationHelper::FetchExternalActorDependencies()
 {
 	if (!IsValid())
@@ -1435,6 +1480,21 @@ void FGenerationHelper::PreGarbageCollectGCLifetimeData()
 	// Starts at one because the caller of PreGarbageCollect has a ref
 	uint32 HoldForGCRefCounts = 1;
 	HoldForGCRefCounts += ReferenceFromKeepForAllSavedOrGC.IsValid() ? 1 : 0;
+	// Every stalled package will be holding a refcount, but is allowed to be demoted and released
+	// if we are going to garbage collect.
+	// If the generator package is stalled, that's a complex case that we don't need to handle optimally;
+	// just keep the entire generation helper referenced.
+	if (OwnerInfo.PackageData->IsStalled())
+	{
+		return;
+	}
+	for (FCookGenerationInfo& Info : PackagesToGenerate)
+	{
+		if (Info.PackageData->GetParentGenerationHelper() && Info.PackageData->IsStalled())
+		{
+			HoldForGCRefCounts += 1;
+		}
+	}
 
 	check(GetRefCount() >= HoldForGCRefCounts);
 	if (GetRefCount() > HoldForGCRefCounts)
@@ -1447,7 +1507,8 @@ void FGenerationHelper::PreGarbageCollectGCLifetimeData()
 	// We should either uninitialize or destroy after the garbage collect.
 	// We should not uninitialize unless the Generator package is going to be collected, but we are in a state
 	// where nothing in the cooker is depending on the package anymore (all generator and generated packages are
-	// not in the save state) so we do expect the generator package to be garbage collected by the upcoming GC
+	// not in the save state or are stalled) so we do expect the generator package to be garbage collected by the
+	// upcoming GC.
 	// But for that to happen we have to drop our references to it from this FGenerationHelper, so we need
 	// to uninitialize. Also mark that we should check for generator garbage collect in PostGarbageCollect.
 	// Depending on the Splitter class, it may tolerate failure to GC the Generator package, in which case we
diff --git a/Engine/Source/Editor/UnrealEd/Private/Cooker/CookGenerationHelper.h b/Engine/Source/Editor/UnrealEd/Private/Cooker/CookGenerationHelper.h
--- a/Engine/Source/Editor/UnrealEd/Private/Cooker/CookGenerationHelper.h
+++ b/Engine/Source/Editor/UnrealEd/Private/Cooker/CookGenerationHelper.h
@@ -340,6 +340,14 @@ public:
 	void ResetSaveState(FCookGenerationInfo& Info, UPackage* Package, EStateChangeReason ReleaseSaveReason,
 		EPackageState NewState);
 
+	/**
+	 * Called when a package is retracted from a CookWorker, or from the Director's local worker. Default behavior
+	 * for a retracted pacakge is to demote the package to idle, but if the state of the package is too advanced to
+	 * recreate without a garbage collect, then the GenerationHelper instead needs to stall the package in case it
+	 * is assigned back to this worker later.
+	 */
+	bool ShouldRetractionStallRatherThanDemote(FPackageData& PackageData);
+
 	/**
 	 * Record all dependencies from the generator package that are ExternalActor dependencies and
 	 * store them on this.
@@ -455,6 +463,7 @@ private:
 	void Uninitialize();
 	void ModifyNumSaved(int32 Delta);
 	void VerifyGeneratorPackageGarbageCollected();
+	void DemoteStalledPackages(UCookOnTheFlyServer& COTFS);
 
 private:
 	// When adding a new variable, add it to Uninitialize as well
diff --git a/Engine/Source/Editor/UnrealEd/Private/Cooker/CookPackageData.cpp b/Engine/Source/Editor/UnrealEd/Private/Cooker/CookPackageData.cpp
--- a/Engine/Source/Editor/UnrealEd/Private/Cooker/CookPackageData.cpp
+++ b/Engine/Source/Editor/UnrealEd/Private/Cooker/CookPackageData.cpp
@@ -824,6 +824,55 @@ void FPackageData::SendToState(EPackageState NextState, ESendFlags SendFlags, ES
 	PackageDatas.GetMonitor().OnStateChanged(*this, OldState);
 }
 
+void FPackageData::Stall(EPackageState TargetState, ESendFlags SendFlags)
+{
+	switch (TargetState)
+	{
+	case EPackageState::SaveStalledAssignedToWorker:
+	case EPackageState::SaveStalledRetracted:
+		if (GetState() != EPackageState::SaveActive)
+		{
+			return;
+		}
+		break;
+	default:
+		return;
+	}
+
+	SendToState(TargetState, SendFlags, EStateChangeReason::Retraction);
+}
+
+void FPackageData::UnStall(ESendFlags SendFlags)
+{
+	EPackageState TargetState = EPackageState::Idle;
+
+	switch (GetState())
+	{
+	case EPackageState::SaveStalledAssignedToWorker:
+	case EPackageState::SaveStalledRetracted:
+		TargetState = EPackageState::SaveActive;
+		break;
+	default:
+		return;
+	}
+
+	UE_LOG(LogCook, Display, TEXT("Unstalling package %s; it will resume saving from the point at which it was retracted."),
+		*WriteToString<256>(GetPackageName()));
+	SendToState(TargetState, SendFlags, EStateChangeReason::Retraction);
+}
+
+bool FPackageData::IsStalled() const
+{
+	switch (GetState())
+	{
+	case EPackageState::SaveStalledAssignedToWorker:
+	case EPackageState::SaveStalledRetracted:
+		return true;
+	default:
+		return false;
+	}
+}
+
 void FPackageData::CheckInContainer() const
 {
 	switch (GetState())
diff --git a/Engine/Source/Editor/UnrealEd/Private/Cooker/CookPackageData.h b/Engine/Source/Editor/UnrealEd/Private/Cooker/CookPackageData.h
--- a/Engine/Source/Editor/UnrealEd/Private/Cooker/CookPackageData.h
+++ b/Engine/Source/Editor/UnrealEd/Private/Cooker/CookPackageData.h
@@ -460,6 +460,22 @@ public:
 	 * @param ReleaseSaveReason Explanation for why the state is changing, used for debugging.
 	 */
 	void SendToState(EPackageState NextState, ESendFlags SendFlags, EStateChangeReason ReleaseSaveReason);
+
+	/**
+	 * Stall the package into the target stalled state, if it is in a valid source state for the target stalled state.
+	 * Does nothing if not in a valid source state. SendFlags are passed into SendToState.
+	 */
+	void Stall(EPackageState TargetState, ESendFlags SendFlags);
+
+	/**
+	 * If the package is in a stalled state, returns the package to the active state that is a source state for that
+	 * stalled state. Does nothing if not in a valid source state.SendFlags are passed into SendToState.
+	 */
+	void UnStall(ESendFlags SendFlags);
+
+	/** Return whether the package is in one of the stalled states. */
+	bool IsStalled() const;
+
 	/* Debug-only code to assert that this PackageData is contained by the container matching its current state. */
 	void CheckInContainer() const;
 	/**
diff --git a/Engine/Source/Editor/UnrealEd/Private/Cooker/CookWorkerClient.cpp b/Engine/Source/Editor/UnrealEd/Private/Cooker/CookWorkerClient.cpp
--- a/Engine/Source/Editor/UnrealEd/Private/Cooker/CookWorkerClient.cpp
+++ b/Engine/Source/Editor/UnrealEd/Private/Cooker/CookWorkerClient.cpp
@@ -895,8 +895,8 @@ void FCookWorkerClient::AssignPackages(FAssignPackagesMessage& Message)
 					&NeedCookPlatformsBuffer);
 			if (PackageData.IsInProgress())
 			{
-				// If already in progress and no new platforms, ignore the duplicate
-				// If there are new platforms, we don't currently handle that; the director should have retracted it first
+				// If already in progress but there are new platforms requested, we don't currently handle that;
+				// the director should have retracted it first
 				for (const ITargetPlatform* TargetPlatform : NeedCookPlatforms)
 				{
 					check(TargetPlatform != CookerLoadingPlatformKey);
@@ -904,6 +904,9 @@ void FCookWorkerClient::AssignPackages(FAssignPackagesMessage& Message)
 						TEXT("CookWorker received AssignPackage for package %s which is already in progress, with new platform %s. Adding new platforms to an inprogress package is not yet supported"),
 						*PackageData.GetPackageName().ToString(), *TargetPlatform->PlatformName());
 				}
+				// If no new platforms, just allow the package to continue in its progress. If it was in a stalled-by-retraction state,
+				// return it to active.
+				PackageData.UnStall(ESendFlags::QueueAddAndRemove);
 				continue;
 			}
 
@@ -1041,7 +1044,27 @@ void FCookWorkerClient::HandleRetractionMessage(FMPCollectorClientMessageContext
 		FPackageData* PackageData = COTFS.PackageDatas->FindPackageDataByPackageName(PackageName);
 		check(PackageData);
 		PackageData->ResetReachable();
-		COTFS.DemoteToIdle(*PackageData, ESendFlags::QueueAddAndRemove, ESuppressCookReason::RetractedByCookDirector);
+		TRefCountPtr<FGenerationHelper> GenerationHelper = PackageData->GetGenerationHelper();
+		if (!GenerationHelper)
+		{
+			GenerationHelper = PackageData->GetParentGenerationHelper();
+		}
+		bool bShouldStall = false;
+		if (GenerationHelper)
+		{
+			bShouldStall = GenerationHelper->ShouldRetractionStallRatherThanDemote(*PackageData);
+		}
+		if (bShouldStall)
+		{
+			UE_LOG(LogCook, Display, TEXT("Retracting generated package %s; it will remain in memory on this worker until the generator finishes saving."),
+				*WriteToString<256>(PackageData->GetPackageName()));
+			PackageData->Stall(EPackageState::SaveStalledRetracted, ESendFlags::QueueAddAndRemove);
+		}
+		else
+		{
+			COTFS.DemoteToIdle(*PackageData, ESendFlags::QueueAddAndRemove,
+				ESuppressCookReason::RetractedByCookDirector);
+		}
 	}
 
 	FRetractionResultsMessage ResultsMessage;
