MPCook: Add log statement on CookWorkers for retraction to debug why retraction can go incompleted for minutes at a time.
#rnx

diff --git a/Engine/Source/Editor/UnrealEd/Private/Cooker/CookWorkerClient.cpp b/Engine/Source/Editor/UnrealEd/Private/Cooker/CookWorkerClient.cpp
--- a/Engine/Source/Editor/UnrealEd/Private/Cooker/CookWorkerClient.cpp
+++ b/Engine/Source/Editor/UnrealEd/Private/Cooker/CookWorkerClient.cpp
@@ -1067,6 +1067,9 @@ void FCookWorkerClient::HandleRetractionMessage(FMPCollectorClientMessageContext
 		}
 	}
 
+	UE_LOG(LogCook, Display, TEXT("Retraction message received from director. %d packages retracted."),
+		PackageNames.Num());
+
 	FRetractionResultsMessage ResultsMessage;
 	ResultsMessage.ReturnedPackages = MoveTemp(PackageNames);
 	SendMessage(ResultsMessage);
