[Sequencer]
 - Pickability/Isolate selection in viewport to only what is in Sequencer
 - Contains no visual element feedback besides cursor icon
 - Works with Level Editor, UMG, and Actor/Blueprint Sequencers

diff --git a/Engine/Plugins/MovieScene/ActorSequence/Source/ActorSequenceEditor/Private/ActorSequenceEditorTabSummoner.cpp b/Engine/Plugins/MovieScene/ActorSequence/Source/ActorSequenceEditor/Private/ActorSequenceEditorTabSummoner.cpp
--- a/Engine/Plugins/MovieScene/ActorSequence/Source/ActorSequenceEditor/Private/ActorSequenceEditorTabSummoner.cpp
+++ b/Engine/Plugins/MovieScene/ActorSequence/Source/ActorSequenceEditor/Private/ActorSequenceEditorTabSummoner.cpp
@@ -22,6 +22,9 @@
 #include "Toolkits/IToolkitHost.h"
 #include "UObject/ObjectSaveContext.h"
 #include "SSubobjectEditor.h"
+#include "EditorViewportSelectability.h"
+#include "EditorViewportSelectabilityBridge.h"
+#include "SequencerCommands.h"
 
 #define LOCTEXT_NAMESPACE "ActorSequenceEditorSummoner"
 
@@ -208,6 +211,13 @@ public:
 			{
 				Sequencer->GetSelectionChangedObjectGuids().Remove(OnSelectionChangedHandle);
 			}
+			if (OnSelectionLimitedChangedHandle.IsValid())
+			{
+				Sequencer->OnViewportSelectionLimitedChanged().Remove(OnSelectionLimitedChangedHandle);
+			}
+
+			// Turn off selection limiting if this tab is closed to avoid user confusion
+			Sequencer->SetViewportSelectionLimited(false);
 
 			FLevelEditorSequencerIntegration::Get().RemoveSequencer(Sequencer.ToSharedRef());
 			Sequencer->Close();
@@ -224,14 +234,22 @@ public:
 			Content->SetContent(SNew(STextBlock).Text(LOCTEXT("NothingSelected", "Select a sequence")));
 		}
 
-		if (WeakBlueprintEditor.IsValid() && WeakBlueprintEditor.Pin()->IsHosted())
+		if (const TSharedPtr<FBlueprintEditor> BlueprintEditor = WeakBlueprintEditor.Pin())
 		{
-			const FName CurveEditorTabName = FName(TEXT("SequencerGraphEditor"));
-			TSharedPtr<SDockTab> ExistingTab = WeakBlueprintEditor.Pin()->GetToolkitHost()->GetTabManager()->FindExistingLiveTab(CurveEditorTabName);
-			if (ExistingTab)
+			if (BlueprintEditor->IsHosted())
 			{
-				ExistingTab->RequestCloseTab();
+				const FName CurveEditorTabName = FName(TEXT("SequencerGraphEditor"));
+				TSharedPtr<SDockTab> ExistingTab = WeakBlueprintEditor.Pin()->GetToolkitHost()->GetTabManager()->FindExistingLiveTab(CurveEditorTabName);
+				if (ExistingTab)
+				{
+					ExistingTab->RequestCloseTab();
+				}
 			}
+
+			FEditorViewportSelectabilityBridge& SelectabilityBridge = BlueprintEditor->GetViewportSelectabilityBridge();
+			SelectabilityBridge.OnIsViewportSelectionLimited().Unbind();
+			SelectabilityBridge.OnGetIsObjectSelectableInViewport().Unbind();
+			SelectabilityBridge.OnGetViewportSelectionLimitedText().Unbind();
 		}
 
 		GEditor->UnregisterForUndo(this);
@@ -412,6 +430,7 @@ public:
 
 		OnGlobalTimeChangedHandle = Sequencer->OnGlobalTimeChanged().AddSP(this, &SActorSequenceEditorWidgetImpl::OnGlobalTimeChanged);
 		OnSelectionChangedHandle = Sequencer->GetSelectionChangedObjectGuids().AddSP(this, &SActorSequenceEditorWidgetImpl::OnSelectionChanged);
+		OnSelectionLimitedChangedHandle = Sequencer->OnViewportSelectionLimitedChanged().AddSP(this, &SActorSequenceEditorWidgetImpl::OnViewportSelectionLimitedChanged);
 
 		FLevelEditorSequencerIntegrationOptions Options;
 		Options.bRequiresLevelEvents = true;
@@ -422,6 +441,14 @@ public:
 	
 		FLevelEditorModule& LevelEditorModule = FModuleManager::LoadModuleChecked<FLevelEditorModule>("LevelEditor");
 		LevelEditorModule.OnMapChanged().AddRaw(this, &SActorSequenceEditorWidgetImpl::HandleMapChanged);
+
+		if (const TSharedPtr<FBlueprintEditor> BlueprintEditor = WeakBlueprintEditor.Pin())
+		{
+			FEditorViewportSelectabilityBridge& SelectabilityBridge = BlueprintEditor->GetViewportSelectabilityBridge();
+			SelectabilityBridge.OnIsViewportSelectionLimited().BindSP(this, &SActorSequenceEditorWidgetImpl::IsViewportSelectionLimited);
+			SelectabilityBridge.OnGetIsObjectSelectableInViewport().BindSP(this, &SActorSequenceEditorWidgetImpl::IsObjectSelectableInViewport);
+			SelectabilityBridge.OnGetViewportSelectionLimitedText().BindSP(this, &SActorSequenceEditorWidgetImpl::GetViewportSelectionLimitedText);
+		}
 	}
 
 	void HandleMapChanged(UWorld* NewWorld, EMapChangeType MapChangeType)
@@ -629,6 +656,74 @@ public:
 		OnModifiedIndirectly(Object);
 	}
 
+	bool IsViewportSelectionLimited() const
+	{
+		if (Sequencer.IsValid())
+		{
+			return Sequencer->IsViewportSelectionLimited();
+		}
+		return true;
+	}
+
+	bool IsObjectSelectableInViewport(UObject* const InObject) const
+	{
+		if (Sequencer.IsValid())
+		{
+			return Sequencer->IsObjectSelectableInViewport(InObject);
+		}
+		return true;
+	}
+
+	void OnViewportSelectionLimitedChanged(const bool bInSelectionLimited)
+	{
+		const TSharedPtr<FBlueprintEditor> BlueprintEditor = WeakBlueprintEditor.Pin();
+		if (!BlueprintEditor.IsValid())
+		{
+			return;
+		}
+
+		const TSharedPtr<SSubobjectEditor> SubojectEditor = BlueprintEditor->GetSubobjectEditor();
+		if (!SubojectEditor.IsValid())
+		{
+			return;
+		}
+
+		FEditorViewportSelectabilityBridge& SelectabilityBridge = BlueprintEditor->GetViewportSelectabilityBridge();
+
+		if (SelectabilityBridge.IsViewportSelectionLimited())
+		{
+			// Filter out non-selectable nodes from the selection
+			TSet<FSubobjectEditorTreeNodePtrType> NewSelectedNodes;
+			for (const FSubobjectEditorTreeNodePtrType& Node : SubojectEditor->GetSelectedNodes())
+			{
+				const FSubobjectData* const Data = Node->GetDataSource();
+				const UActorComponent* const Component = Data ? Data->FindComponentInstanceInActor(GetPreviewActor()) : nullptr;
+				if (IsValid(Component) && IsObjectSelectableInViewport(const_cast<UActorComponent*>(Component)))
+				{
+					NewSelectedNodes.Add(Node);
+				}
+			}
+
+			// Clear selection and select the new nodes
+			SubojectEditor->ClearSelection();
+			for (const FSubobjectEditorTreeNodePtrType& Node : NewSelectedNodes)
+			{
+				SubojectEditor->SelectNode(Node, false);
+			}
+		}
+	}
+
+	FText GetViewportSelectionLimitedText() const
+	{
+		if (Sequencer.IsValid())
+		{
+			const FSequencerCommands& SequencerCommands = FSequencerCommands::Get();
+			return FEditorViewportSelectability::GetLimitedSelectionText(SequencerCommands.ToggleLimitViewportSelection
+				, LOCTEXT("SequencerSelectionLimitHelpText", "Sequencer Selection Limited"));
+		}
+		return FText();
+	}
+
 private:
 	TWeakObjectPtr<UActorSequence> WeakSequence;
 
@@ -642,6 +737,7 @@ private:
 
 	FDelegateHandle OnSelectionChangedHandle;
 	FDelegateHandle OnGlobalTimeChangedHandle;
+	FDelegateHandle OnSelectionLimitedChangedHandle;
 };
 
 void SActorSequenceEditorWidget::Construct(const FArguments&, TWeakPtr<FBlueprintEditor> InBlueprintEditor)
diff --git a/Engine/Source/Editor/Kismet/Private/BlueprintEditor.cpp b/Engine/Source/Editor/Kismet/Private/BlueprintEditor.cpp
--- a/Engine/Source/Editor/Kismet/Private/BlueprintEditor.cpp
+++ b/Engine/Source/Editor/Kismet/Private/BlueprintEditor.cpp
@@ -181,6 +181,7 @@
 #include "BlueprintActionDatabase.h"
 #include "Algo/MinElement.h"
 #include "Editor/EditorEngine.h"
+#include "EditorViewportSelectabilityBridge.h"
 
 DEFINE_LOG_CATEGORY_STATIC(LogBlueprintEditor, Log, All);
 
@@ -2459,6 +2460,11 @@ void FBlueprintEditor::InitBlueprintEditor(
 		.AddRaw(this, &FBlueprintEditor::OnBlueprintEditorPreferencesChanged);
 	BlueprintProjectSettingsChangedHandle = GetMutableDefault<UBlueprintEditorProjectSettings>()->OnSettingChanged()
 		.AddRaw(this, &FBlueprintEditor::OnBlueprintProjectSettingsChanged);
+
+	if (const TSharedPtr<SSCSEditorViewport> Viewport = GetSubobjectViewport())
+	{
+		ViewportSelectabilityBridge = MakeUnique<FEditorViewportSelectabilityBridge>(Viewport->GetViewportClient());
+	}
 }
 
 void FBlueprintEditor::InitToolMenuContext(FToolMenuContext& MenuContext)
@@ -10884,6 +10890,12 @@ bool FBlueprintEditor::AreDelegatesAllowed() const
 	return true;
 }
 
+FEditorViewportSelectabilityBridge& FBlueprintEditor::GetViewportSelectabilityBridge()
+{
+	check(ViewportSelectabilityBridge.IsValid());
+	return *ViewportSelectabilityBridge.Get();
+}
+
 /////////////////////////////////////////////////////
 /////////////////////////////////////////////////////
 /////////////////////////////////////////////////////
diff --git a/Engine/Source/Editor/Kismet/Private/SCSEditorViewportClient.cpp b/Engine/Source/Editor/Kismet/Private/SCSEditorViewportClient.cpp
--- a/Engine/Source/Editor/Kismet/Private/SCSEditorViewportClient.cpp
+++ b/Engine/Source/Editor/Kismet/Private/SCSEditorViewportClient.cpp
@@ -22,7 +22,9 @@
 #include "Editor/EditorEngine.h"
 #include "Editor/EditorPerProjectUserSettings.h"
 #include "Editor/UnrealEdEngine.h"
+#include "EditorViewportSelectability.h"
 #include "EditorComponents.h"
+#include "EditorViewportSelectabilityBridge.h"
 #include "Engine/Blueprint.h"
 #include "Engine/Engine.h"
 #include "Engine/EngineTypes.h"
@@ -337,6 +339,18 @@ void FSCSEditorViewportClient::DrawCanvas( FViewport& InViewport, FSceneView& Vi
 			}
 		}
 	}
+
+	// If enabled, draw viewport selection limited text
+	const TSharedPtr<FBlueprintEditor> BlueprintEditor = BlueprintEditorPtr.Pin();
+	if (BlueprintEditor.IsValid())
+	{
+		FEditorViewportSelectabilityBridge& SelectabilityBridge = BlueprintEditor->GetViewportSelectabilityBridge();
+		if (SelectabilityBridge.IsViewportSelectionLimited())
+		{
+			const FText SelectionLimitedText = SelectabilityBridge.GetViewportSelectionLimitedText();
+			FEditorViewportSelectability::DrawEnabledTextNotice(&Canvas, SelectionLimitedText);
+		}
+	}
 }
 
 bool FSCSEditorViewportClient::InputKey(const FInputKeyEventArgs& EventArgs)
@@ -351,6 +365,24 @@ bool FSCSEditorViewportClient::InputKey(const FInputKeyEventArgs& EventArgs)
 	return bHandled;
 }
 
+void FSCSEditorViewportClient::MouseMove(FViewport* InViewport, int32 InX, int32 InY)
+{
+	HHitProxy* const HitResult = InViewport ? InViewport->GetHitProxy(InX, InY) : nullptr;
+	UpdateHoverFromHitProxy(HitResult);
+
+	return FEditorViewportClient::MouseMove(InViewport, InX, InY);
+}
+
+EMouseCursor::Type FSCSEditorViewportClient::GetCursor(FViewport* InViewport, int32 InX, int32 InY)
+{
+	if (MouseCursor.IsSet())
+	{
+		return MouseCursor.GetValue();
+	}
+
+	return FEditorViewportClient::GetCursor(InViewport, InX, InY);
+}
+
 void FSCSEditorViewportClient::ProcessClick(class FSceneView& View, class HHitProxy* HitProxy, FKey Key, EInputEvent Event, uint32 HitX, uint32 HitY)
 {
 	const FViewportClick Click(&View, this, Key, Event, HitX, HitY);
@@ -418,7 +450,8 @@ void FSCSEditorViewportClient::ProcessClick(class FSceneView& View, class HHitPr
 				if (ActorProxy->Actor == PreviewActor)
 				{
 					const UPrimitiveComponent* TestComponent = ActorProxy->PrimComponent;
-					if (ActorProxy->Actor->GetComponents().Contains(TestComponent))
+					if (ActorProxy->Actor->GetComponents().Contains(TestComponent)
+						&& IsObjectSelectableInViewport(const_cast<UPrimitiveComponent*>(TestComponent)))
 					{
 						SelectedCompInstance = TestComponent;
 					}
@@ -1160,3 +1193,77 @@ void FSCSEditorViewportClient::RefreshPreviewBounds()
 		PreviewActorBounds = BoundsBuilder;
 	}
 }
+
+bool FSCSEditorViewportClient::IsViewportSelectionLimited() const
+{
+	if (const TSharedPtr<FBlueprintEditor> BlueprintEditor = BlueprintEditorPtr.Pin())
+	{
+		return BlueprintEditor->GetViewportSelectabilityBridge().IsViewportSelectionLimited();
+	}
+	return false;
+}
+
+bool FSCSEditorViewportClient::IsObjectSelectableInViewport(UObject* const InObject) const
+{
+	if (const TSharedPtr<FBlueprintEditor> BlueprintEditor = BlueprintEditorPtr.Pin())
+	{
+		return BlueprintEditor->GetViewportSelectabilityBridge().IsObjectSelectableInViewport(InObject);
+	}
+	return true;
+}
+
+void FSCSEditorViewportClient::UpdateHoverFromHitProxy(HHitProxy* const InHitProxy)
+{
+	const bool bIsViewportSelectionLimited = IsViewportSelectionLimited();
+
+	UPrimitiveComponent* PrimitiveComponent = nullptr;
+	bool bIsGizmoHit = false;
+	bool bIsActorHit = false;
+
+	if (InHitProxy)
+	{
+		if (InHitProxy->IsA(HWidgetAxis::StaticGetType()))
+		{
+			if (bIsViewportSelectionLimited)
+			{
+				bIsGizmoHit = true;
+			}
+		}
+		else if (InHitProxy->IsA(HActor::StaticGetType()))
+		{
+			const HActor* const ActorHitProxy = static_cast<HActor*>(InHitProxy);
+			if (ActorHitProxy && IsValid(ActorHitProxy->PrimComponent))
+			{
+				if (bIsViewportSelectionLimited)
+				{
+					bIsActorHit = true;
+				}
+				PrimitiveComponent = const_cast<UPrimitiveComponent*>(ActorHitProxy->PrimComponent.Get());
+			}
+		}
+	}
+
+	FEditorViewportSelectability::UpdateHoveredPrimitive(bIsViewportSelectionLimited
+		, PrimitiveComponent
+		, HoveredPrimitiveComponents
+		, [this](UObject* const InObject) -> bool
+		{
+			return IsObjectSelectableInViewport(InObject);
+		});
+
+	// Set mouse cursor after hovered primitive component list has been updated
+	MouseCursor.Reset();
+
+	if (bIsGizmoHit)
+	{
+		MouseCursor = EMouseCursor::CardinalCross;
+	}
+	else if (bIsActorHit)
+	{
+		MouseCursor = HoveredPrimitiveComponents.IsEmpty() ? EMouseCursor::SlashedCircle : EMouseCursor::Crosshairs;
+	}
+	else if (bIsViewportSelectionLimited)
+	{
+		MouseCursor = EMouseCursor::SlashedCircle;
+	}
+}
diff --git a/Engine/Source/Editor/Kismet/Private/SCSEditorViewportClient.h b/Engine/Source/Editor/Kismet/Private/SCSEditorViewportClient.h
--- a/Engine/Source/Editor/Kismet/Private/SCSEditorViewportClient.h
+++ b/Engine/Source/Editor/Kismet/Private/SCSEditorViewportClient.h
@@ -3,6 +3,7 @@
 #pragma once
 
 #include "EditorViewportClient.h"
+#include "EditorViewportSelectability.h"
 #include "Engine/EngineBaseTypes.h"
 #include "InputCoreTypes.h"
 #include "Math/Axis.h"
@@ -15,6 +16,7 @@
 
 class AActor;
 class FCanvas;
+class FEditorViewportSelectability;
 class FPreviewScene;
 class FPrimitiveDrawInterface;
 class FSceneView;
@@ -50,6 +52,8 @@ public:
 	virtual void Draw(const FSceneView* View,FPrimitiveDrawInterface* PDI) override;
 	virtual void DrawCanvas( FViewport& InViewport, FSceneView& View, FCanvas& Canvas ) override;
 	virtual bool InputKey(const FInputKeyEventArgs& EventArgs) override;
+	virtual void MouseMove(FViewport* InViewport, int32 InX, int32 InY) override;
+	virtual EMouseCursor::Type GetCursor(FViewport* InViewport, int32 InX, int32 InY) override;
 	virtual void ProcessClick(class FSceneView& View, class HHitProxy* HitProxy, FKey Key, EInputEvent Event, uint32 HitX, uint32 HitY) override;
 	virtual bool InputWidgetDelta( FViewport* Viewport, EAxisList::Type CurrentAxis, FVector& Drag, FRotator& Rot, FVector& Scale ) override;
 	virtual void TrackingStarted( const struct FInputEventState& InInputState, bool bIsDragging, bool bNudge ) override;
@@ -133,6 +137,11 @@ public:
 	 */
 	AActor* GetPreviewActor() const;
 
+	bool IsViewportSelectionLimited() const;
+
+	/** @return True if the specified object is selectable in the viewport and not made unselectable by the Sequencer selection limiting. */
+	bool IsObjectSelectableInViewport(UObject* const InObject) const;
+
 protected:
 	/**
 	 * Initiates a transaction.
@@ -149,7 +158,10 @@ protected:
 	 */
 	void RefreshPreviewBounds();
 
+	void UpdateHoverFromHitProxy(HHitProxy* const InHitProxy);
+
 private:
+
 	UE::Widget::EWidgetMode WidgetMode;
 	ECoordSystem WidgetCoordSystem;
 
@@ -173,4 +185,9 @@ private:
 
 	bool HandleBeginTransform();
 	bool HandleEndTransform();
+
+	/** Hovered primitives and their last overlay color before we apply the hover overlay */
+	TMap<TWeakObjectPtr<UPrimitiveComponent>, TOptional<FColor>> HoveredPrimitiveComponents;
+
+	TOptional<EMouseCursor::Type> MouseCursor;
 };
diff --git a/Engine/Source/Editor/Kismet/Public/BlueprintEditor.h b/Engine/Source/Editor/Kismet/Public/BlueprintEditor.h
--- a/Engine/Source/Editor/Kismet/Public/BlueprintEditor.h
+++ b/Engine/Source/Editor/Kismet/Public/BlueprintEditor.h
@@ -58,6 +58,8 @@
 class AActor;
 class FBlueprintEditorToolbar;
 class FBlueprintNamespaceHelper;
+class FEditorViewportSelectability;
+class FEditorViewportSelectabilityBridge;
 class FProperty;
 class FReferenceCollector;
 class FSCSEditorTreeNode;
@@ -1384,6 +1386,8 @@ public://@TODO
 	/** Whether delegates are allowed  to be displayed/created in a blueprint */
 	virtual bool AreDelegatesAllowed() const;
 
+	FEditorViewportSelectabilityBridge& GetViewportSelectabilityBridge();
+
 protected:
 
 	/** Should intermediate build products be saved when recompiling? */
@@ -1639,6 +1643,8 @@ private:
 
 	virtual void OnBlueprintProjectSettingsChanged(UObject*, struct FPropertyChangedEvent&);
 	virtual void OnBlueprintEditorPreferencesChanged(UObject*, struct FPropertyChangedEvent&);
+
+	TUniquePtr<FEditorViewportSelectabilityBridge> ViewportSelectabilityBridge;
 };
 
 #undef LOCTEXT_NAMESPACE
diff --git a/Engine/Source/Editor/Sequencer/Private/SSequencer.cpp b/Engine/Source/Editor/Sequencer/Private/SSequencer.cpp
--- a/Engine/Source/Editor/Sequencer/Private/SSequencer.cpp
+++ b/Engine/Source/Editor/Sequencer/Private/SSequencer.cpp
@@ -1569,10 +1569,23 @@ void SSequencer::PopulateToolBar(UToolMenu* InMenu)
 			AllowEditsModeEntry.StyleNameOverride = SequencerToolbarStyleName;
 			Section.AddEntry(AllowEditsModeEntry);
 		}
-
+		
 		Section.AddSeparator(NAME_None);
 	}
 
+	{
+		FToolMenuSection& SelectionSection = InMenu->AddSection(TEXT("Selection"));
+
+		FToolMenuEntry ToggleLockViewportSelectionEntry = FToolMenuEntry::InitToolBarButton(FSequencerCommands::Get().ToggleLimitViewportSelection
+			, FText()
+			, FSequencerCommands::Get().ToggleLimitViewportSelection->GetDescription()
+			, FSlateIcon(FAppStyle::GetAppStyleSetName(), TEXT("ContentBrowser.Local.Small")));
+		ToggleLockViewportSelectionEntry.StyleNameOverride = SequencerToolbarStyleName;
+		SelectionSection.AddEntry(ToggleLockViewportSelectionEntry);
+
+		SelectionSection.AddSeparator(NAME_None);
+	}
+
 	{
 		FToolMenuSection& SnappingSection = InMenu->AddSection("Snapping");
 
diff --git a/Engine/Source/Editor/Sequencer/Private/Sequencer.cpp b/Engine/Source/Editor/Sequencer/Private/Sequencer.cpp
--- a/Engine/Source/Editor/Sequencer/Private/Sequencer.cpp
+++ b/Engine/Source/Editor/Sequencer/Private/Sequencer.cpp
@@ -289,6 +289,7 @@ namespace UE
 	} // namespace Sequencer
 } // namespace UE
 
+bool FSequencer::bSelectionLimited = false;
 
 void FSequencer::InitSequencer(const FSequencerInitParams& InitParams, const TSharedRef<ISequencerObjectChangeListener>& InObjectChangeListener, const TArray<FOnCreateTrackEditor>& TrackEditorDelegates, const TArray<FOnCreateEditorObjectBinding>& EditorObjectBindingDelegates, const TArray<FOnCreateOutlinerColumn>& OutlinerColumnDelegates)
 {
@@ -11503,6 +11504,12 @@ void FSequencer::BindCommands()
 		Commands.RefreshUI,
 		FExecuteAction::CreateSP( this, &FSequencer::RefreshUI));
 
+	SequencerCommandBindings->MapAction(
+		Commands.ToggleLimitViewportSelection,
+		FExecuteAction::CreateSP(this, &FSequencer::ToggleLimitViewportSelection),
+		FCanExecuteAction(),
+		FIsActionChecked::CreateSP(this, &FSequencer::IsViewportSelectionLimited));
+
 	// If this sequencer supports a curve editor, let's add bindings for it.
 	FCurveEditorExtension* CurveEditorExtension = ViewModel->CastDynamic<FCurveEditorExtension>();
 	if (CurveEditorExtension && ensure(CurveEditorExtension->GetCurveEditor()))
@@ -11941,5 +11948,107 @@ void FSequencer::OnCameraCutUpdated(const UE::MovieScene::FOnCameraCutUpdatedPar
 	OnCameraCutEvent.Broadcast(Params.ViewTarget, Params.bIsJumpCut);
 }
 
-#undef LOCTEXT_NAMESPACE
+void FSequencer::ToggleLimitViewportSelection()
+{
+	bSelectionLimited = !bSelectionLimited;
+	
+	SetViewportSelectionLimited(bSelectionLimited);
+}
+
+bool FSequencer::IsViewportSelectionLimited() const
+{
+	return bSelectionLimited;
+}
+
+void FSequencer::SetViewportSelectionLimited(const bool bInSelectionLimited)
+{
+	bSelectionLimited = bInSelectionLimited;
+
+	if (FSequencerEdMode* const SequencerEdMode = (FSequencerEdMode*)GLevelEditorModeTools().GetActiveMode(FSequencerEdMode::EM_SequencerMode))
+	{
+		SequencerEdMode->EnableSelectabilityTool(bSelectionLimited);
+	}
+
+	OnSelectionLimitedChangedDelegate.Broadcast(bSelectionLimited);
+}
+
+bool FSequencer::IsObjectSelectableInViewport(UObject* const InObject)
+{
+	using namespace UE::Sequencer;
+
+	if (!bSelectionLimited)
+	{
+		return true;
+	}
+
+	UMovieSceneSequence* const FocusedSequence = GetFocusedMovieSceneSequence();
+	if (!IsValid(FocusedSequence))
+	{
+		return true;
+	}
+
+	const TSharedRef<UE::MovieScene::FSharedPlaybackState> SharedPlaybackState = GetSharedPlaybackState();
+	const FMovieSceneEvaluationState* const EvaluationState = SharedPlaybackState->FindCapability<FMovieSceneEvaluationState>();
+	if (!EvaluationState)
+	{
+		return true;
+	}
+
+	const UMovieSceneSequence* OutSequence = nullptr;
 
+	// Early out on first sequence the object is found in
+	ForEachSubSequenceRecursively(FocusedSequence,
+		[this, InObject, EvaluationState, &OutSequence](UMovieSceneSequence* const InCurrentSequence)
+		{
+			const FGuid ObjectGuid = FindObjectId(*InObject, EvaluationState->FindSequenceId(InCurrentSequence));
+			if (ObjectGuid.IsValid())
+			{
+				OutSequence = InCurrentSequence;
+				return false; // Stop looping recursively
+			}
+			return true; // Continue loop recursively
+		});
+
+	return IsValid(OutSequence);
+}
+
+void FSequencer::ForEachSubSequenceRecursively(UMovieSceneSequence* const InSequence, const TFunctionRef<bool(UMovieSceneSequence* const InCurrentSequence)>& InFunction)
+{
+	if (!IsValid(InSequence) || !InFunction(InSequence))
+	{
+		return;
+	}
+
+	UMovieScene* const MovieScene = InSequence->GetMovieScene();
+	if (!IsValid(MovieScene))
+	{
+		return;
+	}
+
+	// Converting to TSet as GetAllSections() seems to return multiples of the same object
+	const TSet<UMovieSceneSection*> AllSections = TSet<UMovieSceneSection*>(MovieScene->GetAllSections());
+
+	for (UMovieSceneSection* const Section : AllSections)
+	{
+		UMovieSceneSubSection* const SubSection = Cast<UMovieSceneSubSection>(Section);
+		if (!IsValid(SubSection))
+		{
+			continue;
+		}
+
+		UMovieSceneSequence* const Sequence = SubSection->GetSequence();
+		if (!IsValid(Sequence) || !InFunction(Sequence))
+		{
+			continue;
+		}
+
+		ForEachSubSequenceRecursively(Sequence, InFunction);
+	}
+}
+
+ISequencer::FOnViewportSelectionLimitedChanged& FSequencer::OnViewportSelectionLimitedChanged()
+{
+	return OnSelectionLimitedChangedDelegate;
+}
+
+#undef LOCTEXT_NAMESPACE
diff --git a/Engine/Source/Editor/Sequencer/Private/Sequencer.h b/Engine/Source/Editor/Sequencer/Private/Sequencer.h
--- a/Engine/Source/Editor/Sequencer/Private/Sequencer.h
+++ b/Engine/Source/Editor/Sequencer/Private/Sequencer.h
@@ -1160,11 +1160,27 @@ public:
 
 	TSharedRef<SWidget> MakePlayTimeDisplay(const TSharedRef<INumericTypeInterface<double>>& InNumericTypeInterface);
 
+	/** @return True if the Sequencer is currently limiting selection to only Sequencer key-able objects */
+	bool IsViewportSelectionLimited() const override;
+
+	/** Turns on or off Sequencer selection limiting. */
+	void SetViewportSelectionLimited(const bool bInSelectionLimited) override;
+
+	/** @return True if the object specified is selectable in the viewport */
+	bool IsObjectSelectableInViewport(UObject* const InObject) override;
+
+	ISequencer::FOnViewportSelectionLimitedChanged& OnViewportSelectionLimitedChanged() override;
+
 private:
 
 	/** Update the time bases for the current movie scene */
 	void UpdateTimeBases();
 
+	/** Toggles Sequencer selection limiting. */
+	void ToggleLimitViewportSelection();
+
+	void ForEachSubSequenceRecursively(UMovieSceneSequence* const InSequence, const TFunctionRef<bool(UMovieSceneSequence* const InCurrentSequence)>& InFunction);
+
 	/** User-supplied settings object for this sequencer */
 	TObjectPtr<USequencerSettings> Settings;
 
@@ -1490,6 +1506,12 @@ private:
 	TMap<AActor*, FGuid> CachedCameraActors;
 	uint32 LastKnownStateSerial = 0;
 
-	/* Cached list of supported custom binding types*/
+	/** Cached list of supported custom binding types */
 	TArray<const TSubclassOf<UMovieSceneCustomBinding>> SupportedCustomBindingTypes;
+
+	/** True if limiting selection to only Sequencer objects */
+	static bool bSelectionLimited;
+
+	/** Delegate that is called when selection limiting has been toggled on or off */
+	FOnViewportSelectionLimitedChanged OnSelectionLimitedChangedDelegate;
 };
diff --git a/Engine/Source/Editor/Sequencer/Private/SequencerCommands.cpp b/Engine/Source/Editor/Sequencer/Private/SequencerCommands.cpp
--- a/Engine/Source/Editor/Sequencer/Private/SequencerCommands.cpp
+++ b/Engine/Source/Editor/Sequencer/Private/SequencerCommands.cpp
@@ -195,6 +195,8 @@ void FSequencerCommands::RegisterCommands()
 	UI_COMMAND(TogglePilotCamera, "Pilot Camera", "Toggle piloting the last camera or the camera cut camera.", EUserInterfaceActionType::ToggleButton, FInputChord(EModifierKey::Shift, EKeys::P));
 
 	UI_COMMAND(RefreshUI, "Refresh UI", "Forcibly refresh the UI from source data.", EUserInterfaceActionType::Button, FInputChord(EKeys::F5));
+
+	UI_COMMAND(ToggleLimitViewportSelection, "Toggle Limit Viewport Selection", "Toggles viewport selectability between Sequencer only objects and all objects.", EUserInterfaceActionType::ToggleButton, FInputChord(EModifierKey::Control, EKeys::I));
 }
 
 #undef LOCTEXT_NAMESPACE
diff --git a/Engine/Source/Editor/Sequencer/Private/SequencerEdMode.cpp b/Engine/Source/Editor/Sequencer/Private/SequencerEdMode.cpp
--- a/Engine/Source/Editor/Sequencer/Private/SequencerEdMode.cpp
+++ b/Engine/Source/Editor/Sequencer/Private/SequencerEdMode.cpp
@@ -47,6 +47,9 @@
 #include "UnrealEdMisc.h"
 #include "Editor/UnrealEdEngine.h"
 #include "TextureResource.h"
+#include "EditorModeManager.h"
+#include "Components/PrimitiveComponent.h"
+#include "SequencerSelectabilityTool.h"
 
 const FEditorModeID FSequencerEdMode::EM_SequencerMode(TEXT("EM_SequencerMode"));
 
@@ -155,10 +158,12 @@ struct FTrackTransforms
 
 FSequencerEdMode::FSequencerEdMode()
 {
-	FSequencerEdModeTool* SequencerEdModeTool = new FSequencerEdModeTool(this);
+	DefaultTool = MakeShared<FSequencerEdModeTool>(this);
+	SelectabilityTool = MakeShared<FSequencerSelectabilityTool>(FOnIsObjectSelectableInViewport::CreateRaw(this, &FSequencerEdMode::IsObjectSelectableInViewport));
 
-	Tools.Add( SequencerEdModeTool );
-	SetCurrentTool( SequencerEdModeTool );
+	Tools.Add(DefaultTool.Get());
+	Tools.Add(SelectabilityTool.Get());
+	SetCurrentTool(DefaultTool.Get());
 
 	bDrawMeshTrails = CVarDrawMeshTrails->GetBool();
 	CVarDrawMeshTrails.AsVariable()->SetOnChangedCallback(FConsoleVariableDelegate::CreateLambda([this](IConsoleVariable* Var) 
@@ -199,30 +204,23 @@ bool FSequencerEdMode::IsCompatibleWith(FEditorModeID OtherModeID) const
 
 bool FSequencerEdMode::InputKey( FEditorViewportClient* ViewportClient, FViewport* Viewport, FKey Key, EInputEvent Event )
 {
-	TSharedPtr<ISequencer> ActiveSequencer;
-
-	for (TWeakPtr<ISequencer> WeakSequencer : Sequencers)
+	if (Event != IE_Released)
 	{
-		ActiveSequencer = WeakSequencer.Pin();
-		if (ActiveSequencer.IsValid())
+		if (const TSharedPtr<ISequencer> ActiveSequencer = GetFirstActiveSequencer())
 		{
-			break;
-		}
-	}
+			FModifierKeysState KeyState = FSlateApplication::Get().GetModifierKeys();
 
-	if (ActiveSequencer.IsValid() && Event != IE_Released)
-	{
-		FModifierKeysState KeyState = FSlateApplication::Get().GetModifierKeys();
-
- 		if (ActiveSequencer->GetCommandBindings(ESequencerCommandBindings::Shared).Get()->ProcessCommandBindings(Key, KeyState, (Event == IE_Repeat) ))
-		{
-			return true;
-		}
-		if (IsPressingMoveTimeSlider(Viewport)) //this is needed to make sure we get all of the processed mouse events, for some reason the above may not return true
-		{
-			return true;
+			if (ActiveSequencer->GetCommandBindings(ESequencerCommandBindings::Shared).Get()->ProcessCommandBindings(Key, KeyState, (Event == IE_Repeat) ))
+			{
+				return true;
+			}
+			if (IsPressingMoveTimeSlider(Viewport)) //this is needed to make sure we get all of the processed mouse events, for some reason the above may not return true
+			{
+				return true;
+			}
 		}
 	}
+
 	return FEdMode::InputKey(ViewportClient, Viewport, Key, Event);
 }
 
@@ -247,14 +245,9 @@ bool FSequencerEdMode::IsPressingMoveTimeSlider(FViewport* InViewport) const
 //just get the first one.
 USequencerSettings* FSequencerEdMode:: GetSequencerSettings() const
 {
-	TSharedPtr<FSequencer> ActiveSequencer;
-	for (TWeakPtr<FSequencer> WeakSequencer : Sequencers)
+	if (const TSharedPtr<ISequencer> ActiveSequencer = GetFirstActiveSequencer())
 	{
-		ActiveSequencer = WeakSequencer.Pin();
-		if (ActiveSequencer.IsValid())
-		{
-			return ActiveSequencer->GetSequencerSettings();
-		}
+		return ActiveSequencer->GetSequencerSettings();
 	}
 	return nullptr;
 }
@@ -275,11 +268,23 @@ bool FSequencerEdMode::IsDoingDrag(FViewport* InViewport) const
 	const bool bIsAltKeyDown = InViewport->KeyState(EKeys::LeftAlt) || InViewport->KeyState(EKeys::RightAlt);
 	EAxisList::Type CurrentAxis = GetCurrentWidgetAxis();
 
-	//if shfit is down we still want to drag
+	//if shift is down we still want to drag
 
 	return LeftMouseButtonDown && (CurrentAxis == EAxisList::None) && !bIsCtrlKeyDown  && !bIsAltKeyDown && (SequencerSettings ? SequencerSettings->GetLeftMouseDragDoesMarquee() : false);
 }
 
+TSharedPtr<ISequencer> FSequencerEdMode::GetFirstActiveSequencer() const
+{
+	for (const TWeakPtr<ISequencer> SequencerWeak : Sequencers)
+	{
+		if (const TSharedPtr<ISequencer> Sequencer = SequencerWeak.Pin())
+		{
+			return Sequencer;
+		}
+	}
+	return nullptr;
+}
+
 bool FSequencerEdMode::StartTracking(FEditorViewportClient* InViewportClient, FViewport* InViewport)
 {
 	if (IsPressingMoveTimeSlider(InViewport))
@@ -370,15 +375,7 @@ bool FSequencerEdMode::InputDelta(FEditorViewportClient* InViewportClient, FView
 
 bool FSequencerEdMode::ProcessCapturedMouseMoves(FEditorViewportClient* InViewportClient, FViewport* InViewport, const TArrayView<FIntPoint>& CapturedMouseMoves)
 {
-	TSharedPtr<FSequencer> ActiveSequencer;
-	for (TWeakPtr<FSequencer> WeakSequencer : Sequencers)
-	{
-		ActiveSequencer = WeakSequencer.Pin();
-		if (ActiveSequencer.IsValid())
-		{
-			break;
-		}
-	}
+	const TSharedPtr<ISequencer> ActiveSequencer = GetFirstActiveSequencer();
 	const bool bTimeChange = IsPressingMoveTimeSlider(InViewport);
 	if (CapturedMouseMoves.Num() > 0)
 	{
@@ -455,12 +452,58 @@ bool FSequencerEdMode::ProcessCapturedMouseMoves(FEditorViewportClient* InViewpo
 	return false;
 }
 
-//Currently this is handled by the processed mouse events above, but don't fully trust ed modes this so leaving around for now
-bool FSequencerEdMode::MouseMove(FEditorViewportClient* ViewportClient, FViewport* InViewport, int32 X, int32 Y)
+bool FSequencerEdMode::HandleClick(FEditorViewportClient* InViewportClient, HHitProxy *InHitProxy, const FViewportClick &InClick)
 {
-	return false;
+	if (SelectabilityTool->IsSelectionLimited())
+	{
+		return SelectabilityTool->HandleClick(InViewportClient, InHitProxy, InClick);
+	}
+
+	return FEdMode::HandleClick(InViewportClient, InHitProxy, InClick);
+}
+
+bool FSequencerEdMode::BoxSelect(FBox& InBox, bool InSelect)
+{
+	if (SelectabilityTool->IsSelectionLimited())
+	{
+		return SelectabilityTool->BoxSelect(InBox, InSelect);
+	}
+
+	return FEdMode::BoxSelect(InBox, InSelect);
 }
 
+bool FSequencerEdMode::FrustumSelect(const FConvexVolume& InFrustum, FEditorViewportClient* InViewportClient, bool InSelect)
+{
+	if (SelectabilityTool->IsSelectionLimited())
+	{
+		return SelectabilityTool->FrustumSelect(InFrustum, InViewportClient, InSelect);
+	}
+
+	return FEdMode::FrustumSelect(InFrustum, InViewportClient, InSelect);
+}
+
+bool FSequencerEdMode::GetCursor(EMouseCursor::Type& OutCursor) const
+{
+	if (SelectabilityTool->GetCursorForHovered(OutCursor))
+	{
+		return true;
+	}
+
+	OutCursor = EMouseCursor::Default;
+
+	return true;
+}
+
+bool FSequencerEdMode::MouseMove(FEditorViewportClient* InViewportClient, FViewport* InViewport, int32 InX, int32 InY)
+{
+	if (InViewportClient && InViewportClient->Viewport && SelectabilityTool->IsSelectionLimited())
+	{
+		HHitProxy* const HitResult = InViewportClient->Viewport->GetHitProxy(InX, InY);
+		SelectabilityTool->UpdateHoverFromHitProxy(HitResult);
+	}
+
+	return FEdMode::MouseMove(InViewportClient, InViewport, InX, InY);
+}
 
 void FSequencerEdMode::Render(const FSceneView* View, FViewport* Viewport, FPrimitiveDrawInterface* PDI)
 {
@@ -488,7 +531,7 @@ void FSequencerEdMode::Render(const FSceneView* View, FViewport* Viewport, FPrim
 #endif
 }
 
-void FSequencerEdMode::DrawHUD(FEditorViewportClient* ViewportClient,FViewport* Viewport,const FSceneView* View,FCanvas* Canvas)
+void FSequencerEdMode::DrawHUD(FEditorViewportClient* ViewportClient, FViewport* Viewport, const FSceneView* View, FCanvas* Canvas)
 {
 	FEdMode::DrawHUD(ViewportClient,Viewport,View,Canvas);
 
@@ -506,6 +549,11 @@ void FSequencerEdMode::DrawHUD(FEditorViewportClient* ViewportClient,FViewport*
 		FIntRect SubtitleRegion(FMath::TruncToInt(SizeX * MinPos.X), FMath::TruncToInt(SizeY * MinPos.Y), FMath::TruncToInt(SizeX * MaxPos.X), FMath::TruncToInt(SizeY * MaxPos.Y));
 		FSubtitleManager::GetSubtitleManager()->DisplaySubtitles( Canvas, SubtitleRegion, ViewportClient->GetWorld()->GetAudioTimeSeconds() );
 	}
+
+	if (SelectabilityTool->IsSelectionLimited())
+	{
+		SelectabilityTool->DrawHUD(ViewportClient, Viewport, View, Canvas);
+	}
 }
 
 void FSequencerEdMode::AddReferencedObjects(FReferenceCollector& Collector)
@@ -966,6 +1014,28 @@ void FSequencerEdMode::DrawAudioTracks(FPrimitiveDrawInterface* PDI)
 	}
 }
 
+bool FSequencerEdMode::IsViewportSelectionLimited() const
+{
+	return SelectabilityTool.IsValid() && SelectabilityTool->IsSelectionLimited();
+}
+
+void FSequencerEdMode::EnableSelectabilityTool(const bool bInEnabled)
+{
+	SelectabilityTool->EnableLimitedSelection(bInEnabled);
+
+	SetCurrentTool(bInEnabled ? SelectabilityTool->GetID() : DefaultTool->GetID());
+}
+
+bool FSequencerEdMode::IsObjectSelectableInViewport(UObject* const InObject) const
+{
+	if (const TSharedPtr<ISequencer> Sequencer = GetFirstActiveSequencer())
+	{
+		return Sequencer->IsObjectSelectableInViewport(InObject);
+	}
+	return true;
+}
+
+
 FSequencerEdModeTool::FSequencerEdModeTool(FSequencerEdMode* InSequencerEdMode) :
 	SequencerEdMode(InSequencerEdMode)
 {
diff --git a/Engine/Source/Editor/Sequencer/Private/SequencerEdMode.h b/Engine/Source/Editor/Sequencer/Private/SequencerEdMode.h
--- a/Engine/Source/Editor/Sequencer/Private/SequencerEdMode.h
+++ b/Engine/Source/Editor/Sequencer/Private/SequencerEdMode.h
@@ -14,6 +14,7 @@ class FCanvas;
 class FEditorViewportClient;
 class FPrimitiveDrawInterface;
 class FSceneView;
+class FEditorViewportSelectability;
 class ISequencer;
 class FSequencer;
 class FViewport;
@@ -23,6 +24,7 @@ class UMovieScene3DTransformTrack;
 struct FMovieSceneEvaluationTrack;
 struct FMovieSceneInterrogationData;
 class USequencerSettings;
+class FSequencerSelectabilityTool;
 
 // This struct wraps up functionality for creating a marquee(frustum or box) selection drag tool
 // It does so based upon the type of viewport being drawn.
@@ -86,20 +88,23 @@ public:
 	virtual void Enter() override;
 	virtual void Exit() override;
 	virtual bool IsCompatibleWith(FEditorModeID OtherModeID) const override;
-	virtual bool InputKey( FEditorViewportClient* ViewportClient, FViewport* Viewport, FKey Key, EInputEvent Event ) override;
+	virtual bool InputKey(FEditorViewportClient* ViewportClient, FViewport* Viewport, FKey Key, EInputEvent Event) override;
 	virtual void Render(const FSceneView* View,FViewport* Viewport,FPrimitiveDrawInterface* PDI) override;
-	virtual void DrawHUD(FEditorViewportClient* ViewportClient,FViewport* Viewport,const FSceneView* View,FCanvas* Canvas) override;
+	virtual void DrawHUD(FEditorViewportClient* ViewportClient, FViewport* Viewport, const FSceneView* View, FCanvas* Canvas) override;
 	virtual bool UsesTransformWidget() const override { return false; }
 	virtual bool UsesTransformWidget(UE::Widget::EWidgetMode CheckMode) const override { return false; }
 	virtual void AddReferencedObjects(FReferenceCollector& Collector) override;
 
-
 	virtual bool StartTracking(FEditorViewportClient* InViewportClient, FViewport* InViewport) override;
 	virtual bool EndTracking(FEditorViewportClient* InViewportClient, FViewport* InViewport) override;
-	virtual bool MouseMove(FEditorViewportClient* ViewportClient, FViewport* Viewport, int32 x, int32 y) override;
+	virtual bool GetCursor(EMouseCursor::Type& OutCursor) const override;
+	virtual bool MouseMove(FEditorViewportClient* InViewportClient, FViewport* InViewport, int32 InX, int32 InY) override;
 	virtual bool ProcessCapturedMouseMoves(FEditorViewportClient* InViewportClient, FViewport* InViewport, const TArrayView<FIntPoint>& CapturedMouseMoves) override;
+	virtual bool HandleClick(FEditorViewportClient* InViewportClient, HHitProxy *InHitProxy, const FViewportClick &InClick) override;
+	virtual bool BoxSelect(FBox& InBox, bool InSelect) override;
+	virtual bool FrustumSelect(const FConvexVolume& InFrustum, FEditorViewportClient* InViewportClient, bool InSelect) override;
 	virtual bool InputDelta(FEditorViewportClient* InViewportClient, FViewport* InViewport, FVector& InDrag, FRotator& InRot, FVector& InScale) override;
-	virtual void Tick(FEditorViewportClient* ViewportClient,float DeltaTime);
+	virtual void Tick(FEditorViewportClient* ViewportClient,float DeltaTime) override;
 
 	bool IsPressingMoveTimeSlider(FViewport* InViewport) const;
 	bool IsDoingDrag(FViewport* InViewport) const;
@@ -117,6 +122,14 @@ public:
 	/** Clean up any mesh trails and their associated key actors */
 	void CleanUpMeshTrails();
 
+	bool IsViewportSelectionLimited() const;
+
+	void EnableSelectabilityTool(const bool bInEnabled);
+
+	bool IsObjectSelectableInViewport(UObject* const InObject) const;
+
+	TSharedPtr<ISequencer> GetFirstActiveSequencer() const;
+
 protected:
 	void DrawTracks3D(FPrimitiveDrawInterface* PDI);
 	void DrawTransformTrack(const TSharedPtr<ISequencer>& Sequencer, FPrimitiveDrawInterface* PDI, UMovieScene3DTransformTrack* TransformTrack, TArrayView<const TWeakObjectPtr<>> BoundObjects, const bool bIsSelected);
@@ -149,6 +162,9 @@ private:
 
 	/** If the pivot location needs to be updated */
 	bool bUpdatePivot = false;
+
+	TSharedPtr<class FSequencerEdModeTool> DefaultTool;
+	TSharedPtr<FSequencerSelectabilityTool> SelectabilityTool;
 };
 
 /**
diff --git a/Engine/Source/Editor/Sequencer/Private/SequencerSelectabilityTool.cpp b/Engine/Source/Editor/Sequencer/Private/SequencerSelectabilityTool.cpp
new file mode 100644
--- /dev/null
+++ b/Engine/Source/Editor/Sequencer/Private/SequencerSelectabilityTool.cpp
@@ -0,0 +1,62 @@
+﻿// Copyright Epic Games, Inc. All Rights Reserved.
+
+#include "SequencerSelectabilityTool.h"
+#include "LevelEditorViewport.h"
+#include "ScopedTransaction.h"
+#include "Settings/LevelEditorViewportSettings.h"
+#include "SequencerCommands.h"
+
+#define LOCTEXT_NAMESPACE "SequencerSelectabilityTool"
+
+FSequencerSelectabilityTool::FSequencerSelectabilityTool(const FOnIsObjectSelectableInViewport& InOnIsObjectSelectableInViewport)
+	: FModeTool()
+	, FEditorViewportSelectability(InOnIsObjectSelectableInViewport)
+{
+}
+
+bool FSequencerSelectabilityTool::BoxSelect(FBox& InBox, const bool InSelect)
+{
+	if (!bSelectionLimited
+		|| !GCurrentLevelEditingViewportClient
+		|| GCurrentLevelEditingViewportClient->IsInGameView())
+	{
+		return false;
+	}
+
+	return BoxSelectWorldActors(InBox, GCurrentLevelEditingViewportClient, InSelect);
+}
+
+bool FSequencerSelectabilityTool::FrustumSelect(const FConvexVolume& InFrustum, FEditorViewportClient* const InEditorViewportClient, const bool InSelect)
+{
+	if (!bSelectionLimited
+		|| !InEditorViewportClient
+		|| InEditorViewportClient->IsInGameView())
+	{
+		return false;
+	}
+
+	// Need to check for a zero frustum since ComponentIsTouchingSelectionFrustum will return true, selecting everything, when this is the case
+	const bool bMalformedFrustum = (InFrustum.Planes[0].IsNearlyZero() && InFrustum.Planes[2].IsNearlyZero())
+		|| (InFrustum.Planes[3].IsNearlyZero() && InFrustum.Planes[4].IsNearlyZero());
+	if (bMalformedFrustum)
+	{
+		return false;
+	}
+
+	return FrustumSelectWorldActors(InFrustum, InEditorViewportClient, InSelect);
+}
+
+void FSequencerSelectabilityTool::DrawHUD(FEditorViewportClient* const InEditorViewportClient, FViewport* const InViewport, const FSceneView* const InView, FCanvas* InCanvas)
+{
+	if (!bSelectionLimited)
+	{
+		return;
+	}
+
+	const FSequencerCommands& SequencerCommands = FSequencerCommands::Get();
+	const FText HelpText = GetLimitedSelectionText(SequencerCommands.ToggleLimitViewportSelection
+		, LOCTEXT("SequencerSelectionLimitHelpText", "Sequencer Selection Limited"));
+	DrawEnabledTextNotice(InCanvas, HelpText);
+}
+
+#undef LOCTEXT_NAMESPACE
diff --git a/Engine/Source/Editor/Sequencer/Private/SequencerSelectabilityTool.h b/Engine/Source/Editor/Sequencer/Private/SequencerSelectabilityTool.h
new file mode 100644
--- /dev/null
+++ b/Engine/Source/Editor/Sequencer/Private/SequencerSelectabilityTool.h
@@ -0,0 +1,23 @@
+﻿// Copyright Epic Games, Inc. All Rights Reserved.
+
+#pragma once
+
+#include "EditorModeTools.h"
+#include "EditorViewportSelectability.h"
+
+class FSequencerSelectabilityTool : public FModeTool, public FEditorViewportSelectability
+{
+public:
+	FSequencerSelectabilityTool() = delete;
+	FSequencerSelectabilityTool(const FOnIsObjectSelectableInViewport& InOnIsObjectSelectableInViewport);
+
+	//~ Begin FModeTool
+	virtual FString GetName() const override
+	{
+		return TEXT("Sequencer Selectability");
+	}
+	virtual bool BoxSelect(FBox& InBox, const bool InSelect) override;
+	virtual bool FrustumSelect(const FConvexVolume& InFrustum, FEditorViewportClient* const InEditorViewportClient, const bool InSelect) override;
+	virtual void DrawHUD(FEditorViewportClient* const InEditorViewportClient, FViewport* const InViewport, const FSceneView* InView, FCanvas* const InCanvas) override;
+	//~ End FModeTool
+};
diff --git a/Engine/Source/Editor/Sequencer/Public/ISequencer.h b/Engine/Source/Editor/Sequencer/Public/ISequencer.h
--- a/Engine/Source/Editor/Sequencer/Public/ISequencer.h
+++ b/Engine/Source/Editor/Sequencer/Public/ISequencer.h
@@ -236,6 +236,8 @@ public:
 
 	DECLARE_MULTICAST_DELEGATE(FOnTreeViewChanged);
 
+	DECLARE_MULTICAST_DELEGATE_OneParam(FOnViewportSelectionLimitedChanged, const bool /*bInEnabled*/);
+
 public:
 
 	/** Close the sequencer. */
@@ -874,6 +876,18 @@ public:
 	*/
 	virtual TArrayView<const TSubclassOf<UMovieSceneCustomBinding>> GetSupportedCustomBindingTypes() const { static TArray<TSubclassOf<UMovieSceneCustomBinding>> EmptyArray; return EmptyArray; }
 
+	/** @return True if the Sequencer is currently limiting viewport selection to only Sequencer objects. */
+	virtual bool IsViewportSelectionLimited() const = 0;
+
+	/** Turns on or off Sequencer selection limiting. */
+	virtual void SetViewportSelectionLimited(const bool bInSelectionLimited) = 0;
+
+	/** @return True if the specified object is selectable in the viewport and not made unselectable by the Sequencer selection limiting. */
+	virtual bool IsObjectSelectableInViewport(UObject* const InObject) = 0;
+
+	/** @return Delegate executed when Sequencer selection limiting is enabled or disabled. */
+	virtual FOnViewportSelectionLimitedChanged& OnViewportSelectionLimitedChanged() = 0;
+
 protected:
 	FOnInitializeDetailsPanel InitializeDetailsPanelEvent;
 	FOnCameraAddedToSequencer CameraAddedToSequencer;
diff --git a/Engine/Source/Editor/Sequencer/Public/SequencerCommands.h b/Engine/Source/Editor/Sequencer/Public/SequencerCommands.h
--- a/Engine/Source/Editor/Sequencer/Public/SequencerCommands.h
+++ b/Engine/Source/Editor/Sequencer/Public/SequencerCommands.h
@@ -475,6 +475,9 @@ public:
 	/** Forcibly refresh the UI */
 	TSharedPtr< FUICommandInfo > RefreshUI;
 
+	/** Toggles selectability of viewport objects to Sequencer items only */
+	TSharedPtr< FUICommandInfo > ToggleLimitViewportSelection;
+
 	/**
 	 * Initialize commands
 	 */
diff --git a/Engine/Source/Editor/UMGEditor/Private/Designer/SDesignerView.cpp b/Engine/Source/Editor/UMGEditor/Private/Designer/SDesignerView.cpp
--- a/Engine/Source/Editor/UMGEditor/Private/Designer/SDesignerView.cpp
+++ b/Engine/Source/Editor/UMGEditor/Private/Designer/SDesignerView.cpp
@@ -377,11 +377,16 @@ void SDesignerView::Construct(const FArguments& InArgs, TSharedPtr<FWidgetBluepr
 		]
 	);
 
-	auto PinnedBlueprintEditor = BlueprintEditor.Pin();
-	PinnedBlueprintEditor->OnSelectedWidgetsChanged.AddRaw(this, &SDesignerView::OnEditorSelectionChanged);
-	PinnedBlueprintEditor->OnHoveredWidgetSet.AddRaw(this, &SDesignerView::OnHoveredWidgetSet);
-	PinnedBlueprintEditor->OnHoveredWidgetCleared.AddRaw(this, &SDesignerView::OnHoveredWidgetCleared);
-	PinnedBlueprintEditor->OnWidgetPreviewUpdated.AddRaw(this, &SDesignerView::OnPreviewNeedsRecreation);
+	InBlueprintEditor->OnSelectedWidgetsChanged.AddSP(this, &SDesignerView::OnEditorSelectionChanged);
+	InBlueprintEditor->OnHoveredWidgetSet.AddSP(this, &SDesignerView::OnHoveredWidgetSet);
+	InBlueprintEditor->OnHoveredWidgetCleared.AddSP(this, &SDesignerView::OnHoveredWidgetCleared);
+	InBlueprintEditor->OnWidgetPreviewUpdated.AddSP(this, &SDesignerView::OnPreviewNeedsRecreation);
+	InBlueprintEditor->OnSelectedAnimationChanged.AddSP(this, &SDesignerView::OnSelectedAnimationChanged);
+
+	if (const TSharedPtr<ISequencer>& Sequencer = InBlueprintEditor->GetSequencer())
+	{
+		Sequencer->OnViewportSelectionLimitedChanged().AddSP(this, &SDesignerView::OnSelectionLimitedChanged);
+	}
 
 	DesignerHittestGrid = MakeShared<FHittestGrid>();
 
@@ -464,11 +469,25 @@ TSharedRef<SWidget> SDesignerView::CreateOverlayUI()
 		.VAlign(VAlign_Center)
 		.Padding(6.0f, 2.0f, 0.0f, 0.0f)
 		[
-			SNew(STextBlock)
-			.TextStyle(FAppStyle::Get(), "Graph.ZoomText")
-			.Text(this, &SDesignerView::GetZoomText)
-			.ColorAndOpacity(this, &SDesignerView::GetZoomTextColorAndOpacity)
-			.Visibility(EVisibility::SelfHitTestInvisible)
+			SNew(SVerticalBox)
+			+ SVerticalBox::Slot()
+			.AutoHeight()
+			[
+				SNew(STextBlock)
+				.TextStyle(FAppStyle::Get(), TEXT("Graph.ZoomText"))
+				.Text(this, &SDesignerView::GetZoomText)
+				.ColorAndOpacity(this, &SDesignerView::GetZoomTextColorAndOpacity)
+				.Visibility(EVisibility::SelfHitTestInvisible)
+			]
+			+ SVerticalBox::Slot()
+			.AutoHeight()
+			[
+				SNew(STextBlock)
+				.TextStyle(FAppStyle::Get(), TEXT("Graph.ZoomText"))
+				.Text(LOCTEXT("", "Sequencer Selection Limited"))
+				.ColorAndOpacity(FLinearColor(1.f, 1.f, 1.f, 0.25f))
+				.Visibility(this, &SDesignerView::GetSelectionLimitedTextVisibility)
+			]
 		]
 
 		+ SHorizontalBox::Slot()
@@ -499,6 +518,7 @@ TSharedRef<SWidget> SDesignerView::CreateOverlayUI()
 
 		+ SHorizontalBox::Slot()
 		.FillWidth(1.0f)
+		.VAlign(VAlign_Top)
 		[
 			SNew(SSpacer)
 			.Size(FVector2D(1.0f, 1.0f))
@@ -507,6 +527,7 @@ TSharedRef<SWidget> SDesignerView::CreateOverlayUI()
 		+ SHorizontalBox::Slot()
 		.Padding(0.0f, 1.0f)
 		.AutoWidth()
+		.VAlign(VAlign_Top)
 		[
 			SNew(SDesignerToolBar)
 			.CommandList(CommandList)
@@ -514,6 +535,7 @@ TSharedRef<SWidget> SDesignerView::CreateOverlayUI()
 		+ SHorizontalBox::Slot()
 		.Padding(0.0f, 1.0f)
 		.AutoWidth()
+		.VAlign(VAlign_Top)
 		[
 			SNew(SButton)
 			.ButtonStyle(&ToolBarStyle.ButtonStyle)
@@ -529,42 +551,45 @@ TSharedRef<SWidget> SDesignerView::CreateOverlayUI()
 		]
 
 		+ SHorizontalBox::Slot()
-			.AutoWidth()
+		.AutoWidth()
+		.VAlign(VAlign_Top)
+		[
+			SNew(SButton)
+			.ButtonStyle(&ToolBarStyle.ButtonStyle)
+			.ToolTipText(LOCTEXT("SwapAspectRatio_ToolTip", "Switch between Landscape and Portrait"))
+			.OnClicked(this, &SDesignerView::HandleSwapAspectRatioClicked)
+			.ContentPadding(ToolBarStyle.ButtonPadding)
+			.IsEnabled(this, &SDesignerView::GetAspectRatioSwitchEnabled)
+			.VAlign(VAlign_Center)
 			[
-				SNew(SButton)
-				.ButtonStyle(&ToolBarStyle.ButtonStyle)
-				.ToolTipText(LOCTEXT("SwapAspectRatio_ToolTip", "Switch between Landscape and Portrait"))
-				.OnClicked(this, &SDesignerView::HandleSwapAspectRatioClicked)
-				.ContentPadding(ToolBarStyle.ButtonPadding)
-				.IsEnabled(this, &SDesignerView::GetAspectRatioSwitchEnabled)
-				.VAlign(VAlign_Center)
-				[
-					SNew(SImage)
-					.Image(this, &SDesignerView::GetAspectRatioSwitchImage)
-					.ColorAndOpacity(FSlateColor::UseForeground())
-				]
+				SNew(SImage)
+				.Image(this, &SDesignerView::GetAspectRatioSwitchImage)
+				.ColorAndOpacity(FSlateColor::UseForeground())
 			]
+		]
 		+ SHorizontalBox::Slot()
-			.AutoWidth()
+		.AutoWidth()
+		.VAlign(VAlign_Top)
+		[
+			SNew(SButton)
+			.ButtonStyle(&ToolBarStyle.ButtonStyle)
+			.ToolTipText(LOCTEXT("Mirror_ToolTip", "Flip the current safe zones"))
+			.OnClicked(this, &SDesignerView::HandleFlipSafeZonesClicked)
+			.ContentPadding(ToolBarStyle.ButtonPadding)
+			.IsEnabled(this, &SDesignerView::GetFlipDeviceEnabled)
+			.VAlign(VAlign_Center)
 			[
-				SNew(SButton)
-				.ButtonStyle(&ToolBarStyle.ButtonStyle)
-				.ToolTipText(LOCTEXT("Mirror_ToolTip", "Flip the current safe zones"))
-				.OnClicked(this, &SDesignerView::HandleFlipSafeZonesClicked)
-				.ContentPadding(ToolBarStyle.ButtonPadding)
-				.IsEnabled(this, &SDesignerView::GetFlipDeviceEnabled)
-				.VAlign(VAlign_Center)
-				[
-					SNew(SImage)
-					.Image(FAppStyle::Get().GetBrush("UMGEditor.Mirror"))
-					.ColorAndOpacity(FSlateColor::UseForeground())
-				]
+				SNew(SImage)
+				.Image(FAppStyle::Get().GetBrush("UMGEditor.Mirror"))
+				.ColorAndOpacity(FSlateColor::UseForeground())
 			]
+		]
 
 		// Preview Screen Size
 		+ SHorizontalBox::Slot()
 		.Padding(2.0f,0.0f)
 		.AutoWidth()
+		.VAlign(VAlign_Top)
 		[
 			SNew(SComboButton)
 			.ButtonStyle(&ToolBarStyle.ButtonStyle)
@@ -583,6 +608,7 @@ TSharedRef<SWidget> SDesignerView::CreateOverlayUI()
 		+ SHorizontalBox::Slot()
 		.Padding(2.0f, 0.0f)
 		.AutoWidth()
+		.VAlign(VAlign_Top)
 		[
 			SNew(SComboButton)
 			.ButtonStyle(&ToolBarStyle.ButtonStyle)
@@ -770,6 +796,12 @@ SDesignerView::~SDesignerView()
 		PinnedEditor->OnHoveredWidgetSet.RemoveAll(this);
 		PinnedEditor->OnHoveredWidgetCleared.RemoveAll(this);
 		PinnedEditor->OnWidgetPreviewUpdated.RemoveAll(this);
+		PinnedEditor->OnSelectedAnimationChanged.RemoveAll(this);
+
+		if (const TSharedPtr<ISequencer> Sequencer = PinnedEditor->GetSequencer())
+		{
+			Sequencer->OnViewportSelectionLimitedChanged().RemoveAll(this);
+		}
 	}
 
 	if ( GEditor )
@@ -1683,6 +1715,12 @@ FReply SDesignerView::OnMouseButtonDown(const FGeometry& MyGeometry, const FPoin
 
 	if (bFoundWidgetUnderCursor)
 	{
+		if (!IsSelectableInSequencer(PendingSelectedWidget.GetPreview()))
+		{
+			PendingSelectedWidget = FWidgetReference();
+			return FReply::Handled();
+		}
+
 		if ( MouseEvent.GetEffectingButton() == EKeys::LeftMouseButton )
 		{
 			const TSet<FWidgetReference>& SelectedWidgets = GetSelectedWidgets();
@@ -1836,15 +1874,18 @@ FReply SDesignerView::OnMouseMove(const FGeometry& MyGeometry, const FPointerEve
 	
 
 	// Update the hovered widget under the mouse
-	auto PinnedBlueprintEditor = BlueprintEditor.Pin();
-	FWidgetHitResult HitResult;
-	if ( FindWidgetUnderCursor(MyGeometry, MouseEvent, UWidget::StaticClass(), HitResult) )
+	if (const TSharedPtr<FWidgetBlueprintEditor> PinnedBlueprintEditor = BlueprintEditor.Pin())
 	{
-		PinnedBlueprintEditor->SetHoveredWidget(HitResult.Widget);
-	}
-	else if (PinnedBlueprintEditor->GetHoveredWidget().IsValid())
-	{
-		PinnedBlueprintEditor->ClearHoveredWidget();
+		FWidgetHitResult HitResult;
+		if (FindWidgetUnderCursor(MyGeometry, MouseEvent, UWidget::StaticClass(), HitResult)
+			&& IsSelectableInSequencer(HitResult.Widget.GetPreview()))
+		{
+			PinnedBlueprintEditor->SetHoveredWidget(HitResult.Widget);
+		}
+		else if (PinnedBlueprintEditor->GetHoveredWidget().IsValid())
+		{
+			PinnedBlueprintEditor->ClearHoveredWidget();
+		}
 	}
 
 	return FReply::Unhandled();
@@ -3836,4 +3877,109 @@ EVisibility SDesignerView::GetRulerVisibility() const
 	return EVisibility::Visible;
 }
 
+bool SDesignerView::IsSelectableInSequencer(UWidget* const InWidget) const
+{
+	const TSharedPtr<FWidgetBlueprintEditor> WidgetBlueprintEditor = BlueprintEditor.Pin();
+	if (!WidgetBlueprintEditor.IsValid()
+		|| WidgetBlueprintEditor->GetCurrentAnimation() == UWidgetAnimation::GetNullAnimation())
+	{
+		return true;
+	}
+
+	const TSharedPtr<ISequencer>& Sequencer = WidgetBlueprintEditor->GetSequencer();
+	if (!Sequencer.IsValid() || !Sequencer->IsViewportSelectionLimited())
+	{
+		return true;
+	}
+
+	return Sequencer->IsObjectSelectableInViewport(InWidget);
+}
+
+void SDesignerView::OnSelectedAnimationChanged()
+{
+	const TSharedPtr<FWidgetBlueprintEditor> WidgetBlueprintEditor = BlueprintEditor.Pin();
+	if (!WidgetBlueprintEditor.IsValid())
+	{
+		return;
+	}
+
+	const TSharedPtr<ISequencer>& ActiveSequencer = WidgetBlueprintEditor->GetSequencer();
+	if (!ActiveSequencer.IsValid())
+	{
+		return;
+	}
+
+	if (ActiveSequencer->IsViewportSelectionLimited())
+	{
+		DeselectNonSequencerWidgets();
+	}
+}
+
+void SDesignerView::OnSelectionLimitedChanged(const bool bInEnabled)
+{
+	if (bInEnabled)
+	{
+		DeselectNonSequencerWidgets();
+	}
+}
+
+void SDesignerView::DeselectNonSequencerWidgets()
+{
+	const TSharedPtr<FWidgetBlueprintEditor> WidgetBlueprintEditor = BlueprintEditor.Pin();
+	if (!WidgetBlueprintEditor.IsValid())
+	{
+		return;
+	}
+
+	const TSharedPtr<ISequencer>& ActiveSequencer = WidgetBlueprintEditor->GetSequencer();
+	if (!ActiveSequencer.IsValid())
+	{
+		return;
+	}
+
+	const TSet<FWidgetReference>& SelectedWidgets = WidgetBlueprintEditor->GetSelectedWidgets();
+	if (SelectedWidgets.IsEmpty())
+	{
+		return;
+	}
+
+	TSet<FWidgetReference> NewSelection;
+
+	for (const FWidgetReference& WidgetReference : SelectedWidgets)
+	{
+		UWidget* const Widget = WidgetReference.GetPreview();
+		if (!IsValid(Widget))
+		{
+			continue;
+		}
+
+		if (!ActiveSequencer->IsViewportSelectionLimited()
+			|| WidgetBlueprintEditor->GetCurrentAnimation() == UWidgetAnimation::GetNullAnimation()
+			|| ActiveSequencer->IsObjectSelectableInViewport(Widget))
+		{
+			NewSelection.Add(WidgetReference);
+		}
+	}
+
+	WidgetBlueprintEditor->SelectWidgets(NewSelection, false);
+}
+
+EVisibility SDesignerView::GetSelectionLimitedTextVisibility() const
+{
+	const TSharedPtr<FWidgetBlueprintEditor> WidgetBlueprintEditor = BlueprintEditor.Pin();
+	if (!WidgetBlueprintEditor.IsValid())
+	{
+		return EVisibility::Collapsed;
+	}
+
+	const TSharedPtr<ISequencer>& ActiveSequencer = WidgetBlueprintEditor->GetSequencer();
+	if (!ActiveSequencer.IsValid() || !ActiveSequencer->IsViewportSelectionLimited())
+	{
+		return EVisibility::Collapsed;
+	}
+
+	return (WidgetBlueprintEditor->GetCurrentAnimation() == UWidgetAnimation::GetNullAnimation())
+		? EVisibility::Collapsed : EVisibility::SelfHitTestInvisible;
+}
+
 #undef LOCTEXT_NAMESPACE
diff --git a/Engine/Source/Editor/UMGEditor/Public/Designer/SDesignerView.h b/Engine/Source/Editor/UMGEditor/Public/Designer/SDesignerView.h
--- a/Engine/Source/Editor/UMGEditor/Public/Designer/SDesignerView.h
+++ b/Engine/Source/Editor/UMGEditor/Public/Designer/SDesignerView.h
@@ -299,6 +299,16 @@ private:
 
 	void SwapSafeZoneTypes();
 
+	bool IsSelectableInSequencer(UWidget* const InWidget) const;
+
+	void OnSelectedAnimationChanged();
+
+	void OnSelectionLimitedChanged(const bool bInEnabled);
+
+	void DeselectNonSequencerWidgets();
+
+	EVisibility GetSelectionLimitedTextVisibility() const;
+
 private:
 	/** A reference to the BP Editor that owns this designer */
 	TWeakPtr<FWidgetBlueprintEditor> BlueprintEditor;
diff --git a/Engine/Source/Editor/UnrealEd/Private/EditorViewportSelectability.cpp b/Engine/Source/Editor/UnrealEd/Private/EditorViewportSelectability.cpp
new file mode 100644
--- /dev/null
+++ b/Engine/Source/Editor/UnrealEd/Private/EditorViewportSelectability.cpp
@@ -0,0 +1,552 @@
+﻿// Copyright Epic Games, Inc. All Rights Reserved.
+
+#include "EditorViewportSelectability.h"
+#include "CanvasTypes.h"
+#include "Engine/Canvas.h"
+#include "EngineUtils.h"
+#include "LevelEditorViewport.h"
+#include "Math/Color.h"
+#include "ScopedTransaction.h"
+#include "Selection.h"
+#include "Settings/LevelEditorViewportSettings.h"
+#include "UnrealWidget.h"
+#include "UObject/WeakObjectPtr.h"
+
+#define LOCTEXT_NAMESPACE "SequencerSelectabilityTool"
+
+const FText FEditorViewportSelectability::DefaultLimitedSelectionText = LOCTEXT("DefaultSelectionLimitedHelp", "Viewport Selection Limited");
+
+FEditorViewportSelectability::FEditorViewportSelectability(const FOnIsObjectSelectableInViewport& InOnIsObjectSelectableInViewport)
+	: OnIsObjectSelectableInViewportDelegate(InOnIsObjectSelectableInViewport)
+{
+}
+
+void FEditorViewportSelectability::EnableLimitedSelection(const bool bInEnabled)
+{
+	bSelectionLimited = bInEnabled;
+
+	if (bSelectionLimited)
+	{
+		DeselectNonSelectableActors();
+	}
+
+	UpdateSelectionLimitedVisuals(!bInEnabled);
+}
+
+bool FEditorViewportSelectability::IsObjectSelectableInViewport(UObject* const InObject) const
+{
+	if (OnIsObjectSelectableInViewportDelegate.IsBound())
+	{
+		return OnIsObjectSelectableInViewportDelegate.Execute(InObject);
+	}
+	return true;
+}
+
+void FEditorViewportSelectability::UpdatePrimitiveVisuals(const bool bInSelectedLimited, UPrimitiveComponent* const InPrimitive, const TOptional<FColor>& InColor)
+{
+	if (bInSelectedLimited && InColor.IsSet())
+	{
+		// @TODO: Need to resolve rendering issue before this can be used
+		//InPrimitive->SetOverlayColor(InColor.GetValue());
+		InPrimitive->PushHoveredToProxy(true);
+	}
+	else
+	{
+		// @TODO: Need to resolve rendering issue before this can be used
+		//InPrimitive->RemoveOverlayColor();
+		InPrimitive->PushHoveredToProxy(false);
+	}
+}
+
+bool FEditorViewportSelectability::UpdateHoveredPrimitive(const bool bInSelectedLimited
+	, UPrimitiveComponent* const InPrimitiveComponent
+	, TMap<TWeakObjectPtr<UPrimitiveComponent>, TOptional<FColor>>& InOutHoveredPrimitiveComponents
+	, const TFunctionRef<bool(UObject*)>& InSelectablePredicate)
+{
+	bool bValid = IsValid(InPrimitiveComponent);
+
+	// Save the current overlay color to restore when unhovered
+	TMap<UPrimitiveComponent*, TOptional<FColor>> PrimitiveComponentsToAdd;
+
+	if (bValid && bInSelectedLimited && InSelectablePredicate(InPrimitiveComponent))
+	{
+		TOptional<FColor> UnhoveredColor;
+		if (InPrimitiveComponent->bWantsEditorEffects)
+		{
+			UnhoveredColor = InPrimitiveComponent->OverlayColor;
+		}
+		PrimitiveComponentsToAdd.Add(const_cast<UPrimitiveComponent*>(InPrimitiveComponent), UnhoveredColor);
+
+		bValid = true;
+	}
+
+	// Get the set of components to remove that aren't in the newly hovered set from the currently hovered compenents
+	TMap<UPrimitiveComponent*, TOptional<FColor>> PrimitiveComponentsToRemove;
+	for (const TPair<TWeakObjectPtr<UPrimitiveComponent>, TOptional<FColor>>& HoveredPair : InOutHoveredPrimitiveComponents)
+	{
+		UPrimitiveComponent* const PrimitiveComponent = HoveredPair.Key.Get();
+		if (IsValid(PrimitiveComponent) && !PrimitiveComponentsToAdd.Contains(PrimitiveComponent))
+		{
+			PrimitiveComponentsToRemove.Add(PrimitiveComponent, HoveredPair.Value);
+		}
+	}
+
+	// Hover new primitives, unhover old primitives
+	InOutHoveredPrimitiveComponents.Empty();
+
+	for (const TPair<UPrimitiveComponent*, TOptional<FColor>>& AddPair : PrimitiveComponentsToAdd)
+	{
+		InOutHoveredPrimitiveComponents.Add(AddPair);
+
+		// Using white becaue we've stripped out the visual element until the rendering issue can be resolved
+		UpdatePrimitiveVisuals(bInSelectedLimited, AddPair.Key, FColor::White);
+	}
+
+	for (const TPair<UPrimitiveComponent*, TOptional<FColor>>& RemovePair : PrimitiveComponentsToRemove)
+	{
+		UpdatePrimitiveVisuals(bInSelectedLimited, RemovePair.Key);
+	}
+
+	return bValid;
+}
+
+bool FEditorViewportSelectability::UpdateHoveredActorPrimitives(const bool bInSelectedLimited
+	, AActor* const InActor
+	, TMap<TWeakObjectPtr<UPrimitiveComponent>, TOptional<FColor>>& InOutHoveredPrimitiveComponents
+	, const TFunctionRef<bool(UObject*)>& InSelectablePredicate)
+{
+	bool bValid = false;
+
+	// Save the current overlay color to restore when unhovered
+	TMap<UPrimitiveComponent*, TOptional<FColor>> PrimitiveComponentsToAdd;
+
+	if (IsValid(InActor) && bInSelectedLimited)
+	{
+		if (InSelectablePredicate(InActor))
+		{
+			bValid = true;
+		}
+		InActor->ForEachComponent<UPrimitiveComponent>(/*bIncludeFromChildActors=*/true,
+			[&InSelectablePredicate, &bValid, &PrimitiveComponentsToAdd](UPrimitiveComponent* const InPrimitiveComponent)
+			{
+				if (bValid || InSelectablePredicate(InPrimitiveComponent))
+				{
+					TOptional<FColor> UnhoveredColor;
+					if (InPrimitiveComponent->bWantsEditorEffects)
+					{
+						UnhoveredColor = InPrimitiveComponent->OverlayColor;
+					}
+					PrimitiveComponentsToAdd.Add(InPrimitiveComponent, UnhoveredColor);
+
+					bValid = true;
+				}
+			});
+	}
+
+	// Get the set of components to remove that aren't in the newly hovered set from the currently hovered compenents
+	TMap<UPrimitiveComponent*, TOptional<FColor>> PrimitiveComponentsToRemove;
+	for (const TPair<TWeakObjectPtr<UPrimitiveComponent>, TOptional<FColor>>& HoveredPair : InOutHoveredPrimitiveComponents)
+	{
+		UPrimitiveComponent* const PrimitiveComponent = HoveredPair.Key.Get();
+		if (IsValid(PrimitiveComponent) && !PrimitiveComponentsToAdd.Contains(PrimitiveComponent))
+		{
+			PrimitiveComponentsToRemove.Add(PrimitiveComponent, HoveredPair.Value);
+		}
+	}
+
+	// Hover new primitives, unhover old primitives
+	InOutHoveredPrimitiveComponents.Empty();
+
+	for (const TPair<UPrimitiveComponent*, TOptional<FColor>>& AddPair : PrimitiveComponentsToAdd)
+	{
+		InOutHoveredPrimitiveComponents.Add(AddPair);
+
+		// Using white becaue we've stripped out the visual element until the rendering issue can be resolved
+		UpdatePrimitiveVisuals(bInSelectedLimited, AddPair.Key, FColor::White);
+	}
+
+	for (const TPair<UPrimitiveComponent*, TOptional<FColor>>& RemovePair : PrimitiveComponentsToRemove)
+	{
+		UpdatePrimitiveVisuals(bInSelectedLimited, RemovePair.Key);
+	}
+
+	return bValid;
+}
+
+void FEditorViewportSelectability::UpdateHoveredActorPrimitives(AActor* const InActor)
+{
+	UpdateHoveredActorPrimitives(bSelectionLimited, InActor, HoveredPrimitiveComponents,
+		[this](UObject* const InObject) -> bool
+		{
+			return IsObjectSelectableInViewport(InObject);
+		});
+}
+
+void FEditorViewportSelectability::UpdateSelectionLimitedVisuals(const bool bInClearHovered)
+{
+	if (bInClearHovered)
+	{
+		UpdateHoveredActorPrimitives(nullptr);
+	}
+
+	for (const TPair<TWeakObjectPtr<UPrimitiveComponent>, TOptional<FColor>>& HoveredPair : HoveredPrimitiveComponents)
+	{
+		UPrimitiveComponent* const PrimitiveComponent = HoveredPair.Key.Get();
+		if (IsValid(PrimitiveComponent))
+		{
+			if (bSelectionLimited
+				&& (IsObjectSelectableInViewport(PrimitiveComponent) || IsObjectSelectableInViewport(PrimitiveComponent->GetOwner())))
+			{
+				UpdatePrimitiveVisuals(bSelectionLimited, PrimitiveComponent, HoveredPair.Value);
+			}
+			else
+			{
+				UpdatePrimitiveVisuals(bSelectionLimited, PrimitiveComponent);
+			}
+		}
+	}
+}
+
+void FEditorViewportSelectability::DeselectNonSelectableActors()
+{
+	if (!bSelectionLimited)
+	{
+		return;
+	}
+
+	USelection* const ActorSelection = GEditor ? GEditor->GetSelectedActors() : nullptr;
+	if (!ActorSelection || ActorSelection->Num() == 0)
+	{
+		return;
+	}
+
+	TArray<AActor*> SelectedActors;
+	ActorSelection->GetSelectedObjects<AActor>(SelectedActors);
+
+	SelectActorsByPredicate(false, false
+		, [this](AActor* const InActor) -> bool
+		{
+			return IsObjectSelectableInViewport(InActor);
+		}
+		, SelectedActors);
+}
+
+bool FEditorViewportSelectability::SelectActorsByPredicate(const bool bInSelect
+	, const bool bInClearSelection
+	, const TFunctionRef<bool(AActor*)> InPredicate
+	, const TArray<AActor*>& InActors)
+{
+	USelection* const ActorSelection = GEditor ? GEditor->GetSelectedActors() : nullptr;
+	if (!ActorSelection || ActorSelection->Num() == 0)
+	{
+		return false;
+	}
+
+	const FText TransactionText = bInSelect ? LOCTEXT("SelectActors_Internal", "Select Actor(s)") : LOCTEXT("DeselectActors_Internal", "Deselect Actor(s)");
+	FScopedTransaction ScopedTransaction(TransactionText, !GIsTransacting);
+
+	bool bSomethingSelected = false;
+
+	ActorSelection->BeginBatchSelectOperation();
+	ActorSelection->Modify();
+
+	if (bInClearSelection)
+	{
+		ActorSelection->DeselectAll();
+	}
+
+	auto SelectIfPossible = [bInSelect, ActorSelection, &InPredicate, &bSomethingSelected](AActor* const InActor)
+	{
+		if (IsValid(InActor) && !ActorSelection->IsSelected(InActor) && InPredicate(InActor))
+		{
+			bSomethingSelected = true;
+			GEditor->SelectActor(InActor, bInSelect, true);
+		}
+	};
+
+	if (InActors.IsEmpty())
+	{
+		for (FActorIterator Iter(ActorSelection->GetWorld()); Iter; ++Iter)
+		{
+			AActor* const Actor = *Iter;
+			SelectIfPossible(Actor);
+		}
+	}
+	else
+	{
+		for (AActor* const Actor : InActors)
+		{
+			SelectIfPossible(Actor);
+		}
+	}
+
+	ActorSelection->EndBatchSelectOperation();
+	GEditor->NoteSelectionChange();
+
+	if (!bSomethingSelected)
+	{
+		ScopedTransaction.Cancel();
+	}
+
+	return bSomethingSelected;
+}
+
+bool FEditorViewportSelectability::IsActorInLevelHiddenLayer(AActor& InActor, FLevelEditorViewportClient* const InLevelEditorViewportClient)
+{
+	if (!InLevelEditorViewportClient)
+	{
+		return false;
+	}
+
+	for (const FName Layer : InActor.Layers)
+	{
+		if (InLevelEditorViewportClient->ViewHiddenLayers.Contains(Layer))
+		{
+			return true;
+		}
+	}
+
+	return false;
+}
+
+bool FEditorViewportSelectability::DoesActorIntersectBox(AActor& InActor, const FBox& InBox, FEditorViewportClient* const InEditorViewportClient, const bool bInUseStrictSelection)
+{
+	if (InActor.IsHiddenEd())
+	{
+		return false;
+	}
+
+	if (GCurrentLevelEditingViewportClient
+		&& InEditorViewportClient == GCurrentLevelEditingViewportClient
+		&& IsActorInLevelHiddenLayer(InActor, GCurrentLevelEditingViewportClient))
+	{
+		return false;
+	}
+
+	// Iterate over all actor components, selecting our primitive components
+	for (const UActorComponent* const Component : InActor.GetComponents())
+	{
+		const UPrimitiveComponent* const PrimitiveComponent = Cast<UPrimitiveComponent>(Component);
+		if (PrimitiveComponent && PrimitiveComponent->IsRegistered() && PrimitiveComponent->IsVisibleInEditor())
+		{
+			if (PrimitiveComponent->IsShown(InEditorViewportClient->EngineShowFlags)
+				&& PrimitiveComponent->ComponentIsTouchingSelectionBox(InBox, false, bInUseStrictSelection))
+			{
+				return true;
+			}
+		}
+	}
+
+	return false;
+}
+
+bool FEditorViewportSelectability::DoesActorIntersectFrustum(AActor& InActor, const FConvexVolume& InFrustum, FEditorViewportClient* const InEditorViewportClient, const bool bInUseStrictSelection)
+{
+	if (InActor.IsHiddenEd())
+	{
+		return false;
+	}
+
+	if (GCurrentLevelEditingViewportClient
+		&& InEditorViewportClient == GCurrentLevelEditingViewportClient
+		&& IsActorInLevelHiddenLayer(InActor, GCurrentLevelEditingViewportClient))
+	{
+		return false;
+	}
+
+	// Iterate over all actor components, selecting our primitive components
+	for (const UActorComponent* const Component : InActor.GetComponents())
+	{
+		const UPrimitiveComponent* const PrimitiveComponent = Cast<UPrimitiveComponent>(Component);
+		if (PrimitiveComponent && PrimitiveComponent->IsRegistered() && PrimitiveComponent->IsVisibleInEditor())
+		{
+			if (PrimitiveComponent->IsShown(InEditorViewportClient->EngineShowFlags)
+				&& PrimitiveComponent->ComponentIsTouchingSelectionFrustum(InFrustum, false, bInUseStrictSelection))
+			{
+				return true;
+			}
+		}
+	}
+
+	return false;
+}
+
+bool FEditorViewportSelectability::GetCursorForHovered(EMouseCursor::Type& OutCursor) const
+{
+	if (bSelectionLimited && MouseCursor.IsSet())
+	{
+		OutCursor = MouseCursor.GetValue();
+		return true;
+	}
+
+	return false;
+}
+
+void FEditorViewportSelectability::UpdateHoverFromHitProxy(HHitProxy* const InHitProxy)
+{
+	AActor* Actor = nullptr;
+	bool bIsGizmoHit = false;
+	bool bIsActorHit = false;
+
+	if (InHitProxy)
+	{
+		if (InHitProxy->IsA(HWidgetAxis::StaticGetType()))
+		{
+			if (bSelectionLimited)
+			{
+				bIsGizmoHit = true;
+			}
+		}
+		else if (InHitProxy->IsA(HActor::StaticGetType()))
+		{
+			const HActor* const ActorHitProxy = static_cast<HActor*>(InHitProxy);
+			if (ActorHitProxy && IsValid(ActorHitProxy->PrimComponent))
+			{
+				if (bSelectionLimited)
+				{
+					bIsActorHit = true;
+				}
+				Actor = ActorHitProxy->Actor;
+			}
+		}
+	}
+
+	UpdateHoveredActorPrimitives(Actor);
+
+	// Set mouse cursor after hovered primitive component list has been updated
+	MouseCursor.Reset();
+
+	if (bIsGizmoHit)
+	{
+		MouseCursor = EMouseCursor::CardinalCross;
+	}
+	else if (bIsActorHit)
+	{
+		MouseCursor = HoveredPrimitiveComponents.IsEmpty() ? EMouseCursor::SlashedCircle : EMouseCursor::Crosshairs;
+	}
+	else if (bSelectionLimited)
+	{
+		MouseCursor = EMouseCursor::SlashedCircle;
+	}
+}
+
+bool FEditorViewportSelectability::HandleClick(FEditorViewportClient* const InViewportClient, HHitProxy* const InHitProxy, const FViewportClick& InClick)
+{
+	// Disable actor selection when sequencer is limiting selection
+	const int32 HitX = InViewportClient->Viewport->GetMouseX();
+	const int32 HitY = InViewportClient->Viewport->GetMouseY();
+
+	HHitProxy* const HitResult = InViewportClient->Viewport->GetHitProxy(HitX, HitY);
+	if (!HitResult)
+	{
+		return false;
+	}
+
+	if (HitResult->IsA(HActor::StaticGetType()))
+	{
+		const HActor* const ActorHitProxy = static_cast<HActor*>(HitResult);
+		if (ActorHitProxy
+			&& IsValid(ActorHitProxy->Actor)
+			&& !IsObjectSelectableInViewport(ActorHitProxy->Actor))
+		{
+			return true;
+		}
+	}
+
+	return false;
+}
+
+bool FEditorViewportSelectability::BoxSelectWorldActors(FBox& InBox, FEditorViewportClient* const InEditorViewportClient, const bool bInSelect)
+{
+	if (!GEditor
+		|| !InEditorViewportClient
+		|| InEditorViewportClient->IsInGameView())
+	{
+		return false;
+	}
+
+	UWorld* const World = InEditorViewportClient->GetWorld();
+	if (!IsValid(World))
+	{
+		return false;
+	}
+
+	const ULevelEditorViewportSettings* const LevelEditorViewportSettings = GetDefault<ULevelEditorViewportSettings>();
+	const bool bUseStrictSelection = IsValid(LevelEditorViewportSettings) ? LevelEditorViewportSettings->bStrictBoxSelection : false;
+
+	const auto Predicate = [this, &InBox, InEditorViewportClient, bUseStrictSelection](AActor* const InActor) -> bool
+	{
+		return IsObjectSelectableInViewport(InActor)
+			&& DoesActorIntersectBox(*InActor, InBox, InEditorViewportClient, bUseStrictSelection);
+	};
+
+	const bool bShiftDown = InEditorViewportClient->Viewport->KeyState(EKeys::LeftShift)
+		|| InEditorViewportClient->Viewport->KeyState(EKeys::RightShift);
+
+	return SelectActorsByPredicate(bInSelect, bShiftDown, Predicate);
+}
+
+bool FEditorViewportSelectability::FrustumSelectWorldActors(const FConvexVolume& InFrustum, FEditorViewportClient* const InEditorViewportClient, const bool bInSelect)
+{
+	if (!GEditor
+		|| !InEditorViewportClient
+		|| InEditorViewportClient->IsInGameView())
+	{
+		return false;
+	}
+
+	UWorld* const World = InEditorViewportClient->GetWorld();
+	if (!IsValid(World))
+	{
+		return false;
+	}
+
+	const ULevelEditorViewportSettings* const LevelEditorViewportSettings = GetDefault<ULevelEditorViewportSettings>();
+	const bool bUseStrictSelection = IsValid(LevelEditorViewportSettings) ? LevelEditorViewportSettings->bStrictBoxSelection : false;
+
+	const auto Predicate = [this, &InFrustum, InEditorViewportClient, bUseStrictSelection](AActor* const InActor) -> bool
+	{
+		return IsObjectSelectableInViewport(InActor)
+			&& DoesActorIntersectFrustum(*InActor, InFrustum, InEditorViewportClient, bUseStrictSelection);
+	};
+
+	const bool bShiftDown = InEditorViewportClient->Viewport->KeyState(EKeys::LeftShift)
+		|| InEditorViewportClient->Viewport->KeyState(EKeys::RightShift);
+
+	return SelectActorsByPredicate(bInSelect, bShiftDown, Predicate);
+}
+
+void FEditorViewportSelectability::DrawEnabledTextNotice(FCanvas* const InCanvas, const FText& InText)
+{
+	const FString HelpString = InText.ToString();
+
+	FTextSizingParameters SizingParameters(GEngine->GetLargeFont(), 1.f, 1.f);
+	UCanvas::CanvasStringSize(SizingParameters, HelpString);
+
+	const float ViewWidth = InCanvas->GetViewRect().Width() / InCanvas->GetDPIScale();
+	const float DrawX = FMath::FloorToFloat((ViewWidth - SizingParameters.DrawXL) * 0.5f);
+	InCanvas->DrawShadowedString(DrawX, 34.f, *HelpString, GEngine->GetLargeFont(), FLinearColor::White);
+}
+
+FText FEditorViewportSelectability::GetLimitedSelectionText(const TSharedPtr<FUICommandInfo>& InToggleAction, const FText& InDefaultText)
+{
+	if (!InToggleAction.IsValid())
+	{
+		return FText();
+	}
+
+	FText HelpText = InDefaultText.IsEmpty() ? DefaultLimitedSelectionText : InDefaultText;
+
+	if (InToggleAction.IsValid())
+	{
+		const TSharedRef<const FInputChord> ActiveChord = InToggleAction->GetFirstValidChord();
+		if (ActiveChord->IsValidChord())
+		{
+			HelpText = FText::Format(LOCTEXT("LimitedSelectionActionKeyHelp", "{0}  ({1} to toggle)"), HelpText, ActiveChord->GetInputText(true));
+		}
+	}
+
+	return HelpText;
+}
+
+#undef LOCTEXT_NAMESPACE
diff --git a/Engine/Source/Editor/UnrealEd/Private/EditorViewportSelectabilityBridge.cpp b/Engine/Source/Editor/UnrealEd/Private/EditorViewportSelectabilityBridge.cpp
new file mode 100644
--- /dev/null
+++ b/Engine/Source/Editor/UnrealEd/Private/EditorViewportSelectabilityBridge.cpp
@@ -0,0 +1,50 @@
+// Copyright Epic Games, Inc. All Rights Reserved.
+
+#include "EditorViewportSelectabilityBridge.h"
+
+FEditorViewportSelectabilityBridge::FEditorViewportSelectabilityBridge(const TWeakPtr<FEditorViewportClient>& InEditorViewportClientWeak)
+	: EditorViewportClientWeak(InEditorViewportClientWeak)
+{
+}
+
+FOnIsViewportSelectionLimited& FEditorViewportSelectabilityBridge::OnIsViewportSelectionLimited()
+{
+	return IsViewportSelectionLimitedDelegate;
+}
+
+FOnIsObjectSelectableInViewport& FEditorViewportSelectabilityBridge::OnGetIsObjectSelectableInViewport()
+{
+	return GetIsObjectSelectableInViewportDelegate;
+}
+
+FOnGetViewportSelectionLimitedText& FEditorViewportSelectabilityBridge::OnGetViewportSelectionLimitedText()
+{
+	return GetViewportSelectionLimitedTextDelegate;
+}
+
+bool FEditorViewportSelectabilityBridge::IsViewportSelectionLimited() const
+{
+	if (IsViewportSelectionLimitedDelegate.IsBound())
+	{
+		return IsViewportSelectionLimitedDelegate.Execute();
+	}
+	return false;
+}
+
+bool FEditorViewportSelectabilityBridge::IsObjectSelectableInViewport(UObject* const InObject) const
+{
+	if (GetIsObjectSelectableInViewportDelegate.IsBound())
+	{
+		return GetIsObjectSelectableInViewportDelegate.Execute(InObject);
+	}
+	return true;
+}
+
+FText FEditorViewportSelectabilityBridge::GetViewportSelectionLimitedText() const
+{
+	if (GetViewportSelectionLimitedTextDelegate.IsBound())
+	{
+		return GetViewportSelectionLimitedTextDelegate.Execute();
+	}
+	return FText();
+}
diff --git a/Engine/Source/Editor/UnrealEd/Public/EditorViewportSelectability.h b/Engine/Source/Editor/UnrealEd/Public/EditorViewportSelectability.h
new file mode 100644
--- /dev/null
+++ b/Engine/Source/Editor/UnrealEd/Public/EditorViewportSelectability.h
@@ -0,0 +1,114 @@
+﻿// Copyright Epic Games, Inc. All Rights Reserved.
+
+#pragma once
+
+#include "Components/PrimitiveComponent.h"
+#include "GenericPlatform/ICursor.h"
+#include "Internationalization/Text.h"
+#include "Math/Color.h"
+#include "Misc/Optional.h"
+
+class AActor;
+class FCanvas;
+class FEditorViewportClient;
+class FLevelEditorViewportClient;
+class FUICommandInfo;
+class HHitProxy;
+struct FConvexVolume;
+struct FViewportClick;
+
+DECLARE_DELEGATE_RetVal_OneParam(bool, FOnIsObjectSelectableInViewport, UObject* /*InObject*/);
+
+/**
+ * Manages level actor viewport selectability and hovered visual states.
+ * Contains static methods to enable outside modules to implement their own management.
+ */
+class UNREALED_API FEditorViewportSelectability
+{
+public:
+	/** Default text to display in the viewport when selection is limited as a helpful reminder to the user. */
+	static const FText DefaultLimitedSelectionText;
+
+	/** Updates a single primitive component's hovered state and visuals. */
+	static void UpdatePrimitiveVisuals(const bool bInSelectedLimited, UPrimitiveComponent* const InPrimitive, const TOptional<FColor>& InColor = TOptional<FColor>());
+
+	/** Updates a list of hovered primitive component's hovered state and visuals */
+	static bool UpdateHoveredPrimitive(const bool bInSelectedLimited
+		, UPrimitiveComponent* const InPrimitiveComponent
+		, TMap<TWeakObjectPtr<UPrimitiveComponent>, TOptional<FColor>>& InOutHoveredPrimitiveComponents
+		, const TFunctionRef<bool(UObject*)>& InSelectablePredicate);
+
+	/** Updates an actors hovered state and visuals. */
+	static bool UpdateHoveredActorPrimitives(const bool bInSelectedLimited
+		, AActor* const InActor
+		, TMap<TWeakObjectPtr<UPrimitiveComponent>, TOptional<FColor>>& InOutHoveredPrimitiveComponents
+		, const TFunctionRef<bool(UObject*)>& InSelectablePredicate);
+
+	static FText GetLimitedSelectionText(const TSharedPtr<FUICommandInfo>& InToggleAction, const FText& InDefaultText = DefaultLimitedSelectionText);
+
+	static void DrawEnabledTextNotice(FCanvas* const InCanvas, const FText& InText);
+
+	FEditorViewportSelectability() = delete;
+	FEditorViewportSelectability(const FOnIsObjectSelectableInViewport& InOnIsObjectSelectableInViewport);
+
+	/** Enables or disables the selectability tool */
+	void EnableLimitedSelection(const bool bInEnabled);
+
+	bool IsSelectionLimited() const
+    {
+    	return bSelectionLimited;
+    }
+
+	/** @return True if the specified object is selectable in the viewport and not made unselectable by the Sequencer selection limiting. */
+	bool IsObjectSelectableInViewport(UObject* const InObject) const;
+
+	/** Updates hover visual states based on current selection limiting settings */
+	void UpdateSelectionLimitedVisuals(const bool bInClearHovered);
+
+	void DeselectNonSelectableActors();
+
+	bool GetCursorForHovered(EMouseCursor::Type& OutCursor) const;
+
+	void UpdateHoverFromHitProxy(HHitProxy* const InHitProxy);
+
+	bool HandleClick(FEditorViewportClient* const InViewportClient, HHitProxy* const InHitProxy, const FViewportClick& InClick);
+
+	/** Selects or deselects all actors in a level world that are inside a defined box. */
+	bool BoxSelectWorldActors(FBox& InBox, FEditorViewportClient* const InEditorViewportClient, const bool bInSelect);
+	/** Selects or deselects all actors in a level world that are inside a defined convex volume. */
+	bool FrustumSelectWorldActors(const FConvexVolume& InFrustum, FEditorViewportClient* const InEditorViewportClient, const bool bInSelect);
+
+protected:
+	static bool IsActorInLevelHiddenLayer(AActor& InActor, FLevelEditorViewportClient* const InLevelEditorViewportClient);
+
+    static bool DoesActorIntersectBox(AActor& InActor, const FBox& InBox, FEditorViewportClient* const InEditorViewportClient, const bool bInUseStrictSelection);
+    static bool DoesActorIntersectFrustum(AActor& InActor, const FConvexVolume& InFrustum, FEditorViewportClient* const InEditorViewportClient, const bool bInUseStrictSelection);
+
+	/**
+	 * Selects or deselects actors in a world. If no actors are specified, uses all the actors in the level.
+	 * 
+	 * @param InPredicate Function to use to check if the actor should be selected/deselected
+	 * @param InActors Optional list of actors to selected/deselected
+	 * @param bInSelect If true, selects the actors. If false, deselects the actors
+	 * @param bInClearSelection If true, clears the current selection before selecting the new actors
+	 * @return True if atleast one new actor was selected/deselected
+	 */
+	static bool SelectActorsByPredicate(const bool bInSelect
+		, const bool bInClearSelection
+		, const TFunctionRef<bool(AActor*)> InPredicate
+		, const TArray<AActor*>& InActors = {});
+
+	/** Updates an actors hovered state and visuals. */
+	void UpdateHoveredActorPrimitives(AActor* const InActor);
+
+	bool bSelectionLimited = false;
+
+	/** Hovered primitives and their last overlay color before we apply the hover overlay */
+	TMap<TWeakObjectPtr<UPrimitiveComponent>, TOptional<FColor>> HoveredPrimitiveComponents;
+
+	/** Mouse cursor to display for the viewport when selection is limited. */
+	TOptional<EMouseCursor::Type> MouseCursor;
+
+	/** Delegate used to check if an object is selectable in the viewport */
+	FOnIsObjectSelectableInViewport OnIsObjectSelectableInViewportDelegate;
+};
diff --git a/Engine/Source/Editor/UnrealEd/Public/EditorViewportSelectabilityBridge.h b/Engine/Source/Editor/UnrealEd/Public/EditorViewportSelectabilityBridge.h
new file mode 100644
--- /dev/null
+++ b/Engine/Source/Editor/UnrealEd/Public/EditorViewportSelectabilityBridge.h
@@ -0,0 +1,43 @@
+// Copyright Epic Games, Inc. All Rights Reserved.
+
+#pragma once
+
+#include "Internationalization/Text.h"
+#include "Templates/SharedPointer.h"
+
+class FEditorViewportClient;
+
+DECLARE_DELEGATE_RetVal(bool, FOnIsViewportSelectionLimited);
+DECLARE_DELEGATE_RetVal_OneParam(bool, FOnIsObjectSelectableInViewport, UObject* /*InObject*/);
+DECLARE_DELEGATE_RetVal(FText, FOnGetViewportSelectionLimitedText);
+
+/**
+ * Creates a link between a viewport and an outside module without requiring extra dependencies.
+ * This could be moved to Editor/UnrealEd module to allow other modules that may need this functionality to access.
+ */
+class UNREALED_API FEditorViewportSelectabilityBridge
+{
+public:
+	FEditorViewportSelectabilityBridge() = delete;
+	FEditorViewportSelectabilityBridge(const TWeakPtr<FEditorViewportClient>& InEditorViewportClientWeak);
+
+	/** @return Delegate used to check if viewport selection is limited */
+	FOnIsViewportSelectionLimited& OnIsViewportSelectionLimited();
+	bool IsViewportSelectionLimited() const;
+
+	/** @return Delegate used to check if an object is selectable in the viewport */
+	FOnIsObjectSelectableInViewport& OnGetIsObjectSelectableInViewport();
+	/** @return True if the specified object is selectable in the viewport and not made unselectable by the Sequencer selection limiting */
+	bool IsObjectSelectableInViewport(UObject* const InObject) const;
+
+	/** @return Delegate used to get the text to display in the viewport when selection is limited */
+	FOnGetViewportSelectionLimitedText& OnGetViewportSelectionLimitedText();
+	FText GetViewportSelectionLimitedText() const;
+
+private:
+	TWeakPtr<FEditorViewportClient> EditorViewportClientWeak;
+
+	FOnIsViewportSelectionLimited IsViewportSelectionLimitedDelegate;
+	FOnIsObjectSelectableInViewport GetIsObjectSelectableInViewportDelegate;
+	FOnGetViewportSelectionLimitedText GetViewportSelectionLimitedTextDelegate;
+};
