Fix SVT commandlets. Flush all rendering commands issued by UpdateAllPrimitiveSceneInfos inside the FGlobalComponentRecreateRenderStateContext.
Some rendering commands may trigger some shader compilations that needs to be issued and wait for completion before rendering the RVT.

#tests Editor + cmdlet

diff --git a/Engine/Source/Editor/UnrealEd/Private/WorldPartition/WorldPartitionRuntimeVirtualTextureBuilder.cpp b/Engine/Source/Editor/UnrealEd/Private/WorldPartition/WorldPartitionRuntimeVirtualTextureBuilder.cpp
--- a/Engine/Source/Editor/UnrealEd/Private/WorldPartition/WorldPartitionRuntimeVirtualTextureBuilder.cpp
+++ b/Engine/Source/Editor/UnrealEd/Private/WorldPartition/WorldPartitionRuntimeVirtualTextureBuilder.cpp
@@ -122,6 +122,10 @@ bool UWorldPartitionRuntimeVirtualTextureBuilder::RunInternal(UWorld* World, con
 			FGlobalComponentRecreateRenderStateContext Context;
 		}
 
+		// Flush all rendering commands issued by UpdateAllPrimitiveSceneInfos inside the FGlobalComponentRecreateRenderStateContext. 
+		// Some rendering commands may trigger some shader compilations that we need to be issued and wait for completion before rendering the RVT.
+		FlushRenderingCommands();
+
 		// FGlobalComponentRecreateRenderStateContext can create new shaderJobs, make sure to wait on them.
 		FAssetCompilingManager::Get().FinishAllCompilation();
 		FAssetCompilingManager::Get().ProcessAsyncTasks();
