Fix SVT generation.
Flush all rendering commands issued by UpdateAllPrimitiveSceneInfos inside the FGlobalComponentRecreateRenderStateContext.
Some rendering commands may trigger some shader compilations that needs to be issued and wait for completion before rendering the RVT.

#tests Editor + cmdlet

diff --git a/Engine/Source/Editor/UnrealEd/Private/EditorBuildUtils.cpp b/Engine/Source/Editor/UnrealEd/Private/EditorBuildUtils.cpp
--- a/Engine/Source/Editor/UnrealEd/Private/EditorBuildUtils.cpp
+++ b/Engine/Source/Editor/UnrealEd/Private/EditorBuildUtils.cpp
@@ -1826,6 +1826,10 @@ bool FEditorBuildUtils::EditorBuildVirtualTexture(UWorld* InWorld)
 				FGlobalComponentRecreateRenderStateContext Context;
 			}
 
+			// Flush all rendering commands issued by UpdateAllPrimitiveSceneInfos inside the FGlobalComponentRecreateRenderStateContext. 
+			// Some rendering commands may trigger some shader compilations that we need to be issued and wait for completion before rendering the RVT.
+			FlushRenderingCommands();
+
 			// FGlobalComponentRecreateRenderStateContext can create new shaderJobs, make sure to wait on them.
 			FAssetCompilingManager::Get().FinishAllCompilation();
 			FAssetCompilingManager::Get().ProcessAsyncTasks();
