FlushAsyncLoading before running a GC when gathering text from assets

This aims to avoid detaching a linker from an object that was still being loaded

#rnx

diff --git a/Engine/Source/Editor/UnrealEd/Private/Commandlets/GatherTextFromAssetsCommandlet.cpp b/Engine/Source/Editor/UnrealEd/Private/Commandlets/GatherTextFromAssetsCommandlet.cpp
--- a/Engine/Source/Editor/UnrealEd/Private/Commandlets/GatherTextFromAssetsCommandlet.cpp
+++ b/Engine/Source/Editor/UnrealEd/Private/Commandlets/GatherTextFromAssetsCommandlet.cpp
@@ -373,6 +373,8 @@ void UGatherTextFromAssetsCommandlet::PurgeGarbage(const bool bPurgeReferencedPa
 {
 	check(ObjectsToKeepAlive.Num() == 0);
 
+	FlushAsyncLoading();
+
 	TSet<FName> LoadedPackageNames;
 	TSet<FName> PackageNamesToKeepAlive;
 
