[Material Editor] Fx crash in UMaterialEditorInstanceConstant::ClearInvalidParameterOverrides() triggered when Parent is null.

diff --git a/Engine/Source/Editor/UnrealEd/Private/PreviewMaterial.cpp b/Engine/Source/Editor/UnrealEd/Private/PreviewMaterial.cpp
--- a/Engine/Source/Editor/UnrealEd/Private/PreviewMaterial.cpp
+++ b/Engine/Source/Editor/UnrealEd/Private/PreviewMaterial.cpp
@@ -793,7 +793,7 @@ void UMaterialEditorInstanceConstant::ResetOverrides(int32 Index, EMaterialParam
 
 void UMaterialEditorInstanceConstant::ClearInvalidParameterOverrides()
 {
-	const FMaterialCachedExpressionData& CachedExpressionData = Parent->GetCachedExpressionData();
+	const FMaterialCachedExpressionData* CachedExpressionData = Parent ? &Parent->GetCachedExpressionData() : nullptr;
 
 	// Look for all Atlas Scalar parameters in each parameter group, then if a parameter has
 	// an override, disable it unless the atlas texture matches that originally set in the parent Material. 
@@ -814,7 +814,7 @@ void UMaterialEditorInstanceConstant::ClearInvalidParameterOverrides()
 				// Get parent's parameter atlas texture. If identical to the one the editor parameter is using,
 				// we can keep the override on (as the selected atlas curve will still make sense).
 				FMaterialParameterMetadata Value;
-				if (CachedExpressionData.GetParameterValue(EMaterialParameterType::Scalar, ScalarParameter->ParameterInfo, Value)
+				if (CachedExpressionData && CachedExpressionData->GetParameterValue(EMaterialParameterType::Scalar, ScalarParameter->ParameterInfo, Value)
 					&& Value.ScalarAtlas == ScalarParameter->AtlasData.Atlas)
 				{
 					continue;
