[World Partition] Fix convert commandlet failing to create a new ini file when converting a level.

diff --git a/Engine/Source/Editor/UnrealEd/Private/Commandlets/WorldPartitionConvertCommandlet.cpp b/Engine/Source/Editor/UnrealEd/Private/Commandlets/WorldPartitionConvertCommandlet.cpp
--- a/Engine/Source/Editor/UnrealEd/Private/Commandlets/WorldPartitionConvertCommandlet.cpp
+++ b/Engine/Source/Editor/UnrealEd/Private/Commandlets/WorldPartitionConvertCommandlet.cpp
@@ -1614,6 +1614,8 @@ int32 UWorldPartitionConvertCommandlet::Main(const FString& Params)
 	{
 		if (bGenerateIni || !FPlatformFileManager::Get().GetPlatformFile().FileExists(*LevelConfigFilename))
 		{
+			GConfig->AddNewBranch(LevelConfigFilename);
+	
 			SaveConfig(CPF_Config, *LevelConfigFilename);
 
 			if (!bOnlyMergeSubLevels)
