Fix for FBX exporter not exporting control rig track data when there is no equivalent animation track

diff --git a/Engine/Source/Editor/UnrealEd/Private/Fbx/FbxMainExport.cpp b/Engine/Source/Editor/UnrealEd/Private/Fbx/FbxMainExport.cpp
--- a/Engine/Source/Editor/UnrealEd/Private/Fbx/FbxMainExport.cpp
+++ b/Engine/Source/Editor/UnrealEd/Private/Fbx/FbxMainExport.cpp
@@ -2083,6 +2083,7 @@ bool FFbxExporter::ExportLevelSequenceTracks(UMovieSceneSequence* MovieSceneSequ
 
 	// Look for the tracks that we currently support
 	UMovieSceneSkeletalAnimationTrack* SkeletalAnimationTrack = nullptr;
+	bool IsControlRigTrack = false;
 	for (UMovieSceneTrack* Track : Tracks)
 	{
 		if (Track->IsA(UMovieScene3DTransformTrack::StaticClass()))
@@ -2105,6 +2106,10 @@ bool FFbxExporter::ExportLevelSequenceTracks(UMovieSceneSequence* MovieSceneSequ
 		{
 			ExportLevelSequenceVectorTrack(FbxActor, *Cast<UMovieSceneDoubleVectorTrack>(Track), BoundObject, MovieScene->GetPlaybackRange(), RootToLocalTransform);
 		}
+		else if(Cast<INodeAndChannelMappings>(Track))
+		{
+			IsControlRigTrack = true;
+		}
 		else
 		{
 			bool bBakeChannels = false;
@@ -2121,7 +2126,7 @@ bool FFbxExporter::ExportLevelSequenceTracks(UMovieSceneSequence* MovieSceneSequ
 	}
 
 	// Export all of the skeletal animation components for this actor
-	if (SkeletalMeshComp && SkeletalAnimationTrack)
+	if (SkeletalMeshComp && (SkeletalAnimationTrack || IsControlRigTrack))
 	{
 		TArray<USkeletalMeshComponent*> SkeletalMeshComponents;
 		SkeletalMeshComp->GetOwner()->GetComponents(SkeletalMeshComponents);
