IncrementalCook: IterativeValidate: Add a PackageIgnoreList so we can collect diagnostics on packages that are known to have falsepositives, but still mark the job as succeeded. Add a cookforcediff parameter for testing.
#rnx

diff --git a/Engine/Config/BaseEditor.ini b/Engine/Config/BaseEditor.ini
--- a/Engine/Config/BaseEditor.ini
+++ b/Engine/Config/BaseEditor.ini
@@ -484,6 +484,7 @@ SaveUnversioned=true
 ;    Allow a class and all its superclasses.
 ; +IterativeClassDenyList=/Long/PackageName.ClassName
 ;    Block a class (and all its subclasses) that is in the allow list.
+
 +IterativeClassScriptPackageAllowList=Allow,<EngineRoot>
 
 +IterativeClassDenyList=/Script/ComputeFramework.ComputeGraph				; UE-203846
@@ -494,6 +495,9 @@ SaveUnversioned=true
 +IterativeClassDenyList=/Script/PoseSearch.PoseSearchDatabase				; UE-209084
 +IterativeClassDenyList=/Script/RigVM.RigVM                                 ; UE-209204
 
+; IterativeValidatePackageIgnoreList - package names that will be displayed but not treated as an error by IterativeValidate
+; +IterativeValidatePackageIgnoreList=/Game/PackageName
+
 [CookPlatformDataCacheSettings]
 ; Add <ShortClassName>=<Integer> here to limit the number of instances that can have BeginCacheForCookedPlatformData called at once
 ; This was originally created for textures and materials, but we now have thread pools that handle them in a system-specific way
diff --git a/Engine/Source/Editor/UnrealEd/Private/Cooker/DiffWriterArchive.cpp b/Engine/Source/Editor/UnrealEd/Private/Cooker/DiffWriterArchive.cpp
--- a/Engine/Source/Editor/UnrealEd/Private/Cooker/DiffWriterArchive.cpp
+++ b/Engine/Source/Editor/UnrealEd/Private/Cooker/DiffWriterArchive.cpp
@@ -1356,6 +1356,14 @@ void FAccumulator::GenerateDiffMap()
 
 	bIdentical = bHeaderIdentical && bExportsIdentical;
 	bHasDifferences = !bIdentical;
+	static bool bForceDiff = []()
+		{
+			return FParse::Param(FCommandLine::Get(), TEXT("cookforcediff"));
+		}();
+	if (bForceDiff)
+	{
+		bHasDifferences = true;
+	}
 }
 
 FLinkerLoad* CreateLinkerForPackage(FUObjectSerializeContext* LoadContext, const FString& InPackageName, const FString& InFilename, const FPackageData& PackageData)
diff --git a/Engine/Source/Editor/UnrealEd/Private/Cooker/IterativeValidatePackageWriter.cpp b/Engine/Source/Editor/UnrealEd/Private/Cooker/IterativeValidatePackageWriter.cpp
--- a/Engine/Source/Editor/UnrealEd/Private/Cooker/IterativeValidatePackageWriter.cpp
+++ b/Engine/Source/Editor/UnrealEd/Private/Cooker/IterativeValidatePackageWriter.cpp
@@ -10,6 +10,7 @@
 #include "HAL/FileManager.h"
 #include "Logging/LogMacros.h"
 #include "Misc/CommandLine.h"
+#include "Misc/ConfigCacheIni.h"
 #include "Misc/OutputDeviceHelper.h"
 #include "Misc/Parse.h"
 #include "Misc/Paths.h"
@@ -287,6 +288,15 @@ FIterativeValidatePackageWriter::FIterativeValidatePackageWriter(UCookOnTheFlySe
 
 	Indent = FCString::Spc(FOutputDeviceHelper::FormatLogLine(ELogVerbosity::Warning,
 		LogIterativeValidate.GetCategoryName(), TEXT(""), GPrintLogTimes).Len());
+
+	TArray<FString> IterativeValidatePackageIgnoreList;
+	GConfig->GetArray(TEXT("TargetDomain"), TEXT("IterativeValidatePackageIgnoreList"),
+		IterativeValidatePackageIgnoreList, GEditorIni);
+	PackageIgnoreList.Reserve(IterativeValidatePackageIgnoreList.Num());
+	for (const FString& PackageName : IterativeValidatePackageIgnoreList)
+	{
+		PackageIgnoreList.Add(FName(FStringView(PackageName)));
+	}
 }
 
 void FIterativeValidatePackageWriter::BeginPackage(const FBeginPackageInfo& Info)
@@ -322,7 +332,8 @@ void FIterativeValidatePackageWriter::BeginPackage(const FBeginPackageInfo& Info
 	case EPhase::Phase2:
 		{
 			EPackageStatus PackageStatus = GetPackageStatus(Info.PackageName);
-			if (PackageStatus == EPackageStatus::DeclaredUnmodified_ConfirmedUnmodified)
+			if (PackageStatus == EPackageStatus::DeclaredUnmodified_ConfirmedUnmodified ||
+				PackageStatus == EPackageStatus::DeclaredUnmodified_FoundModified_OnIgnoreList)
 			{
 				// Already saved in Phase 1; no need to diff it or save it now
 				SaveAction = ESaveAction::IgnoreResults;
@@ -613,14 +624,19 @@ void FIterativeValidatePackageWriter::UpdatePackageModificationStatus(FName Pack
 		}
 		break;
 	case EPhase::Phase2:
+	{
 		// Ignore the Unmodified flag from this cook phase. Skip the packages that were found to 
 		// be IterativeValidated from Phase1. Save the packages that phase1 found modified; this phase
 		// is responsible for getting those resaved. Reexecute save for the packages that were IterativeFailed
 		// from phase1, so we can test whether they are indeterministic. Always save generators so we can
 		// test their generated packages.
-		bInOutShouldIterativelySkip = (GetPackageStatus(PackageName) == EPackageStatus::DeclaredUnmodified_ConfirmedUnmodified) 
-										&& !bKnownGenerator;
+		EPackageStatus PackageStatus = GetPackageStatus(PackageName);
+		bInOutShouldIterativelySkip =
+			(PackageStatus == EPackageStatus::DeclaredUnmodified_ConfirmedUnmodified
+				|| PackageStatus == EPackageStatus::DeclaredUnmodified_FoundModified_OnIgnoreList)
+			&& !bKnownGenerator;
 		break;
+	}
 	default:
 		checkNoEntry();
 		break;
@@ -700,13 +716,15 @@ void FIterativeValidatePackageWriter::EndCook(const FCookInfo& Info)
 	case EPhase::AllInOnePhase:
 	{
 		int32 DetectedUnmodified = StatusCounts[EPackageStatus::DeclaredUnmodified_ConfirmedUnmodified]
-			+ StatusCounts[EPackageStatus::DeclaredUnmodified_FoundModified_FalsePositive];
+			+ StatusCounts[EPackageStatus::DeclaredUnmodified_FoundModified_FalsePositive]
+			+ StatusCounts[EPackageStatus::DeclaredUnmodified_FoundModified_OnIgnoreList];
 		UE_CLOG(COTFS.GetCookMode() != ECookMode::CookWorker,
 			LogIterativeValidate, Display,
 			TEXT("Modified: %d. DetectedUnmodified: %d. ValidatedUnmodified: %d. IterativeSkipFalsePositive: %d."),
 			StatusCounts[EPackageStatus::DeclaredModified_WillNotVerify], 
 			DetectedUnmodified,
-			StatusCounts[EPackageStatus::DeclaredUnmodified_ConfirmedUnmodified],
+			StatusCounts[EPackageStatus::DeclaredUnmodified_ConfirmedUnmodified]
+			+ StatusCounts[EPackageStatus::DeclaredUnmodified_FoundModified_OnIgnoreList],
 			StatusCounts[EPackageStatus::DeclaredUnmodified_FoundModified_FalsePositive]);
 		FString Message = FString::Printf(TEXT("Packages Iteratively Skipped: %d: IterativeSkipFalsePositive: %d."),
 			DetectedUnmodified,
@@ -716,6 +734,14 @@ void FIterativeValidatePackageWriter::EndCook(const FCookInfo& Info)
 			UE_CLOG(COTFS.GetCookMode() != ECookMode::CookWorker,
 				LogIterativeValidate, Error, TEXT("%s"), *Message);
 		}
+		else if (StatusCounts[EPackageStatus::DeclaredUnmodified_FoundModified_OnIgnoreList] > 0)
+		{
+			Message = FString::Printf(TEXT("Packages Iteratively Skipped: %d: IterativeSkipFalsePositive (Ignored): %d."),
+				DetectedUnmodified,
+				StatusCounts[EPackageStatus::DeclaredUnmodified_FoundModified_OnIgnoreList]);
+			UE_CLOG(COTFS.GetCookMode() != ECookMode::CookWorker,
+				LogIterativeValidate, Warning, TEXT("%s"), *Message);
+		}
 		else
 		{
 			UE_CLOG(COTFS.GetCookMode() != ECookMode::CookWorker,
@@ -726,13 +752,15 @@ void FIterativeValidatePackageWriter::EndCook(const FCookInfo& Info)
 	case EPhase::Phase1:
 	{
 		int32 DetectedUnmodified = StatusCounts[EPackageStatus::DeclaredUnmodified_ConfirmedUnmodified]
-			+ StatusCounts[EPackageStatus::DeclaredUnmodified_FoundModified_IndeterminismOrFalsePositive];
+			+ StatusCounts[EPackageStatus::DeclaredUnmodified_FoundModified_IndeterminismOrFalsePositive]
+			+ StatusCounts[EPackageStatus::DeclaredUnmodified_FoundModified_OnIgnoreList];
 		UE_CLOG(COTFS.GetCookMode() != ECookMode::CookWorker,
 			LogIterativeValidate, Display,
 			TEXT("Modified: %d. DetectedUnmodified: %d. ValidatedUnmodified: %d. IterativeSkipFalsePositiveOrIndeterminism: %d."),
 			StatusCounts[EPackageStatus::DeclaredModified_WillNotVerify],
 			DetectedUnmodified,
-			StatusCounts[EPackageStatus::DeclaredUnmodified_ConfirmedUnmodified],
+			StatusCounts[EPackageStatus::DeclaredUnmodified_ConfirmedUnmodified]
+			+ StatusCounts[EPackageStatus::DeclaredUnmodified_FoundModified_OnIgnoreList],
 			StatusCounts[EPackageStatus::DeclaredUnmodified_FoundModified_IndeterminismOrFalsePositive]);
 		Save();
 		break;
@@ -742,13 +770,15 @@ void FIterativeValidatePackageWriter::EndCook(const FCookInfo& Info)
 		int32 DetectedUnmodified = StatusCounts[EPackageStatus::DeclaredUnmodified_ConfirmedUnmodified]
 			+ StatusCounts[EPackageStatus::DeclaredUnmodified_FoundModified_IndeterminismOrFalsePositive]
 			+ StatusCounts[EPackageStatus::DeclaredUnmodified_FoundModified_Indeterminism]
-			+ StatusCounts[EPackageStatus::DeclaredUnmodified_FoundModified_FalsePositive];
+			+ StatusCounts[EPackageStatus::DeclaredUnmodified_FoundModified_FalsePositive]
+			+ StatusCounts[EPackageStatus::DeclaredUnmodified_FoundModified_OnIgnoreList];
 		UE_CLOG(COTFS.GetCookMode() != ECookMode::CookWorker,
 			LogIterativeValidate, Display,
 			TEXT("Modified: %d. DetectedUnmodified: %d. ValidatedUnmodified: %d. Indeterminism: %d. IterativeSkipFalsePositive: %d."),
 			StatusCounts[EPackageStatus::DeclaredModified_WillNotVerify],
 			DetectedUnmodified,
-			StatusCounts[EPackageStatus::DeclaredUnmodified_ConfirmedUnmodified],
+			StatusCounts[EPackageStatus::DeclaredUnmodified_ConfirmedUnmodified]
+			+ StatusCounts[EPackageStatus::DeclaredUnmodified_FoundModified_OnIgnoreList],
 			StatusCounts[EPackageStatus::DeclaredUnmodified_FoundModified_Indeterminism],
 			StatusCounts[EPackageStatus::DeclaredUnmodified_FoundModified_FalsePositive]
 			);
@@ -760,6 +790,14 @@ void FIterativeValidatePackageWriter::EndCook(const FCookInfo& Info)
 			UE_CLOG(COTFS.GetCookMode() != ECookMode::CookWorker,
 				LogIterativeValidate, Error, TEXT("%s"), *Message);
 		}
+		else if (StatusCounts[EPackageStatus::DeclaredUnmodified_FoundModified_OnIgnoreList] > 0)
+		{
+			Message = FString::Printf(TEXT("Packages Iteratively Skipped: %d: IterativeSkipFalsePositive (Ignored): %d."),
+				DetectedUnmodified,
+				StatusCounts[EPackageStatus::DeclaredUnmodified_FoundModified_FalsePositive]);
+			UE_CLOG(COTFS.GetCookMode() != ECookMode::CookWorker,
+				LogIterativeValidate, Warning, TEXT("%s"), *Message);
+		}
 		else
 		{
 			UE_CLOG(COTFS.GetCookMode() != ECookMode::CookWorker,
@@ -834,7 +872,14 @@ bool FIterativeValidatePackageWriter::IsAnotherSaveNeeded(FSavePackageResultStru
 
 			if (bIsDifferent && !bNewPackage)
 			{
-				SetPackageStatus(BeginInfo.PackageName, EPackageStatus::DeclaredUnmodified_FoundModified_FalsePositive);
+				if (!PackageIgnoreList.Contains(BeginInfo.PackageName))
+				{
+					SetPackageStatus(BeginInfo.PackageName, EPackageStatus::DeclaredUnmodified_FoundModified_FalsePositive);
+				}
+				else
+				{
+					SetPackageStatus(BeginInfo.PackageName, EPackageStatus::DeclaredUnmodified_FoundModified_OnIgnoreList);
+				}
 			}
 			else if (!bNewPackage)
 			{
@@ -875,9 +920,16 @@ bool FIterativeValidatePackageWriter::IsAnotherSaveNeeded(FSavePackageResultStru
 			Inner->BeginPackage(BeginInfo);
 			SaveAction = ESaveAction::SaveToInner;
 
-			// Mark that the iterative validation failed if it was not already marked by log or warning messages.
-			// We need to record it for an indeterminism test
-			SetPackageStatus(BeginInfo.PackageName, EPackageStatus::DeclaredUnmodified_FoundModified_IndeterminismOrFalsePositive);
+			if (!PackageIgnoreList.Contains(BeginInfo.PackageName))
+			{
+				// Mark that the iterative validation failed if it was not already marked by log or warning messages.
+				// We need to record it for an indeterminism test
+				SetPackageStatus(BeginInfo.PackageName, EPackageStatus::DeclaredUnmodified_FoundModified_IndeterminismOrFalsePositive);
+			}
+			else
+			{
+				SetPackageStatus(BeginInfo.PackageName, EPackageStatus::DeclaredUnmodified_FoundModified_OnIgnoreList);
+			}
 
 			return true;
 		}
diff --git a/Engine/Source/Editor/UnrealEd/Private/Cooker/IterativeValidatePackageWriter.h b/Engine/Source/Editor/UnrealEd/Private/Cooker/IterativeValidatePackageWriter.h
--- a/Engine/Source/Editor/UnrealEd/Private/Cooker/IterativeValidatePackageWriter.h
+++ b/Engine/Source/Editor/UnrealEd/Private/Cooker/IterativeValidatePackageWriter.h
@@ -54,15 +54,6 @@ public:
 protected:
 	enum class EPackageStatus : uint8;
 
-	virtual void OnDiffWriterMessage(ELogVerbosity::Type Verbosity, FStringView Message) override;
-	void LogIterativeDifferences();
-	void Save();
-	void Load();
-	void Serialize(FArchive& Ar);
-	FString GetIterativeValidatePath() const;
-	EPackageStatus GetPackageStatus(FName PackageName) const;
-	void SetPackageStatus(FName PackageName, EPackageStatus NewStatus);
-
 	enum class ESaveAction : uint8
 	{
 		CheckForDiffs,
@@ -77,6 +68,7 @@ protected:
 		DeclaredUnmodified_FoundModified_IndeterminismOrFalsePositive,
 		DeclaredUnmodified_FoundModified_Indeterminism,
 		DeclaredUnmodified_FoundModified_FalsePositive,
+		DeclaredUnmodified_FoundModified_OnIgnoreList,
 		DeclaredUnmodified_NotYetProcessed,
 		DeclaredModified_WillNotVerify,
 		Count
@@ -89,9 +81,6 @@ protected:
 	};
 	friend FArchive& operator<<(FArchive& Ar, FMessage& Message);
 
-	TMap<FName, EPackageStatus> PackageStatusMap;
-	TMap<FName, TArray<FMessage>> PackageMessageMap;
-
 	struct FStatusCounts
 	{
 		FStatusCounts()
@@ -108,8 +97,21 @@ protected:
 		TStaticArray<uint32, (uint32)EPackageStatus::Count> Data;
 	};
 
+protected:
+	virtual void OnDiffWriterMessage(ELogVerbosity::Type Verbosity, FStringView Message) override;
+	void LogIterativeDifferences();
+	void Save();
+	void Load();
+	void Serialize(FArchive& Ar);
+	FString GetIterativeValidatePath() const;
+	EPackageStatus GetPackageStatus(FName PackageName) const;
+	void SetPackageStatus(FName PackageName, EPackageStatus NewStatus);
 	FStatusCounts CountPackagesByStatus();
 
+protected:
+	TMap<FName, EPackageStatus> PackageStatusMap;
+	TMap<FName, TArray<FMessage>> PackageMessageMap;
+	TSet<FName> PackageIgnoreList;
 	FString MetadataPath;
 	UCookOnTheFlyServer& COTFS;
 	EPhase Phase = EPhase::AllInOnePhase;
