[Selection] Fix an issue from the last submit that removed an assumption that if there is a BP actor, we only grab the top level actor in sync to CB code

#rnx

diff --git a/Engine/Source/Editor/UnrealEd/Private/EditorEngine.cpp b/Engine/Source/Editor/UnrealEd/Private/EditorEngine.cpp
--- a/Engine/Source/Editor/UnrealEd/Private/EditorEngine.cpp
+++ b/Engine/Source/Editor/UnrealEd/Private/EditorEngine.cpp
@@ -3305,7 +3305,7 @@ void UEditorEngine::GetAssetsToSyncToContentBrowser(TArray<FAssetData>& Assets,
 
 				if (TTypedElement<ITypedElementAssetDataInterface> AssetDataInterface = UTypedElementRegistry::GetInstance()->GetElement<ITypedElementAssetDataInterface>(SelectedHandle))
 				{
-					Assets.Append(AssetDataInterface.GetAllReferencedAssetDatas());
+					Assets.Append(AssetDataInterface.GetAllReferencedAssetDatas(FTypedElementAssetDataReferencedOptions().SetOnlyTopLevelAsset(true)));
 				}
 
 				return true;
diff --git a/Engine/Source/Editor/UnrealEd/Private/Elements/Actor/ActorElementEditorAssetDataInterface.cpp b/Engine/Source/Editor/UnrealEd/Private/Elements/Actor/ActorElementEditorAssetDataInterface.cpp
--- a/Engine/Source/Editor/UnrealEd/Private/Elements/Actor/ActorElementEditorAssetDataInterface.cpp
+++ b/Engine/Source/Editor/UnrealEd/Private/Elements/Actor/ActorElementEditorAssetDataInterface.cpp
@@ -7,7 +7,7 @@
 #include "Elements/Actor/ActorElementData.h"
 #include "GameFramework/Actor.h"
 
-TArray<FAssetData> UActorElementEditorAssetDataInterface::GetAllReferencedAssetDatas(const FTypedElementHandle& InElementHandle)
+TArray<FAssetData> UActorElementEditorAssetDataInterface::GetAllReferencedAssetDatas(const FTypedElementHandle& InElementHandle, const FTypedElementAssetDataReferencedOptions& InOptions)
 {
 	TArray<FAssetData> AssetDatas;
 
@@ -15,28 +15,49 @@ TArray<FAssetData> UActorElementEditorAssetDataInterface::GetAllReferencedAssetD
 	{
 		TArray<UObject*> ReferencedContentObjects;
 		RawActorPtr->GetReferencedContentObjects(ReferencedContentObjects);
-		for (const UObject* ContentObject : ReferencedContentObjects)
+
+		if (InOptions.OnlyTopLevelAsset())
 		{
-			FAssetData ObjectAssetData = FAssetData(ContentObject);
-			if (ObjectAssetData.IsValid())
+			UObject** BPObject = ReferencedContentObjects.FindByPredicate([](const UObject* InObject)
 			{
-				AssetDatas.Emplace(ObjectAssetData);
+				if (InObject && (InObject->IsA<UBlueprint>() || InObject->IsA<UBlueprintGeneratedClass>()))
+				{
+					return true;
+				}
+				return false;
+			});
+			
+			if (BPObject)
+			{
+				AssetDatas.Add(FAssetData(*BPObject, FAssetData::ECreationFlags::SkipAssetRegistryTagsGathering));
 			}
 		}
 
-		TArray<FSoftObjectPath> SoftObjects;
-		RawActorPtr->GetSoftReferencedContentObjects(SoftObjects);
-		if (SoftObjects.Num())
+		if (AssetDatas.IsEmpty())
 		{
-			IAssetRegistry& AssetRegistry = IAssetRegistry::GetChecked();
+			for (const UObject* ContentObject : ReferencedContentObjects)
+			{
+				FAssetData ObjectAssetData = FAssetData(ContentObject, FAssetData::ECreationFlags::SkipAssetRegistryTagsGathering);
+				if (ObjectAssetData.IsValid())
+				{
+					AssetDatas.Emplace(ObjectAssetData);
+				}
+			}
 
-			for (const FSoftObjectPath& SoftObject : SoftObjects)
+			TArray<FSoftObjectPath> SoftObjects;
+			RawActorPtr->GetSoftReferencedContentObjects(SoftObjects);
+			if (SoftObjects.Num())
 			{
-				FAssetData AssetData = AssetRegistry.GetAssetByObjectPath(SoftObject);
+				IAssetRegistry& AssetRegistry = IAssetRegistry::GetChecked();
 
-				if (AssetData.IsValid())
+				for (const FSoftObjectPath& SoftObject : SoftObjects)
 				{
-					AssetDatas.Add(AssetData);
+					FAssetData AssetData = AssetRegistry.GetAssetByObjectPath(SoftObject);
+
+					if (AssetData.IsValid())
+					{
+						AssetDatas.Add(AssetData);
+					}
 				}
 			}
 		}
diff --git a/Engine/Source/Editor/UnrealEd/Public/Elements/Actor/ActorElementEditorAssetDataInterface.h b/Engine/Source/Editor/UnrealEd/Public/Elements/Actor/ActorElementEditorAssetDataInterface.h
--- a/Engine/Source/Editor/UnrealEd/Public/Elements/Actor/ActorElementEditorAssetDataInterface.h
+++ b/Engine/Source/Editor/UnrealEd/Public/Elements/Actor/ActorElementEditorAssetDataInterface.h
@@ -19,5 +19,5 @@ class UActorElementEditorAssetDataInterface : public UActorElementAssetDataInter
 	GENERATED_BODY()
 
 public:
-	UNREALED_API virtual TArray<FAssetData> GetAllReferencedAssetDatas(const FTypedElementHandle& InElementHandle) override;
+	UNREALED_API virtual TArray<FAssetData> GetAllReferencedAssetDatas(const FTypedElementHandle& InElementHandle, const FTypedElementAssetDataReferencedOptions& InOptions) override;
 };
diff --git a/Engine/Source/Runtime/TypedElementRuntime/Private/Elements/Interfaces/TypedElementAssetDataInterface.cpp b/Engine/Source/Runtime/TypedElementRuntime/Private/Elements/Interfaces/TypedElementAssetDataInterface.cpp
--- a/Engine/Source/Runtime/TypedElementRuntime/Private/Elements/Interfaces/TypedElementAssetDataInterface.cpp
+++ b/Engine/Source/Runtime/TypedElementRuntime/Private/Elements/Interfaces/TypedElementAssetDataInterface.cpp
@@ -6,7 +6,7 @@
 
 #include UE_INLINE_GENERATED_CPP_BY_NAME(TypedElementAssetDataInterface)
 
-TArray<FAssetData> ITypedElementAssetDataInterface::GetAllReferencedAssetDatas(const FTypedElementHandle& InElementHandle)
+TArray<FAssetData> ITypedElementAssetDataInterface::GetAllReferencedAssetDatas(const FTypedElementHandle& InElementHandle, const FTypedElementAssetDataReferencedOptions& InOptions)
 {
 	TArray<FAssetData> AssetDatas;
 
diff --git a/Engine/Source/Runtime/TypedElementRuntime/Public/Elements/Interfaces/TypedElementAssetDataInterface.h b/Engine/Source/Runtime/TypedElementRuntime/Public/Elements/Interfaces/TypedElementAssetDataInterface.h
--- a/Engine/Source/Runtime/TypedElementRuntime/Public/Elements/Interfaces/TypedElementAssetDataInterface.h
+++ b/Engine/Source/Runtime/TypedElementRuntime/Public/Elements/Interfaces/TypedElementAssetDataInterface.h
@@ -14,6 +14,20 @@
 class UObject;
 struct FFrame;
 
+USTRUCT(BlueprintType)
+struct FTypedElementAssetDataReferencedOptions
+{
+	GENERATED_BODY()
+
+public:
+	FTypedElementAssetDataReferencedOptions& SetOnlyTopLevelAsset(const bool InTopLevelAsset) { bOnlyTopLevelAsset = InTopLevelAsset; return *this; }
+	bool OnlyTopLevelAsset() const { return bOnlyTopLevelAsset; }
+
+private:
+	UPROPERTY(BlueprintReadWrite, Category = "TypedElementInterfaces|AssetData|GetReferencedOptions", meta = (AllowPrivateAccess = true))
+	bool bOnlyTopLevelAsset = false;
+};
+
 UINTERFACE(MinimalAPI, BlueprintType, meta = (CannotImplementInterfaceInBlueprint))
 class UTypedElementAssetDataInterface : public UInterface
 {
@@ -31,7 +45,7 @@ public:
 	 *
 	 * @returns An array of valid asset datas.
 	 */
-	TYPEDELEMENTRUNTIME_API virtual TArray<FAssetData> GetAllReferencedAssetDatas(const FTypedElementHandle& InElementHandle);
+	TYPEDELEMENTRUNTIME_API virtual TArray<FAssetData> GetAllReferencedAssetDatas(const FTypedElementHandle& InElementHandle, const FTypedElementAssetDataReferencedOptions& InOptions = FTypedElementAssetDataReferencedOptions());
 
 	/**
 	 * Returns the asset data for the given handle, if it exists.
@@ -61,6 +75,6 @@ public:
 template <>
 struct TTypedElement<ITypedElementAssetDataInterface> : public TTypedElementBase<ITypedElementAssetDataInterface>
 {
-	TArray<FAssetData> GetAllReferencedAssetDatas() const { return InterfacePtr->GetAllReferencedAssetDatas(*this); }
+	TArray<FAssetData> GetAllReferencedAssetDatas(const FTypedElementAssetDataReferencedOptions& InOptions = FTypedElementAssetDataReferencedOptions()) const { return InterfacePtr->GetAllReferencedAssetDatas(*this, InOptions); }
 	FAssetData GetAssetData() const { return InterfacePtr->GetAssetData(*this); }
 };
