[Backout] - CL34745785 - backed out due to errors
Original CL Desc
-----------------------------------------------------------------
[Selection] Fix an issue from the last submit that removed an assumption that if there is a BP actor, we only grab the top level actor in sync to CB code

#rnx

diff --git a/Engine/Source/Editor/UnrealEd/Private/EditorEngine.cpp b/Engine/Source/Editor/UnrealEd/Private/EditorEngine.cpp
--- a/Engine/Source/Editor/UnrealEd/Private/EditorEngine.cpp
+++ b/Engine/Source/Editor/UnrealEd/Private/EditorEngine.cpp
@@ -3305,7 +3305,7 @@ void UEditorEngine::GetAssetsToSyncToContentBrowser(TArray<FAssetData>& Assets,
 
 				if (TTypedElement<ITypedElementAssetDataInterface> AssetDataInterface = UTypedElementRegistry::GetInstance()->GetElement<ITypedElementAssetDataInterface>(SelectedHandle))
 				{
-					Assets.Append(AssetDataInterface.GetAllReferencedAssetDatas(FTypedElementAssetDataReferencedOptions().SetOnlyTopLevelAsset(true)));
+					Assets.Append(AssetDataInterface.GetAllReferencedAssetDatas());
 				}
 
 				return true;
diff --git a/Engine/Source/Editor/UnrealEd/Private/Elements/Actor/ActorElementEditorAssetDataInterface.cpp b/Engine/Source/Editor/UnrealEd/Private/Elements/Actor/ActorElementEditorAssetDataInterface.cpp
--- a/Engine/Source/Editor/UnrealEd/Private/Elements/Actor/ActorElementEditorAssetDataInterface.cpp
+++ b/Engine/Source/Editor/UnrealEd/Private/Elements/Actor/ActorElementEditorAssetDataInterface.cpp
@@ -7,7 +7,7 @@
 #include "Elements/Actor/ActorElementData.h"
 #include "GameFramework/Actor.h"
 
-TArray<FAssetData> UActorElementEditorAssetDataInterface::GetAllReferencedAssetDatas(const FTypedElementHandle& InElementHandle, const FTypedElementAssetDataReferencedOptions& InOptions)
+TArray<FAssetData> UActorElementEditorAssetDataInterface::GetAllReferencedAssetDatas(const FTypedElementHandle& InElementHandle)
 {
 	TArray<FAssetData> AssetDatas;
 
@@ -15,49 +15,28 @@ TArray<FAssetData> UActorElementEditorAssetDataInterface::GetAllReferencedAssetD
 	{
 		TArray<UObject*> ReferencedContentObjects;
 		RawActorPtr->GetReferencedContentObjects(ReferencedContentObjects);
-
-		if (InOptions.OnlyTopLevelAsset())
+		for (const UObject* ContentObject : ReferencedContentObjects)
 		{
-			UObject** BPObject = ReferencedContentObjects.FindByPredicate([](const UObject* InObject)
+			FAssetData ObjectAssetData = FAssetData(ContentObject);
+			if (ObjectAssetData.IsValid())
 			{
-				if (InObject && (InObject->IsA<UBlueprint>() || InObject->IsA<UBlueprintGeneratedClass>()))
-				{
-					return true;
-				}
-				return false;
-			});
-			
-			if (BPObject)
-			{
-				AssetDatas.Add(FAssetData(*BPObject, FAssetData::ECreationFlags::SkipAssetRegistryTagsGathering));
+				AssetDatas.Emplace(ObjectAssetData);
 			}
 		}
 
-		if (AssetDatas.IsEmpty())
+		TArray<FSoftObjectPath> SoftObjects;
+		RawActorPtr->GetSoftReferencedContentObjects(SoftObjects);
+		if (SoftObjects.Num())
 		{
-			for (const UObject* ContentObject : ReferencedContentObjects)
-			{
-				FAssetData ObjectAssetData = FAssetData(ContentObject, FAssetData::ECreationFlags::SkipAssetRegistryTagsGathering);
-				if (ObjectAssetData.IsValid())
-				{
-					AssetDatas.Emplace(ObjectAssetData);
-				}
-			}
+			IAssetRegistry& AssetRegistry = IAssetRegistry::GetChecked();
 
-			TArray<FSoftObjectPath> SoftObjects;
-			RawActorPtr->GetSoftReferencedContentObjects(SoftObjects);
-			if (SoftObjects.Num())
+			for (const FSoftObjectPath& SoftObject : SoftObjects)
 			{
-				IAssetRegistry& AssetRegistry = IAssetRegistry::GetChecked();
+				FAssetData AssetData = AssetRegistry.GetAssetByObjectPath(SoftObject);
 
-				for (const FSoftObjectPath& SoftObject : SoftObjects)
+				if (AssetData.IsValid())
 				{
-					FAssetData AssetData = AssetRegistry.GetAssetByObjectPath(SoftObject);
-
-					if (AssetData.IsValid())
-					{
-						AssetDatas.Add(AssetData);
-					}
+					AssetDatas.Add(AssetData);
 				}
 			}
 		}
diff --git a/Engine/Source/Editor/UnrealEd/Public/Elements/Actor/ActorElementEditorAssetDataInterface.h b/Engine/Source/Editor/UnrealEd/Public/Elements/Actor/ActorElementEditorAssetDataInterface.h
--- a/Engine/Source/Editor/UnrealEd/Public/Elements/Actor/ActorElementEditorAssetDataInterface.h
+++ b/Engine/Source/Editor/UnrealEd/Public/Elements/Actor/ActorElementEditorAssetDataInterface.h
@@ -19,5 +19,5 @@ class UActorElementEditorAssetDataInterface : public UActorElementAssetDataInter
 	GENERATED_BODY()
 
 public:
-	UNREALED_API virtual TArray<FAssetData> GetAllReferencedAssetDatas(const FTypedElementHandle& InElementHandle, const FTypedElementAssetDataReferencedOptions& InOptions) override;
+	UNREALED_API virtual TArray<FAssetData> GetAllReferencedAssetDatas(const FTypedElementHandle& InElementHandle) override;
 };
diff --git a/Engine/Source/Runtime/TypedElementRuntime/Private/Elements/Interfaces/TypedElementAssetDataInterface.cpp b/Engine/Source/Runtime/TypedElementRuntime/Private/Elements/Interfaces/TypedElementAssetDataInterface.cpp
--- a/Engine/Source/Runtime/TypedElementRuntime/Private/Elements/Interfaces/TypedElementAssetDataInterface.cpp
+++ b/Engine/Source/Runtime/TypedElementRuntime/Private/Elements/Interfaces/TypedElementAssetDataInterface.cpp
@@ -6,7 +6,7 @@
 
 #include UE_INLINE_GENERATED_CPP_BY_NAME(TypedElementAssetDataInterface)
 
-TArray<FAssetData> ITypedElementAssetDataInterface::GetAllReferencedAssetDatas(const FTypedElementHandle& InElementHandle, const FTypedElementAssetDataReferencedOptions& InOptions)
+TArray<FAssetData> ITypedElementAssetDataInterface::GetAllReferencedAssetDatas(const FTypedElementHandle& InElementHandle)
 {
 	TArray<FAssetData> AssetDatas;
 
diff --git a/Engine/Source/Runtime/TypedElementRuntime/Public/Elements/Interfaces/TypedElementAssetDataInterface.h b/Engine/Source/Runtime/TypedElementRuntime/Public/Elements/Interfaces/TypedElementAssetDataInterface.h
--- a/Engine/Source/Runtime/TypedElementRuntime/Public/Elements/Interfaces/TypedElementAssetDataInterface.h
+++ b/Engine/Source/Runtime/TypedElementRuntime/Public/Elements/Interfaces/TypedElementAssetDataInterface.h
@@ -14,20 +14,6 @@
 class UObject;
 struct FFrame;
 
-USTRUCT(BlueprintType)
-struct FTypedElementAssetDataReferencedOptions
-{
-	GENERATED_BODY()
-
-public:
-	FTypedElementAssetDataReferencedOptions& SetOnlyTopLevelAsset(const bool InTopLevelAsset) { bOnlyTopLevelAsset = InTopLevelAsset; return *this; }
-	bool OnlyTopLevelAsset() const { return bOnlyTopLevelAsset; }
-
-private:
-	UPROPERTY(BlueprintReadWrite, Category = "TypedElementInterfaces|AssetData|GetReferencedOptions", meta = (AllowPrivateAccess = true))
-	bool bOnlyTopLevelAsset = false;
-};
-
 UINTERFACE(MinimalAPI, BlueprintType, meta = (CannotImplementInterfaceInBlueprint))
 class UTypedElementAssetDataInterface : public UInterface
 {
@@ -45,7 +31,7 @@ public:
 	 *
 	 * @returns An array of valid asset datas.
 	 */
-	TYPEDELEMENTRUNTIME_API virtual TArray<FAssetData> GetAllReferencedAssetDatas(const FTypedElementHandle& InElementHandle, const FTypedElementAssetDataReferencedOptions& InOptions = FTypedElementAssetDataReferencedOptions());
+	TYPEDELEMENTRUNTIME_API virtual TArray<FAssetData> GetAllReferencedAssetDatas(const FTypedElementHandle& InElementHandle);
 
 	/**
 	 * Returns the asset data for the given handle, if it exists.
@@ -75,6 +61,6 @@ public:
 template <>
 struct TTypedElement<ITypedElementAssetDataInterface> : public TTypedElementBase<ITypedElementAssetDataInterface>
 {
-	TArray<FAssetData> GetAllReferencedAssetDatas(const FTypedElementAssetDataReferencedOptions& InOptions = FTypedElementAssetDataReferencedOptions()) const { return InterfacePtr->GetAllReferencedAssetDatas(*this, InOptions); }
+	TArray<FAssetData> GetAllReferencedAssetDatas() const { return InterfacePtr->GetAllReferencedAssetDatas(*this); }
 	FAssetData GetAssetData() const { return InterfacePtr->GetAssetData(*this); }
 };
