[Viewport Toolbar] Work on the "Snapping" submenu for the new level editor viewport toolbar

Adds submenus to select increments and options for:
- Grid
- Rotation
- Scale

Commit removes code duplication from the following toolbars, which are now using the same code as the New Level Editor Viewport Toolbar:
- Level Viewport
- UV Editor
- Custom Object Editor
- UMG Editor

diff --git a/Engine/Plugins/Editor/UVEditor/Source/UVEditor/Private/SUVEditor2DViewportToolBar.cpp b/Engine/Plugins/Editor/UVEditor/Source/UVEditor/Private/SUVEditor2DViewportToolBar.cpp
--- a/Engine/Plugins/Editor/UVEditor/Source/UVEditor/Private/SUVEditor2DViewportToolBar.cpp
+++ b/Engine/Plugins/Editor/UVEditor/Source/UVEditor/Private/SUVEditor2DViewportToolBar.cpp
@@ -11,7 +11,7 @@
 #include "UVEditor2DViewportClient.h"
 #include "Settings/LevelEditorViewportSettings.h"
 #include "SViewportToolBarComboMenu.h"
-#include "Widgets/Layout/SUniformGridPanel.h"
+#include "ViewportToolbar/UnrealEdViewportToolbar.h"
 
 #define LOCTEXT_NAMESPACE "UVEditor2DViewportToolbar"
 
@@ -318,124 +318,80 @@ TSharedRef<SWidget> SUVEditor2DViewportToolBar::FillLocationGridSnapMenu()
 	{
 		GridSizes.Add(FUVEditorUXSettings::LocationSnapValue(Index));
 	}
-	return BuildLocationGridCheckBoxList("Snap", LOCTEXT("LocationSnapText", "Snap Delta Distances"), GridSizes);
-}
 
-TSharedRef<SWidget> SUVEditor2DViewportToolBar::BuildLocationGridCheckBoxList(FName InExtentionHook, const FText& InHeading, const TArray<float>& InGridSizes) const
-{
-	const ULevelEditorViewportSettings* ViewportSettings = GetDefault<ULevelEditorViewportSettings>();
+	using namespace UE::UnrealEd;
 
-	const bool bShouldCloseWindowAfterMenuSelection = true;
-	FMenuBuilder LocationGridMenuBuilder(bShouldCloseWindowAfterMenuSelection, CommandList);
+	FLocationGridCheckboxListExecuteActionDelegate ExecuteDelegate =
+		FLocationGridCheckboxListExecuteActionDelegate::CreateLambda(
+			[this, GridSizes](int CurrGridSizeIndex)
+			{
+				const float CurrGridSize = GridSizes[CurrGridSizeIndex];
+				Viewport2DClient->SetLocationGridSnapValue(CurrGridSize);
+			}
+		);
 
-	LocationGridMenuBuilder.BeginSection(InExtentionHook, InHeading);
-	for (int32 CurGridSizeIndex = 0; CurGridSizeIndex < InGridSizes.Num(); ++CurGridSizeIndex)
-	{
-		const float CurGridSize = InGridSizes[CurGridSizeIndex];
-
-		LocationGridMenuBuilder.AddMenuEntry(
-			FText::AsNumber(CurGridSize),
-			FText::Format(LOCTEXT("LocationGridSize_ToolTip", "Sets snap delta to {0}"), FText::AsNumber(CurGridSize)),
-			FSlateIcon(),
-			FUIAction(FExecuteAction::CreateLambda([this, CurGridSize]() {Viewport2DClient->SetLocationGridSnapValue(CurGridSize); }),
-				FCanExecuteAction(),
-				FIsActionChecked::CreateLambda([this, CurGridSize]() { return FMath::IsNearlyEqual(Viewport2DClient->GetLocationGridSnapValue(), CurGridSize); })),
-			NAME_None,
-			EUserInterfaceActionType::RadioButton);
-	}
-	LocationGridMenuBuilder.EndSection();
+	FLocationGridCheckboxListIsCheckedDelegate IsCheckedDelegate = FLocationGridCheckboxListIsCheckedDelegate::CreateLambda(
+		[this, GridSizes](int CurrGridSizeIndex)
+		{
+			const float CurrGridSize = GridSizes[CurrGridSizeIndex];
+			return FMath::IsNearlyEqual(Viewport2DClient->GetLocationGridSnapValue(), CurrGridSize);
+		}
+	);
 
-	return LocationGridMenuBuilder.MakeWidget();
+	return CreateLocationGridSnapMenu(ExecuteDelegate, IsCheckedDelegate, GridSizes, CommandList);
 }
 
 TSharedRef<SWidget> SUVEditor2DViewportToolBar::FillRotationGridSnapMenu()
 {
 	const ULevelEditorViewportSettings* ViewportSettings = GetDefault<ULevelEditorViewportSettings>();
 
-	return SNew(SUniformGridPanel)
-
-		+ SUniformGridPanel::Slot(0, 0)
-		[
-			BuildRotationGridCheckBoxList("Common", LOCTEXT("RotationCommonText", "Common"), ViewportSettings->CommonRotGridSizes)
-		]
-
-	+ SUniformGridPanel::Slot(1, 0)
-		[
-			BuildRotationGridCheckBoxList("Div360", LOCTEXT("RotationDivisions360DegreesText", "Divisions of 360\u00b0"), ViewportSettings->DivisionsOf360RotGridSizes)
-		];
-}
+	TArray<float> GridSizes = ViewportSettings->CommonRotGridSizes;
 
-TSharedRef<SWidget> SUVEditor2DViewportToolBar::BuildRotationGridCheckBoxList(FName InExtentionHook, const FText& InHeading, const TArray<float>& InGridSizes) const
-{
-	const bool bShouldCloseWindowAfterMenuSelection = true;
-	FMenuBuilder RotationGridMenuBuilder(bShouldCloseWindowAfterMenuSelection, CommandList);
+	using namespace UE::UnrealEd;
+	FRotationGridCheckboxListExecuteActionDelegate ExecuteDelegate =
+		FRotationGridCheckboxListExecuteActionDelegate::CreateLambda(
+			[this, GridSizes](int InCurrGridAngleIndex, ERotationGridMode InGridMode)
+			{
+				const float CurrGridAngle = GridSizes[InCurrGridAngleIndex];
+				Viewport2DClient->SetRotationGridSnapValue(CurrGridAngle);
+			}
+		);
 
-	RotationGridMenuBuilder.BeginSection(InExtentionHook, InHeading);
-	for (int32 CurGridAngleIndex = 0; CurGridAngleIndex < InGridSizes.Num(); ++CurGridAngleIndex)
-	{
-		const float CurGridAngle = InGridSizes[CurGridAngleIndex];
-
-		FText MenuName = FText::Format(LOCTEXT("RotationGridAngle", "{0}\u00b0"), FText::AsNumber(CurGridAngle)); /*degree symbol*/
-		FText ToolTipText = FText::Format(LOCTEXT("RotationGridAngle_ToolTip", "Sets rotation snap angle to {0}"), MenuName); /*degree symbol*/
-
-		RotationGridMenuBuilder.AddMenuEntry(
-			MenuName,
-			ToolTipText,
-			FSlateIcon(),
-			FUIAction(FExecuteAction::CreateLambda([this, CurGridAngle]() {Viewport2DClient->SetRotationGridSnapValue(CurGridAngle); }),
-				FCanExecuteAction(),
-				FIsActionChecked::CreateLambda([this, CurGridAngle]() {return FMath::IsNearlyEqual(Viewport2DClient->GetRotationGridSnapValue(), CurGridAngle); })),
-			NAME_None,
-			EUserInterfaceActionType::RadioButton);
-	}
-	RotationGridMenuBuilder.EndSection();
+	FRotationGridCheckboxListIsCheckedDelegate IsCheckedDelegate = FRotationGridCheckboxListIsCheckedDelegate::CreateLambda(
+		[this, GridSizes](int InCurrGridAngleIndex, ERotationGridMode InGridMode)
+		{
+			const float CurrGridAngle = GridSizes[InCurrGridAngleIndex];
+			return FMath::IsNearlyEqual(Viewport2DClient->GetRotationGridSnapValue(), CurrGridAngle);
+		}
+	);
 
-	return RotationGridMenuBuilder.MakeWidget();
+	return CreateRotationGridSnapMenu(ExecuteDelegate, IsCheckedDelegate, CommandList);
 }
 
 TSharedRef<SWidget> SUVEditor2DViewportToolBar::FillScaleGridSnapMenu()
 {
 	const ULevelEditorViewportSettings* ViewportSettings = GetDefault<ULevelEditorViewportSettings>();
-	const bool bShouldCloseWindowAfterMenuSelection = true;
-
-	FNumberFormattingOptions NumberFormattingOptions;
-	NumberFormattingOptions.MaximumFractionalDigits = 5;
-
-	FMenuBuilder ScaleGridMenuBuilder(bShouldCloseWindowAfterMenuSelection, CommandList);
-
-	ScaleGridMenuBuilder.BeginSection("ScaleSnapOptions", LOCTEXT("ScaleSnapOptions", "Scale Snap"));
-
-	for (int32 CurGridAmountIndex = 0; CurGridAmountIndex < ViewportSettings->ScalingGridSizes.Num(); ++CurGridAmountIndex)
-	{
-		const float CurGridAmount = ViewportSettings->ScalingGridSizes[CurGridAmountIndex];
-
-		FText MenuText;
-		FText ToolTipText;
-
-		if (GEditor->UsePercentageBasedScaling())
-		{
-			MenuText = FText::AsPercent(CurGridAmount / 100.0f, &NumberFormattingOptions);
-			ToolTipText = FText::Format(LOCTEXT("ScaleGridAmountOld_ToolTip", "Snaps scale values to {0}"), MenuText);
+	TArray<float> GridSizes = ViewportSettings->ScalingGridSizes;
+
+	UE::UnrealEd::FScaleGridCheckboxListExecuteActionDelegate ExecuteDelegate =
+		UE::UnrealEd::FScaleGridCheckboxListExecuteActionDelegate::CreateLambda(
+			[this, GridSizes](int CurrGridScaleIndex)
+			{
+				const float CurGridAmount = GridSizes[CurrGridScaleIndex];
+				Viewport2DClient->SetScaleGridSnapValue(CurGridAmount);
 		}
-		else
+		);
+
+	UE::UnrealEd::FScaleGridCheckboxListIsCheckedDelegate IsCheckedDelegate =
+		UE::UnrealEd::FScaleGridCheckboxListIsCheckedDelegate::CreateLambda(
+			[this, GridSizes](int CurrGridScaleIndex)
 		{
-			MenuText = FText::AsNumber(CurGridAmount, &NumberFormattingOptions);
-			ToolTipText = FText::Format(LOCTEXT("ScaleGridAmount_ToolTip", "Snaps scale values to increments of {0}"), MenuText);
+				const float CurrGridAmount = GridSizes[CurrGridScaleIndex];
+				return FMath::IsNearlyEqual(Viewport2DClient->GetScaleGridSnapValue(), CurrGridAmount);
 		}
+		);
 
-		ScaleGridMenuBuilder.AddMenuEntry(
-			MenuText,
-			ToolTipText,
-			FSlateIcon(),
-			FUIAction(FExecuteAction::CreateLambda([this, CurGridAmount]() {Viewport2DClient->SetScaleGridSnapValue(CurGridAmount); }),
-				FCanExecuteAction(),
-				FIsActionChecked::CreateLambda([this, CurGridAmount]() {return FMath::IsNearlyEqual(Viewport2DClient->GetScaleGridSnapValue(), CurGridAmount); })),
-			NAME_None,
-			EUserInterfaceActionType::RadioButton);
-	}
-	ScaleGridMenuBuilder.EndSection();
-
-	return ScaleGridMenuBuilder.MakeWidget();
+	return UE::UnrealEd::CreateScaleGridSnapMenu(ExecuteDelegate, IsCheckedDelegate, GridSizes, CommandList);
 }
 
 ECheckBoxState SUVEditor2DViewportToolBar::IsLocationGridSnapChecked() const
@@ -468,4 +424,4 @@ void SUVEditor2DViewportToolBar::HandleToggleScaleGridSnap(ECheckBoxState InStat
 	Viewport2DClient->SetScaleGridSnapEnabled(InState == ECheckBoxState::Checked);
 }
 
-#undef LOCTEXT_NAMESPACE
\ No newline at end of file
+#undef LOCTEXT_NAMESPACE
diff --git a/Engine/Plugins/Editor/UVEditor/Source/UVEditor/Public/SUVEditor2DViewportToolBar.h b/Engine/Plugins/Editor/UVEditor/Source/UVEditor/Public/SUVEditor2DViewportToolBar.h
--- a/Engine/Plugins/Editor/UVEditor/Source/UVEditor/Public/SUVEditor2DViewportToolBar.h
+++ b/Engine/Plugins/Editor/UVEditor/Source/UVEditor/Public/SUVEditor2DViewportToolBar.h
@@ -31,10 +31,6 @@ private:
 	TSharedRef<SWidget> MakeGizmoToolBar(const TSharedPtr<FExtender> InExtenders);
 	TSharedRef<SWidget> MakeTransformToolBar(const TSharedPtr< FExtender > InExtenders);
 
-	TSharedRef<SWidget> BuildLocationGridCheckBoxList(FName InExtentionHook, const FText& InHeading, const TArray<float>& InGridSizes) const;
-	TSharedRef<SWidget> BuildRotationGridCheckBoxList(FName InExtentionHook, const FText& InHeading, const TArray<float>& InGridSizes) const;
-
-
 	/** Grid snap label callbacks */
 	FText GetLocationGridLabel() const;
 	FText GetRotationGridLabel() const;	
diff --git a/Engine/Plugins/Experimental/Mutable/Source/CustomizableObjectEditor/Private/MuCOE/SCustomizableObjectEditorViewportToolBar.cpp b/Engine/Plugins/Experimental/Mutable/Source/CustomizableObjectEditor/Private/MuCOE/SCustomizableObjectEditorViewportToolBar.cpp
--- a/Engine/Plugins/Experimental/Mutable/Source/CustomizableObjectEditor/Private/MuCOE/SCustomizableObjectEditorViewportToolBar.cpp
+++ b/Engine/Plugins/Experimental/Mutable/Source/CustomizableObjectEditor/Private/MuCOE/SCustomizableObjectEditorViewportToolBar.cpp
@@ -26,6 +26,7 @@
 #include "MuCO/CustomizableObject.h"
 #include "MuCO/CustomizableObjectPrivate.h"
 #include "MuCO/CustomizableObjectSystem.h"
+#include "ViewportToolbar/UnrealEdViewportToolbar.h"
 #include "Widgets/Input/SButton.h"
 #include "Widgets/Input/SMenuAnchor.h"
 #include "Widgets/Input/SSpinBox.h"
@@ -499,51 +500,19 @@ FText SCustomizableObjectEditorViewportToolBar::GetRotationGridLabel() const
 
 TSharedRef<SWidget> SCustomizableObjectEditorViewportToolBar::FillRotationGridSnapMenu()
 {
-	const ULevelEditorViewportSettings* ViewportSettings = GetDefault<ULevelEditorViewportSettings>();
+	UE::UnrealEd::FRotationGridCheckboxListExecuteActionDelegate ExecuteDelegate =
+		UE::UnrealEd::FRotationGridCheckboxListExecuteActionDelegate::CreateStatic(
+			&SCustomizableObjectEditorViewportTabBody::SetRotationGridSize
+		);
 
-	return SNew(SUniformGridPanel)
+	UE::UnrealEd::FRotationGridCheckboxListIsCheckedDelegate IsCheckedDelegate =
+		UE::UnrealEd::FRotationGridCheckboxListIsCheckedDelegate::CreateStatic(
+			&SCustomizableObjectEditorViewportTabBody::IsRotationGridSizeChecked
+		);
 
-		+ SUniformGridPanel::Slot(0, 0)
-		[
-			BuildRotationGridCheckBoxList("Common", LOCTEXT("RotationCommonText", "Common"), ViewportSettings->CommonRotGridSizes, GridMode_Common)
-		]
-
-		+ SUniformGridPanel::Slot(1, 0)
-		[
-			BuildRotationGridCheckBoxList("Div360", LOCTEXT("RotationDivisions360DegreesText", "Divisions of 360\u00b0"), ViewportSettings->DivisionsOf360RotGridSizes, GridMode_DivisionsOf360)
-		];
-}
-
-
-TSharedRef<SWidget> SCustomizableObjectEditorViewportToolBar::BuildRotationGridCheckBoxList(FName InExtentionHook, const FText& InHeading, const TArray<float>& InGridSizes, ERotationGridMode InGridMode) const
-{
-	const bool bShouldCloseWindowAfterMenuSelection = true;
-	FMenuBuilder RotationGridMenuBuilder(bShouldCloseWindowAfterMenuSelection, Viewport.Pin()->GetCommandList());
-
-	RotationGridMenuBuilder.BeginSection(InExtentionHook, InHeading);
-	for (int32 CurGridAngleIndex = 0; CurGridAngleIndex < InGridSizes.Num(); ++CurGridAngleIndex)
-	{
-		const float CurGridAngle = InGridSizes[CurGridAngleIndex];
-
-		FText MenuName = FText::Format(LOCTEXT("RotationGridAngle", "{0}\u00b0"), FText::AsNumber(CurGridAngle)); /*degree symbol*/
-		FText ToolTipText = FText::Format(LOCTEXT("RotationGridAngle_ToolTip", "Sets rotation grid angle to {0}"), MenuName); /*degree symbol*/
-
-		RotationGridMenuBuilder.AddMenuEntry(
-			MenuName,
-			ToolTipText,
-			FSlateIcon(),
-			FUIAction(FExecuteAction::CreateStatic(&SCustomizableObjectEditorViewportTabBody::SetRotationGridSize, CurGridAngleIndex, InGridMode),
-				FCanExecuteAction(),
-				FIsActionChecked::CreateStatic(&SCustomizableObjectEditorViewportTabBody::IsRotationGridSizeChecked, CurGridAngleIndex, InGridMode)),
-			NAME_None,
-			EUserInterfaceActionType::RadioButton);
-	}
-	RotationGridMenuBuilder.EndSection();
-
-	return RotationGridMenuBuilder.MakeWidget();
+	return UE::UnrealEd::CreateRotationGridSnapMenu(ExecuteDelegate, IsCheckedDelegate, Viewport.Pin()->GetCommandList());
 }
 
-
 FReply SCustomizableObjectEditorViewportToolBar::OnMenuClicked()
 {
 	// If the menu button is clicked toggle the state of the menu anchor which will open or close the menu
diff --git a/Engine/Plugins/Experimental/Mutable/Source/CustomizableObjectEditor/Public/MuCOE/SCustomizableObjectEditorViewportToolBar.h b/Engine/Plugins/Experimental/Mutable/Source/CustomizableObjectEditor/Public/MuCOE/SCustomizableObjectEditorViewportToolBar.h
--- a/Engine/Plugins/Experimental/Mutable/Source/CustomizableObjectEditor/Public/MuCOE/SCustomizableObjectEditorViewportToolBar.h
+++ b/Engine/Plugins/Experimental/Mutable/Source/CustomizableObjectEditor/Public/MuCOE/SCustomizableObjectEditorViewportToolBar.h
@@ -93,12 +93,11 @@ private:
 	float OnGetFOVValue() const;
 	/** Called when the FOV slider is adjusted in the perspective viewport */
 	void OnFOVValueChanged(float NewValue) const;
-	
+
 	ECheckBoxState IsRotationGridSnapChecked() const;
 	void HandleToggleRotationGridSnap(ECheckBoxState InState);
 	FText GetRotationGridLabel() const;
 	TSharedRef<SWidget> FillRotationGridSnapMenu();
-	TSharedRef<SWidget> BuildRotationGridCheckBoxList(FName InExtentionHook, const FText& InHeading, const TArray<float>& InGridSizes, ERotationGridMode InGridMode) const;
 
 	/** Callback for drop-down menu with FOV and high resolution screenshot options currently */
 	FReply OnMenuClicked();
diff --git a/Engine/Source/Editor/LevelEditor/Private/SLevelViewport.cpp b/Engine/Source/Editor/LevelEditor/Private/SLevelViewport.cpp
--- a/Engine/Source/Editor/LevelEditor/Private/SLevelViewport.cpp
+++ b/Engine/Source/Editor/LevelEditor/Private/SLevelViewport.cpp
@@ -1940,7 +1940,7 @@ TSharedPtr<SWidget> SLevelViewport::MakeViewportToolbar()
 
 				// Add the "Snapping" sub menu.
 				{
-					FToolMenuEntry SnappingSubmenu = UE::LevelEditor::CreateViewportToolbarSnappingSubmenu();
+					FToolMenuEntry SnappingSubmenu = UE::UnrealEd::CreateViewportToolbarSnappingSubmenu();
 					SnappingSubmenu.InsertPosition.Position = EToolMenuInsertType::First;
 					LeftSection.AddEntry(SnappingSubmenu);
 				}
diff --git a/Engine/Source/Editor/LevelEditor/Private/ViewportToolbar/LevelEditorViewportToolbarSections.cpp b/Engine/Source/Editor/LevelEditor/Private/ViewportToolbar/LevelEditorViewportToolbarSections.cpp
--- a/Engine/Source/Editor/LevelEditor/Private/ViewportToolbar/LevelEditorViewportToolbarSections.cpp
+++ b/Engine/Source/Editor/LevelEditor/Private/ViewportToolbar/LevelEditorViewportToolbarSections.cpp
@@ -128,208 +128,6 @@ void PopulateShowLayersSubmenu(UToolMenu* InMenu, TWeakPtr<::SLevelViewport> InV
 	}
 }
 
-// TODO: Maybe export CreateSurfaceSnapOffsetEntry function, so that it can be used elsewhere, e.g. STransformViewportToolbar.cpp
-FToolMenuEntry CreateSurfaceSnapOffsetEntry()
-{
-	FText Label = LOCTEXT("SurfaceOffsetLabel", "Surface Offset");
-	FText Tooltip = LOCTEXT("SurfaceOffsetTooltip", "The amount of offset to apply when snapping to surfaces");
-
-	FToolMenuEntry SurfaceOffset = FToolMenuEntry::InitMenuEntry(
-		"SurfaceOffset",
-		Label,
-		Tooltip,
-		FSlateIcon(),
-		FUIAction());
-
-	const FMargin WidgetsMargin(8.0f, 0.0f, 0.0f, 0.0f);
-
-	SurfaceOffset.MakeCustomWidget.BindLambda(
-		[Label, Tooltip, WidgetsMargin](const FToolMenuContext& InContext, const FToolMenuCustomWidgetContext& InWidgetContext)
-		{
-			// clang-format off
-			return SNew(SHorizontalBox).IsEnabled_Lambda([]()
-				{
-					return GetDefault<ULevelEditorViewportSettings>()->SnapToSurface.bEnabled;
-				})
-				+ SHorizontalBox::Slot().VAlign(VAlign_Center).Padding(WidgetsMargin).AutoWidth()
-				[
-					SNew(STextBlock).Text(Label).TextStyle(InWidgetContext.StyleSet,
-					                                       ISlateStyle::Join(InWidgetContext.StyleName, ".Label"))
-				]
-				+ SHorizontalBox::Slot().VAlign(VAlign_Center).Padding(WidgetsMargin).AutoWidth()
-				[
-					SNew(SBox).Padding(WidgetsMargin).MinDesiredWidth(100.0f)
-					[
-						// Min/Max/Slider values taken from STransformViewportToolbar.cpp
-						SNew(SNumericEntryBox<float>)
-						.ToolTipText(Tooltip)
-						.MinValue(0.0f)
-						.MaxValue(static_cast<float>(HALF_WORLD_MAX))
-						.MaxSliderValue(1000.0f)
-						.AllowSpin(true)
-						.MaxFractionalDigits(2)
-						.Font(FAppStyle::GetFontStyle(TEXT("MenuItem.Font")))
-						.OnValueChanged_Lambda([](float InNewValue)
-						{
-							auto* Settings = GetMutableDefault<ULevelEditorViewportSettings>();
-							Settings->SnapToSurface.SnapOffsetExtent = InNewValue;
-						})
-						.Value_Lambda([]()
-						{
-							return GetDefault<ULevelEditorViewportSettings>()->SnapToSurface.SnapOffsetExtent;
-						})
-					]
-				];
-			// clang-format on
-		});
-
-	return SurfaceOffset;
-}
-
-FToolMenuEntry CreateSurfaceSnapMenu()
-{
-	FNewToolMenuDelegate MakeMenuDelegate = FNewToolMenuDelegate::CreateLambda([](UToolMenu* Submenu)
-	{
-		FToolMenuSection& SurfaceSnappingSection = Submenu->FindOrAddSection("SurfaceSnapping", LOCTEXT("SurfaceSnappingLabel", "Surface Snapping"));
-
-		// Add "Rotate to surface normal" checkbox.
-		{
-			FToolMenuEntry RotateToSurfaceNormalSnapping = FToolMenuEntry::InitMenuEntry(FEditorViewportCommands::Get().RotateToSurfaceNormal);
-			SurfaceSnappingSection.AddEntry(RotateToSurfaceNormalSnapping);
-		}
-
-		// Add "Surface offset" widget.
-		{
-			SurfaceSnappingSection.AddEntry(CreateSurfaceSnapOffsetEntry());
-		}
-	});
-
-	FToolUIAction CheckboxMenuAction;
-	{
-		CheckboxMenuAction.ExecuteAction = FToolMenuExecuteAction::CreateLambda(
-			[](const FToolMenuContext& InContext)
-			{
-				auto* Settings = GetMutableDefault<ULevelEditorViewportSettings>();
-				Settings->SnapToSurface.bEnabled = !Settings->SnapToSurface.bEnabled;
-			}
-		);
-		CheckboxMenuAction.GetActionCheckState = FToolMenuGetActionCheckState::CreateLambda(
-			[](const FToolMenuContext& InContext)
-			{
-				return GetDefault<ULevelEditorViewportSettings>()->SnapToSurface.bEnabled ? ECheckBoxState::Checked : ECheckBoxState::Unchecked;
-			}
-		);
-	}
-
-	FToolMenuEntry SurfaceSnapping = FToolMenuEntry::InitSubMenu(
-		"SurfaceSnapping",
-		LOCTEXT("SurfaceSnapLabel", "Surface"),
-		FEditorViewportCommands::Get().SurfaceSnapping->MakeTooltip()->GetTextTooltip(),
-		MakeMenuDelegate,
-		CheckboxMenuAction,
-		EUserInterfaceActionType::ToggleButton
-	);
-
-	return SurfaceSnapping;
-}
-
-FToolMenuEntry CreateActorSnapDistanceEntry()
-{
-	const FText Label = LOCTEXT("ActorSnapDistanceLabel", "Snap Distance");
-	const FText Tooltip = LOCTEXT("ActorSnapDistanceTooltip", "The amount of offset to apply when snapping to surfaces");
-
-	FToolMenuEntry SnapDistance = FToolMenuEntry::InitMenuEntry(
-		"ActorSnapDistance",
-		Label,
-		Tooltip,
-		FSlateIcon(),
-		FUIAction());
-
-	const FMargin WidgetsMargin(8.0f, 0.0f, 0.0f, 0.0f);
-
-	SnapDistance.MakeCustomWidget.BindLambda(
-		[Label, Tooltip, WidgetsMargin](const FToolMenuContext& InContext, const FToolMenuCustomWidgetContext& InWidgetContext)
-		{
-			// clang-format off
-			return SNew(SHorizontalBox).IsEnabled_Lambda([]()
-				{
-					return !!GetDefault<ULevelEditorViewportSettings>()->bEnableActorSnap;
-				})
-				+ SHorizontalBox::Slot().VAlign(VAlign_Center).Padding(WidgetsMargin).AutoWidth()
-				[
-					SNew(STextBlock).Text(Label).TextStyle(InWidgetContext.StyleSet,
-					                                       ISlateStyle::Join(InWidgetContext.StyleName, ".Label"))
-				]
-				+ SHorizontalBox::Slot().VAlign(VAlign_Center).Padding(WidgetsMargin).AutoWidth()
-				[
-					SNew(SBox).Padding(WidgetsMargin).MinDesiredWidth(100.0f)
-					[
-						// TODO: Check how to improve performance for this widget OnValueChanged.
-						// Same functionality in LevelEditorToolBar.cpp seems to have better performance
-						SNew(SNumericEntryBox<float>)
-						.ToolTipText(Tooltip)
-						.MinValue(0.0f)
-						.MaxValue(1.0f)
-						.MaxSliderValue(1.0f)
-						.AllowSpin(true)
-						.MaxFractionalDigits(1)
-						.Font(FAppStyle::GetFontStyle(TEXT("MenuItem.Font")))
-						.OnValueChanged_Static(&FLevelEditorActionCallbacks::SetActorSnapSetting)
-						.Value_Lambda([]()
-						{
-							return FLevelEditorActionCallbacks::GetActorSnapSetting();
-						})
-					]
-				];
-			// clang-format on
-		});
-
-	return SnapDistance;
-}
-
-FToolMenuEntry CreateActorSnapMenu()
-{
-	FNewToolMenuDelegate MakeMenuDelegate = FNewToolMenuDelegate::CreateLambda([](UToolMenu* Submenu)
-	{
-		FToolMenuSection& ActorSnappingSection = Submenu->FindOrAddSection("ActorSnapping", LOCTEXT("ActorSnappingLabel", "Actor Snapping"));
-
-		// Add "Actor snapping" widget.
-		{
-			ActorSnappingSection.AddEntry(CreateActorSnapDistanceEntry());
-		}
-	});
-
-	FToolUIAction CheckboxMenuAction;
-	{
-		CheckboxMenuAction.ExecuteAction = FToolMenuExecuteAction::CreateLambda(
-			[](const FToolMenuContext& InContext)
-			{
-				if (ULevelEditorViewportSettings* Settings = GetMutableDefault<ULevelEditorViewportSettings>())
-				{
-					Settings->bEnableActorSnap = !Settings->bEnableActorSnap;
-				}
-			}
-		);
-		CheckboxMenuAction.GetActionCheckState = FToolMenuGetActionCheckState::CreateLambda(
-			[](const FToolMenuContext& InContext)
-			{
-				return GetDefault<ULevelEditorViewportSettings>()->bEnableActorSnap ? ECheckBoxState::Checked : ECheckBoxState::Unchecked;
-			}
-		);
-	}
-
-	FToolMenuEntry ActorSnapping = FToolMenuEntry::InitSubMenu(
-		"ActorSnapping",
-		LOCTEXT("ActorSnapLabel", "Actor"),
-		FLevelEditorCommands::Get().EnableActorSnap->MakeTooltip()->GetTextTooltip(),
-		MakeMenuDelegate,
-		CheckboxMenuAction,
-		EUserInterfaceActionType::ToggleButton
-	);
-
-	return ActorSnapping;
-}
-
 } // namespace UE::LevelEditor::Private
 
 namespace UE::LevelEditor
@@ -1863,54 +1661,6 @@ FToolMenuEntry CreateLevelEditorViewportToolbarSettingsSubmenu()
 	);
 }
 
-FToolMenuEntry CreateViewportToolbarSnappingSubmenu()
-{
-	return FToolMenuEntry::InitSubMenu(
-		"Snapping",
-		LOCTEXT("SnappingSubmenuLabel", "Snapping"),
-		LOCTEXT("SnappingSubmenuTooltip", "Viewport-related snapping settings"),
-		FNewToolMenuDelegate::CreateLambda(
-			[](UToolMenu* Submenu) -> void
-			{
-				FToolMenuSection& SnappingSection =
-					Submenu->FindOrAddSection("Snapping", LOCTEXT("SnappingLabel", "Snapping"));
-
-				SnappingSection.AddEntry(Private::CreateSurfaceSnapMenu());
-
-				FToolMenuEntry GridSnapping = FToolMenuEntry::InitMenuEntry(FEditorViewportCommands::Get().LocationGridSnap);
-				GridSnapping.UserInterfaceActionType = EUserInterfaceActionType::ToggleButton;
-				GridSnapping.Label = LOCTEXT("GridSnapLabel", "Grid");
-				GridSnapping.SetShowInToolbarTopLevel(true);
-				SnappingSection.AddEntry(GridSnapping);
-
-				FToolMenuEntry RotationSnapping = FToolMenuEntry::InitMenuEntry(FEditorViewportCommands::Get().RotationGridSnap);
-				RotationSnapping.UserInterfaceActionType = EUserInterfaceActionType::ToggleButton;
-				RotationSnapping.Label = LOCTEXT("RotationSnapLabel", "Rotation");
-				SnappingSection.AddEntry(RotationSnapping);
-
-				FToolMenuEntry ScaleSnapping = FToolMenuEntry::InitMenuEntry(FEditorViewportCommands::Get().ScaleGridSnap);
-				ScaleSnapping.UserInterfaceActionType = EUserInterfaceActionType::ToggleButton;
-				ScaleSnapping.Label = LOCTEXT("ScaleSnapLabel", "Scale");
-				SnappingSection.AddEntry(ScaleSnapping);
-
-				SnappingSection.AddEntry(Private::CreateActorSnapMenu());
-
-				FToolMenuEntry SocketSnapping = FToolMenuEntry::InitMenuEntry(FLevelEditorCommands::Get().ToggleSocketSnapping);
-				SocketSnapping.UserInterfaceActionType = EUserInterfaceActionType::ToggleButton;
-				SocketSnapping.Label = LOCTEXT("SocketSnapLabel", "Socket");
-				SnappingSection.AddEntry(SocketSnapping);
-
-				FToolMenuEntry VertexSnapping = FToolMenuEntry::InitMenuEntry(FLevelEditorCommands::Get().EnableVertexSnap);
-				VertexSnapping.UserInterfaceActionType = EUserInterfaceActionType::ToggleButton;
-				VertexSnapping.Label = LOCTEXT("VertexSnapLabel", "Vertex");
-				SnappingSection.AddEntry(VertexSnapping);
-
-				// TODO: add Planar Snapping
-			}
-		)
-	);
-}
-
 } // namespace UE::LevelEditor
 
 #undef LOCTEXT_NAMESPACE
diff --git a/Engine/Source/Editor/LevelEditor/Private/ViewportToolbar/LevelEditorViewportToolbarSections.h b/Engine/Source/Editor/LevelEditor/Private/ViewportToolbar/LevelEditorViewportToolbarSections.h
--- a/Engine/Source/Editor/LevelEditor/Private/ViewportToolbar/LevelEditorViewportToolbarSections.h
+++ b/Engine/Source/Editor/LevelEditor/Private/ViewportToolbar/LevelEditorViewportToolbarSections.h
@@ -39,6 +39,4 @@ void GenerateViewportLayoutsMenu(UToolMenu* InMenu, TSharedPtr<::SLevelViewport>
 TSharedRef<SWidget> BuildVolumeControlCustomWidget();
 FToolMenuEntry CreateLevelEditorViewportToolbarSettingsSubmenu();
 
-FToolMenuEntry CreateViewportToolbarSnappingSubmenu();
-
 } // namespace UE::LevelEditor
diff --git a/Engine/Source/Editor/UMGEditor/Private/Designer/SDesignerToolBar.cpp b/Engine/Source/Editor/UMGEditor/Private/Designer/SDesignerToolBar.cpp
--- a/Engine/Source/Editor/UMGEditor/Private/Designer/SDesignerToolBar.cpp
+++ b/Engine/Source/Editor/UMGEditor/Private/Designer/SDesignerToolBar.cpp
@@ -19,6 +19,7 @@
 
 #include "Designer/DesignerCommands.h"
 #include "SViewportToolBarComboMenu.h"
+#include "ViewportToolbar/UnrealEdViewportToolbar.h"
 
 #define LOCTEXT_NAMESPACE "UMG"
 
@@ -124,9 +125,7 @@ FText SDesignerToolBar::GetLocationGridLabel() const
 
 TSharedRef<SWidget> SDesignerToolBar::FillLocationGridSnapMenu()
 {
-	const UWidgetDesignerSettings* ViewportSettings = GetDefault<UWidgetDesignerSettings>();
-
-	TArray<int32> GridSizes;
+	TArray<float> GridSizes;
 	GridSizes.Add(1);
 	GridSizes.Add(2);
 	GridSizes.Add(3);
@@ -136,34 +135,27 @@ TSharedRef<SWidget> SDesignerToolBar::FillLocationGridSnapMenu()
 	GridSizes.Add(15);
 	GridSizes.Add(25);
 
-	return BuildLocationGridCheckBoxList("Snap", LOCTEXT("LocationSnapText", "Snap Sizes"), GridSizes);
-}
-
-TSharedRef<SWidget> SDesignerToolBar::BuildLocationGridCheckBoxList(FName InExtentionHook, const FText& InHeading, const TArray<int32>& InGridSizes) const
-{
-	const UWidgetDesignerSettings* ViewportSettings = GetDefault<UWidgetDesignerSettings>();
-
-	const bool bShouldCloseWindowAfterMenuSelection = true;
-	FMenuBuilder LocationGridMenuBuilder(bShouldCloseWindowAfterMenuSelection, CommandList);
-
-	LocationGridMenuBuilder.BeginSection(InExtentionHook, InHeading);
-	for ( int32 CurGridSizeIndex = 0; CurGridSizeIndex < InGridSizes.Num(); ++CurGridSizeIndex )
-	{
-		const int32 CurGridSize = InGridSizes[CurGridSizeIndex];
+	UE::UnrealEd::FLocationGridCheckboxListExecuteActionDelegate ExecuteDelegate =
+		UE::UnrealEd::FLocationGridCheckboxListExecuteActionDelegate::CreateLambda(
+			[GridSizes](int CurrGridSizeIndex)
+			{
+				UWidgetDesignerSettings* ViewportSettings = GetMutableDefault<UWidgetDesignerSettings>();
+				const int32 CurrGridSize = GridSizes[CurrGridSizeIndex];
+				ViewportSettings->GridSnapSize = CurrGridSize;
+			}
+		);
 
-		LocationGridMenuBuilder.AddMenuEntry(
-			FText::AsNumber(CurGridSize),
-			FText::Format(LOCTEXT("LocationGridSize_ToolTip", "Sets grid size to {0}"), FText::AsNumber(CurGridSize)),
-			FSlateIcon(),
-			FUIAction(FExecuteAction::CreateStatic(&SDesignerToolBar::SetGridSize, CurGridSize),
-			FCanExecuteAction(),
-			FIsActionChecked::CreateStatic(&SDesignerToolBar::IsGridSizeChecked, CurGridSize)),
-			NAME_None,
-			EUserInterfaceActionType::RadioButton);
-	}
-	LocationGridMenuBuilder.EndSection();
+	UE::UnrealEd::FLocationGridCheckboxListIsCheckedDelegate IsCheckedDelegate =
+		UE::UnrealEd::FLocationGridCheckboxListIsCheckedDelegate::CreateLambda(
+			[GridSizes](int CurrGridSizeIndex)
+			{
+				const UWidgetDesignerSettings* ViewportSettings = GetDefault<UWidgetDesignerSettings>();
+				const int32 CurrGridSize = GridSizes[CurrGridSizeIndex];
+				return (ViewportSettings->GridSnapSize == CurrGridSize);
+			}
+		);
 
-	return LocationGridMenuBuilder.MakeWidget();
+	return UE::UnrealEd::CreateLocationGridSnapMenu(ExecuteDelegate, IsCheckedDelegate, GridSizes, CommandList);
 }
 
 void SDesignerToolBar::SetGridSize(int32 InGridSize)
diff --git a/Engine/Source/Editor/UMGEditor/Private/Designer/SDesignerToolBar.h b/Engine/Source/Editor/UMGEditor/Private/Designer/SDesignerToolBar.h
--- a/Engine/Source/Editor/UMGEditor/Private/Designer/SDesignerToolBar.h
+++ b/Engine/Source/Editor/UMGEditor/Private/Designer/SDesignerToolBar.h
@@ -43,7 +43,6 @@ private:
 	void HandleToggleLocationGridSnap(ECheckBoxState InState);
 	FText GetLocationGridLabel() const;
 	TSharedRef<SWidget> FillLocationGridSnapMenu();
-	TSharedRef<SWidget> BuildLocationGridCheckBoxList(FName InExtentionHook, const FText& InHeading, const TArray<int32>& InGridSizes) const;
 	static void SetGridSize(int32 InGridSize);
 	static bool IsGridSizeChecked(int32 InGridSnapSize);
 	// End Grid Snapping
diff --git a/Engine/Source/Editor/UnrealEd/Private/STransformViewportToolbar.cpp b/Engine/Source/Editor/UnrealEd/Private/STransformViewportToolbar.cpp
--- a/Engine/Source/Editor/UnrealEd/Private/STransformViewportToolbar.cpp
+++ b/Engine/Source/Editor/UnrealEd/Private/STransformViewportToolbar.cpp
@@ -34,6 +34,7 @@
 #include "LevelEditorActions.h"
 #include "Styling/ToolBarStyle.h"
 #include "SEditorViewportToolBarMenu.h"
+#include "ViewportToolbar/UnrealEdViewportToolbar.h"
 
 #define LOCTEXT_NAMESPACE "TransformToolBar"
 
@@ -769,52 +770,40 @@ bool STransformViewportToolBar::IsPreserveNonUniformScaleChecked()
 TSharedRef<SWidget> STransformViewportToolBar::FillLocationGridSnapMenu()
 {
 	const ULevelEditorViewportSettings* ViewportSettings = GetDefault<ULevelEditorViewportSettings>();
+	TArray<float> GridSizes = ViewportSettings->bUsePowerOf2SnapSize ? ViewportSettings->Pow2GridSizes
+																	 : ViewportSettings->DecimalGridSizes;
 
-	return BuildLocationGridCheckBoxList("Snap", LOCTEXT("LocationSnapText", "Snap Sizes"), ViewportSettings->bUsePowerOf2SnapSize ? ViewportSettings->Pow2GridSizes : ViewportSettings->DecimalGridSizes );
-}
-
-TSharedRef<SWidget> STransformViewportToolBar::BuildLocationGridCheckBoxList(FName InExtentionHook, const FText& InHeading, const TArray<float>& InGridSizes) const
-{
-	const ULevelEditorViewportSettings* ViewportSettings = GetDefault<ULevelEditorViewportSettings>();
-
-	const bool bShouldCloseWindowAfterMenuSelection = true;
-	FMenuBuilder LocationGridMenuBuilder( bShouldCloseWindowAfterMenuSelection, CommandList );
-
-	LocationGridMenuBuilder.BeginSection(InExtentionHook, InHeading);
-	for( int32 CurGridSizeIndex = 0; CurGridSizeIndex < InGridSizes.Num(); ++CurGridSizeIndex )
-	{
-		const float CurGridSize = InGridSizes[ CurGridSizeIndex ];
+	UE::UnrealEd::FLocationGridCheckboxListExecuteActionDelegate ExecuteDelegate =
+		UE::UnrealEd::FLocationGridCheckboxListExecuteActionDelegate::CreateLambda(
+			[](int CurrGridSizeIndex)
+			{
+				STransformViewportToolBar::SetGridSize(CurrGridSizeIndex);
+			}
+		);
 
-		LocationGridMenuBuilder.AddMenuEntry(
-			FText::AsNumber( CurGridSize ),
-			FText::Format( LOCTEXT("LocationGridSize_ToolTip", "Sets grid size to {0}"), FText::AsNumber( CurGridSize ) ),
-			FSlateIcon(),
-			FUIAction( FExecuteAction::CreateStatic( &STransformViewportToolBar::SetGridSize, CurGridSizeIndex ),
-			FCanExecuteAction(),
-			FIsActionChecked::CreateStatic( &STransformViewportToolBar::IsGridSizeChecked, CurGridSizeIndex ) ),
-			NAME_None,
-			EUserInterfaceActionType::RadioButton );
-	}
-	LocationGridMenuBuilder.EndSection();
+	UE::UnrealEd::FLocationGridCheckboxListIsCheckedDelegate IsCheckedDelegate =
+		UE::UnrealEd::FLocationGridCheckboxListIsCheckedDelegate::CreateLambda(
+			[](int CurrGridSizeIndex)
+			{
+				return STransformViewportToolBar::IsGridSizeChecked(CurrGridSizeIndex);
+			}
+		);
 
-	return LocationGridMenuBuilder.MakeWidget();
+	return UE::UnrealEd::CreateLocationGridSnapMenu(ExecuteDelegate, IsCheckedDelegate, GridSizes, CommandList);
 }
 
 TSharedRef<SWidget> STransformViewportToolBar::FillRotationGridSnapMenu()
 {
-	const ULevelEditorViewportSettings* ViewportSettings = GetDefault<ULevelEditorViewportSettings>();
-
-	return SNew(SUniformGridPanel)
+	UE::UnrealEd::FRotationGridCheckboxListExecuteActionDelegate ExecuteDelegate =
+		UE::UnrealEd::FRotationGridCheckboxListExecuteActionDelegate::CreateStatic(&STransformViewportToolBar::SetRotationGridSize
+		);
 
-		+ SUniformGridPanel::Slot(0, 0)
-		[
-			BuildRotationGridCheckBoxList("Common", LOCTEXT("RotationCommonText", "Common"), ViewportSettings->CommonRotGridSizes, GridMode_Common)
-		]
+	UE::UnrealEd::FRotationGridCheckboxListIsCheckedDelegate IsCheckedDelegate =
+		UE::UnrealEd::FRotationGridCheckboxListIsCheckedDelegate::CreateStatic(
+			&STransformViewportToolBar::IsRotationGridSizeChecked
+		);
 
-	+ SUniformGridPanel::Slot(1, 0)
-		[
-			BuildRotationGridCheckBoxList("Div360", LOCTEXT("RotationDivisions360DegreesText", "Divisions of 360\u00b0"), ViewportSettings->DivisionsOf360RotGridSizes, GridMode_DivisionsOf360)
-		];
+	return UE::UnrealEd::CreateRotationGridSnapMenu(ExecuteDelegate, IsCheckedDelegate, CommandList);
 }
 
 TSharedRef<SWidget> STransformViewportToolBar::FillLayer2DSnapMenu()
@@ -869,95 +858,51 @@ TSharedRef<SWidget> STransformViewportToolBar::FillLayer2DSnapMenu()
 	return ShowMenuBuilder.MakeWidget();
 }
 
-TSharedRef<SWidget> STransformViewportToolBar::BuildRotationGridCheckBoxList(FName InExtentionHook, const FText& InHeading, const TArray<float>& InGridSizes, ERotationGridMode InGridMode) const
-{
-	const bool bShouldCloseWindowAfterMenuSelection = true;
-	FMenuBuilder RotationGridMenuBuilder( bShouldCloseWindowAfterMenuSelection, CommandList );
-
-	RotationGridMenuBuilder.BeginSection(InExtentionHook, InHeading);
-	for( int32 CurGridAngleIndex = 0; CurGridAngleIndex < InGridSizes.Num(); ++CurGridAngleIndex )
-	{
-		const float CurGridAngle = InGridSizes[ CurGridAngleIndex ];
-
-		FText MenuName = FText::Format( LOCTEXT("RotationGridAngle", "{0}\u00b0"), FText::AsNumber( CurGridAngle ) ); /*degree symbol*/
-		FText ToolTipText = FText::Format( LOCTEXT("RotationGridAngle_ToolTip", "Sets rotation grid angle to {0}"), MenuName ) ; /*degree symbol*/
-
-		RotationGridMenuBuilder.AddMenuEntry(
-			MenuName,
-			ToolTipText,
-			FSlateIcon(),
-			FUIAction( FExecuteAction::CreateStatic( &STransformViewportToolBar::SetRotationGridSize, CurGridAngleIndex, InGridMode ),
-			FCanExecuteAction(),
-			FIsActionChecked::CreateStatic( &STransformViewportToolBar::IsRotationGridSizeChecked, CurGridAngleIndex, InGridMode ) ),
-			NAME_None,
-			EUserInterfaceActionType::RadioButton );
-	}
-	RotationGridMenuBuilder.EndSection();
-
-	return RotationGridMenuBuilder.MakeWidget();
-}
-
 TSharedRef<SWidget> STransformViewportToolBar::FillScaleGridSnapMenu()
 {
 	const ULevelEditorViewportSettings* ViewportSettings = GetDefault<ULevelEditorViewportSettings>();
-	const bool bShouldCloseWindowAfterMenuSelection = true;
-
-	FNumberFormattingOptions NumberFormattingOptions;
-	NumberFormattingOptions.MaximumFractionalDigits = 5;
-
-	FMenuBuilder ScaleGridMenuBuilder( bShouldCloseWindowAfterMenuSelection, CommandList );
-
-	ScaleGridMenuBuilder.BeginSection("ScaleSnapOptions", LOCTEXT("ScaleSnapOptions", "Scale Snap"));
+	TArray<float> GridSizes = ViewportSettings->ScalingGridSizes;
 
-	for( int32 CurGridAmountIndex = 0; CurGridAmountIndex < ViewportSettings->ScalingGridSizes.Num(); ++CurGridAmountIndex )
-	{
-		const float CurGridAmount = ViewportSettings->ScalingGridSizes[ CurGridAmountIndex ];
-
-		FText MenuText;
-		FText ToolTipText;
+	UE::UnrealEd::FScaleGridCheckboxListExecuteActionDelegate ExecuteDelegate =
+		UE::UnrealEd::FScaleGridCheckboxListExecuteActionDelegate::CreateLambda(
+			[](int CurrGridScaleIndex)
+			{
+				STransformViewportToolBar::SetScaleGridSize(CurrGridScaleIndex);
+			}
+		);
 
-		if( GEditor->UsePercentageBasedScaling() )
-		{
-			MenuText = FText::AsPercent( CurGridAmount / 100.0f, &NumberFormattingOptions );
-			ToolTipText = FText::Format( LOCTEXT("ScaleGridAmountOld_ToolTip", "Snaps scale values to {0}"), MenuText );
-		}
-		else
+	UE::UnrealEd::FScaleGridCheckboxListIsCheckedDelegate IsCheckedDelegate =
+		UE::UnrealEd::FScaleGridCheckboxListIsCheckedDelegate::CreateLambda(
+			[](int CurrGridScaleIndex)
 		{
-			MenuText = FText::AsNumber( CurGridAmount, &NumberFormattingOptions );
-			ToolTipText = FText::Format( LOCTEXT("ScaleGridAmount_ToolTip", "Snaps scale values to increments of {0}"), MenuText );
+				return STransformViewportToolBar::IsScaleGridSizeChecked(CurrGridScaleIndex);
 		}
+		);
 
-		ScaleGridMenuBuilder.AddMenuEntry(
-			MenuText,
-			ToolTipText,
-			FSlateIcon(),
-			FUIAction( FExecuteAction::CreateStatic( &STransformViewportToolBar::SetScaleGridSize, CurGridAmountIndex ),
-			FCanExecuteAction(),
-			FIsActionChecked::CreateStatic( &STransformViewportToolBar::IsScaleGridSizeChecked, CurGridAmountIndex ) ),
-			NAME_None,
-			EUserInterfaceActionType::RadioButton );
-	}
-	ScaleGridMenuBuilder.EndSection();
-
-	if( !GEditor->UsePercentageBasedScaling() )
-	{
-		ScaleGridMenuBuilder.BeginSection("ScaleGeneralOptions", LOCTEXT("ScaleOptions", "Scaling Options"));
-
-		ScaleGridMenuBuilder.AddMenuEntry(
-			LOCTEXT("ScaleGridPreserveNonUniformScale", "Preserve Non-Uniform Scale"),
-			LOCTEXT("ScaleGridPreserveNonUniformScale_ToolTip", "When this option is checked, scaling objects that have a non-uniform scale will preserve the ratios between each axis, snapping the axis with the largest value."),
-			FSlateIcon(),
-			FUIAction( FExecuteAction::CreateStatic( &STransformViewportToolBar::TogglePreserveNonUniformScale ),
+	return UE::UnrealEd::CreateScaleGridSnapMenu(
+		ExecuteDelegate,
+		IsCheckedDelegate,
+		GridSizes,
+		true,
+		CommandList,
+		true,
+		FUIAction(
+			FExecuteAction::CreateLambda(
+				[]()
+				{
+					ULevelEditorViewportSettings* Settings = GetMutableDefault<ULevelEditorViewportSettings>();
+					Settings->PreserveNonUniformScale = !Settings->PreserveNonUniformScale;
+				}
+			),
 			FCanExecuteAction(),
-			FIsActionChecked::CreateStatic( &STransformViewportToolBar::IsPreserveNonUniformScaleChecked ) ),
-			NAME_None,
-			EUserInterfaceActionType::Check );
-
-		ScaleGridMenuBuilder.EndSection();
-	}
-
-
-	return ScaleGridMenuBuilder.MakeWidget();
+			FIsActionChecked::CreateLambda(
+				[]()
+				{
+					return GetDefault<ULevelEditorViewportSettings>()->PreserveNonUniformScale;
+				}
+			)
+		)
+	);
 }
 
 ECheckBoxState STransformViewportToolBar::IsLocationGridSnapChecked() const
diff --git a/Engine/Source/Editor/UnrealEd/Private/ViewportToolbar/UnrealEdViewportToolbar.cpp b/Engine/Source/Editor/UnrealEd/Private/ViewportToolbar/UnrealEdViewportToolbar.cpp
--- a/Engine/Source/Editor/UnrealEd/Private/ViewportToolbar/UnrealEdViewportToolbar.cpp
+++ b/Engine/Source/Editor/UnrealEd/Private/ViewportToolbar/UnrealEdViewportToolbar.cpp
@@ -18,13 +18,494 @@
 #include "ToolMenuEntry.h"
 #include "ToolMenuSection.h"
 #include "ViewportToolbar/UnrealEdViewportToolbarContext.h"
+#include "Widgets/Input/SNumericEntryBox.h"
 #include "Widgets/Input/SSpinBox.h"
+#include "Widgets/Layout/SUniformGridPanel.h"
+#include "Widgets/SToolTip.h"
 
 #define LOCTEXT_NAMESPACE "UnrealEdViewportToolbar"
 
 namespace UE::UnrealEd::Private
 {
 int32 CVarToolMenusViewportToolbarsValue = 0;
+
+FToolUIAction DisabledAction()
+{
+	static FToolUIAction Action;
+	Action.CanExecuteAction = FToolMenuCanExecuteAction::CreateLambda(
+		[](const FToolMenuContext&)
+		{
+			return false;
+		}
+	);
+
+	return Action;
+}
+
+// TODO: Maybe export CreateSurfaceSnapOffsetEntry function, so that it can be used elsewhere, e.g. STransformViewportToolbar.cpp
+FToolMenuEntry CreateSurfaceSnapOffsetEntry()
+{
+	FText Label = LOCTEXT("SurfaceOffsetLabel", "Surface Offset");
+	FText Tooltip = LOCTEXT("SurfaceOffsetTooltip", "The amount of offset to apply when snapping to surfaces");
+
+	const FMargin WidgetsMargin(2.0f, 0.0f, 3.0f, 0.0f);
+
+	FToolMenuEntry SurfaceOffset = FToolMenuEntry::InitMenuEntry(
+		"SurfaceOffset",
+		FUIAction(
+			FExecuteAction(),
+			FCanExecuteAction::CreateLambda(
+				[]()
+				{
+					return GetDefault<ULevelEditorViewportSettings>()->SnapToSurface.bEnabled;
+				}
+			)
+		),
+		// clang-format off
+		SNew(SHorizontalBox)
+		+ SHorizontalBox::Slot()
+		.VAlign(VAlign_Center)
+		.HAlign(HAlign_Left)
+		.Padding(WidgetsMargin)
+		.AutoWidth()
+		[
+			SNew(STextBlock)
+			.Text(Label)
+		]
+		+ SHorizontalBox::Slot()
+		.VAlign(VAlign_Center)
+		.Padding(WidgetsMargin)
+		.AutoWidth()
+		[
+			SNew(SBox)
+			.Padding(WidgetsMargin)
+			.MinDesiredWidth(100.0f)
+			[
+				// Min/Max/Slider values taken from STransformViewportToolbar.cpp
+				SNew(SNumericEntryBox<float>)
+				.ToolTipText(Tooltip)
+				.MinValue(0.0f)
+				.MaxValue(static_cast<float>(HALF_WORLD_MAX))
+				.MaxSliderValue(1000.0f)
+				.AllowSpin(true)
+				.MaxFractionalDigits(2)
+				.Font(FAppStyle::GetFontStyle(TEXT("MenuItem.Font")))
+				.OnValueChanged_Lambda([](float InNewValue)
+				{
+					ULevelEditorViewportSettings* Settings =
+						GetMutableDefault<ULevelEditorViewportSettings>();
+					Settings->SnapToSurface.SnapOffsetExtent = InNewValue;
+				})
+				.Value_Lambda([]()
+				{
+					return GetDefault<ULevelEditorViewportSettings>()->SnapToSurface.SnapOffsetExtent;
+				})
+			]
+		]
+		// clang-format on
+	);
+
+	return SurfaceOffset;
+}
+
+FToolMenuEntry CreateSurfaceSnapCheckboxMenu()
+{
+	FNewToolMenuDelegate MakeMenuDelegate = FNewToolMenuDelegate::CreateLambda(
+		[](UToolMenu* Submenu)
+		{
+			FToolMenuSection& SurfaceSnappingSection =
+				Submenu->FindOrAddSection("SurfaceSnapping", LOCTEXT("SurfaceSnappingLabel", "Surface Snapping"));
+
+			// Add "Rotate to surface normal" checkbox.
+			{
+				FToolMenuEntry RotateToSurfaceNormalSnapping =
+					FToolMenuEntry::InitMenuEntry(FEditorViewportCommands::Get().RotateToSurfaceNormal);
+				SurfaceSnappingSection.AddEntry(RotateToSurfaceNormalSnapping);
+			}
+
+			// Add "Surface offset" widget.
+			{
+				SurfaceSnappingSection.AddEntry(CreateSurfaceSnapOffsetEntry());
+			}
+		}
+	);
+
+	return UnrealEd::CreateCheckboxSubmenu(
+		"SurfaceSnapping",
+		LOCTEXT("SurfaceSnapLabel", "Surface"),
+		FEditorViewportCommands::Get().SurfaceSnapping->MakeTooltip()->GetTextTooltip(),
+		FToolMenuExecuteAction::CreateLambda(
+			[](const FToolMenuContext& InContext)
+			{
+				ULevelEditorViewportSettings* Settings = GetMutableDefault<ULevelEditorViewportSettings>();
+				Settings->SnapToSurface.bEnabled = !Settings->SnapToSurface.bEnabled;
+			}
+		),
+		FToolMenuCanExecuteAction(),
+		FToolMenuGetActionCheckState::CreateLambda(
+			[](const FToolMenuContext& InContext)
+			{
+				return GetDefault<ULevelEditorViewportSettings>()->SnapToSurface.bEnabled ? ECheckBoxState::Checked
+																						  : ECheckBoxState::Unchecked;
+			}
+		),
+		MakeMenuDelegate
+	);
+}
+
+FToolMenuEntry CreateActorSnapDistanceEntry()
+{
+	const FText Label = LOCTEXT("ActorSnapDistanceLabel", "Snap Distance");
+	const FText Tooltip = LOCTEXT("ActorSnapDistanceTooltip", "The amount of offset to apply when snapping to surfaces");
+
+	const FMargin WidgetsMargin(2.0f, 0.0f, 3.0f, 0.0f);
+
+	FToolMenuEntry SnapDistance = FToolMenuEntry::InitMenuEntry(
+		"ActorSnapDistance",
+		FUIAction(
+			FExecuteAction(),
+			FCanExecuteAction::CreateLambda(
+				[]()
+				{
+					return !!GetDefault<ULevelEditorViewportSettings>()->bEnableActorSnap;
+				}
+			)
+		),
+		// clang-format off
+		SNew(SHorizontalBox)
+		+ SHorizontalBox::Slot()
+		.VAlign(VAlign_Center)
+		.Padding(WidgetsMargin)
+		.AutoWidth()
+		[
+			SNew(STextBlock)
+			.Text(Label)
+		]
+		+ SHorizontalBox::Slot()
+		.VAlign(VAlign_Center)
+		.Padding(WidgetsMargin)
+		.AutoWidth()
+		[
+			SNew(SBox).Padding(WidgetsMargin).MinDesiredWidth(100.0f)
+			[
+				// TODO: Check how to improve performance for this widget OnValueChanged.
+				// Same functionality in LevelEditorToolBar.cpp seems to have better performance
+				SNew(SNumericEntryBox<float>)
+				.ToolTipText(Tooltip)
+				.MinValue(0.0f)
+				.MaxValue(1.0f)
+				.MaxSliderValue(1.0f)
+				.AllowSpin(true)
+				.MaxFractionalDigits(1)
+				.Font(FAppStyle::GetFontStyle(TEXT("MenuItem.Font")))
+				.OnValueChanged_Static(&FLevelEditorActionCallbacks::SetActorSnapSetting)
+				.Value_Lambda([]()
+				{
+					return FLevelEditorActionCallbacks::GetActorSnapSetting();
+				})
+			]
+		]
+		// clang-format on
+	);
+
+	return SnapDistance;
+}
+
+FToolMenuEntry CreateActorSnapCheckboxMenu()
+{
+	FNewToolMenuDelegate MakeMenuDelegate = FNewToolMenuDelegate::CreateLambda(
+		[](UToolMenu* Submenu)
+		{
+			// Add "Actor snapping" widget.
+			{
+				FToolMenuSection& ActorSnappingSection =
+					Submenu->FindOrAddSection("ActorSnapping", LOCTEXT("ActorSnappingLabel", "Actor Snapping"));
+				ActorSnappingSection.AddEntry(CreateActorSnapDistanceEntry());
+			}
+		}
+	);
+
+	FToolUIAction CheckboxMenuAction;
+	{
+		CheckboxMenuAction.ExecuteAction = FToolMenuExecuteAction::CreateLambda(
+			[](const FToolMenuContext& InContext)
+			{
+				if (ULevelEditorViewportSettings* Settings = GetMutableDefault<ULevelEditorViewportSettings>())
+				{
+					Settings->bEnableActorSnap = !Settings->bEnableActorSnap;
+				}
+			}
+		);
+		CheckboxMenuAction.GetActionCheckState = FToolMenuGetActionCheckState::CreateLambda(
+			[](const FToolMenuContext& InContext)
+			{
+				return GetDefault<ULevelEditorViewportSettings>()->bEnableActorSnap ? ECheckBoxState::Checked
+																					: ECheckBoxState::Unchecked;
+			}
+		);
+	}
+
+	FToolMenuEntry ActorSnapping = FToolMenuEntry::InitSubMenu(
+		"ActorSnapping",
+		LOCTEXT("ActorSnapLabel", "Actor"),
+		FLevelEditorCommands::Get().EnableActorSnap->MakeTooltip()->GetTextTooltip(),
+		MakeMenuDelegate,
+		CheckboxMenuAction,
+		EUserInterfaceActionType::ToggleButton
+	);
+
+	return ActorSnapping;
+}
+
+FToolMenuEntry CreateLocationSnapCheckboxMenu()
+{
+	const FName LocationSnapName = "LocationSnap";
+	const FText LocationSnapLabel = LOCTEXT("LocationSnapLabel", "Location");
+
+	if (!GEditor)
+	{
+		return FToolMenuEntry::InitMenuEntry(
+			LocationSnapName, LocationSnapLabel, FText(), FSlateIcon(), UnrealEd::Private::DisabledAction()
+		);
+	}
+
+	FNewToolMenuDelegate MakeMenuDelegate = FNewToolMenuDelegate::CreateLambda(
+		[LocationSnapName](UToolMenu* InToolMenu)
+		{
+			UnrealEd::FLocationGridCheckboxListExecuteActionDelegate ExecuteDelegate =
+				UnrealEd::FLocationGridCheckboxListExecuteActionDelegate::CreateUObject(GEditor, &UEditorEngine::SetGridSize);
+
+			UnrealEd::FLocationGridCheckboxListIsCheckedDelegate IsCheckedDelegate =
+				UnrealEd::FLocationGridCheckboxListIsCheckedDelegate::CreateLambda(
+					[](int CurrGridSizeIndex)
+					{
+						const ULevelEditorViewportSettings* ViewportSettings = GetDefault<ULevelEditorViewportSettings>();
+						return ViewportSettings->CurrentPosGridSize == CurrGridSizeIndex;
+					}
+				);
+
+			const ULevelEditorViewportSettings* ViewportSettings = GetDefault<ULevelEditorViewportSettings>();
+			TArray<float> GridSizes = ViewportSettings->bUsePowerOf2SnapSize ? ViewportSettings->Pow2GridSizes
+																			 : ViewportSettings->DecimalGridSizes;
+
+			InToolMenu->AddMenuEntry(
+				LocationSnapName,
+				FToolMenuEntry::InitWidget(
+					LocationSnapName,
+					UnrealEd::CreateLocationGridSnapMenu(
+						ExecuteDelegate,
+						IsCheckedDelegate,
+						GridSizes,
+						TAttribute<bool>::CreateLambda(
+							[]()
+							{
+								return FLevelEditorActionCallbacks::LocationGridSnap_IsChecked();
+							}
+						)
+					),
+					FText()
+				)
+			);
+		}
+	);
+
+	return UnrealEd::CreateCheckboxSubmenu(
+		"GridSnapping",
+		LocationSnapLabel,
+		FEditorViewportCommands::Get().SurfaceSnapping->MakeTooltip()->GetTextTooltip(),
+		FToolMenuExecuteAction::CreateLambda(
+			[](const FToolMenuContext& InContext)
+			{
+				FLevelEditorActionCallbacks::LocationGridSnap_Clicked();
+			}
+		),
+		FToolMenuCanExecuteAction(),
+		FToolMenuGetActionCheckState::CreateLambda(
+			[](const FToolMenuContext& InContext)
+			{
+				return FLevelEditorActionCallbacks::LocationGridSnap_IsChecked() ? ECheckBoxState::Checked
+																				 : ECheckBoxState::Unchecked;
+			}
+		),
+		MakeMenuDelegate
+	);
+}
+
+FToolMenuEntry CreateRotationSnapCheckboxMenu()
+{
+	const FName RotationSnapName = "RotationSnap";
+	const FText RotationSnapLabel = LOCTEXT("RotationSnapLabel", "Rotation");
+
+	if (!GEditor)
+	{
+		return FToolMenuEntry::InitMenuEntry(
+			RotationSnapName, RotationSnapLabel, FText(), FSlateIcon(), UnrealEd::Private::DisabledAction()
+		);
+	}
+
+	FNewToolMenuDelegate MakeMenuDelegate = FNewToolMenuDelegate::CreateLambda(
+		[RotationSnapName](UToolMenu* InToolMenu)
+		{
+			const ULevelEditorViewportSettings* ViewportSettings = GetDefault<ULevelEditorViewportSettings>();
+			TArray<float> GridSizes = ViewportSettings->bUsePowerOf2SnapSize ? ViewportSettings->Pow2GridSizes
+																			 : ViewportSettings->DecimalGridSizes;
+
+			UnrealEd::FRotationGridCheckboxListExecuteActionDelegate ExecuteDelegate =
+				UnrealEd::FRotationGridCheckboxListExecuteActionDelegate::CreateUObject(
+					GEditor, &UEditorEngine::SetRotGridSize
+				);
+
+			UnrealEd::FRotationGridCheckboxListIsCheckedDelegate IsCheckedDelegate =
+				UnrealEd::FRotationGridCheckboxListIsCheckedDelegate::CreateLambda(
+					[](int CurrGridAngleIndex, ERotationGridMode InGridMode)
+					{
+						const ULevelEditorViewportSettings* ViewportSettings = GetDefault<ULevelEditorViewportSettings>();
+						return ViewportSettings->CurrentRotGridSize == CurrGridAngleIndex
+							&& ViewportSettings->CurrentRotGridMode == InGridMode;
+					}
+				);
+
+			InToolMenu->AddMenuEntry(
+				RotationSnapName,
+				FToolMenuEntry::InitWidget(
+					RotationSnapName,
+					UnrealEd::CreateRotationGridSnapMenu(
+						ExecuteDelegate,
+						IsCheckedDelegate,
+						TAttribute<bool>::CreateLambda(
+							[]()
+							{
+								return FLevelEditorActionCallbacks::RotationGridSnap_IsChecked();
+							}
+						)
+					),
+					FText()
+				)
+			);
+		}
+	);
+
+	return UnrealEd::CreateCheckboxSubmenu(
+		"RotationSnapping",
+		RotationSnapLabel,
+		FEditorViewportCommands::Get().SurfaceSnapping->MakeTooltip()->GetTextTooltip(),
+		FToolMenuExecuteAction::CreateLambda(
+			[](const FToolMenuContext& InContext)
+			{
+				FLevelEditorActionCallbacks::RotationGridSnap_Clicked();
+			}
+		),
+		FToolMenuCanExecuteAction(),
+		FToolMenuGetActionCheckState::CreateLambda(
+			[](const FToolMenuContext& InContext)
+			{
+				return FLevelEditorActionCallbacks::RotationGridSnap_IsChecked() ? ECheckBoxState::Checked
+																				 : ECheckBoxState::Unchecked;
+			}
+		),
+		MakeMenuDelegate
+	);
+}
+
+FToolMenuEntry CreateScaleSnapCheckboxMenu()
+{
+	const FName ScaleSnapName = "ScaleSnap";
+	const FText ScaleSnapLabel = LOCTEXT("ScaleSnapLabel", "Scale");
+
+	if (!GEditor)
+	{
+		return FToolMenuEntry::InitMenuEntry(
+			ScaleSnapName, ScaleSnapLabel, FText(), FSlateIcon(), UnrealEd::Private::DisabledAction()
+		);
+	}
+
+	FNewToolMenuDelegate MakeMenuDelegate = FNewToolMenuDelegate::CreateLambda(
+		[ScaleSnapName](UToolMenu* InToolMenu)
+		{
+			FCanExecuteAction CanExecuteScaleSnapping = FCanExecuteAction::CreateLambda(
+				[]()
+				{
+					return FLevelEditorActionCallbacks::ScaleGridSnap_IsChecked();
+				}
+			);
+
+			const ULevelEditorViewportSettings* ViewportSettings = GetDefault<ULevelEditorViewportSettings>();
+			TArray<float> GridSizes = ViewportSettings->ScalingGridSizes;
+
+			UnrealEd::FScaleGridCheckboxListExecuteActionDelegate ExecuteDelegate =
+				UnrealEd::FScaleGridCheckboxListExecuteActionDelegate::CreateUObject(GEditor, &UEditorEngine::SetScaleGridSize);
+
+			UnrealEd::FScaleGridCheckboxListIsCheckedDelegate IsCheckedDelegate =
+				UnrealEd::FScaleGridCheckboxListIsCheckedDelegate::CreateLambda(
+					[](int CurrGridSizeIndex)
+					{
+						const ULevelEditorViewportSettings* ViewportSettings = GetDefault<ULevelEditorViewportSettings>();
+						return ViewportSettings->CurrentScalingGridSize == CurrGridSizeIndex;
+					}
+				);
+
+			InToolMenu->AddMenuEntry(
+				ScaleSnapName,
+				FToolMenuEntry::InitWidget(
+					ScaleSnapName,
+					UnrealEd::CreateScaleGridSnapMenu(
+						ExecuteDelegate,
+						IsCheckedDelegate,
+						GridSizes,
+						TAttribute<bool>::CreateLambda(
+							[]()
+							{
+								return FLevelEditorActionCallbacks::ScaleGridSnap_IsChecked();
+							}
+						),
+						{},
+						true,
+						FUIAction(
+							FExecuteAction::CreateLambda(
+								[]()
+								{
+									ULevelEditorViewportSettings* Settings =
+										GetMutableDefault<ULevelEditorViewportSettings>();
+									Settings->PreserveNonUniformScale = !Settings->PreserveNonUniformScale;
+								}
+							),
+							CanExecuteScaleSnapping,
+							FIsActionChecked::CreateLambda(
+								[]()
+								{
+									return GetDefault<ULevelEditorViewportSettings>()->PreserveNonUniformScale;
+								}
+							)
+						)
+					),
+					FText()
+				)
+			);
+		}
+	);
+
+	return UnrealEd::CreateCheckboxSubmenu(
+		"ScaleSnapping",
+		ScaleSnapLabel,
+		FEditorViewportCommands::Get().SurfaceSnapping->MakeTooltip()->GetTextTooltip(),
+		FToolMenuExecuteAction::CreateLambda(
+			[](const FToolMenuContext& InContext)
+			{
+				FLevelEditorActionCallbacks::ScaleGridSnap_Clicked();
+			}
+		),
+		FToolMenuCanExecuteAction(),
+		FToolMenuGetActionCheckState::CreateLambda(
+			[](const FToolMenuContext& InContext)
+			{
+				return FLevelEditorActionCallbacks::ScaleGridSnap_IsChecked() ? ECheckBoxState::Checked
+																			  : ECheckBoxState::Unchecked;
+			}
+		),
+		MakeMenuDelegate
+	);
+}
+
 }
 
 static FAutoConsoleVariableRef CVarToolMenusViewportToolbars(
@@ -208,6 +689,41 @@ FToolMenuEntry CreateViewportToolbarSelectionSection()
 	);
 }
 
+FToolMenuEntry CreateViewportToolbarSnappingSubmenu()
+{
+	return FToolMenuEntry::InitSubMenu(
+		"Snapping",
+		LOCTEXT("SnappingSubmenuLabel", "Snapping"),
+		LOCTEXT("SnappingSubmenuTooltip", "Viewport-related snapping settings"),
+		FNewToolMenuDelegate::CreateLambda(
+			[](UToolMenu* Submenu) -> void
+			{
+				FToolMenuSection& SnappingSection =
+					Submenu->FindOrAddSection("Snapping", LOCTEXT("SnappingLabel", "Snapping"));
+
+				SnappingSection.AddEntry(Private::CreateSurfaceSnapCheckboxMenu());
+				SnappingSection.AddEntry(Private::CreateLocationSnapCheckboxMenu()).SetShowInToolbarTopLevel(true);
+				SnappingSection.AddEntry(Private::CreateRotationSnapCheckboxMenu());
+				SnappingSection.AddEntry(Private::CreateScaleSnapCheckboxMenu());
+				SnappingSection.AddEntry(Private::CreateActorSnapCheckboxMenu());
+
+				FToolMenuEntry SocketSnapping =
+					FToolMenuEntry::InitMenuEntry(FLevelEditorCommands::Get().ToggleSocketSnapping);
+				SocketSnapping.UserInterfaceActionType = EUserInterfaceActionType::ToggleButton;
+				SocketSnapping.Label = LOCTEXT("SocketSnapLabel", "Socket");
+				SnappingSection.AddEntry(SocketSnapping);
+
+				FToolMenuEntry VertexSnapping = FToolMenuEntry::InitMenuEntry(FLevelEditorCommands::Get().EnableVertexSnap);
+				VertexSnapping.UserInterfaceActionType = EUserInterfaceActionType::ToggleButton;
+				VertexSnapping.Label = LOCTEXT("VertexSnapLabel", "Vertex");
+				SnappingSection.AddEntry(VertexSnapping);
+
+				// TODO: add Planar Snapping
+			}
+		)
+	);
+}
+
 FText GetViewModesSubmenuLabel(TWeakPtr<SEditorViewport> InViewport)
 {
 	FText Label = LOCTEXT("ViewMenuTitle_Default", "View");
@@ -555,6 +1071,254 @@ FToolMenuEntry CreateViewportToolbarViewModesSubmenu()
 	);
 }
 
+TSharedRef<SWidget> BuildRotationGridCheckBoxList(
+	FName InExtentionHook,
+	const FText& InHeading,
+	const TArray<float>& InGridSizes,
+	ERotationGridMode InGridMode,
+	const FRotationGridCheckboxListExecuteActionDelegate& InExecuteAction,
+	const FRotationGridCheckboxListIsCheckedDelegate& InIsActionChecked,
+	const TSharedPtr<FUICommandList>& InCommandList
+)
+{
+	constexpr bool bShouldCloseWindowAfterMenuSelection = true;
+	FMenuBuilder RotationGridMenuBuilder(bShouldCloseWindowAfterMenuSelection, InCommandList);
+
+	RotationGridMenuBuilder.BeginSection(InExtentionHook, InHeading);
+	for (int32 CurrGridAngleIndex = 0; CurrGridAngleIndex < InGridSizes.Num(); ++CurrGridAngleIndex)
+	{
+		const float CurrGridAngle = InGridSizes[CurrGridAngleIndex];
+
+		FText MenuName =
+			FText::Format(LOCTEXT("RotationGridAngle", "{0}\u00b0"), FText::AsNumber(CurrGridAngle)); /*degree symbol*/
+		FText ToolTipText = FText::Format(
+			LOCTEXT("RotationGridAngle_ToolTip", "Sets rotation grid angle to {0}"), MenuName
+		); /*degree symbol*/
+
+		RotationGridMenuBuilder.AddMenuEntry(
+			MenuName,
+			ToolTipText,
+			FSlateIcon(),
+			FUIAction(
+				FExecuteAction::CreateLambda(
+					[CurrGridAngleIndex, InGridMode, InExecuteAction]()
+					{
+						InExecuteAction.Execute(CurrGridAngleIndex, InGridMode);
+					}
+				),
+				FCanExecuteAction(),
+				FIsActionChecked::CreateLambda(
+					[CurrGridAngleIndex, InGridMode, InIsActionChecked]()
+					{
+						return InIsActionChecked.Execute(CurrGridAngleIndex, InGridMode);
+					}
+				)
+			),
+			NAME_None,
+			EUserInterfaceActionType::RadioButton
+		);
+	}
+	RotationGridMenuBuilder.EndSection();
+
+	return RotationGridMenuBuilder.MakeWidget();
+}
+
+TSharedRef<SWidget> CreateRotationGridSnapMenu(
+	const FRotationGridCheckboxListExecuteActionDelegate& InExecuteDelegate,
+	const FRotationGridCheckboxListIsCheckedDelegate& InIsCheckedDelegate,
+	const TAttribute<bool>& InIsEnabledDelegate,
+	const TSharedPtr<FUICommandList>& InCommandList
+)
+{
+	const ULevelEditorViewportSettings* ViewportSettings = GetDefault<ULevelEditorViewportSettings>();
+	TArray<float> GridSizes = ViewportSettings->bUsePowerOf2SnapSize ? ViewportSettings->Pow2GridSizes
+																	 : ViewportSettings->DecimalGridSizes;
+
+	// clang-format off
+	return SNew(SUniformGridPanel)
+		.IsEnabled(InIsEnabledDelegate)
+		+ SUniformGridPanel::Slot(0, 0)
+		[
+			UnrealEd::BuildRotationGridCheckBoxList("Common",LOCTEXT("RotationCommonText", "Rotation Increment")
+				, ViewportSettings->CommonRotGridSizes, GridMode_Common,
+				InExecuteDelegate,
+				InIsCheckedDelegate,
+				InCommandList
+			)
+		]
+		+ SUniformGridPanel::Slot(1, 0)
+		[
+			UnrealEd::BuildRotationGridCheckBoxList("Div360",LOCTEXT("RotationDivisions360DegreesText", "Divisions of 360\u00b0")
+				, ViewportSettings->DivisionsOf360RotGridSizes, GridMode_DivisionsOf360,
+				InExecuteDelegate,
+				InIsCheckedDelegate,
+				InCommandList
+			)
+		];
+	// clang-format on
+}
+
+TSharedRef<SWidget> CreateLocationGridSnapMenu(
+	const FLocationGridCheckboxListExecuteActionDelegate& InExecuteDelegate,
+	const FLocationGridCheckboxListIsCheckedDelegate& InIsCheckedDelegate,
+	const TArray<float>& InGridSizes,
+	const TAttribute<bool>& InIsEnabledDelegate,
+	const TSharedPtr<FUICommandList>& InCommandList
+)
+{
+	constexpr bool bShouldCloseWindowAfterMenuSelection = true;
+	FMenuBuilder LocationGridMenuBuilder(bShouldCloseWindowAfterMenuSelection, InCommandList);
+
+	LocationGridMenuBuilder.BeginSection("Snap", LOCTEXT("LocationSnapText", "Snap Sizes"));
+	for (int32 CurrGridSizeIndex = 0; CurrGridSizeIndex < InGridSizes.Num(); ++CurrGridSizeIndex)
+	{
+		const float CurGridSize = InGridSizes[CurrGridSizeIndex];
+
+		LocationGridMenuBuilder.AddMenuEntry(
+			FText::AsNumber(CurGridSize),
+			FText::Format(LOCTEXT("LocationGridSize_ToolTip", "Sets grid size to {0}"), FText::AsNumber(CurGridSize)),
+			FSlateIcon(),
+			FUIAction(
+				FExecuteAction::CreateLambda(
+					[CurrGridSizeIndex, InExecuteDelegate]()
+					{
+						InExecuteDelegate.Execute(CurrGridSizeIndex);
+					}
+				),
+				FCanExecuteAction::CreateLambda(
+					[InIsEnabledDelegate]()
+					{
+						return InIsEnabledDelegate.Get();
+					}
+				),
+				FIsActionChecked::CreateLambda(
+					[CurrGridSizeIndex, InIsCheckedDelegate]()
+					{
+						return InIsCheckedDelegate.Execute(CurrGridSizeIndex);
+					}
+				)
+			),
+			NAME_None,
+			EUserInterfaceActionType::RadioButton
+		);
+	}
+	LocationGridMenuBuilder.EndSection();
+
+	return LocationGridMenuBuilder.MakeWidget();
+}
+
+TSharedRef<SWidget> CreateScaleGridSnapMenu(
+	const FScaleGridCheckboxListExecuteActionDelegate& InExecuteDelegate,
+	const FScaleGridCheckboxListIsCheckedDelegate& InIsCheckedDelegate,
+	const TArray<float>& InGridSizes,
+	const TAttribute<bool>& InIsEnabledDelegate,
+	const TSharedPtr<FUICommandList>& InCommandList,
+	const TAttribute<bool>& ShowPreserveNonUniformScaleOption,
+	const FUIAction& PreserveNonUniformScaleUIAction
+)
+{
+	FNumberFormattingOptions NumberFormattingOptions;
+	NumberFormattingOptions.MaximumFractionalDigits = 5;
+
+	constexpr bool bShouldCloseWindowAfterMenuSelection = true;
+	FMenuBuilder ScaleGridMenuBuilder(bShouldCloseWindowAfterMenuSelection, InCommandList);
+
+	ScaleGridMenuBuilder.BeginSection("ScaleSnapOptions", LOCTEXT("ScaleSnapOptions", "Scale Snap"));
+
+	for (int32 CurrGridAmountIndex = 0; CurrGridAmountIndex < InGridSizes.Num(); ++CurrGridAmountIndex)
+	{
+		const float CurGridAmount = InGridSizes[CurrGridAmountIndex];
+
+		FText MenuText;
+		FText ToolTipText;
+
+		if (GEditor->UsePercentageBasedScaling())
+		{
+			MenuText = FText::AsPercent(CurGridAmount / 100.0f, &NumberFormattingOptions);
+			ToolTipText = FText::Format(LOCTEXT("ScaleGridAmountOld_ToolTip", "Snaps scale values to {0}"), MenuText);
+		}
+		else
+		{
+			MenuText = FText::AsNumber(CurGridAmount, &NumberFormattingOptions);
+			ToolTipText =
+				FText::Format(LOCTEXT("ScaleGridAmount_ToolTip", "Snaps scale values to increments of {0}"), MenuText);
+		}
+
+		ScaleGridMenuBuilder.AddMenuEntry(
+			MenuText,
+			ToolTipText,
+			FSlateIcon(),
+			FUIAction(
+				FExecuteAction::CreateLambda(
+					[CurrGridAmountIndex, InExecuteDelegate]()
+					{
+						InExecuteDelegate.Execute(CurrGridAmountIndex);
+					}
+				),
+				FCanExecuteAction::CreateLambda(
+					[InIsEnabledDelegate]()
+					{
+						return InIsEnabledDelegate.Get();
+					}
+				),
+				FIsActionChecked::CreateLambda(
+					[CurrGridAmountIndex, InIsCheckedDelegate]()
+					{
+						return InIsCheckedDelegate.Execute(CurrGridAmountIndex);
+					}
+				)
+			),
+			NAME_None,
+			EUserInterfaceActionType::RadioButton
+		);
+	}
+	ScaleGridMenuBuilder.EndSection();
+
+	if (!GEditor->UsePercentageBasedScaling() && ShowPreserveNonUniformScaleOption.Get())
+	{
+		ScaleGridMenuBuilder.BeginSection("ScaleGeneralOptions", LOCTEXT("ScaleOptions", "Scaling Options"));
+
+		ScaleGridMenuBuilder.AddMenuEntry(
+			LOCTEXT("ScaleGridPreserveNonUniformScale", "Preserve Non-Uniform Scale"),
+			LOCTEXT(
+				"ScaleGridPreserveNonUniformScale_ToolTip", "When this option is checked, scaling objects that have a non-uniform scale will preserve the ratios between each axis, snapping the axis with the largest value."
+			),
+			FSlateIcon(),
+			PreserveNonUniformScaleUIAction,
+			NAME_None,
+			EUserInterfaceActionType::ToggleButton
+		);
+
+		ScaleGridMenuBuilder.EndSection();
+	}
+
+	return ScaleGridMenuBuilder.MakeWidget();
+}
+
+FToolMenuEntry CreateCheckboxSubmenu(
+	const FName InName,
+	const TAttribute<FText>& InLabel,
+	const TAttribute<FText>& InToolTip,
+	const FToolMenuExecuteAction& InCheckboxExecuteAction,
+	const FToolMenuCanExecuteAction& InCheckboxCanExecuteAction,
+	const FToolMenuGetActionCheckState& InCheckboxActionCheckState,
+	const FNewToolMenuChoice& InMakeMenu
+)
+{
+	FToolUIAction CheckboxMenuAction;
+	{
+		CheckboxMenuAction.ExecuteAction = InCheckboxExecuteAction;
+		CheckboxMenuAction.CanExecuteAction = InCheckboxCanExecuteAction;
+		CheckboxMenuAction.GetActionCheckState = InCheckboxActionCheckState;
+	}
+
+	FToolMenuEntry CheckBoxSubmenu = FToolMenuEntry::InitSubMenu(
+		InName, InLabel, InToolTip, InMakeMenu, CheckboxMenuAction, EUserInterfaceActionType::ToggleButton
+	);
+
+	return CheckBoxSubmenu;
+}
+
 } // namespace UE::UnrealEd
 
 #undef LOCTEXT_NAMESPACE
diff --git a/Engine/Source/Editor/UnrealEd/Public/STransformViewportToolbar.h b/Engine/Source/Editor/UnrealEd/Public/STransformViewportToolbar.h
--- a/Engine/Source/Editor/UnrealEd/Public/STransformViewportToolbar.h
+++ b/Engine/Source/Editor/UnrealEd/Public/STransformViewportToolbar.h
@@ -113,10 +113,6 @@ private:
 	/** Callbacks for preserving non-uniform scaling when snapping */
 	static void TogglePreserveNonUniformScale();
 	static bool IsPreserveNonUniformScaleChecked();
-	
-	/** Methods to build more complex duel lists */
-	TSharedRef<SWidget> BuildLocationGridCheckBoxList(FName InExtentionHook, const FText& InHeading, const TArray<float>& InGridSizes) const;
-	TSharedRef<SWidget> BuildRotationGridCheckBoxList(FName InExtentionHook, const FText& InHeading, const TArray<float>& InGridSizes, ERotationGridMode InGridMode) const;
 
 	/** Make the surface snapping toolbar checkbox button */
 	TSharedRef< SWidget > MakeSurfaceSnappingButton();
diff --git a/Engine/Source/Editor/UnrealEd/Public/ViewportToolbar/UnrealEdViewportToolbar.h b/Engine/Source/Editor/UnrealEd/Public/ViewportToolbar/UnrealEdViewportToolbar.h
--- a/Engine/Source/Editor/UnrealEd/Public/ViewportToolbar/UnrealEdViewportToolbar.h
+++ b/Engine/Source/Editor/UnrealEd/Public/ViewportToolbar/UnrealEdViewportToolbar.h
@@ -5,12 +5,15 @@
 #include "Delegates/DelegateCombinations.h"
 #include "Engine/EngineBaseTypes.h"
 #include "Templates/SharedPointerFwd.h"
+#include "ToolMenuDelegates.h"
 #include "UnrealEdViewportToolbarContext.h"
 
 class FText;
 class SEditorViewport;
 class UToolMenu;
+struct FNewToolMenuChoice;
 struct FToolMenuEntry;
+enum ERotationGridMode : int;
 
 namespace UE::UnrealEd
 {
@@ -25,6 +28,8 @@ UNREALED_API FToolMenuEntry CreateViewportToolbarTransformsSection();
 
 UNREALED_API FToolMenuEntry CreateViewportToolbarSelectionSection();
 
+UNREALED_API FToolMenuEntry CreateViewportToolbarSnappingSubmenu();
+
 UNREALED_API FText GetViewModesSubmenuLabel(TWeakPtr<SEditorViewport> InViewport);
 
 /**
@@ -40,4 +45,58 @@ UNREALED_API void PopulateViewModesMenu(UToolMenu* InMenu,
 
 UNREALED_API FToolMenuEntry CreateViewportToolbarViewModesSubmenu();
 
+DECLARE_DELEGATE_TwoParams(FRotationGridCheckboxListExecuteActionDelegate, int, ERotationGridMode);
+DECLARE_DELEGATE_RetVal_TwoParams(bool, FRotationGridCheckboxListIsCheckedDelegate, int, ERotationGridMode);
+
+DECLARE_DELEGATE_OneParam(FLocationGridCheckboxListExecuteActionDelegate, int);
+DECLARE_DELEGATE_RetVal_OneParam(bool, FLocationGridCheckboxListIsCheckedDelegate, int);
+
+DECLARE_DELEGATE_OneParam(FScaleGridCheckboxListExecuteActionDelegate, int);
+DECLARE_DELEGATE_RetVal_OneParam(bool, FScaleGridCheckboxListIsCheckedDelegate, int);
+
+UNREALED_API TSharedRef<SWidget> BuildRotationGridCheckBoxList(
+	FName InExtentionHook,
+	const FText& InHeading,
+	const TArray<float>& InGridSizes,
+	ERotationGridMode InGridMode,
+	const FRotationGridCheckboxListExecuteActionDelegate& InExecuteAction,
+	const FRotationGridCheckboxListIsCheckedDelegate& InIsActionChecked,
+	const TSharedPtr<FUICommandList>& InCommandList = {}
+);
+
+UNREALED_API TSharedRef<SWidget> CreateRotationGridSnapMenu(
+	const FRotationGridCheckboxListExecuteActionDelegate& InExecuteDelegate,
+	const FRotationGridCheckboxListIsCheckedDelegate& InIsCheckedDelegate,
+	const TAttribute<bool>& InIsEnabledDelegate = TAttribute<bool>(true),
+	const TSharedPtr<FUICommandList>& InCommandList = {}
+);
+
+UNREALED_API TSharedRef<SWidget> CreateLocationGridSnapMenu(
+	const FLocationGridCheckboxListExecuteActionDelegate& InExecuteDelegate,
+	const FLocationGridCheckboxListIsCheckedDelegate& InIsCheckedDelegate,
+	const TArray<float>& InGridSizes,
+	const TAttribute<bool>& InIsEnabledDelegate = TAttribute<bool>(true),
+	const TSharedPtr<FUICommandList>& InCommandList = {}
+);
+
+UNREALED_API TSharedRef<SWidget> CreateScaleGridSnapMenu(
+	const FScaleGridCheckboxListExecuteActionDelegate& InExecuteDelegate,
+	const FScaleGridCheckboxListIsCheckedDelegate& InIsCheckedDelegate,
+	const TArray<float>& InGridSizes,
+	const TAttribute<bool>& InIsEnabledDelegate = TAttribute<bool>(true),
+	const TSharedPtr<FUICommandList>& InCommandList = {},
+	const TAttribute<bool>& ShowPreserveNonUniformScaleOption = TAttribute<bool>(false),
+	const FUIAction& PreserveNonUniformScaleUIAction = FUIAction()
+);
+
+UNREALED_API FToolMenuEntry CreateCheckboxSubmenu(
+	const FName InName,
+	const TAttribute<FText>& InLabel,
+	const TAttribute<FText>& InToolTip,
+	const FToolMenuExecuteAction& InCheckboxExecuteAction,
+	const FToolMenuCanExecuteAction& InCheckboxCanExecuteAction,
+	const FToolMenuGetActionCheckState& InCheckboxActionCheckState,
+	const FNewToolMenuChoice& InMakeMenu
+);
+
 } // namespace UE::UnrealEd
