[Sequencer] Added interface for AControlRigShapeActor to interact with Sequencer and fix cursor display on hover of selectable control rig handle when Sequencer selection limiting.

diff --git a/Engine/Plugins/Animation/ControlRig/Source/ControlRig/Public/ControlRigGizmoActor.h b/Engine/Plugins/Animation/ControlRig/Source/ControlRig/Public/ControlRigGizmoActor.h
--- a/Engine/Plugins/Animation/ControlRig/Source/ControlRig/Public/ControlRigGizmoActor.h
+++ b/Engine/Plugins/Animation/ControlRig/Source/ControlRig/Public/ControlRigGizmoActor.h
@@ -5,6 +5,7 @@
 #include "CoreMinimal.h"
 #include "GameFramework/Actor.h"
 #include "Engine/StaticMesh.h"
+#include "IViewportSelectableObject.h"
 #include "Materials/Material.h"
 #include "Rigs/RigControlHierarchy.h"
 #include "ControlRigGizmoActor.generated.h"
@@ -51,7 +52,7 @@ struct FControlShapeActorCreationParam
 
 /** An actor used to represent a rig control */
 UCLASS(NotPlaceable, Transient)
-class CONTROLRIG_API AControlRigShapeActor : public AActor
+class CONTROLRIG_API AControlRigShapeActor : public AActor, public IViewportSelectableObject
 {
 	GENERATED_BODY()
 
diff --git a/Engine/Source/Editor/Kismet/Private/SCSEditorViewportClient.cpp b/Engine/Source/Editor/Kismet/Private/SCSEditorViewportClient.cpp
--- a/Engine/Source/Editor/Kismet/Private/SCSEditorViewportClient.cpp
+++ b/Engine/Source/Editor/Kismet/Private/SCSEditorViewportClient.cpp
@@ -1232,7 +1232,7 @@ void FSCSEditorViewportClient::UpdateHoverFromHitProxy(HHitProxy* const InHitPro
 		else if (InHitProxy->IsA(HActor::StaticGetType()))
 		{
 			const HActor* const ActorHitProxy = static_cast<HActor*>(InHitProxy);
-			if (ActorHitProxy && IsValid(ActorHitProxy->PrimComponent))
+			if (ActorHitProxy && IsValid(ActorHitProxy->Actor))
 			{
 				if (bIsViewportSelectionLimited)
 				{
@@ -1252,8 +1252,6 @@ void FSCSEditorViewportClient::UpdateHoverFromHitProxy(HHitProxy* const InHitPro
 		});
 
 	// Set mouse cursor after hovered primitive component list has been updated
-	MouseCursor.Reset();
-
 	if (bIsGizmoHit)
 	{
 		MouseCursor = EMouseCursor::CardinalCross;
@@ -1266,4 +1264,8 @@ void FSCSEditorViewportClient::UpdateHoverFromHitProxy(HHitProxy* const InHitPro
 	{
 		MouseCursor = EMouseCursor::SlashedCircle;
 	}
+	else
+	{
+		MouseCursor.Reset();
+	}
 }
diff --git a/Engine/Source/Editor/Sequencer/Private/Sequencer.cpp b/Engine/Source/Editor/Sequencer/Private/Sequencer.cpp
--- a/Engine/Source/Editor/Sequencer/Private/Sequencer.cpp
+++ b/Engine/Source/Editor/Sequencer/Private/Sequencer.cpp
@@ -163,6 +163,7 @@
 #include "Editor/TransBuffer.h"
 
 #include "EngineModule.h"
+#include "IViewportSelectableObject.h"
 #include "Tracks/MovieSceneBindingLifetimeTrack.h"
 #include "Sections/MovieSceneBindingLifetimeSection.h"
 #include "Bindings/MovieSceneSpawnableBinding.h"
@@ -11981,6 +11982,11 @@ bool FSequencer::IsObjectSelectableInViewport(UObject* const InObject)
 		return true;
 	}
 
+	if (const IViewportSelectableObject* const SelectableObject = Cast<IViewportSelectableObject>(InObject))
+	{
+		return SelectableObject->IsSelectable();
+	}
+
 	UMovieSceneSequence* const FocusedSequence = GetFocusedMovieSceneSequence();
 	if (!IsValid(FocusedSequence))
 	{
@@ -11988,19 +11994,14 @@ bool FSequencer::IsObjectSelectableInViewport(UObject* const InObject)
 	}
 
 	const TSharedRef<UE::MovieScene::FSharedPlaybackState> SharedPlaybackState = GetSharedPlaybackState();
-	const FMovieSceneEvaluationState* const EvaluationState = SharedPlaybackState->FindCapability<FMovieSceneEvaluationState>();
-	if (!EvaluationState)
-	{
-		return true;
-	}
 
 	const UMovieSceneSequence* OutSequence = nullptr;
 
 	// Early out on first sequence the object is found in
 	ForEachSubSequenceRecursively(FocusedSequence,
-		[this, InObject, EvaluationState, &OutSequence](UMovieSceneSequence* const InCurrentSequence)
+		[this, InObject, &SharedPlaybackState, &OutSequence](UMovieSceneSequence* const InCurrentSequence)
 		{
-			const FGuid ObjectGuid = FindObjectId(*InObject, EvaluationState->FindSequenceId(InCurrentSequence));
+			const FGuid ObjectGuid = InCurrentSequence->FindBindingFromObject(InObject, SharedPlaybackState);
 			if (ObjectGuid.IsValid())
 			{
 				OutSequence = InCurrentSequence;
diff --git a/Engine/Source/Editor/Sequencer/Private/SequencerCommands.cpp b/Engine/Source/Editor/Sequencer/Private/SequencerCommands.cpp
--- a/Engine/Source/Editor/Sequencer/Private/SequencerCommands.cpp
+++ b/Engine/Source/Editor/Sequencer/Private/SequencerCommands.cpp
@@ -196,7 +196,7 @@ void FSequencerCommands::RegisterCommands()
 
 	UI_COMMAND(RefreshUI, "Refresh UI", "Forcibly refresh the UI from source data.", EUserInterfaceActionType::Button, FInputChord(EKeys::F5));
 
-	UI_COMMAND(ToggleLimitViewportSelection, "Toggle Limit Viewport Selection", "Toggles viewport selectability between Sequencer only objects and all objects.", EUserInterfaceActionType::ToggleButton, FInputChord(EModifierKey::Control, EKeys::I));
+	UI_COMMAND(ToggleLimitViewportSelection, "Toggle Limit Viewport Selection", "Toggles viewport selectability between Sequencer only objects and all objects.", EUserInterfaceActionType::ToggleButton, FInputChord(EModifierKey::Control, EKeys::Y));
 }
 
 #undef LOCTEXT_NAMESPACE
diff --git a/Engine/Source/Editor/UnrealEd/Private/EditorViewportSelectability.cpp b/Engine/Source/Editor/UnrealEd/Private/EditorViewportSelectability.cpp
--- a/Engine/Source/Editor/UnrealEd/Private/EditorViewportSelectability.cpp
+++ b/Engine/Source/Editor/UnrealEd/Private/EditorViewportSelectability.cpp
@@ -266,10 +266,14 @@ bool FEditorViewportSelectability::SelectActorsByPredicate(const bool bInSelect
 
 	if (InActors.IsEmpty())
 	{
-		for (FActorIterator Iter(ActorSelection->GetWorld()); Iter; ++Iter)
+		const UWorld* const World = ActorSelection->GetWorld();
+		if (IsValid(World))
 		{
-			AActor* const Actor = *Iter;
-			SelectIfPossible(Actor);
+			for (FActorIterator Iter(World); Iter; ++Iter)
+			{
+				AActor* const Actor = *Iter;
+				SelectIfPossible(Actor);
+			}
 		}
 	}
 	else
@@ -400,7 +404,7 @@ void FEditorViewportSelectability::UpdateHoverFromHitProxy(HHitProxy* const InHi
 		else if (InHitProxy->IsA(HActor::StaticGetType()))
 		{
 			const HActor* const ActorHitProxy = static_cast<HActor*>(InHitProxy);
-			if (ActorHitProxy && IsValid(ActorHitProxy->PrimComponent))
+			if (ActorHitProxy && IsValid(ActorHitProxy->Actor))
 			{
 				if (bSelectionLimited)
 				{
@@ -414,8 +418,6 @@ void FEditorViewportSelectability::UpdateHoverFromHitProxy(HHitProxy* const InHi
 	UpdateHoveredActorPrimitives(Actor);
 
 	// Set mouse cursor after hovered primitive component list has been updated
-	MouseCursor.Reset();
-
 	if (bIsGizmoHit)
 	{
 		MouseCursor = EMouseCursor::CardinalCross;
@@ -428,6 +430,10 @@ void FEditorViewportSelectability::UpdateHoverFromHitProxy(HHitProxy* const InHi
 	{
 		MouseCursor = EMouseCursor::SlashedCircle;
 	}
+	else
+	{
+		MouseCursor.Reset();
+	}
 }
 
 bool FEditorViewportSelectability::HandleClick(FEditorViewportClient* const InViewportClient, HHitProxy* const InHitProxy, const FViewportClick& InClick)
@@ -435,25 +441,21 @@ bool FEditorViewportSelectability::HandleClick(FEditorViewportClient* const InVi
 	// Disable actor selection when sequencer is limiting selection
 	const int32 HitX = InViewportClient->Viewport->GetMouseX();
 	const int32 HitY = InViewportClient->Viewport->GetMouseY();
-
 	HHitProxy* const HitResult = InViewportClient->Viewport->GetHitProxy(HitX, HitY);
 	if (!HitResult)
 	{
 		return false;
 	}
 
-	if (HitResult->IsA(HActor::StaticGetType()))
+	if (HitResult->IsA(HWidgetAxis::StaticGetType()) || !HitResult->IsA(HActor::StaticGetType()))
 	{
-		const HActor* const ActorHitProxy = static_cast<HActor*>(HitResult);
-		if (ActorHitProxy
-			&& IsValid(ActorHitProxy->Actor)
-			&& !IsObjectSelectableInViewport(ActorHitProxy->Actor))
-		{
-			return true;
-		}
+		return false;
 	}
 
-	return false;
+	const HActor* const ActorHitProxy = static_cast<HActor*>(HitResult);
+	return ActorHitProxy
+		&& IsValid(ActorHitProxy->Actor)
+		&& !IsObjectSelectableInViewport(ActorHitProxy->Actor);
 }
 
 bool FEditorViewportSelectability::BoxSelectWorldActors(FBox& InBox, FEditorViewportClient* const InEditorViewportClient, const bool bInSelect)
diff --git a/Engine/Source/Runtime/Engine/Public/IViewportSelectableObject.h b/Engine/Source/Runtime/Engine/Public/IViewportSelectableObject.h
new file mode 100644
--- /dev/null
+++ b/Engine/Source/Runtime/Engine/Public/IViewportSelectableObject.h
@@ -0,0 +1,26 @@
+﻿// Copyright Epic Games, Inc. All Rights Reserved.
+
+#pragma once
+
+#include "UObject/Interface.h"
+#include "IViewportSelectableObject.generated.h"
+
+/**
+ * Interface for objects selectable in the viewport.
+ * Only used by Sequencer in editor at the moment to allow AControlRigShapeActor (runtime)
+ * to be selectable when Sequencer is selection limiting.
+ * Could be extended or removed in the future where there is a more generalized, low level selection system.
+ **/
+UINTERFACE(MinimalApi, meta=(CannotImplementInterfaceInBlueprint))
+class UViewportSelectableObject : public UInterface
+{
+	GENERATED_BODY()
+};
+
+class IViewportSelectableObject
+{
+	GENERATED_BODY()
+
+public:
+	ENGINE_API virtual bool IsSelectable() const { return true; }
+};
