[Viewport Toolbar] Complete the "Selection" submenu for the new level editor viewport toolbar

diff --git a/Engine/Plugins/Experimental/GeometryCollectionPlugin/Source/GeometryCollectionEditor/Private/GeometryCollection/GeometryCollectionEditorPlugin.cpp b/Engine/Plugins/Experimental/GeometryCollectionPlugin/Source/GeometryCollectionEditor/Private/GeometryCollection/GeometryCollectionEditorPlugin.cpp
--- a/Engine/Plugins/Experimental/GeometryCollectionPlugin/Source/GeometryCollectionEditor/Private/GeometryCollection/GeometryCollectionEditorPlugin.cpp
+++ b/Engine/Plugins/Experimental/GeometryCollectionPlugin/Source/GeometryCollectionEditor/Private/GeometryCollection/GeometryCollectionEditorPlugin.cpp
@@ -319,6 +319,27 @@ const ISlateStyle* IGeometryCollectionEditorPlugin::GetEditorStyle()
 	return FSlateStyleRegistry::FindSlateStyle(GetEditorStyleName());
 }
 
+void IGeometryCollectionEditorPlugin::CreateGeometrySelectionMenu(FMenuBuilder& InMenuBuilder)
+{
+	InMenuBuilder.AddMenuEntry(FGeometryCollectionSelectionCommands::Get().SelectAllGeometry);
+	InMenuBuilder.AddMenuEntry(FGeometryCollectionSelectionCommands::Get().SelectNone);
+	InMenuBuilder.AddMenuEntry(FGeometryCollectionSelectionCommands::Get().SelectInverseGeometry);
+}
+
+TSharedRef<FExtender> IGeometryCollectionEditorPlugin::OnExtendLevelEditorViewMenu(const TSharedRef<FUICommandList> InCommandList)
+{
+	TSharedRef<FExtender> Extender(new FExtender());
+
+	Extender->AddMenuExtension(
+		"SelectGeometryCollections",
+		EExtensionHook::After,
+		nullptr,
+		FMenuExtensionDelegate::CreateRaw(this, &IGeometryCollectionEditorPlugin::CreateGeometrySelectionMenu)
+	);
+
+	return Extender;
+}
+
 void IGeometryCollectionEditorPlugin::RegisterMenus()
 {
 	FToolMenuOwnerScoped OwnerScoped(this);
@@ -330,4 +351,14 @@ void IGeometryCollectionEditorPlugin::RegisterMenus()
 		Section.AddMenuEntry(FGeometryCollectionSelectionCommands::Get().SelectNone);
 		Section.AddMenuEntry(FGeometryCollectionSelectionCommands::Get().SelectInverseGeometry);
 	}
+
+	ViewMenuExtender = FLevelEditorModule::FLevelEditorMenuExtender::CreateRaw(
+		this, &IGeometryCollectionEditorPlugin::OnExtendLevelEditorViewMenu
+	);
+
+	FLevelEditorModule& LevelEditor = FModuleManager::LoadModuleChecked<FLevelEditorModule>(TEXT("LevelEditor"));
+	TArray<FLevelEditorModule::FLevelEditorMenuExtender>& MenuExtenders =
+		LevelEditor.GetAllLevelEditorToolbarViewMenuExtenders();
+	MenuExtenders.Add(ViewMenuExtender);
+	ViewMenuExtenderHandle = MenuExtenders.Last().GetHandle();
 }
diff --git a/Engine/Plugins/Experimental/GeometryCollectionPlugin/Source/GeometryCollectionEditor/Public/GeometryCollection/GeometryCollectionEditorPlugin.h b/Engine/Plugins/Experimental/GeometryCollectionPlugin/Source/GeometryCollectionEditor/Public/GeometryCollection/GeometryCollectionEditorPlugin.h
--- a/Engine/Plugins/Experimental/GeometryCollectionPlugin/Source/GeometryCollectionEditor/Public/GeometryCollection/GeometryCollectionEditorPlugin.h
+++ b/Engine/Plugins/Experimental/GeometryCollectionPlugin/Source/GeometryCollectionEditor/Public/GeometryCollection/GeometryCollectionEditorPlugin.h
@@ -3,6 +3,7 @@
 #pragma once
 
 #include "CoreMinimal.h"
+#include "LevelEditor.h"
 #include "Modules/ModuleInterface.h"
 #include "Modules/ModuleManager.h"
 #include "GeometryCollection/GeometryCollectionProviderEditor.h"
@@ -56,6 +57,9 @@ public:
 private:
 	void RegisterMenus();
 
+	void CreateGeometrySelectionMenu(FMenuBuilder& InMenuBuilder);
+	TSharedRef<FExtender> OnExtendLevelEditorViewMenu(const TSharedRef<FUICommandList> InCommandList);
+
 private:
 	TArray<IConsoleObject*> EditorCommands;
 
@@ -70,5 +74,8 @@ private:
 
 	// Styleset for geom collection tool brushes/fonts etc.
 	TSharedPtr<FSlateStyleSet> StyleSet;
+
+	FLevelEditorModule::FLevelEditorMenuExtender ViewMenuExtender;
+	FDelegateHandle ViewMenuExtenderHandle;
 };
 
diff --git a/Engine/Source/Editor/UnrealEd/Private/ViewportToolbar/UnrealEdViewportToolbar.cpp b/Engine/Source/Editor/UnrealEd/Private/ViewportToolbar/UnrealEdViewportToolbar.cpp
--- a/Engine/Source/Editor/UnrealEd/Private/ViewportToolbar/UnrealEdViewportToolbar.cpp
+++ b/Engine/Source/Editor/UnrealEd/Private/ViewportToolbar/UnrealEdViewportToolbar.cpp
@@ -13,6 +13,7 @@
 #include "RayTracingDebugVisualizationMenuCommands.h"
 #include "SEditorViewport.h"
 #include "Settings/LevelEditorViewportSettings.h"
+#include "Styling/SlateIconFinder.h"
 #include "Templates/SharedPointer.h"
 #include "ToolMenu.h"
 #include "ToolMenuEntry.h"
@@ -646,13 +647,64 @@ FToolMenuEntry CreateViewportToolbarSelectionSection()
 				{
 					FToolMenuSection& UnnamedSection = Submenu->FindOrAddSection(NAME_None);
 
-					UnnamedSection.AddMenuEntry(FGenericCommands::Get().SelectAll);
-					UnnamedSection.AddMenuEntry(FLevelEditorCommands::Get().SelectNone);
-					UnnamedSection.AddMenuEntry(FLevelEditorCommands::Get().InvertSelection);
+					UnnamedSection.AddMenuEntry(
+						FGenericCommands::Get().SelectAll,
+						FGenericCommands::Get().SelectAll->GetLabel(),
+						FGenericCommands::Get().SelectAll->MakeTooltip()->GetTextTooltip(),
+						FSlateIconFinder::FindIcon("FoliageEditMode.SelectAll")
+					);
+
+					UnnamedSection.AddMenuEntry(
+						FLevelEditorCommands::Get().SelectNone,
+						FLevelEditorCommands::Get().SelectNone->GetLabel(),
+						FLevelEditorCommands::Get().SelectNone->MakeTooltip()->GetTextTooltip(),
+						FSlateIconFinder::FindIcon("Cross")
+					);
+
+					UnnamedSection.AddMenuEntry(
+						FLevelEditorCommands::Get().InvertSelection,
+						FLevelEditorCommands::Get().InvertSelection->GetLabel(),
+						FLevelEditorCommands::Get().InvertSelection->MakeTooltip()->GetTextTooltip(),
+						FSlateIconFinder::FindIcon("FoliageEditMode.DeselectAll")
+					);
+
+					// Hierarchy based selection
+					{
+						UnnamedSection.AddSubMenu(
+							"Hierarchy",
+							LOCTEXT("HierarchyLabel", "Hierarchy"),
+							LOCTEXT("HierarchyTooltip", "Hierarchy selection tools"),
+							FNewToolMenuDelegate::CreateLambda(
+								[](UToolMenu* HierarchyMenu)
+								{
+									FToolMenuSection& HierarchySection = HierarchyMenu->FindOrAddSection(
+										"SelectAllHierarchy", LOCTEXT("SelectAllHierarchyLabel", "Hierarchy")
+									);
+
+									HierarchySection.AddMenuEntry(
+										FLevelEditorCommands::Get().SelectImmediateChildren,
+										LOCTEXT("HierarchySelectImmediateChildrenLabel", "Immediate Children")
+									);
+
+									HierarchySection.AddMenuEntry(
+										FLevelEditorCommands::Get().SelectAllDescendants,
+										LOCTEXT("HierarchySelectAllDescendantsLabel", "All Descendants")
+									);
+								}
+							),
+							false,
+							FSlateIconFinder::FindIcon("BTEditor.SwitchToBehaviorTreeMode")
+						);
+					}
 
 					UnnamedSection.AddSeparator("Advanced");
 
-					UnnamedSection.AddMenuEntry(FLevelEditorCommands::Get().SelectAllActorsOfSameClass);
+					UnnamedSection.AddMenuEntry(
+						FLevelEditorCommands::Get().SelectAllActorsOfSameClass,
+						LOCTEXT("AdvancedSelectAllActorsOfSameClassLabel", "All of Same Class"),
+						FLevelEditorCommands::Get().SelectAllActorsOfSameClass->MakeTooltip()->GetTextTooltip(),
+						FSlateIconFinder::FindIcon("PlacementBrowser.Icons.All")
+					);
 				}
 
 				{
@@ -670,11 +722,196 @@ FToolMenuEntry CreateViewportToolbarSelectionSection()
 									"SelectAllBSP", LOCTEXT("SelectAllBSPLabel", "Select All BSP")
 								);
 
-								SelectAllSection.AddMenuEntry(FLevelEditorCommands::Get().SelectAllAddditiveBrushes);
-								SelectAllSection.AddMenuEntry(FLevelEditorCommands::Get().SelectAllSubtractiveBrushes);
-								SelectAllSection.AddMenuEntry(FLevelEditorCommands::Get().SelectAllSurfaces);
+								SelectAllSection.AddMenuEntry(
+									FLevelEditorCommands::Get().SelectAllAddditiveBrushes,
+									LOCTEXT("BSPSelectAllAdditiveBrushesLabel", "Addditive Brushes")
+								);
+
+								SelectAllSection.AddMenuEntry(
+									FLevelEditorCommands::Get().SelectAllSubtractiveBrushes,
+									LOCTEXT("BSPSelectAllSubtractiveBrushesLabel", "Subtractive Brushes")
+								);
+
+								SelectAllSection.AddMenuEntry(
+									FLevelEditorCommands::Get().SelectAllSurfaces,
+									LOCTEXT("BSPSelectAllAllSurfacesLabel", "Surfaces")
+								);
 							}
-						)
+						),
+						false,
+						FSlateIconFinder::FindIcon("ShowFlagsMenu.BSP")
+					);
+
+					ByTypeSection.AddSubMenu(
+						"Emitters",
+						LOCTEXT("EmittersLabel", "Emitters"),
+						LOCTEXT("EmittersTooltip", "Emitters-related tools"),
+						FNewToolMenuDelegate::CreateLambda(
+							[](UToolMenu* EmittersMenu)
+							{
+								FToolMenuSection& SelectAllSection = EmittersMenu->FindOrAddSection(
+									"SelectAllEmitters", LOCTEXT("SelectAllEmittersLabel", "Select All Emitters")
+								);
+
+								SelectAllSection.AddMenuEntry(
+									FLevelEditorCommands::Get().SelectMatchingEmitter,
+									LOCTEXT("EmittersSelectMatchingEmitterLabel", "Matching Emitters")
+								);
+							}
+						),
+						false,
+						FSlateIconFinder::FindIcon("ClassIcon.Emitter")
+					);
+
+					ByTypeSection.AddSubMenu(
+						"GeometryCollections",
+						LOCTEXT("GeometryCollectionsLabel", "Geometry Collections"),
+						LOCTEXT("GeometryCollectionsTooltip", "GeometryCollections-related tools"),
+						FNewToolMenuDelegate::CreateLambda(
+							[](UToolMenu* GeometryCollectionsMenu)
+							{
+								// This one will be filled by extensions from GeometryCollectionEditorPlugin
+								// Hook is "SelectGeometryCollections"
+								FToolMenuSection& SelectAllSection = GeometryCollectionsMenu->FindOrAddSection(
+									"SelectGeometryCollections",
+									LOCTEXT("SelectGeometryCollectionsLabel", "Geometry Collections")
+								);
+							}
+						),
+						false,
+						FSlateIconFinder::FindIcon("ClassIcon.GeometryCollection")
+					);
+
+					ByTypeSection.AddSubMenu(
+						"HLOD",
+						LOCTEXT("HLODLabel", "HLOD"),
+						LOCTEXT("HLODTooltip", "HLOD-related tools"),
+						FNewToolMenuDelegate::CreateLambda(
+							[](UToolMenu* HLODMenu)
+							{
+								FToolMenuSection& SelectAllSection = HLODMenu->FindOrAddSection(
+									"SelectAllHLOD", LOCTEXT("SelectAllHLODLabel", "Select All HLOD")
+								);
+
+								SelectAllSection.AddMenuEntry(
+									FLevelEditorCommands::Get().SelectOwningHierarchicalLODCluster,
+									LOCTEXT("HLODSelectOwningHierarchicalLODClusterLabel", "Owning HLOD Cluster")
+								);
+							}
+						),
+						false,
+						FSlateIconFinder::FindIcon("WorldPartition.ShowHLODActors")
+					);
+
+					ByTypeSection.AddSubMenu(
+						"Lights",
+						LOCTEXT("LightsLabel", "Lights"),
+						LOCTEXT("LightsTooltip", "Lights-related tools"),
+						FNewToolMenuDelegate::CreateLambda(
+							[](UToolMenu* LightsMenu)
+							{
+								FToolMenuSection& SelectAllSection = LightsMenu->FindOrAddSection(
+									"SelectAllLights", LOCTEXT("SelectAllLightsLabel", "Select All Lights")
+								);
+
+								SelectAllSection.AddMenuEntry(
+									FLevelEditorCommands::Get().SelectAllLights,
+									LOCTEXT("LightsSelectAllLightsLabel", "All Lights")
+								);
+
+								SelectAllSection.AddMenuEntry(
+									FLevelEditorCommands::Get().SelectRelevantLights,
+									LOCTEXT("LightsSelectRelevantLightsLabel", "Relevant Lights")
+								);
+
+								SelectAllSection.AddMenuEntry(
+									FLevelEditorCommands::Get().SelectStationaryLightsExceedingOverlap,
+									LOCTEXT("LightsSelectStationaryLightsExceedingOverlapLabel", "Stationary Lights Exceeding Overlap")
+								);
+							}
+						),
+						false,
+						FSlateIconFinder::FindIcon("PlacementBrowser.Icons.Lights")
+					);
+
+					ByTypeSection.AddSubMenu(
+						"Material",
+						LOCTEXT("MaterialLabel", "Material"),
+						LOCTEXT("MaterialTooltip", "Material-related tools"),
+						FNewToolMenuDelegate::CreateLambda(
+							[](UToolMenu* MaterialMenu)
+							{
+								FToolMenuSection& SelectAllSection = MaterialMenu->FindOrAddSection(
+									"SelectAllMaterial", LOCTEXT("SelectAllMaterialLabel", "Select All Material")
+								);
+
+								SelectAllSection.AddMenuEntry(
+									FLevelEditorCommands::Get().SelectAllWithSameMaterial,
+									LOCTEXT("MaterialSelectAllWithSameMaterialLabel", "With Same Material")
+								);
+							}
+						),
+						false,
+						FSlateIconFinder::FindIcon("ClassIcon.Material")
+					);
+
+					ByTypeSection.AddSubMenu(
+						"SkeletalMeshes",
+						LOCTEXT("SkeletalMeshesLabel", "Skeletal Meshes"),
+						LOCTEXT("SkeletalMeshesTooltip", "SkeletalMeshes-related tools"),
+						FNewToolMenuDelegate::CreateLambda(
+							[](UToolMenu* SkeletalMeshesMenu)
+							{
+								FToolMenuSection& SelectAllSection = SkeletalMeshesMenu->FindOrAddSection(
+									"SelectAllSkeletalMeshes",
+									LOCTEXT("SelectAllSkeletalMeshesLabel", "Select All SkeletalMeshes")
+								);
+
+								SelectAllSection.AddMenuEntry(
+									FLevelEditorCommands::Get().SelectSkeletalMeshesOfSameClass,
+									LOCTEXT(
+										"SkeletalMeshesSelectSkeletalMeshesOfSameClassLabel",
+										"Using Selected Skeletal Meshes (Selected Actor Types)"
+									)
+								);
+
+								SelectAllSection.AddMenuEntry(
+									FLevelEditorCommands::Get().SelectSkeletalMeshesAllClasses,
+									LOCTEXT(
+										"SkeletalMeshesSelectSkeletalMeshesAllClassesLabel",
+										"Using Selected Skeletal Meshes (All Actor Types)"
+									)
+								);
+							}
+						),
+						false,
+						FSlateIconFinder::FindIcon("SkeletonTree.Bone")
+					);
+
+					ByTypeSection.AddSubMenu(
+						"StaticMeshes",
+						LOCTEXT("StaticMeshesLabel", "Static Meshes"),
+						LOCTEXT("StaticMeshesTooltip", "StaticMeshes-related tools"),
+						FNewToolMenuDelegate::CreateLambda(
+							[](UToolMenu* StaticMeshesMenu)
+							{
+								FToolMenuSection& SelectAllSection = StaticMeshesMenu->FindOrAddSection(
+									"SelectAllStaticMeshes", LOCTEXT("SelectAllStaticMeshesLabel", "Select All StaticMeshes")
+								);
+
+								SelectAllSection.AddMenuEntry(
+									FLevelEditorCommands::Get().SelectStaticMeshesOfSameClass,
+									LOCTEXT("StaticMeshesSelectStaticMeshesOfSameClassLabel", "Matching Selected Class")
+								);
+
+								SelectAllSection.AddMenuEntry(
+									FLevelEditorCommands::Get().SelectStaticMeshesAllClasses,
+									LOCTEXT("StaticMeshesSelectStaticMeshesAllClassesLabel", "Matching All Classes")
+								);
+							}
+						),
+						false,
+						FSlateIconFinder::FindIcon("ShowFlagsMenu.StaticMeshes")
 					);
 				}
 
@@ -682,7 +919,25 @@ FToolMenuEntry CreateViewportToolbarSelectionSection()
 					FToolMenuSection& OptionsSection =
 						Submenu->FindOrAddSection("Options", LOCTEXT("OptionsLabel", "Options"));
 
-					OptionsSection.AddMenuEntry(FLevelEditorCommands::Get().AllowTranslucentSelection);
+					OptionsSection.AddMenuEntry(
+						FLevelEditorCommands::Get().AllowTranslucentSelection,
+						LOCTEXT("OptionsAllowTranslucentSelectionLabel", "Translucent Objects")
+					);
+
+					OptionsSection.AddMenuEntry(
+						FLevelEditorCommands::Get().AllowGroupSelection,
+						LOCTEXT("OptionsAllowGroupSelectionLabel", "Select Groups")
+					);
+
+					OptionsSection.AddMenuEntry(
+						FLevelEditorCommands::Get().StrictBoxSelect,
+						LOCTEXT("OptionsStrictBoxSelectLabel", "Strict Marquee Selection")
+					);
+
+					OptionsSection.AddMenuEntry(
+						FLevelEditorCommands::Get().TransparentBoxSelect,
+						LOCTEXT("OptionsTransparentBoxSelectLabel", "Marquee Select Occluded")
+					);
 				}
 			}
 		)
@@ -717,8 +972,6 @@ FToolMenuEntry CreateViewportToolbarSnappingSubmenu()
 				VertexSnapping.UserInterfaceActionType = EUserInterfaceActionType::ToggleButton;
 				VertexSnapping.Label = LOCTEXT("VertexSnapLabel", "Vertex");
 				SnappingSection.AddEntry(VertexSnapping);
-
-				// TODO: add Planar Snapping
 			}
 		)
 	);
