[Backout] - CL34812542
Original CL Desc
-----------------------------------------------------------------
[Sequencer] Added interface for AControlRigShapeActor to interact with Sequencer and fix cursor display on hover of selectable control rig handle when Sequencer selection limiting.

diff --git a/Engine/Plugins/Animation/ControlRig/Source/ControlRig/Public/ControlRigGizmoActor.h b/Engine/Plugins/Animation/ControlRig/Source/ControlRig/Public/ControlRigGizmoActor.h
--- a/Engine/Plugins/Animation/ControlRig/Source/ControlRig/Public/ControlRigGizmoActor.h
+++ b/Engine/Plugins/Animation/ControlRig/Source/ControlRig/Public/ControlRigGizmoActor.h
@@ -5,7 +5,6 @@
 #include "CoreMinimal.h"
 #include "GameFramework/Actor.h"
 #include "Engine/StaticMesh.h"
-#include "IViewportSelectableObject.h"
 #include "Materials/Material.h"
 #include "Rigs/RigControlHierarchy.h"
 #include "ControlRigGizmoActor.generated.h"
@@ -52,7 +51,7 @@ struct FControlShapeActorCreationParam
 
 /** An actor used to represent a rig control */
 UCLASS(NotPlaceable, Transient)
-class CONTROLRIG_API AControlRigShapeActor : public AActor, public IViewportSelectableObject
+class CONTROLRIG_API AControlRigShapeActor : public AActor
 {
 	GENERATED_BODY()
 
diff --git a/Engine/Source/Editor/Kismet/Private/SCSEditorViewportClient.cpp b/Engine/Source/Editor/Kismet/Private/SCSEditorViewportClient.cpp
--- a/Engine/Source/Editor/Kismet/Private/SCSEditorViewportClient.cpp
+++ b/Engine/Source/Editor/Kismet/Private/SCSEditorViewportClient.cpp
@@ -1232,7 +1232,7 @@ void FSCSEditorViewportClient::UpdateHoverFromHitProxy(HHitProxy* const InHitPro
 		else if (InHitProxy->IsA(HActor::StaticGetType()))
 		{
 			const HActor* const ActorHitProxy = static_cast<HActor*>(InHitProxy);
-			if (ActorHitProxy && IsValid(ActorHitProxy->Actor))
+			if (ActorHitProxy && IsValid(ActorHitProxy->PrimComponent))
 			{
 				if (bIsViewportSelectionLimited)
 				{
@@ -1252,6 +1252,8 @@ void FSCSEditorViewportClient::UpdateHoverFromHitProxy(HHitProxy* const InHitPro
 		});
 
 	// Set mouse cursor after hovered primitive component list has been updated
+	MouseCursor.Reset();
+
 	if (bIsGizmoHit)
 	{
 		MouseCursor = EMouseCursor::CardinalCross;
@@ -1264,8 +1266,4 @@ void FSCSEditorViewportClient::UpdateHoverFromHitProxy(HHitProxy* const InHitPro
 	{
 		MouseCursor = EMouseCursor::SlashedCircle;
 	}
-	else
-	{
-		MouseCursor.Reset();
-	}
 }
diff --git a/Engine/Source/Editor/Sequencer/Private/Sequencer.cpp b/Engine/Source/Editor/Sequencer/Private/Sequencer.cpp
--- a/Engine/Source/Editor/Sequencer/Private/Sequencer.cpp
+++ b/Engine/Source/Editor/Sequencer/Private/Sequencer.cpp
@@ -164,7 +164,6 @@
 #include "Editor/TransBuffer.h"
 
 #include "EngineModule.h"
-#include "IViewportSelectableObject.h"
 #include "Tracks/MovieSceneBindingLifetimeTrack.h"
 #include "Sections/MovieSceneBindingLifetimeSection.h"
 #include "Bindings/MovieSceneSpawnableBinding.h"
@@ -12012,11 +12011,6 @@ bool FSequencer::IsObjectSelectableInViewport(UObject* const InObject)
 		return true;
 	}
 
-	if (const IViewportSelectableObject* const SelectableObject = Cast<IViewportSelectableObject>(InObject))
-	{
-		return SelectableObject->IsSelectable();
-	}
-
 	UMovieSceneSequence* const FocusedSequence = GetFocusedMovieSceneSequence();
 	if (!IsValid(FocusedSequence))
 	{
@@ -12024,14 +12018,19 @@ bool FSequencer::IsObjectSelectableInViewport(UObject* const InObject)
 	}
 
 	const TSharedRef<UE::MovieScene::FSharedPlaybackState> SharedPlaybackState = GetSharedPlaybackState();
+	const FMovieSceneEvaluationState* const EvaluationState = SharedPlaybackState->FindCapability<FMovieSceneEvaluationState>();
+	if (!EvaluationState)
+	{
+		return true;
+	}
 
 	const UMovieSceneSequence* OutSequence = nullptr;
 
 	// Early out on first sequence the object is found in
 	ForEachSubSequenceRecursively(FocusedSequence,
-		[this, InObject, &SharedPlaybackState, &OutSequence](UMovieSceneSequence* const InCurrentSequence)
+		[this, InObject, EvaluationState, &OutSequence](UMovieSceneSequence* const InCurrentSequence)
 		{
-			const FGuid ObjectGuid = InCurrentSequence->FindBindingFromObject(InObject, SharedPlaybackState);
+			const FGuid ObjectGuid = FindObjectId(*InObject, EvaluationState->FindSequenceId(InCurrentSequence));
 			if (ObjectGuid.IsValid())
 			{
 				OutSequence = InCurrentSequence;
diff --git a/Engine/Source/Editor/Sequencer/Private/SequencerCommands.cpp b/Engine/Source/Editor/Sequencer/Private/SequencerCommands.cpp
--- a/Engine/Source/Editor/Sequencer/Private/SequencerCommands.cpp
+++ b/Engine/Source/Editor/Sequencer/Private/SequencerCommands.cpp
@@ -196,7 +196,7 @@ void FSequencerCommands::RegisterCommands()
 
 	UI_COMMAND(RefreshUI, "Refresh UI", "Forcibly refresh the UI from source data.", EUserInterfaceActionType::Button, FInputChord(EKeys::F5));
 
-	UI_COMMAND(ToggleLimitViewportSelection, "Toggle Limit Viewport Selection", "Toggles viewport selectability between Sequencer only objects and all objects.", EUserInterfaceActionType::ToggleButton, FInputChord(EModifierKey::Control, EKeys::Y));
+	UI_COMMAND(ToggleLimitViewportSelection, "Toggle Limit Viewport Selection", "Toggles viewport selectability between Sequencer only objects and all objects.", EUserInterfaceActionType::ToggleButton, FInputChord(EModifierKey::Control, EKeys::I));
 }
 
 #undef LOCTEXT_NAMESPACE
diff --git a/Engine/Source/Editor/UnrealEd/Private/EditorViewportSelectability.cpp b/Engine/Source/Editor/UnrealEd/Private/EditorViewportSelectability.cpp
--- a/Engine/Source/Editor/UnrealEd/Private/EditorViewportSelectability.cpp
+++ b/Engine/Source/Editor/UnrealEd/Private/EditorViewportSelectability.cpp
@@ -266,14 +266,10 @@ bool FEditorViewportSelectability::SelectActorsByPredicate(const bool bInSelect
 
 	if (InActors.IsEmpty())
 	{
-		const UWorld* const World = ActorSelection->GetWorld();
-		if (IsValid(World))
+		for (FActorIterator Iter(ActorSelection->GetWorld()); Iter; ++Iter)
 		{
-			for (FActorIterator Iter(World); Iter; ++Iter)
-			{
-				AActor* const Actor = *Iter;
-				SelectIfPossible(Actor);
-			}
+			AActor* const Actor = *Iter;
+			SelectIfPossible(Actor);
 		}
 	}
 	else
@@ -404,7 +400,7 @@ void FEditorViewportSelectability::UpdateHoverFromHitProxy(HHitProxy* const InHi
 		else if (InHitProxy->IsA(HActor::StaticGetType()))
 		{
 			const HActor* const ActorHitProxy = static_cast<HActor*>(InHitProxy);
-			if (ActorHitProxy && IsValid(ActorHitProxy->Actor))
+			if (ActorHitProxy && IsValid(ActorHitProxy->PrimComponent))
 			{
 				if (bSelectionLimited)
 				{
@@ -418,6 +414,8 @@ void FEditorViewportSelectability::UpdateHoverFromHitProxy(HHitProxy* const InHi
 	UpdateHoveredActorPrimitives(Actor);
 
 	// Set mouse cursor after hovered primitive component list has been updated
+	MouseCursor.Reset();
+
 	if (bIsGizmoHit)
 	{
 		MouseCursor = EMouseCursor::CardinalCross;
@@ -430,10 +428,6 @@ void FEditorViewportSelectability::UpdateHoverFromHitProxy(HHitProxy* const InHi
 	{
 		MouseCursor = EMouseCursor::SlashedCircle;
 	}
-	else
-	{
-		MouseCursor.Reset();
-	}
 }
 
 bool FEditorViewportSelectability::HandleClick(FEditorViewportClient* const InViewportClient, HHitProxy* const InHitProxy, const FViewportClick& InClick)
@@ -441,21 +435,25 @@ bool FEditorViewportSelectability::HandleClick(FEditorViewportClient* const InVi
 	// Disable actor selection when sequencer is limiting selection
 	const int32 HitX = InViewportClient->Viewport->GetMouseX();
 	const int32 HitY = InViewportClient->Viewport->GetMouseY();
+
 	HHitProxy* const HitResult = InViewportClient->Viewport->GetHitProxy(HitX, HitY);
 	if (!HitResult)
 	{
 		return false;
 	}
 
-	if (HitResult->IsA(HWidgetAxis::StaticGetType()) || !HitResult->IsA(HActor::StaticGetType()))
+	if (HitResult->IsA(HActor::StaticGetType()))
 	{
-		return false;
+		const HActor* const ActorHitProxy = static_cast<HActor*>(HitResult);
+		if (ActorHitProxy
+			&& IsValid(ActorHitProxy->Actor)
+			&& !IsObjectSelectableInViewport(ActorHitProxy->Actor))
+		{
+			return true;
+		}
 	}
 
-	const HActor* const ActorHitProxy = static_cast<HActor*>(HitResult);
-	return ActorHitProxy
-		&& IsValid(ActorHitProxy->Actor)
-		&& !IsObjectSelectableInViewport(ActorHitProxy->Actor);
+	return false;
 }
 
 bool FEditorViewportSelectability::BoxSelectWorldActors(FBox& InBox, FEditorViewportClient* const InEditorViewportClient, const bool bInSelect)
diff --git a/Engine/Source/Runtime/Engine/Public/IViewportSelectableObject.h b/Engine/Source/Runtime/Engine/Public/IViewportSelectableObject.h
deleted file mode 100644
