Add ODSC support to landscape materials

There needs to be a specific way of loading these materials as they can be part of the World Partition system. This means that they can be embedded in a package not immmediately reachable by the AssetRegistry / LoadObject

From the material path, it's actually possible to retrieve the directory that contains the package with the material we're interested in. See FWorldPartitionHelpers::ConvertRuntimePathToEditorPath. From this path, we then need to do IAssetRegistry::ScanSynchronous to retrieve the list of all the packages + the corresponding object paths

Unfortunately, this only gives the root paths, ie the World Partition actors. It seems that landscape material instances follow this pattern: /Path/To/Map/Map.Map:PersistentLevel.LandscapeStreamingProxy_[ID1].LandscapeMaterialInstanceConstant_[ID2]
What precedes LandscapeMaterialInstanceConstant_[ID2] is actually the path to the actor, and ODSC relies on this convention to find the correct root object.

The last problem is that loading these packages can take a long time. In some situations, AWorldSettings::DefaultGameMode can point to a blueprint class with a lot of dependencies. To skip it, FCoreRedirects::AddRedirectList can be used such that the package points to an invalid path and fails to load. The list of base classes for which packages are filtered out is set through odsc.excludedclasses. Only packages which are uasset inheriting from one of these classes will be filtered out. This means that /Script/Engine which contains GameMode will not be exluded, but something like /Game/MyGame/MyGame_GameMode containing MyGame_GameMode_C will. Any asset inside this package will fail to load, so one must be careful which package ends up in that list. The assumption is that the client shouldn't pull data from a game mode.

Finally, other materials that are part of World partition, such as HLOD materials, still don't work

diff --git a/Engine/Source/Editor/UnrealEd/Private/CookOnTheFlyServer.cpp b/Engine/Source/Editor/UnrealEd/Private/CookOnTheFlyServer.cpp
--- a/Engine/Source/Editor/UnrealEd/Private/CookOnTheFlyServer.cpp
+++ b/Engine/Source/Editor/UnrealEd/Private/CookOnTheFlyServer.cpp
@@ -679,6 +679,7 @@ bool UCookOnTheFlyServer::StartCookOnTheFly(FCookOnTheFlyStartupOptions InCookOn
 
 				const void* ConnectionPtr = &Connection;
 				RecompileData.LoadedMaterialsToRecompile = &LoadedMaterialsToRecompile;
+				RecompileData.ODSCCustomLoadMaterial = &FODSCClientData::FindMaterial;
 
 				FEventRef RecompileCompletedEvent;
 				UE::Cook::FRecompileShaderCompletedCallback RecompileCompleted = [this, &RecompileCompletedEvent, ConnectionPtr, &LoadedMaterialsToRecompile, RecompileDataCommandType = RecompileData.CommandType]()
@@ -9502,6 +9503,8 @@ void UCookOnTheFlyServer::SaveGlobalShaderMapFiles(const TArrayView<const ITarge
 
 		TArray<uint8> GlobalShaderMap;
 		FShaderRecompileData RecompileData(PlatformName, SP_NumPlatforms, RecompileCommand, nullptr, nullptr, &GlobalShaderMap);
+		RecompileData.ODSCCustomLoadMaterial = &UE::Cook::FODSCClientData::FindMaterial;
+		
 		RecompileShadersForRemote(RecompileData, GetSandboxDirectory(PlatformName));
 	}
 }
diff --git a/Engine/Source/Editor/UnrealEd/Private/Cooker/OnDemandShaderCompilation.cpp b/Engine/Source/Editor/UnrealEd/Private/Cooker/OnDemandShaderCompilation.cpp
--- a/Engine/Source/Editor/UnrealEd/Private/Cooker/OnDemandShaderCompilation.cpp
+++ b/Engine/Source/Editor/UnrealEd/Private/Cooker/OnDemandShaderCompilation.cpp
@@ -4,6 +4,12 @@
 
 #include "Async/TaskGraphInterfaces.h"
 #include "Async/Async.h"
+#include "AssetRegistry/IAssetRegistry.h"
+#include "AssetRegistry/AssetData.h"
+#include "LandscapeMaterialInstanceConstant.h"
+#include "UObject/CoreRedirects.h"
+#include "PackageTools.h"
+#include "WorldPartition/WorldPartitionHelpers.h"
 
 
 int32 GODSCShaderMapsLifetime = 25;
@@ -22,6 +28,14 @@ static FAutoConsoleVariableRef CVarODSCNumShaderMapsBeforeGC(
 	TEXT("-1 means we never delete shader maps\n"),
 	ECVF_Default);
 
+FString GODSCExludedClasses = "/Script/Engine.GameMode";
+static FAutoConsoleVariableRef CVarODSCExludedClasses(
+	TEXT("odsc.excludedclasses"),
+	GODSCExludedClasses,
+	TEXT("Controls what packages will be ignored during material loading to speed it up. This will exclude uassets inheriting from these classes\n")
+	TEXT("This list can contain multiple classes, separated by '|'\n"),
+	ECVF_Default);
+
 
 namespace UE::Cook
 {
@@ -93,7 +107,231 @@ void FODSCClientData::KeepClientPersistentData(const void* ConnectionPtr, const
 				Lifetime = 0;
 			}
 		}
+
+		CleanupWorldPartitionAssets();
+	}
+}
+
+TMap<FString, FODSCClientData::FWorldPartitionAssets> FODSCClientData::WorldPartitionAssets;
+TSet<FString> FODSCClientData::ScannedWorldPartitionPaths;
+TSet<FName> FODSCClientData::ExcludedPackageNames;
+
+void FODSCClientData::CleanupWorldPartitionAssets()
+{
+	TSet<UPackage*> PackagesToUnloadSet;
+	for (auto Iter = WorldPartitionAssets.CreateIterator(); Iter; ++Iter)
+	{
+		FWorldPartitionAssets& DynamicMaterialData = Iter.Value();
+		PackagesToUnloadSet.FindOrAdd(DynamicMaterialData.PackagePtr);
+		DynamicMaterialData.PackagePtr = nullptr;
+	}
+
+	TArray<UPackage*> PackagesToUnloadArray;
+	for (UPackage* Package : PackagesToUnloadSet)
+	{
+		PackagesToUnloadArray.Add(Package);
+	}
+
+	UPackageTools::UnloadPackages(PackagesToUnloadArray);
+}
+
+bool ExtractMaterialPath(FSoftObjectPath& MaterialPath, const FString& MaterialKey)
+{
+	FString MaterialPathString = MaterialKey;
+	if (!FWorldPartitionHelpers::ConvertRuntimePathToEditorPath(MaterialPathString, MaterialPath))
+	{
+		MaterialPathString.ReplaceInline(TEXT("/_Generated_/"), TEXT("/"));
+		MaterialPath = MaterialPathString;
+		return false;
+	}
+	return true;
+}
+
+UMaterialInterface* FODSCClientData::FindMaterial(const FString& InMaterialKey)
+{
+	TRACE_CPUPROFILER_EVENT_SCOPE_TEXT(*FString::Printf(TEXT("FODSCClientData_FindMaterial %s"), *InMaterialKey));
+
+	UE_CALL_ONCE([](){ SetupClassExclusionList(); });
+
+	// Setup some packages exclusion as loading some packages may take a long time and are not necessary in the context of ODSC
+	TArray<FCoreRedirect> CoreRedirectExcludedPackages; 
+	FName InvalidPath(TEXT("/ODSC/Invalid/Path"));
+	{
+		CoreRedirectExcludedPackages.Reserve(ExcludedPackageNames.Num());
+		for (FName ExcludedPackage : ExcludedPackageNames)
+		{
+			CoreRedirectExcludedPackages.Add(FCoreRedirect(ECoreRedirectFlags::Type_Package, ExcludedPackage.ToString(), InvalidPath.ToString()));
+		}
+	}
+
+	FCoreRedirects::AddRedirectList(CoreRedirectExcludedPackages, InvalidPath.ToString());
+
+	ON_SCOPE_EXIT
+	{
+		FCoreRedirects::RemoveRedirectList(CoreRedirectExcludedPackages, InvalidPath.ToString());
+	};
+
+	FSoftObjectPath MaterialPath;
+	bool bIsWorldPartitionPath = ExtractMaterialPath(MaterialPath, InMaterialKey);
+
+	UMaterialInterface* MaterialInterface = nullptr;
+
+	if (bIsWorldPartitionPath)
+	{
+		MaterialInterface = TryFindWorldPartitionMaterial(MaterialPath);
+		if (MaterialInterface)
+		{
+			return MaterialInterface;
+		}
+	}
+
+	MaterialInterface = FindObject<UMaterialInterface>(nullptr, *MaterialPath.ToString());
+	if (MaterialInterface)
+	{
+		return MaterialInterface;
 	}
+
+	MaterialInterface = LoadObject<UMaterialInterface>(nullptr, *MaterialPath.ToString());	
+
+	return MaterialInterface;
 }
 
-}
\ No newline at end of file
+FString GetWorldPartitionActorPath(const FString& MaterialEditorFullPath)
+{
+	FString PathToProxy = MaterialEditorFullPath;
+	// Remove the landscape prefix since Landscape MIC are embedded in the package of their proxy
+	FString LandscapeMaterialInstanceConstantClassName = ULandscapeMaterialInstanceConstant::StaticClass()->GetFName().ToString();
+	int32 SubObjectNameChopIndex = PathToProxy.Find(FString(".") + LandscapeMaterialInstanceConstantClassName, ESearchCase::IgnoreCase, ESearchDir::FromEnd);
+	if (SubObjectNameChopIndex != INDEX_NONE)
+	{
+		PathToProxy.LeftChopInline(PathToProxy.Len() - SubObjectNameChopIndex);
+	}
+
+	return PathToProxy;
+}
+
+UMaterialInterface* FODSCClientData::TryFindWorldPartitionMaterial(const FSoftObjectPath& MaterialSoftPath)
+{
+	FString AssetPath = MaterialSoftPath.GetAssetPath().GetPackageName().ToString();
+
+	bool bAlreadySeenPath;
+	ScannedWorldPartitionPaths.FindOrAdd(AssetPath, &bAlreadySeenPath);
+	if (!bAlreadySeenPath)
+	{
+		ScanWorldPartitionAssets(AssetPath);
+	}
+
+	FString MaterialEditorFullPath = *MaterialSoftPath.ToString();
+	FString PathToProxy = GetWorldPartitionActorPath(MaterialEditorFullPath);
+
+	FWorldPartitionAssets* DynamicMaterialData = WorldPartitionAssets.Find(PathToProxy);
+	if (DynamicMaterialData == nullptr)
+	{
+		return nullptr;
+	}
+
+	if (DynamicMaterialData->PackagePtr == nullptr)
+	{
+		TRACE_CPUPROFILER_EVENT_SCOPE_TEXT(*FString::Printf(TEXT("FODSCClientData_LoadPackage %s"), *DynamicMaterialData->PackageName));
+		DynamicMaterialData->PackagePtr = LoadPackage(nullptr, *DynamicMaterialData->PackageName, LOAD_None);
+	}
+
+	return FindObject<UMaterialInterface>(nullptr, *MaterialEditorFullPath);
+}
+
+void FODSCClientData::ScanWorldPartitionAssets(const FString& AssetPath)
+{
+	if (AssetPath.IsEmpty())
+	{
+		return;
+	}
+
+	TArray<FString> PathsToScan;
+	PathsToScan.Add(AssetPath);
+	PathsToScan.Append(ULevel::GetExternalObjectsPaths(AssetPath));
+
+	// Do a synchronous scan of the level external actors path.					
+	IAssetRegistry& AssetRegistry = IAssetRegistry::GetChecked();
+	{
+		TRACE_CPUPROFILER_EVENT_SCOPE(FODSCClientData_ScanSynchronous);
+		AssetRegistry.ScanSynchronous(PathsToScan, TArray<FString>());
+	}
+
+	FARFilter Filter;
+	Filter.bRecursivePaths = true;
+	Filter.bIncludeOnlyOnDiskAssets = true;
+	for (FString& PackagePath : PathsToScan)
+	{
+		Filter.PackagePaths.Add(*PackagePath);
+	}
+
+	TRACE_CPUPROFILER_EVENT_SCOPE(GetAssets);
+	TArray<FAssetData> Assets;
+	AssetRegistry.GetAssets(Filter, Assets);
+
+	for (FAssetData& AssetData : Assets)
+	{
+		FWorldPartitionAssets DynamicMaterialData{};
+		DynamicMaterialData.PackageName = AssetData.PackageName.ToString();
+		WorldPartitionAssets.Add(AssetData.GetSoftObjectPath().ToString(), DynamicMaterialData);
+
+        // Uncomment to see all the WP actors collected
+//		UE_LOG(LogShaders, Display, TEXT("ScanAssets: %s->%s"), *AssetData.GetSoftObjectPath().ToString(), *DynamicMaterialData.PackageName);
+
+	}
+}
+
+static bool DetectIsUAssetByNames(FStringView PackageName, FStringView ObjectPathName)
+{
+	FStringView PackageBaseName;
+	{
+		// Get everything after the last slash
+		int32 IndexOfLastSlash = INDEX_NONE;
+		PackageName.FindLastChar(TEXT('/'), IndexOfLastSlash);
+		PackageBaseName = PackageName.Mid(IndexOfLastSlash + 1);
+	}
+
+	return PackageBaseName.Equals(ObjectPathName, ESearchCase::IgnoreCase);
+}
+
+static void AddUClassPackagesToExclusionList(TSet<FName>& ExcludedPackageNames, const FString& ClassName)
+{
+	IAssetRegistry& AssetRegistry = IAssetRegistry::GetChecked();
+	FARFilter Filter;
+	Filter.bRecursiveClasses = true;
+	Filter.ClassPaths.Add(FTopLevelAssetPath(ClassName));
+	TSet<FTopLevelAssetPath> DerivedClassNames;
+	AssetRegistry.GetDerivedClassNames(Filter.ClassPaths, {}, DerivedClassNames);
+
+	// Still keep script packages
+	for (FTopLevelAssetPath& DerivedClassName : DerivedClassNames)
+	{
+		FString ObjectName = DerivedClassName.GetAssetName().ToString();
+		// Remove suffix for compiled blueprint classes
+		ObjectName.RemoveFromEnd(TEXT("_C"));
+		
+		bool bIsUAsset = DetectIsUAssetByNames(DerivedClassName.GetPackageName().ToString(), ObjectName);
+		if (!bIsUAsset)
+		{
+			continue;
+		}
+
+		UE_LOG(LogShaders, Display, TEXT("FODSCClientData excluding package '%s' because '%s' is a uasset inheriting from %s"), 
+			   *DerivedClassName.GetPackageName().ToString(), *DerivedClassName.ToString(), *ClassName);
+
+		ExcludedPackageNames.Add(FName(DerivedClassName.GetPackageName().ToString()));
+	}
+}
+
+void FODSCClientData::SetupClassExclusionList()
+{
+	TArray<FString> ClassNames;
+	
+	GODSCExludedClasses.ParseIntoArray(ClassNames, TEXT("|"));
+	for (FString& ClassName : ClassNames)
+	{
+		AddUClassPackagesToExclusionList(ExcludedPackageNames, ClassName);
+	}
+}
+
+}
diff --git a/Engine/Source/Editor/UnrealEd/Private/Cooker/OnDemandShaderCompilation.h b/Engine/Source/Editor/UnrealEd/Private/Cooker/OnDemandShaderCompilation.h
--- a/Engine/Source/Editor/UnrealEd/Private/Cooker/OnDemandShaderCompilation.h
+++ b/Engine/Source/Editor/UnrealEd/Private/Cooker/OnDemandShaderCompilation.h
@@ -18,7 +18,25 @@ public:
 	void KeepClientPersistentData(const void* ConnectionPtr, const TArray<TStrongObjectPtr<UMaterialInterface>>& LoadedMaterialsToRecompile);
 	void FlushClientPersistentData(const void* ConnectionPtr);
 
+	static UMaterialInterface* FindMaterial(const FString& MaterialKey);
+
 private:
+	
+	static UMaterialInterface* TryFindWorldPartitionMaterial(const FSoftObjectPath& MaterialSoftPath);
+
+	static void ScanWorldPartitionAssets(const FString& AssetPath);
+	static void SetupClassExclusionList();
+
+	void CleanupWorldPartitionAssets();
+	struct FWorldPartitionAssets
+	{
+		FString PackageName;
+		TObjectPtr<UPackage> PackagePtr;
+	};
+
+	static TMap<FString, FWorldPartitionAssets> WorldPartitionAssets;
+	static TSet<FString> ScannedWorldPartitionPaths;
+	static TSet<FName> ExcludedPackageNames;
 
 	struct FODSCClientPersistentData
 	{
diff --git a/Engine/Source/Runtime/Engine/Private/Materials/Material.cpp b/Engine/Source/Runtime/Engine/Private/Materials/Material.cpp
--- a/Engine/Source/Runtime/Engine/Private/Materials/Material.cpp
+++ b/Engine/Source/Runtime/Engine/Private/Materials/Material.cpp
@@ -5699,7 +5699,8 @@ void UMaterial::CompileMaterialsForRemoteRecompile(
 	}
 }
 
-void UMaterial::CompileODSCMaterialsForRemoteRecompile(TArray<FODSCRequestPayload> ShadersToRecompile, TMap<FString, TArray<TRefCountPtr<class FMaterialShaderMap>>>& OutShaderMaps)
+void UMaterial::CompileODSCMaterialsForRemoteRecompile(TArray<FODSCRequestPayload> ShadersToRecompile, TMap<FString, TArray<TRefCountPtr<class FMaterialShaderMap>>>& OutShaderMaps, 
+													   TFunction<UMaterialInterface*(const FString&)> ODSCCustomLoadMaterial)
 {
 	TRACE_CPUPROFILER_EVENT_SCOPE(UMaterial::CompileODSCMaterialsForRemoteRecompile);
 
@@ -5720,15 +5721,31 @@ void UMaterial::CompileODSCMaterialsForRemoteRecompile(TArray<FODSCRequestPayloa
 		TArray<const FVertexFactoryType*> VFTypes;
 		TArray<const FShaderPipelineType*> PipelineTypes;
 		TArray<const FShaderType*> ShaderTypes;
+		FString OriginalMaterialName;
 	};
 	TMap<UMaterialInterface*, FShadersToCompile> CoalescedShadersToCompile;
 
+	TSet<FString> MaterialsFailingToLoad;
 	for (const FODSCRequestPayload& payload : ShadersToRecompile)
 	{
-		UE_LOG(LogShaders, Display, TEXT(""));
-		UE_LOG(LogShaders, Display, TEXT("Material:    %s "), *payload.MaterialName);
+		UE_LOG(LogShaders, Verbose, TEXT(""));
+		UE_LOG(LogShaders, Verbose, TEXT("Material:    %s "), *payload.MaterialName);
+
+		UMaterialInterface* MaterialInterface = nullptr;
+		if (ODSCCustomLoadMaterial)
+		{
+			MaterialInterface = ODSCCustomLoadMaterial(payload.MaterialName);
+		}
+		else
+		{
+			MaterialInterface = LoadObject<UMaterialInterface>(nullptr, *payload.MaterialName);
+		}
+
+		if (!MaterialInterface)
+		{
+			MaterialsFailingToLoad.Add(payload.MaterialName);
+		}
 
-		UMaterialInterface* MaterialInterface = LoadObject<UMaterialInterface>(nullptr, *payload.MaterialName);
 		if (MaterialInterface)
 		{
 			FShadersToCompile& Shaders = CoalescedShadersToCompile.FindOrAdd(MaterialInterface);
@@ -5738,15 +5755,16 @@ void UMaterial::CompileODSCMaterialsForRemoteRecompile(TArray<FODSCRequestPayloa
 			Shaders.ShaderPlatform = payload.ShaderPlatform;
 			Shaders.FeatureLevel = payload.FeatureLevel;
 			Shaders.QualityLevel = payload.QualityLevel;
+			Shaders.OriginalMaterialName = payload.MaterialName;
 
 			if (VFType)
 			{
-				UE_LOG(LogShaders, Display, TEXT("VF Type:     %s "), *payload.VertexFactoryName);
+				UE_LOG(LogShaders, Verbose, TEXT("VF Type:     %s "), *payload.VertexFactoryName);
 			}
 
 			if (PipelineType)
 			{
-				UE_LOG(LogShaders, Display, TEXT("Pipeline Type: %s"), *payload.PipelineName);
+				UE_LOG(LogShaders, Verbose, TEXT("Pipeline Type: %s"), *payload.PipelineName);
 
 				Shaders.VFTypes.Add(VFType);
 				Shaders.PipelineTypes.Add(PipelineType);
@@ -5759,7 +5777,7 @@ void UMaterial::CompileODSCMaterialsForRemoteRecompile(TArray<FODSCRequestPayloa
 					const FShaderType* ShaderType = FShaderType::GetShaderTypeByName(*ShaderTypeName);
 					if (ShaderType)
 					{
-						UE_LOG(LogShaders, Display, TEXT("\tShader Type: %s"), *ShaderTypeName);
+						UE_LOG(LogShaders, Verbose, TEXT("\tShader Type: %s"), *ShaderTypeName);
 
 						Shaders.VFTypes.Add(VFType);
 						Shaders.PipelineTypes.Add(nullptr);
@@ -5784,7 +5802,7 @@ void UMaterial::CompileODSCMaterialsForRemoteRecompile(TArray<FODSCRequestPayloa
 		UMaterialInterface* MaterialInterface = Entry.Key;
 		const FShadersToCompile& Shaders = Entry.Value;
 
-		TArray<FMaterialResource*>& ResourceArray = CompilingResources.Add(MaterialInterface->GetPathName(), TArray<FMaterialResource*>());
+		TArray<FMaterialResource*>& ResourceArray = CompilingResources.Add(Shaders.OriginalMaterialName, TArray<FMaterialResource*>());
 		FMaterialResource* MaterialResource = MaterialInterface->GetMaterialResource(Shaders.FeatureLevel, Shaders.QualityLevel);
 		check(MaterialResource);
 		check(MaterialResource->GetFeatureLevel() == Shaders.FeatureLevel);
@@ -5809,6 +5827,12 @@ void UMaterial::CompileODSCMaterialsForRemoteRecompile(TArray<FODSCRequestPayloa
 			OutShaderMapArray.Add(CurrentResource->GetGameThreadShaderMap());
 		}
 	}
+
+	// Report errors last to ensure visibility: some messages may have been output during shader compilation
+	for (FString& MaterialFailingToLoad : MaterialsFailingToLoad)
+	{
+		UE_LOG(LogShaders, Warning, TEXT("Failed to load %s, skipping shader reloading"), *MaterialFailingToLoad);
+	}
 }
 #endif // WITH_EDITOR
 
diff --git a/Engine/Source/Runtime/Engine/Private/Materials/MaterialShared.cpp b/Engine/Source/Runtime/Engine/Private/Materials/MaterialShared.cpp
--- a/Engine/Source/Runtime/Engine/Private/Materials/MaterialShared.cpp
+++ b/Engine/Source/Runtime/Engine/Private/Materials/MaterialShared.cpp
@@ -3727,7 +3727,7 @@ bool FMaterial::TryGetShaders(const FMaterialShaderTypes& InTypes, const FVertex
 	}
 
 #if WITH_ODSC
-	bool bShouldForceRecompile = FODSCManager::ShouldForceRecompile(ShaderMap, this) && !IsDefaultMaterial() && !IsUsedWithLandscape();
+	bool bShouldForceRecompile = FODSCManager::ShouldForceRecompile(ShaderMap, this) && !IsDefaultMaterial();
 #else
 	constexpr bool bShouldForceRecompile = false;
 #endif
diff --git a/Engine/Source/Runtime/Engine/Private/ShaderCompiler/ShaderCompilerEditor.cpp b/Engine/Source/Runtime/Engine/Private/ShaderCompiler/ShaderCompilerEditor.cpp
--- a/Engine/Source/Runtime/Engine/Private/ShaderCompiler/ShaderCompilerEditor.cpp
+++ b/Engine/Source/Runtime/Engine/Private/ShaderCompiler/ShaderCompilerEditor.cpp
@@ -369,17 +369,31 @@ void RecompileShadersForRemote(
 			// UMaterial::CompileODSCMaterialsForRemoteRecompile will call LoadObjects on the material names but doesn't keep them around. Add a GC guard to ensure we can still get them
 			// before they get unloaded, so that the whole chain UMaterial->FMaterial->FMaterialShaderMap is kept intact, and we can merge the next batch of ODSC requests
 			FGCScopeGuard NoGCScopeGuard;
-			UMaterial::CompileODSCMaterialsForRemoteRecompile(Args.ShadersToRecompile, CompiledShaderMaps);
+			UMaterial::CompileODSCMaterialsForRemoteRecompile(Args.ShadersToRecompile, CompiledShaderMaps, Args.ODSCCustomLoadMaterial);
 			if (Args.LoadedMaterialsToRecompile)
 			{
 				for (auto Iter : CompiledShaderMaps)
 				{
-					TStrongObjectPtr<UMaterialInterface> MaterialInterface = TStrongObjectPtr<UMaterialInterface>(FindObject<UMaterialInterface>(nullptr, *Iter.Key));
+					TStrongObjectPtr<UMaterialInterface> MaterialInterface;
+					
+					if (Args.ODSCCustomLoadMaterial)
+					{
+						MaterialInterface = TStrongObjectPtr<UMaterialInterface>(Args.ODSCCustomLoadMaterial(Iter.Key));
+					}
+					else
+					{
+						MaterialInterface = TStrongObjectPtr<UMaterialInterface>(FindObject<UMaterialInterface>(nullptr, *Iter.Key));
+					}
+
 					if (MaterialInterface)
 					{
 						ResetLoaders(MaterialInterface->GetPackage());
 						Args.LoadedMaterialsToRecompile->Add(MaterialInterface);
 					}
+					else
+					{
+						UE_LOG(LogShaders, Warning, TEXT("Failed to find Material %s. Reloading on the client will be skipped"), *Iter.Key);
+					}
 				}
 			}
 		}
diff --git a/Engine/Source/Runtime/Engine/Public/Materials/Material.h b/Engine/Source/Runtime/Engine/Public/Materials/Material.h
--- a/Engine/Source/Runtime/Engine/Public/Materials/Material.h
+++ b/Engine/Source/Runtime/Engine/Public/Materials/Material.h
@@ -1633,7 +1633,7 @@ public:
 		class ITargetPlatform* TargetPlatform,
 		TMap<FString, TArray<TRefCountPtr<class FMaterialShaderMap> > >& OutShaderMaps);
 
-	ENGINE_API static void CompileODSCMaterialsForRemoteRecompile(TArray<FODSCRequestPayload> ShadersToRecompile, TMap<FString, TArray<TRefCountPtr<class FMaterialShaderMap>>>& OutShaderMaps);
+	ENGINE_API static void CompileODSCMaterialsForRemoteRecompile(TArray<FODSCRequestPayload> ShadersToRecompile, TMap<FString, TArray<TRefCountPtr<class FMaterialShaderMap>>>& OutShaderMaps, TFunction<UMaterialInterface*(const FString&)> ODSCCustomLoadMaterial = {});
 
 	/**
 	 * Add an expression node that represents a parameter to the list of material parameters.
diff --git a/Engine/Source/Runtime/Engine/Public/ShaderCompiler.h b/Engine/Source/Runtime/Engine/Public/ShaderCompiler.h
--- a/Engine/Source/Runtime/Engine/Public/ShaderCompiler.h
+++ b/Engine/Source/Runtime/Engine/Public/ShaderCompiler.h
@@ -1252,6 +1252,10 @@ struct FShaderRecompileData
 	/** Optional Array of the loaded materials  */
 	TArray<TStrongObjectPtr<UMaterialInterface>>* LoadedMaterialsToRecompile = nullptr;
 
+#if WITH_EDITOR
+	TFunction<UMaterialInterface*(const FString&)> ODSCCustomLoadMaterial;
+#endif
+
 	/** Default constructor. */
 	FShaderRecompileData() {};
 
