[Viewport Toolbar] Add new viewport toolbar to the Material editor

Adds the new toolbar and the first entries for:

- Camera
- View Modes
- Show
- Performance and Scalability

====
Quick note about refactoring of legacy toolbars:

This CL will move some camera-related code from SCommonEditorViewportToolbarBase to UnrealViewportToolbar.cpp functions:

- Camera perspective selection menu
- Camera FOV
- Camera Far View Plane
- Screen percentage widgets creation

There are other toolbars using the SCommonEditorViewportToolbarBase version of these functions. These calls can be substituted as well in order to remove/deprecate the legacy functions, but in order to reduce the scope of this CL I'll address that in a separate CL.
====

diff --git a/Engine/Source/Editor/MaterialEditor/Private/MaterialEditorViewportToolbarSections.cpp b/Engine/Source/Editor/MaterialEditor/Private/MaterialEditorViewportToolbarSections.cpp
new file mode 100644
--- /dev/null
+++ b/Engine/Source/Editor/MaterialEditor/Private/MaterialEditorViewportToolbarSections.cpp
@@ -0,0 +1,67 @@
+﻿// Copyright Epic Games, Inc. All Rights Reserved.
+
+#include "MaterialEditorViewportToolbarSections.h"
+#include "EditorViewportCommands.h"
+#include "MaterialEditorActions.h"
+#include "SEditorViewport.h"
+#include "SMaterialEditorViewport.h"
+#include "ToolMenu.h"
+
+#define LOCTEXT_NAMESPACE "MaterialEditorViewportToolbarSections"
+
+TSharedRef<SWidget> UE::MaterialEditor::CreateShowMenuWidget(const TSharedRef<SMaterialEditor3DPreviewViewport>& InMaterialEditorViewport, bool bInShowViewportStatsToggle
+)
+{
+	InMaterialEditorViewport->OnFloatingButtonClicked();
+
+	constexpr bool bInShouldCloseWindowAfterMenuSelection = true;
+	FMenuBuilder ShowMenuBuilder(bInShouldCloseWindowAfterMenuSelection, InMaterialEditorViewport->GetCommandList());
+	{
+		FMaterialEditorCommands Commands = FMaterialEditorCommands::Get();
+
+		if (bInShowViewportStatsToggle)
+		{
+			ShowMenuBuilder.AddMenuEntry(
+				FEditorViewportCommands::Get().ToggleStats, "ViewportStats", LOCTEXT("ViewportStatsLabel", "Viewport Stats")
+			);
+
+			ShowMenuBuilder.AddMenuSeparator();
+		}
+
+		ShowMenuBuilder.AddMenuEntry(Commands.ToggleMaterialStats);
+
+		ShowMenuBuilder.AddMenuSeparator();
+
+		ShowMenuBuilder.AddMenuEntry(Commands.TogglePreviewGrid);
+		ShowMenuBuilder.AddMenuEntry(Commands.TogglePreviewBackground);
+	}
+
+	return ShowMenuBuilder.MakeWidget();
+}
+
+FToolMenuEntry UE::MaterialEditor::CreateShowSubmenu(TWeakPtr<SMaterialEditor3DPreviewViewport> InViewport)
+{
+	return FToolMenuEntry::InitSubMenu(
+		"Show",
+		LOCTEXT("ShowSubmenuLabel", "Show"),
+		LOCTEXT("ShowSubmenuTooltip", "Show options"),
+		FNewToolMenuDelegate::CreateLambda(
+			[InViewport](UToolMenu* Submenu) -> void
+			{
+				TSharedPtr<SMaterialEditor3DPreviewViewport> Viewport = InViewport.Pin();
+				if (!Viewport)
+				{
+					return;
+				}
+
+				FToolMenuSection& UnnamedSection = Submenu->FindOrAddSection("", LOCTEXT("UnnamedLabel", ""));
+
+				UnnamedSection.AddEntry(FToolMenuEntry::InitWidget(
+					"ShowMenuItems", UE::MaterialEditor::CreateShowMenuWidget(Viewport.ToSharedRef()), FText(), true
+				));
+			}
+		)
+	);
+}
+
+#undef LOCTEXT_NAMESPACE
diff --git a/Engine/Source/Editor/MaterialEditor/Private/MaterialEditorViewportToolbarSections.h b/Engine/Source/Editor/MaterialEditor/Private/MaterialEditorViewportToolbarSections.h
new file mode 100644
--- /dev/null
+++ b/Engine/Source/Editor/MaterialEditor/Private/MaterialEditorViewportToolbarSections.h
@@ -0,0 +1,14 @@
+﻿// Copyright Epic Games, Inc. All Rights Reserved.
+
+#pragma once
+
+#include "Templates/SharedPointerFwd.h"
+#include "ToolMenuEntry.h"
+
+class SMaterialEditor3DPreviewViewport;
+
+namespace UE::MaterialEditor
+{
+TSharedRef<SWidget> CreateShowMenuWidget(const TSharedRef<SMaterialEditor3DPreviewViewport>& InMaterialEditorViewport, bool bInShowViewportStatsToggle = true);
+FToolMenuEntry CreateShowSubmenu(TWeakPtr<SMaterialEditor3DPreviewViewport> InViewport);
+}
diff --git a/Engine/Source/Editor/MaterialEditor/Private/SMaterialEditorViewport.cpp b/Engine/Source/Editor/MaterialEditor/Private/SMaterialEditorViewport.cpp
--- a/Engine/Source/Editor/MaterialEditor/Private/SMaterialEditorViewport.cpp
+++ b/Engine/Source/Editor/MaterialEditor/Private/SMaterialEditorViewport.cpp
@@ -44,6 +44,9 @@
 #include "ImageUtils.h"
 #include "ISettingsModule.h"
 #include "Framework/Layout/ScrollyZoomy.h"
+#include "MaterialEditorViewportToolbarSections.h"
+#include "ToolMenus.h"
+#include "ViewportToolbar/UnrealEdViewportToolbar.h"
 
 #define LOCTEXT_NAMESPACE "MaterialEditor"
 #include "UnrealWidget.h"
@@ -947,26 +950,160 @@ TSharedRef<FEditorViewportClient> SMaterialEditor3DPreviewViewport::MakeEditorVi
 
 void SMaterialEditor3DPreviewViewport::PopulateViewportOverlays(TSharedRef<class SOverlay> Overlay)
 {
+	const TSharedRef<SMaterialEditorViewportToolBar> OldViewportToolbar =
+		// clang-format off
+		SNew(SMaterialEditorViewportToolBar, SharedThis(this))
+			.Visibility_Lambda([this]() -> EVisibility
+			{
+				if (!UE::UnrealEd::ShowOldViewportToolbars())
+				{
+					return EVisibility::Collapsed;
+				}
+
+				return EVisibility::Visible;
+			});
+	// clang-format on
+
+	const FName MaterialEditorViewportToolbarName = "MaterialEditor.ViewportToolbar";
+
+	// Register the viewport toolbar if another viewport hasn't already (it's shared).
+	{
+		if (!UToolMenus::Get()->IsMenuRegistered(MaterialEditorViewportToolbarName))
+		{
+			UToolMenu* const ViewportToolbarMenu = UToolMenus::Get()->RegisterMenu(
+				MaterialEditorViewportToolbarName, NAME_None /* parent */, EMultiBoxType::SlimHorizontalToolBar
+			);
+
+			// Add the Left-aligned part of the viewport toolbar.
+			{
+				// Adding it even if empty in order to keep proper toolbar layout
+				ViewportToolbarMenu->FindOrAddSection("Left");
+			}
+
+			// Add the right-aligned part of the viewport toolbar.
+			{
+				FToolMenuSection& RightSection = ViewportToolbarMenu->FindOrAddSection("Right");
+				RightSection.Alignment = EToolMenuSectionAlign::Last;
+
+				// Add the "Camera" submenu.
+				{
+					FToolMenuEntry CameraSubmenu = UE::UnrealEd::CreateCameraSubmenu(SharedThis(this));
+					CameraSubmenu.InsertPosition.Position = EToolMenuInsertType::First;
+					RightSection.AddEntry(CameraSubmenu);
+				}
+
+				// Add the View Modes submenu.
+				{
+					FToolMenuEntry ViewModesSubmenu = UE::UnrealEd::CreateViewportToolbarViewModesSubmenu();
+					ViewModesSubmenu.InsertPosition.Position = EToolMenuInsertType::First;
+					RightSection.AddEntry(ViewModesSubmenu);
+				}
+
+				// Add the Show submenu.
+				{
+					FToolMenuEntry ShowSubmenu = UE::MaterialEditor::CreateShowSubmenu(SharedThis(this));
+					ShowSubmenu.InsertPosition.Position = EToolMenuInsertType::First;
+					RightSection.AddEntry(ShowSubmenu);
+				}
+
+				// Add the Performance and Scalability submenu.
+				{
+					FToolMenuEntry PerformanceAndScalabilitySubmenu = UE::UnrealEd::CreatePerformanceAndScalabilitySubmenu(SharedThis(this));
+					PerformanceAndScalabilitySubmenu.InsertPosition.Position = EToolMenuInsertType::First;
+					RightSection.AddEntry(PerformanceAndScalabilitySubmenu);
+				}
+
+				// Add the Settings submenu.
+				{
+					//CreatePerformanceAndScalabilitySubmenu
+					FToolMenuEntry SettingsSubmenu = FToolMenuEntry::InitSubMenu(
+						"Settings",
+						LOCTEXT("SettingsSubmenuLabel", "Settings"),
+						LOCTEXT("SettingsSubmenuTooltip", "Show flags related to the current viewport"),
+						FNewToolMenuDelegate::CreateLambda(
+							[](UToolMenu* InMenu) -> void
+							{
+								//TODO: fill submenu
+							}
+						)
+					);
+					SettingsSubmenu.InsertPosition.Position = EToolMenuInsertType::First;
+					RightSection.AddEntry(SettingsSubmenu);
+				}
+			}
+		}
+	}
+
+	FToolMenuContext ViewportToolbarContext;
+	{
+		ViewportToolbarContext.AppendCommandList(GetCommandList());
+
+		// Add the UnrealEd viewport toolbar context.
+		{
+			UUnrealEdViewportToolbarContext* const ContextObject = NewObject<UUnrealEdViewportToolbarContext>();
+			ContextObject->Viewport = SharedThis(this);
+			ViewportToolbarContext.AddObject(ContextObject);
+		}
+	}
+		const TSharedRef<SWidget> NewViewportToolbar = SNew(SBox)
+		// clang-format off
+		.Visibility_Lambda(
+			[this]() -> EVisibility
+			{
+				if (!UE::UnrealEd::ShowNewViewportToolbars())
+				{
+					return EVisibility::Collapsed;
+				}
+
+				return EVisibility::Visible;
+			}
+		)
+		[
+			UToolMenus::Get()->GenerateWidget(MaterialEditorViewportToolbarName, ViewportToolbarContext)
+		];
+		// clang-format on
+
 	Overlay->AddSlot()
+		// clang-format off
 		.VAlign(VAlign_Top)
 		[
-			SNew(SMaterialEditorViewportToolBar, SharedThis(this))
+			SNew(SVerticalBox)
+			.Visibility( EVisibility::SelfHitTestInvisible )
+			+ SVerticalBox::Slot()
+			.AutoHeight()
+			.Padding(0.0f, 1.0f, 0.0f, 0.0f)
+			.VAlign(VAlign_Top)
+			[
+				OldViewportToolbar
+			]
+			+ SVerticalBox::Slot()
+			.AutoHeight()
+			.Padding(0.0f, 1.0f, 0.0f, 0.0f)
+			.VAlign(VAlign_Top)
+			[
+				NewViewportToolbar
+			]
 		];
+		// clang-format on
 
 	Overlay->AddSlot()
+		// clang-format off
 		.VAlign(VAlign_Bottom)
 		[
 			SNew(SMaterialEditorViewportPreviewShapeToolBar, SharedThis(this))
 		];
+		// clang-format on
 
 	// add the feature level display widget
 	Overlay->AddSlot()
+		// clang-format off
 		.VAlign(VAlign_Top)
 		.HAlign(HAlign_Right)
 		.Padding(5.0f)
 		[
 			BuildFeatureLevelWidget()
 		];
+		// clang-format on
 }
 
 EVisibility SMaterialEditor3DPreviewViewport::OnGetViewportContentVisibility() const
diff --git a/Engine/Source/Editor/MaterialEditor/Private/SMaterialEditorViewportToolBar.cpp b/Engine/Source/Editor/MaterialEditor/Private/SMaterialEditorViewportToolBar.cpp
--- a/Engine/Source/Editor/MaterialEditor/Private/SMaterialEditorViewportToolBar.cpp
+++ b/Engine/Source/Editor/MaterialEditor/Private/SMaterialEditorViewportToolBar.cpp
@@ -4,6 +4,7 @@
 #include "Widgets/Layout/SBorder.h"
 #include "Styling/AppStyle.h"
 #include "MaterialEditorActions.h"
+#include "MaterialEditorViewportToolbarSections.h"
 #include "PreviewProfileController.h"
 
 #define LOCTEXT_NAMESPACE "MaterialEditorViewportToolBar"
@@ -53,29 +54,19 @@ void SMaterialEditorViewportPreviewShapeToolBar::Construct(const FArguments& InA
 
 void SMaterialEditorViewportToolBar::Construct(const FArguments& InArgs, TSharedPtr<class SMaterialEditor3DPreviewViewport> InViewport)
 {
+	MaterialEditorViewportPtr = InViewport;
 	SCommonEditorViewportToolbarBase::Construct(SCommonEditorViewportToolbarBase::FArguments().PreviewProfileController(MakeShared<FPreviewProfileController>()), InViewport);
 }
 
 TSharedRef<SWidget> SMaterialEditorViewportToolBar::GenerateShowMenu() const
 {
-	GetInfoProvider().OnFloatingButtonClicked();
-
-	TSharedRef<SEditorViewport> ViewportRef = GetInfoProvider().GetViewportWidget();
-
-	const bool bInShouldCloseWindowAfterMenuSelection = true;
-	FMenuBuilder ShowMenuBuilder(bInShouldCloseWindowAfterMenuSelection, ViewportRef->GetCommandList());
+	if (MaterialEditorViewportPtr)
 	{
-		auto Commands = FMaterialEditorCommands::Get();
-
-		ShowMenuBuilder.AddMenuEntry(Commands.ToggleMaterialStats);
-
-		ShowMenuBuilder.AddMenuSeparator();
-
-		ShowMenuBuilder.AddMenuEntry(Commands.TogglePreviewGrid);
-		ShowMenuBuilder.AddMenuEntry(Commands.TogglePreviewBackground);
+		constexpr bool bShowViewportStats = false;
+		return UE::MaterialEditor::CreateShowMenuWidget(MaterialEditorViewportPtr.ToSharedRef(), bShowViewportStats);
 	}
 
-	return ShowMenuBuilder.MakeWidget();
+	return SNullWidget::NullWidget;
 }
 
 bool SMaterialEditorViewportToolBar::IsViewModeSupported(EViewModeIndex ViewModeIndex) const 
diff --git a/Engine/Source/Editor/MaterialEditor/Private/SMaterialEditorViewportToolBar.h b/Engine/Source/Editor/MaterialEditor/Private/SMaterialEditorViewportToolBar.h
--- a/Engine/Source/Editor/MaterialEditor/Private/SMaterialEditorViewportToolBar.h
+++ b/Engine/Source/Editor/MaterialEditor/Private/SMaterialEditorViewportToolBar.h
@@ -25,7 +25,10 @@ public:
 	virtual TSharedRef<SWidget> GenerateShowMenu() const override;
 	// End of SCommonEditorViewportToolbarBase
 
-	virtual bool IsViewModeSupported(EViewModeIndex ViewModeIndex) const override; 
+	virtual bool IsViewModeSupported(EViewModeIndex ViewModeIndex) const override;
+
+private:
+	TSharedPtr<SMaterialEditor3DPreviewViewport> MaterialEditorViewportPtr;
 };
 
 ///////////////////////////////////////////////////////////
diff --git a/Engine/Source/Editor/UnrealEd/Private/SCommonEditorViewportToolbarBase.cpp b/Engine/Source/Editor/UnrealEd/Private/SCommonEditorViewportToolbarBase.cpp
--- a/Engine/Source/Editor/UnrealEd/Private/SCommonEditorViewportToolbarBase.cpp
+++ b/Engine/Source/Editor/UnrealEd/Private/SCommonEditorViewportToolbarBase.cpp
@@ -10,7 +10,6 @@
 #include "Styling/AppStyle.h"
 
 #include "STransformViewportToolbar.h"
-#include "EditorShowFlags.h"
 #include "SEditorViewport.h"
 #include "EditorViewportCommands.h"
 #include "SEditorViewportToolBarMenu.h"
@@ -21,13 +20,11 @@
 #include "Scalability.h"
 #include "SceneView.h"
 #include "SScalabilitySettings.h"
-#include "AssetEditorViewportLayout.h"
 #include "SAssetEditorViewport.h"
 #include "ToolMenu.h"
 #include "ToolMenus.h"
-#include "ToolMenuSection.h"
 #include "ShowFlagMenuCommands.h"
-
+#include "ViewportToolbar/UnrealEdViewportToolbar.h"
 
 #define LOCTEXT_NAMESPACE "SCommonEditorViewportToolbarBase"
 
@@ -417,70 +414,24 @@ void SCommonEditorViewportToolbarBase::ConstructScreenPercentageMenu(FMenuBuilde
 	MenuBuilder.BeginSection("Summary", LOCTEXT("Summary", "Summary"));
 	{
 		MenuBuilder.AddWidget(
-			SNew(SBox)
-			.Padding(CommonPadding)
-			[
-				SNew(STextBlock)
-				.ColorAndOpacity(FSlateColor::UseSubduedForeground())
-				.Text_Lambda([&ViewportClient]() {
-					FFormatNamedArguments FormatArguments = GetScreenPercentageFormatArguments(ViewportClient);
-					return FText::Format(LOCTEXT("ScreenPercentageCurrent_Display", "Current Screen Percentage: {CurrentScreenPercentage}"), FormatArguments);
-				})
-				.ToolTip(SNew(SToolTip).Text(LOCTEXT("ScreenPercentageCurrent_ToolTip", "Current Screen Percentage the viewport is rendered with. The primary screen percentage can either be a spatial or temporal upscaler based of your anti-aliasing settings.")))
-			],
+			UE::UnrealEd::CreateCurrentPercentageWidget(ViewportClient),
 			FText::GetEmpty()
 		);
 
 		MenuBuilder.AddWidget(
-			SNew(SBox)
-			.Padding(CommonPadding)
-			[
-				SNew(STextBlock)
-				.ColorAndOpacity(FSlateColor::UseSubduedForeground())
-				.Text_Lambda([&ViewportClient]() {
-					FFormatNamedArguments FormatArguments = GetScreenPercentageFormatArguments(ViewportClient);
-					return FText::Format(LOCTEXT("ScreenPercentageResolutions", "Resolution: {ResolutionFromTo}"), FormatArguments);
-				})
-			],
+			UE::UnrealEd::CreateResolutionsWidget(ViewportClient),
 			FText::GetEmpty()
 		);
 		MenuBuilder.AddWidget(
-			SNew(SBox)
-			.Padding(CommonPadding)
-			[
-				SNew(STextBlock)
-				.ColorAndOpacity(FSlateColor::UseSubduedForeground())
-				.Text_Lambda([&ViewportClient]() {
-					FFormatNamedArguments FormatArguments = GetScreenPercentageFormatArguments(ViewportClient);
-					return FText::Format(LOCTEXT("ScreenPercentageActiveViewport", "Active Viewport: {ViewportMode}"), FormatArguments);
-				})
-			],
+			UE::UnrealEd::CreateActiveViewportWidget(ViewportClient),
 			FText::GetEmpty()
 		);
 		MenuBuilder.AddWidget(
-			SNew(SBox)
-			.Padding(CommonPadding)
-			[
-				SNew(STextBlock)
-				.ColorAndOpacity(FSlateColor::UseSubduedForeground())
-				.Text_Lambda([&ViewportClient]() {
-					FFormatNamedArguments FormatArguments = GetScreenPercentageFormatArguments(ViewportClient);
-					return FText::Format(LOCTEXT("ScreenPercentageSetFrom", "Set From: {SettingSource}"), FormatArguments);
-				})
-			],
+			UE::UnrealEd::CreateSetFromWidget(ViewportClient),
 			FText::GetEmpty()
 		);
 		MenuBuilder.AddWidget(
-			SNew(SBox)
-			.Padding(CommonPadding)
-			[
-				SNew(STextBlock)
-				.ColorAndOpacity(FSlateColor::UseSubduedForeground())
-				.Text_Lambda([&ViewportClient]() {
-					FFormatNamedArguments FormatArguments = GetScreenPercentageFormatArguments(ViewportClient);
-					return FText::Format(LOCTEXT("ScreenPercentageSetting", "Setting: {Setting}"), FormatArguments);
-				})
-			],
+			UE::UnrealEd::CreateCurrentScreenPercentageSettingWidget(ViewportClient),
 			FText::GetEmpty()
 		);
 	}
@@ -490,35 +441,7 @@ void SCommonEditorViewportToolbarBase::ConstructScreenPercentageMenu(FMenuBuilde
 	{
 		MenuBuilder.AddMenuEntry(BaseViewportCommands.ToggleOverrideViewportScreenPercentage);
 		MenuBuilder.AddWidget(
-			SNew(SBox)
-			.HAlign(HAlign_Right)
-			.IsEnabled_Lambda([&ViewportClient]() {
-				return ViewportClient.IsPreviewingScreenPercentage() && ViewportClient.SupportsPreviewResolutionFraction();
-			})
-			[
-				SNew(SBox)
-				.Padding(FMargin(4.0f, 0.0f, 0.0f, 0.0f))
-				.WidthOverride(100.0f)
-				[
-					SNew(SBorder)
-					//.BorderImage(FAppStyle::Get().GetBrush("Menu.WidgetBorder"))
-					.Padding(FMargin(1.0f))
-					[
-						SNew(SSpinBox<int32>)
-						.Style(&FAppStyle::Get(), "Menu.SpinBox")
-						.Font(FAppStyle::GetFontStyle(TEXT("MenuItem.Font")))
-						.MinSliderValue(PreviewScreenPercentageMin)
-						.MaxSliderValue(PreviewScreenPercentageMax)
-						.Value_Lambda([&ViewportClient]() {
-							return ViewportClient.GetPreviewScreenPercentage();
-						})
-						.OnValueChanged_Lambda([&ViewportClient](int32 NewValue) {
-							ViewportClient.SetPreviewScreenPercentage(NewValue);
-							ViewportClient.Invalidate();
-						})
-					]
-				]
-			],
+			UE::UnrealEd::CreateCurrentScreenPercentageWidget(ViewportClient),
 			LOCTEXT("ScreenPercentage", "Screen Percentage")
 		);
 	}
@@ -587,8 +510,8 @@ TSharedRef<SWidget> SCommonEditorViewportToolbarBase::GenerateOptionsMenu() cons
 
 			if (bIsPerspective)
 			{
-				OptionsMenuBuilder.AddWidget( GenerateFOVMenu(), LOCTEXT("FOVAngle", "Field of View (H)") );
-				OptionsMenuBuilder.AddWidget( GenerateFarViewPlaneMenu(), LOCTEXT("FarViewPlane", "Far View Plane") );
+				OptionsMenuBuilder.AddWidget( UE::UnrealEd::CreateFOVMenuWidget(ViewportRef), LOCTEXT("FOVAngle", "Field of View (H)") );
+				OptionsMenuBuilder.AddWidget( UE::UnrealEd::CreateFarViewPlaneMenuWidget(ViewportRef), LOCTEXT("FarViewPlane", "Far View Plane") );
 			}
 
 			OptionsMenuBuilder.AddSubMenu(
@@ -622,22 +545,7 @@ TSharedRef<SWidget> SCommonEditorViewportToolbarBase::GenerateCameraMenu() const
 	GetInfoProvider().OnFloatingButtonClicked();
 	TSharedRef<SEditorViewport> ViewportRef = GetInfoProvider().GetViewportWidget();
 
-	const bool bInShouldCloseWindowAfterMenuSelection = true;
-	FMenuBuilder CameraMenuBuilder( bInShouldCloseWindowAfterMenuSelection, ViewportRef->GetCommandList() );
-
-	// Camera types
-	CameraMenuBuilder.AddMenuEntry( FEditorViewportCommands::Get().Perspective );
-
-	CameraMenuBuilder.BeginSection("LevelViewportCameraType_Ortho", LOCTEXT("CameraTypeHeader_Ortho", "Orthographic") );
-		CameraMenuBuilder.AddMenuEntry(FEditorViewportCommands::Get().Top);
-		CameraMenuBuilder.AddMenuEntry(FEditorViewportCommands::Get().Bottom);
-		CameraMenuBuilder.AddMenuEntry(FEditorViewportCommands::Get().Left);
-		CameraMenuBuilder.AddMenuEntry(FEditorViewportCommands::Get().Right);
-		CameraMenuBuilder.AddMenuEntry(FEditorViewportCommands::Get().Front);
-		CameraMenuBuilder.AddMenuEntry(FEditorViewportCommands::Get().Back);
-		CameraMenuBuilder.EndSection();
-
-	return CameraMenuBuilder.MakeWidget();
+	return UE::UnrealEd::CreateCameraMenuWidget(ViewportRef);
 }
 
 TSharedRef<SWidget> SCommonEditorViewportToolbarBase::GenerateShowMenu() const
diff --git a/Engine/Source/Editor/UnrealEd/Private/ViewportToolbar/UnrealEdViewportToolbar.cpp b/Engine/Source/Editor/UnrealEd/Private/ViewportToolbar/UnrealEdViewportToolbar.cpp
--- a/Engine/Source/Editor/UnrealEd/Private/ViewportToolbar/UnrealEdViewportToolbar.cpp
+++ b/Engine/Source/Editor/UnrealEd/Private/ViewportToolbar/UnrealEdViewportToolbar.cpp
@@ -4,14 +4,19 @@
 
 #include "DataDrivenShaderPlatformInfo.h"
 #include "DebugViewModeHelpers.h"
+#include "Editor/EditorPerformanceSettings.h"
 #include "EditorViewportClient.h"
 #include "EditorViewportCommands.h"
 #include "Framework/Commands/GenericCommands.h"
 #include "GPUSkinCache.h"
 #include "GPUSkinCacheVisualizationMenuCommands.h"
+#include "IPreviewProfileController.h"
 #include "LevelEditorActions.h"
+#include "PreviewProfileController.h"
 #include "RayTracingDebugVisualizationMenuCommands.h"
+#include "SAssetEditorViewport.h"
 #include "SEditorViewport.h"
+#include "Settings/EditorProjectSettings.h"
 #include "Settings/LevelEditorViewportSettings.h"
 #include "Styling/SlateIconFinder.h"
 #include "Templates/SharedPointer.h"
@@ -23,6 +28,7 @@
 #include "Widgets/Input/SSpinBox.h"
 #include "Widgets/Layout/SUniformGridPanel.h"
 #include "Widgets/SToolTip.h"
+#include "Widgets/Input/STextComboBox.h"
 
 #define LOCTEXT_NAMESPACE "UnrealEdViewportToolbar"
 
@@ -1632,6 +1638,526 @@ FToolMenuEntry CreateNumericEntry(
 	return NumericEntry;
 }
 
+TSharedRef<SWidget> CreateCameraMenuWidget(const TSharedRef<SEditorViewport>& InViewport)
+{
+	constexpr bool bInShouldCloseWindowAfterMenuSelection = true;
+	FMenuBuilder CameraMenuBuilder(bInShouldCloseWindowAfterMenuSelection, InViewport->GetCommandList());
+
+	// Camera types
+	CameraMenuBuilder.AddMenuEntry(FEditorViewportCommands::Get().Perspective);
+
+	CameraMenuBuilder.BeginSection("LevelViewportCameraType_Ortho", LOCTEXT("CameraTypeHeader_Ortho", "Orthographic"));
+	CameraMenuBuilder.AddMenuEntry(FEditorViewportCommands::Get().Top);
+	CameraMenuBuilder.AddMenuEntry(FEditorViewportCommands::Get().Bottom);
+	CameraMenuBuilder.AddMenuEntry(FEditorViewportCommands::Get().Left);
+	CameraMenuBuilder.AddMenuEntry(FEditorViewportCommands::Get().Right);
+	CameraMenuBuilder.AddMenuEntry(FEditorViewportCommands::Get().Front);
+	CameraMenuBuilder.AddMenuEntry(FEditorViewportCommands::Get().Back);
+	CameraMenuBuilder.EndSection();
+
+	return CameraMenuBuilder.MakeWidget();
+}
+
+TSharedRef<SWidget> CreateFOVMenuWidget(const TSharedRef<SEditorViewport>& InViewport)
+{
+	constexpr float FOVMin = 5.0f;
+	constexpr float FOVMax = 170.0f;
+
+	TSharedPtr<FEditorViewportClient> ViewportClient = InViewport->GetViewportClient();
+
+	return
+		// clang-format off
+		SNew(SBox)
+		.HAlign(HAlign_Right)
+		[
+			SNew(SBox)
+			.Padding(FMargin(4.0f, 0.0f, 0.0f, 0.0f))
+			.WidthOverride(100.0f)
+			[
+				SNew(SBorder)
+				.BorderImage(FAppStyle::Get().GetBrush("Menu.WidgetBorder"))
+				.Padding(FMargin(1.0f))
+				[
+					SNew(SSpinBox<float>)
+					.Style(&FAppStyle::Get(), "Menu.SpinBox")
+					.Font( FAppStyle::GetFontStyle(TEXT("MenuItem.Font")))
+					.MinValue(FOVMin)
+					.MaxValue(FOVMax)
+					.Value_Lambda([ViewportClient]()
+					{
+						return ViewportClient->ViewFOV;
+					})
+					.OnValueChanged_Lambda([ViewportClient](float InNewValue)
+					{
+						ViewportClient->FOVAngle = InNewValue;
+						ViewportClient->ViewFOV = InNewValue;
+						ViewportClient->Invalidate();
+					})
+				]
+			]
+		];
+	// clang-format on
+}
+
+TSharedRef<SWidget> CreateFarViewPlaneMenuWidget(const TSharedRef<SEditorViewport>& InViewport)
+{
+	TSharedPtr<FEditorViewportClient> ViewportClient = InViewport->GetViewportClient();
+
+	return
+		// clang-format off
+		SNew(SBox)
+		.HAlign(HAlign_Right)
+		[
+			SNew(SBox)
+			.Padding(FMargin(4.0f, 0.0f, 0.0f, 0.0f))
+			.WidthOverride(100.0f)
+			[
+				SNew(SBorder)
+				.BorderImage(FAppStyle::Get().GetBrush("Menu.WidgetBorder"))
+				.Padding(FMargin(1.0f))
+				[
+					SNew(SSpinBox<float>)
+					.Style(&FAppStyle::Get(), "Menu.SpinBox")
+					.ToolTipText(LOCTEXT("FarViewPlaneTooltip", "Distance to use as the far view plane, or zero to enable an infinite far view plane"))
+					.MinValue(0.0f)
+					.MaxValue(100000.0f)
+					.Font(FAppStyle::GetFontStyle(TEXT("MenuItem.Font")))
+					.Value_Lambda([ViewportClient]()
+					{
+						return ViewportClient->GetFarClipPlaneOverride();
+					})
+					.OnValueChanged_Lambda([ViewportClient](float InNewValue)
+					{
+						ViewportClient->OverrideFarClipPlane(InNewValue);
+						ViewportClient->Invalidate();
+					})
+				]
+			]
+		];
+	// clang-format on
+}
+
+FToolMenuEntry CreateCameraSubmenu(TWeakPtr<SEditorViewport> InViewport)
+{
+	return FToolMenuEntry::InitSubMenu(
+		"Camera",
+		LOCTEXT("CameraSubmenuLabel", "Camera"),
+		LOCTEXT("CameraSubmenuTooltip", "Camera options"),
+		FNewToolMenuDelegate::CreateLambda(
+			[InViewport](UToolMenu* Submenu) -> void
+			{
+				FToolMenuSection& UnnamedSection = Submenu->FindOrAddSection("", LOCTEXT("UnnamedLabel", ""));
+
+				UnnamedSection.AddEntry(FToolMenuEntry::InitWidget(
+					"CameraMenuItems", UE::UnrealEd::CreateCameraMenuWidget(InViewport.Pin().ToSharedRef()), FText(), true
+				));
+
+				UnnamedSection.AddSeparator("CameraSubmenuSeparator");
+
+				UnnamedSection.AddEntry(FToolMenuEntry::InitWidget(
+					"CameraFOV",
+					UE::UnrealEd::CreateFOVMenuWidget(InViewport.Pin().ToSharedRef()),
+					LOCTEXT("CameraSubmenu_FieldOfViewLabel", "Field of View"),
+					true
+				));
+
+				UnnamedSection.AddEntry(FToolMenuEntry::InitWidget(
+					"CameraFarViewPlane",
+					UE::UnrealEd::CreateFarViewPlaneMenuWidget(InViewport.Pin().ToSharedRef()),
+					LOCTEXT("CameraSubmenu_FarViewPlaneLabel", "Far View Plane"),
+					true
+				));
+			}
+		)
+	);
+}
+
+static FFormatNamedArguments GetScreenPercentageFormatArguments(const FEditorViewportClient& ViewportClient)
+{
+	static auto CVarEditorViewportDefaultScreenPercentageRealTimeMode =
+		IConsoleManager::Get().FindConsoleVariable(TEXT("r.Editor.Viewport.ScreenPercentageMode.RealTime"));
+	static auto CVarEditorViewportDefaultScreenPercentageMobileMode =
+		IConsoleManager::Get().FindConsoleVariable(TEXT("r.Editor.Viewport.ScreenPercentageMode.Mobile"));
+	static auto CVarEditorViewportDefaultScreenPercentageVRMode =
+		IConsoleManager::Get().FindConsoleVariable(TEXT("r.Editor.Viewport.ScreenPercentageMode.VR"));
+	static auto CVarEditorViewportDefaultScreenPercentagePathTracerMode =
+		IConsoleManager::Get().FindConsoleVariable(TEXT("r.Editor.Viewport.ScreenPercentageMode.PathTracer"));
+	static auto CVarEditorViewportDefaultScreenPercentageMode =
+		IConsoleManager::Get().FindConsoleVariable(TEXT("r.Editor.Viewport.ScreenPercentageMode.NonRealTime"));
+
+	const UEditorPerformanceProjectSettings* EditorProjectSettings = GetDefault<UEditorPerformanceProjectSettings>();
+	const UEditorPerformanceSettings* EditorUserSettings = GetDefault<UEditorPerformanceSettings>();
+	const FEngineShowFlags& EngineShowFlags = ViewportClient.EngineShowFlags;
+
+	const EViewStatusForScreenPercentage ViewportRenderingMode = ViewportClient.GetViewStatusForScreenPercentage();
+	const bool bViewModeSupportsScreenPercentage = ViewportClient.SupportsPreviewResolutionFraction();
+	const bool bIsPreviewScreenPercentage = ViewportClient.IsPreviewingScreenPercentage();
+
+	float DefaultScreenPercentage = FMath::Clamp(
+										ViewportClient.GetDefaultPrimaryResolutionFractionTarget(),
+										ISceneViewFamilyScreenPercentage::kMinTSRResolutionFraction,
+										ISceneViewFamilyScreenPercentage::kMaxTSRResolutionFraction
+									)
+								  * 100.0f;
+	float PreviewScreenPercentage = float(ViewportClient.GetPreviewScreenPercentage());
+	float FinalScreenPercentage = bIsPreviewScreenPercentage ? PreviewScreenPercentage : DefaultScreenPercentage;
+
+	FFormatNamedArguments FormatArguments;
+	FormatArguments.Add(TEXT("ViewportMode"), UEnum::GetDisplayValueAsText(ViewportRenderingMode));
+
+	EScreenPercentageMode ProjectSetting = EScreenPercentageMode::Manual;
+	EEditorUserScreenPercentageModeOverride UserPreference = EEditorUserScreenPercentageModeOverride::ProjectDefault;
+	IConsoleVariable* CVarDefaultScreenPercentage = nullptr;
+	if (ViewportRenderingMode == EViewStatusForScreenPercentage::PathTracer)
+	{
+		ProjectSetting = EditorProjectSettings->PathTracerScreenPercentageMode;
+		UserPreference = EditorUserSettings->PathTracerScreenPercentageMode;
+		CVarDefaultScreenPercentage = CVarEditorViewportDefaultScreenPercentagePathTracerMode;
+	}
+	else if (ViewportRenderingMode == EViewStatusForScreenPercentage::VR)
+	{
+		ProjectSetting = EditorProjectSettings->VRScreenPercentageMode;
+		UserPreference = EditorUserSettings->VRScreenPercentageMode;
+		CVarDefaultScreenPercentage = CVarEditorViewportDefaultScreenPercentageVRMode;
+	}
+	else if (ViewportRenderingMode == EViewStatusForScreenPercentage::Mobile)
+	{
+		ProjectSetting = EditorProjectSettings->MobileScreenPercentageMode;
+		UserPreference = EditorUserSettings->MobileScreenPercentageMode;
+		CVarDefaultScreenPercentage = CVarEditorViewportDefaultScreenPercentageMobileMode;
+	}
+	else if (ViewportRenderingMode == EViewStatusForScreenPercentage::Desktop)
+	{
+		ProjectSetting = EditorProjectSettings->RealtimeScreenPercentageMode;
+		UserPreference = EditorUserSettings->RealtimeScreenPercentageMode;
+		CVarDefaultScreenPercentage = CVarEditorViewportDefaultScreenPercentageRealTimeMode;
+	}
+	else if (ViewportRenderingMode == EViewStatusForScreenPercentage::NonRealtime)
+	{
+		ProjectSetting = EditorProjectSettings->NonRealtimeScreenPercentageMode;
+		UserPreference = EditorUserSettings->NonRealtimeScreenPercentageMode;
+		CVarDefaultScreenPercentage = CVarEditorViewportDefaultScreenPercentageMode;
+	}
+	else
+	{
+		unimplemented();
+	}
+
+	EScreenPercentageMode FinalScreenPercentageMode = EScreenPercentageMode::Manual;
+	if (!bViewModeSupportsScreenPercentage)
+	{
+		FormatArguments.Add(
+			TEXT("SettingSource"), LOCTEXT("ScreenPercentage_SettingSource_UnsupportedByViewMode", "Unsupported by View mode")
+		);
+		FinalScreenPercentageMode = EScreenPercentageMode::Manual;
+		FinalScreenPercentage = 100;
+	}
+	else if (bIsPreviewScreenPercentage)
+	{
+		FormatArguments.Add(
+			TEXT("SettingSource"), LOCTEXT("ScreenPercentage_SettingSource_ViewportOverride", "Viewport Override")
+		);
+		FinalScreenPercentageMode = EScreenPercentageMode::Manual;
+	}
+	else if ((CVarDefaultScreenPercentage->GetFlags() & ECVF_SetByMask) > ECVF_SetByProjectSetting)
+	{
+		FormatArguments.Add(TEXT("SettingSource"), LOCTEXT("ScreenPercentage_SettingSource_Cvar", "Console Variable"));
+		FinalScreenPercentageMode = EScreenPercentageMode(CVarDefaultScreenPercentage->GetInt());
+	}
+	else if (UserPreference == EEditorUserScreenPercentageModeOverride::ProjectDefault)
+	{
+		FormatArguments.Add(
+			TEXT("SettingSource"), LOCTEXT("ScreenPercentage_SettingSource_ProjectSettigns", "Project Settings")
+		);
+		FinalScreenPercentageMode = ProjectSetting;
+	}
+	else
+	{
+		FormatArguments.Add(
+			TEXT("SettingSource"), LOCTEXT("ScreenPercentage_SettingSource_EditorPreferences", "Editor Preferences")
+		);
+		if (UserPreference == EEditorUserScreenPercentageModeOverride::BasedOnDPIScale)
+		{
+			FinalScreenPercentageMode = EScreenPercentageMode::BasedOnDPIScale;
+		}
+		else if (UserPreference == EEditorUserScreenPercentageModeOverride::BasedOnDisplayResolution)
+		{
+			FinalScreenPercentageMode = EScreenPercentageMode::BasedOnDisplayResolution;
+		}
+		else
+		{
+			FinalScreenPercentageMode = EScreenPercentageMode::Manual;
+		}
+	}
+
+	if (FinalScreenPercentageMode == EScreenPercentageMode::BasedOnDPIScale)
+	{
+		FormatArguments.Add(TEXT("Setting"), LOCTEXT("ScreenPercentage_Setting_BasedOnDPIScale", "Based on OS's DPI scale"));
+	}
+	else if (FinalScreenPercentageMode == EScreenPercentageMode::BasedOnDisplayResolution)
+	{
+		FormatArguments.Add(
+			TEXT("Setting"), LOCTEXT("ScreenPercentage_Setting_BasedOnDisplayResolution", "Based on display resolution")
+		);
+	}
+	else
+	{
+		FormatArguments.Add(TEXT("Setting"), LOCTEXT("ScreenPercentage_Setting_Manual", "Manual"));
+	}
+
+	FormatArguments.Add(
+		TEXT("CurrentScreenPercentage"),
+		FText::FromString(FString::Printf(TEXT("%3.1f"), FMath::RoundToFloat(FinalScreenPercentage * 10.0f) / 10.0f))
+	);
+
+	{
+		float FinalResolutionFraction = (FinalScreenPercentage / 100.0f);
+		FIntPoint DisplayResolution = ViewportClient.Viewport->GetSizeXY();
+		FIntPoint RenderingResolution;
+		RenderingResolution.X = FMath::CeilToInt(DisplayResolution.X * FinalResolutionFraction);
+		RenderingResolution.Y = FMath::CeilToInt(DisplayResolution.Y * FinalResolutionFraction);
+
+		FormatArguments.Add(
+			TEXT("ResolutionFromTo"),
+			FText::FromString(FString::Printf(
+				TEXT("%dx%d -> %dx%d"),
+				RenderingResolution.X,
+				RenderingResolution.Y,
+				DisplayResolution.X,
+				DisplayResolution.Y
+			))
+		);
+	}
+
+	return FormatArguments;
+}
+
+static const FMargin ScreenPercentageMenuCommonPadding(26.0f, 3.0f);
+
+TSharedRef<SWidget> CreateCurrentPercentageWidget(FEditorViewportClient& InViewportClient)
+{
+	// clang-format off
+	return SNew(SBox)
+		.Padding(ScreenPercentageMenuCommonPadding)
+		[
+			SNew(STextBlock)
+			.ColorAndOpacity(FSlateColor::UseSubduedForeground())
+			.Text_Lambda([&InViewportClient]()
+			{
+				FFormatNamedArguments FormatArguments = GetScreenPercentageFormatArguments(InViewportClient);
+				return FText::Format(LOCTEXT("ScreenPercentageCurrent_Display", "Current Screen Percentage: {CurrentScreenPercentage}"), FormatArguments);
+			})
+			.ToolTip(SNew(SToolTip).Text(LOCTEXT("ScreenPercentageCurrent_ToolTip", "Current Screen Percentage the viewport is rendered with. The primary screen percentage can either be a spatial or temporal upscaler based of your anti-aliasing settings.")))
+		];
+	// clang-format on
+}
+
+TSharedRef<SWidget> CreateResolutionsWidget(FEditorViewportClient& InViewportClient)
+{
+	// clang-format off
+	return SNew(SBox)
+	.Padding(ScreenPercentageMenuCommonPadding)
+	[
+		SNew(STextBlock)
+		.ColorAndOpacity(FSlateColor::UseSubduedForeground())
+		.Text_Lambda([&InViewportClient]()
+		{
+			FFormatNamedArguments FormatArguments = GetScreenPercentageFormatArguments(InViewportClient);
+			return FText::Format(LOCTEXT("ScreenPercentageResolutions", "Resolution: {ResolutionFromTo}"), FormatArguments);
+		})
+	];
+	// clang-format on
+}
+
+TSharedRef<SWidget> CreateActiveViewportWidget(FEditorViewportClient& InViewPortClient)
+{
+	// clang-format off
+	return SNew(SBox)
+		.Padding(ScreenPercentageMenuCommonPadding)
+		[
+			SNew(STextBlock)
+			.ColorAndOpacity(FSlateColor::UseSubduedForeground())
+			.Text_Lambda([&InViewPortClient]()
+			{
+				FFormatNamedArguments FormatArguments = GetScreenPercentageFormatArguments(InViewPortClient);
+				return FText::Format(LOCTEXT("ScreenPercentageActiveViewport", "Active Viewport: {ViewportMode}"), FormatArguments);
+			})
+		];
+	// clang-format on
+}
+
+TSharedRef<SWidget> CreateSetFromWidget(FEditorViewportClient& InViewPortClient)
+{
+	// clang-format off
+	return SNew(SBox)
+		.Padding(ScreenPercentageMenuCommonPadding)
+		[
+			SNew(STextBlock)
+			.ColorAndOpacity(FSlateColor::UseSubduedForeground())
+			.Text_Lambda([&InViewPortClient]()
+			{
+				FFormatNamedArguments FormatArguments = GetScreenPercentageFormatArguments(InViewPortClient);
+				return FText::Format(LOCTEXT("ScreenPercentageSetFrom", "Set From: {SettingSource}"), FormatArguments);
+			})
+		];
+	// clang-format on
+}
+
+TSharedRef<SWidget> CreateCurrentScreenPercentageSettingWidget(FEditorViewportClient& InViewPortClient)
+{
+	// clang-format off
+	return SNew(SBox)
+		.Padding(ScreenPercentageMenuCommonPadding)
+		[
+			SNew(STextBlock)
+			.ColorAndOpacity(FSlateColor::UseSubduedForeground())
+			.Text_Lambda([&InViewPortClient]()
+			{
+				FFormatNamedArguments FormatArguments = GetScreenPercentageFormatArguments(InViewPortClient);
+				return FText::Format(LOCTEXT("ScreenPercentageSetting", "Setting: {Setting}"), FormatArguments);
+			})
+		];
+	// clang-format on
+}
+
+TSharedRef<SWidget> CreateCurrentScreenPercentageWidget(FEditorViewportClient& InViewPortClient)
+{
+	constexpr int32 PreviewScreenPercentageMin = ISceneViewFamilyScreenPercentage::kMinTSRResolutionFraction * 100.0f;
+	constexpr int32 PreviewScreenPercentageMax = ISceneViewFamilyScreenPercentage::kMaxTSRResolutionFraction * 100.0f;
+
+	// clang-format off
+	return SNew(SBox)
+		.HAlign(HAlign_Right)
+		.IsEnabled_Lambda([&InViewPortClient]()
+		{
+			return InViewPortClient.IsPreviewingScreenPercentage() && InViewPortClient.SupportsPreviewResolutionFraction();
+		})
+		[
+			SNew(SBox)
+			.Padding(FMargin(4.0f, 0.0f, 0.0f, 0.0f))
+			.WidthOverride(100.0f)
+			[
+				SNew(SBorder)
+				.Padding(FMargin(1.0f))
+				[
+					SNew(SSpinBox<int32>)
+					.Style(&FAppStyle::Get(), "Menu.SpinBox")
+					.Font(FAppStyle::GetFontStyle(TEXT("MenuItem.Font")))
+					.MinSliderValue(PreviewScreenPercentageMin)
+					.MaxSliderValue(PreviewScreenPercentageMax)
+					.Value_Lambda([&InViewPortClient]()
+					{
+						return InViewPortClient.GetPreviewScreenPercentage();
+					})
+					.OnValueChanged_Lambda([&InViewPortClient](int32 NewValue)
+					{
+						InViewPortClient.SetPreviewScreenPercentage(NewValue);
+						InViewPortClient.Invalidate();
+					})
+				]
+			]
+		];
+	// clang-format on
+}
+
+static void ConstructScreenPercentageMenu(UToolMenu* InMenu, FEditorViewportClient* InViewportClient)
+{
+	FEditorViewportClient& ViewportClient = *InViewportClient;
+
+	const FEditorViewportCommands& BaseViewportCommands = FEditorViewportCommands::Get();
+
+	// Summary
+	{
+		FToolMenuSection& SummarySection = InMenu->FindOrAddSection("Summary", LOCTEXT("Summary", "Summary"));
+		SummarySection.AddEntry(FToolMenuEntry::InitWidget(
+			"ScreenPercentageCurrent", CreateCurrentPercentageWidget(ViewportClient), FText::GetEmpty()
+		));
+
+		SummarySection.AddEntry(FToolMenuEntry::InitWidget(
+			"ScreenPercentageResolutions", CreateResolutionsWidget(ViewportClient), FText::GetEmpty()
+		));
+
+		SummarySection.AddEntry(FToolMenuEntry::InitWidget(
+			"ScreenPercentageActiveViewport", CreateActiveViewportWidget(ViewportClient), FText::GetEmpty()
+		));
+
+		SummarySection.AddEntry(
+			FToolMenuEntry::InitWidget("ScreenPercentageSetFrom", CreateSetFromWidget(ViewportClient), FText::GetEmpty())
+		);
+
+		SummarySection.AddEntry(FToolMenuEntry::InitWidget(
+			"ScreenPercentageSetting", CreateCurrentScreenPercentageSettingWidget(ViewportClient), FText::GetEmpty()
+		));
+	}
+
+	// Screen Percentage
+	{
+		FToolMenuSection& ScreenPercentageSection =
+			InMenu->FindOrAddSection("ScreenPercentage", LOCTEXT("ScreenPercentage_ViewportOverride", "Viewport Override"));
+
+		ScreenPercentageSection.AddMenuEntry(BaseViewportCommands.ToggleOverrideViewportScreenPercentage);
+
+		ScreenPercentageSection.AddEntry(FToolMenuEntry::InitWidget(
+			"PreviewScreenPercentage",
+			CreateCurrentScreenPercentageWidget(ViewportClient),
+			LOCTEXT("ScreenPercentage", "Screen Percentage")
+		));
+	}
+
+	// Screen Percentage Settings
+	{
+		FToolMenuSection& ScreenPercentageSettingsSection = InMenu->FindOrAddSection(
+			"ScreenPercentageSettings", LOCTEXT("ScreenPercentage_ViewportSettings", "Viewport Settings")
+		);
+
+		ScreenPercentageSettingsSection.AddMenuEntry(
+			BaseViewportCommands.OpenEditorPerformanceProjectSettings,
+			/* InLabelOverride = */ TAttribute<FText>(),
+			/* InToolTipOverride = */ TAttribute<FText>(),
+			FSlateIcon(FAppStyle::GetAppStyleSetName(), "ProjectSettings.TabIcon")
+		);
+
+		ScreenPercentageSettingsSection.AddMenuEntry(
+			BaseViewportCommands.OpenEditorPerformanceEditorPreferences,
+			/* InLabelOverride = */ TAttribute<FText>(),
+			/* InToolTipOverride = */ TAttribute<FText>(),
+			FSlateIcon(FAppStyle::GetAppStyleSetName(), "EditorPreferences.TabIcon")
+		);
+	}
+}
+
+UNREALED_API FToolMenuEntry CreatePerformanceAndScalabilitySubmenu(TWeakPtr<SEditorViewport> InViewport)
+{
+	return FToolMenuEntry::InitSubMenu(
+		"PerformanceAndScalability",
+		LOCTEXT("PerformanceAndScalabilitySubmenuLabel", "Performance and Scalability"),
+		LOCTEXT("PerformanceAndScalabilitySubmenuTooltip", ""),
+		FNewToolMenuDelegate::CreateLambda(
+			[InViewport](UToolMenu* Submenu) -> void
+			{
+				TSharedPtr<SEditorViewport> Viewport = InViewport.Pin();
+				if (!Viewport)
+				{
+					return;
+				}
+
+				FToolMenuSection& UnnamedSection = Submenu->FindOrAddSection("", LOCTEXT("UnnamedLabel", ""));
+
+				UnnamedSection.AddMenuEntry(FEditorViewportCommands::Get().ToggleRealTime);
+
+				TSharedPtr<FEditorViewportClient> ViewportClient = Viewport->GetViewportClient();
+
+				UnnamedSection.AddSubMenu(
+					"ScreenPercentage",
+					LOCTEXT("ScreenPercentageSubMenu", "Screen Percentage"),
+					LOCTEXT("ScreenPercentageSubMenu_ToolTip", "Customize the viewport's screen percentage"),
+					FNewToolMenuDelegate::CreateStatic(&ConstructScreenPercentageMenu, ViewportClient.Get())
+				);
+			}
+		)
+	);
+}
+
 } // namespace UE::UnrealEd
 
 #undef LOCTEXT_NAMESPACE
diff --git a/Engine/Source/Editor/UnrealEd/Public/ViewportToolbar/UnrealEdViewportToolbar.h b/Engine/Source/Editor/UnrealEd/Public/ViewportToolbar/UnrealEdViewportToolbar.h
--- a/Engine/Source/Editor/UnrealEd/Public/ViewportToolbar/UnrealEdViewportToolbar.h
+++ b/Engine/Source/Editor/UnrealEd/Public/ViewportToolbar/UnrealEdViewportToolbar.h
@@ -8,6 +8,8 @@
 #include "ToolMenuDelegates.h"
 #include "UnrealEdViewportToolbarContext.h"
 
+class IPreviewProfileController;
+class FEditorViewportClient;
 class FText;
 class SEditorViewport;
 class UToolMenu;
@@ -113,4 +115,21 @@ UNREALED_API FToolMenuEntry CreateNumericEntry(
 	int32 InMaxFractionalDigits = 2
 );
 
+UNREALED_API FToolMenuEntry CreateCameraSubmenu(TWeakPtr<SEditorViewport> InViewport);
+
+UNREALED_API FToolMenuEntry CreatePerformanceAndScalabilitySubmenu(TWeakPtr<SEditorViewport> InViewport);
+
+// Camera Menu Widgets
+TSharedRef<SWidget> CreateCameraMenuWidget(const TSharedRef<SEditorViewport>& InViewport);
+TSharedRef<SWidget> CreateFOVMenuWidget(const TSharedRef<SEditorViewport>& InViewport);
+TSharedRef<SWidget> CreateFarViewPlaneMenuWidget(const TSharedRef<SEditorViewport>& InViewport);
+
+// Screen Percentage Submenu Widgets
+TSharedRef<SWidget> CreateCurrentPercentageWidget(FEditorViewportClient& InViewportClient);
+TSharedRef<SWidget> CreateResolutionsWidget(FEditorViewportClient& InViewportClient);
+TSharedRef<SWidget> CreateActiveViewportWidget(FEditorViewportClient& InViewportClient);
+TSharedRef<SWidget> CreateSetFromWidget(FEditorViewportClient& InViewportClient);
+TSharedRef<SWidget> CreateCurrentScreenPercentageSettingWidget(FEditorViewportClient& InViewportClient);
+TSharedRef<SWidget> CreateCurrentScreenPercentageWidget(FEditorViewportClient& InViewportClient);
+
 } // namespace UE::UnrealEd
