[Viewport Toolbar] Fix Horde compilation issue

diff --git a/Engine/Source/Editor/UnrealEd/Private/SCommonEditorViewportToolbarBase.cpp b/Engine/Source/Editor/UnrealEd/Private/SCommonEditorViewportToolbarBase.cpp
--- a/Engine/Source/Editor/UnrealEd/Private/SCommonEditorViewportToolbarBase.cpp
+++ b/Engine/Source/Editor/UnrealEd/Private/SCommonEditorViewportToolbarBase.cpp
@@ -269,137 +269,6 @@ void SCommonEditorViewportToolbarBase::Construct(const FArguments& InArgs, TShar
 	SViewportToolBar::Construct(SViewportToolBar::FArguments());
 }
 
-static FFormatNamedArguments GetScreenPercentageFormatArguments(const FEditorViewportClient& ViewportClient)
-{
-	static auto CVarEditorViewportDefaultScreenPercentageRealTimeMode = IConsoleManager::Get().FindConsoleVariable(TEXT("r.Editor.Viewport.ScreenPercentageMode.RealTime"));
-	static auto CVarEditorViewportDefaultScreenPercentageMobileMode = IConsoleManager::Get().FindConsoleVariable(TEXT("r.Editor.Viewport.ScreenPercentageMode.Mobile"));
-	static auto CVarEditorViewportDefaultScreenPercentageVRMode = IConsoleManager::Get().FindConsoleVariable(TEXT("r.Editor.Viewport.ScreenPercentageMode.VR"));
-	static auto CVarEditorViewportDefaultScreenPercentagePathTracerMode = IConsoleManager::Get().FindConsoleVariable(TEXT("r.Editor.Viewport.ScreenPercentageMode.PathTracer"));
-	static auto CVarEditorViewportDefaultScreenPercentageMode = IConsoleManager::Get().FindConsoleVariable(TEXT("r.Editor.Viewport.ScreenPercentageMode.NonRealTime"));
-
-	const UEditorPerformanceProjectSettings* EditorProjectSettings = GetDefault<UEditorPerformanceProjectSettings>();
-	const UEditorPerformanceSettings* EditorUserSettings = GetDefault<UEditorPerformanceSettings>();
-	const FEngineShowFlags& EngineShowFlags = ViewportClient.EngineShowFlags;
-
-	const EViewStatusForScreenPercentage ViewportRenderingMode = ViewportClient.GetViewStatusForScreenPercentage();
-	const bool bViewModeSupportsScreenPercentage = ViewportClient.SupportsPreviewResolutionFraction();
-	const bool bIsPreviewScreenPercentage = ViewportClient.IsPreviewingScreenPercentage();
-
-	float DefaultScreenPercentage = FMath::Clamp(
-		ViewportClient.GetDefaultPrimaryResolutionFractionTarget(),
-		ISceneViewFamilyScreenPercentage::kMinTSRResolutionFraction,
-		ISceneViewFamilyScreenPercentage::kMaxTSRResolutionFraction) * 100.0f;
-	float PreviewScreenPercentage = float(ViewportClient.GetPreviewScreenPercentage());
-	float FinalScreenPercentage = bIsPreviewScreenPercentage ? PreviewScreenPercentage : DefaultScreenPercentage;
-
-	FFormatNamedArguments FormatArguments;
-	FormatArguments.Add(TEXT("ViewportMode"), UEnum::GetDisplayValueAsText(ViewportRenderingMode));
-
-	EScreenPercentageMode ProjectSetting = EScreenPercentageMode::Manual;
-	EEditorUserScreenPercentageModeOverride UserPreference = EEditorUserScreenPercentageModeOverride::ProjectDefault;
-	IConsoleVariable* CVarDefaultScreenPercentage = nullptr;
-	if (ViewportRenderingMode == EViewStatusForScreenPercentage::PathTracer)
-	{
-		ProjectSetting = EditorProjectSettings->PathTracerScreenPercentageMode;
-		UserPreference = EditorUserSettings->PathTracerScreenPercentageMode;
-		CVarDefaultScreenPercentage = CVarEditorViewportDefaultScreenPercentagePathTracerMode;
-	}
-	else if (ViewportRenderingMode == EViewStatusForScreenPercentage::VR)
-	{
-		ProjectSetting = EditorProjectSettings->VRScreenPercentageMode;
-		UserPreference = EditorUserSettings->VRScreenPercentageMode;
-		CVarDefaultScreenPercentage = CVarEditorViewportDefaultScreenPercentageVRMode;
-	}
-	else if (ViewportRenderingMode == EViewStatusForScreenPercentage::Mobile)
-	{
-		ProjectSetting = EditorProjectSettings->MobileScreenPercentageMode;
-		UserPreference = EditorUserSettings->MobileScreenPercentageMode;
-		CVarDefaultScreenPercentage = CVarEditorViewportDefaultScreenPercentageMobileMode;
-	}
-	else if (ViewportRenderingMode == EViewStatusForScreenPercentage::Desktop)
-	{
-		ProjectSetting = EditorProjectSettings->RealtimeScreenPercentageMode;
-		UserPreference = EditorUserSettings->RealtimeScreenPercentageMode;
-		CVarDefaultScreenPercentage = CVarEditorViewportDefaultScreenPercentageRealTimeMode;
-	}
-	else if (ViewportRenderingMode == EViewStatusForScreenPercentage::NonRealtime)
-	{
-		ProjectSetting = EditorProjectSettings->NonRealtimeScreenPercentageMode;
-		UserPreference = EditorUserSettings->NonRealtimeScreenPercentageMode;
-		CVarDefaultScreenPercentage = CVarEditorViewportDefaultScreenPercentageMode;
-	}
-	else
-	{
-		unimplemented();
-	}
-
-	EScreenPercentageMode FinalScreenPercentageMode = EScreenPercentageMode::Manual;
-	if (!bViewModeSupportsScreenPercentage)
-	{
-		FormatArguments.Add(TEXT("SettingSource"), LOCTEXT("ScreenPercentage_SettingSource_UnsupportedByViewMode", "Unsupported by View mode"));
-		FinalScreenPercentageMode = EScreenPercentageMode::Manual;
-		FinalScreenPercentage = 100;
-	}
-	else if (bIsPreviewScreenPercentage)
-	{
-		FormatArguments.Add(TEXT("SettingSource"), LOCTEXT("ScreenPercentage_SettingSource_ViewportOverride", "Viewport Override"));
-		FinalScreenPercentageMode = EScreenPercentageMode::Manual;
-	}
-	else if ((CVarDefaultScreenPercentage->GetFlags() & ECVF_SetByMask) > ECVF_SetByProjectSetting)
-	{
-		FormatArguments.Add(TEXT("SettingSource"), LOCTEXT("ScreenPercentage_SettingSource_Cvar", "Console Variable"));
-		FinalScreenPercentageMode = EScreenPercentageMode(CVarDefaultScreenPercentage->GetInt());
-	}
-	else if (UserPreference == EEditorUserScreenPercentageModeOverride::ProjectDefault)
-	{
-		FormatArguments.Add(TEXT("SettingSource"), LOCTEXT("ScreenPercentage_SettingSource_ProjectSettigns", "Project Settings"));
-		FinalScreenPercentageMode = ProjectSetting;
-	}
-	else
-	{
-		FormatArguments.Add(TEXT("SettingSource"), LOCTEXT("ScreenPercentage_SettingSource_EditorPreferences", "Editor Preferences"));
-		if (UserPreference == EEditorUserScreenPercentageModeOverride::BasedOnDPIScale)
-		{
-			FinalScreenPercentageMode = EScreenPercentageMode::BasedOnDPIScale;
-		}
-		else if (UserPreference == EEditorUserScreenPercentageModeOverride::BasedOnDisplayResolution)
-		{
-			FinalScreenPercentageMode = EScreenPercentageMode::BasedOnDisplayResolution;
-		}
-		else
-		{
-			FinalScreenPercentageMode = EScreenPercentageMode::Manual;
-		}
-	}
-
-	if (FinalScreenPercentageMode == EScreenPercentageMode::BasedOnDPIScale)
-	{
-		FormatArguments.Add(TEXT("Setting"), LOCTEXT("ScreenPercentage_Setting_BasedOnDPIScale", "Based on OS's DPI scale"));
-	}
-	else if (FinalScreenPercentageMode == EScreenPercentageMode::BasedOnDisplayResolution)
-	{
-		FormatArguments.Add(TEXT("Setting"), LOCTEXT("ScreenPercentage_Setting_BasedOnDisplayResolution", "Based on display resolution"));
-	}
-	else
-	{
-		FormatArguments.Add(TEXT("Setting"), LOCTEXT("ScreenPercentage_Setting_Manual", "Manual"));
-	}
-
-	FormatArguments.Add(TEXT("CurrentScreenPercentage"), FText::FromString(FString::Printf(TEXT("%3.1f"), FMath::RoundToFloat(FinalScreenPercentage * 10.0f) / 10.0f)));
-
-	{
-		float FinalResolutionFraction = (FinalScreenPercentage / 100.0f);
-		FIntPoint DisplayResolution = ViewportClient.Viewport->GetSizeXY();
-		FIntPoint RenderingResolution;
-		RenderingResolution.X = FMath::CeilToInt(DisplayResolution.X * FinalResolutionFraction);
-		RenderingResolution.Y = FMath::CeilToInt(DisplayResolution.Y * FinalResolutionFraction);
-
-		FormatArguments.Add(TEXT("ResolutionFromTo"), FText::FromString(FString::Printf(TEXT("%dx%d -> %dx%d"), RenderingResolution.X, RenderingResolution.Y, DisplayResolution.X, DisplayResolution.Y)));
-	}
-
-	return FormatArguments;
-}
-
 void SCommonEditorViewportToolbarBase::ConstructScreenPercentageMenu(FMenuBuilder& MenuBuilder, FEditorViewportClient* InViewportClient)
 {
 	FEditorViewportClient& ViewportClient = *InViewportClient;
diff --git a/Engine/Source/Editor/UnrealEd/Private/ViewportToolbar/UnrealEdViewportToolbar.cpp b/Engine/Source/Editor/UnrealEd/Private/ViewportToolbar/UnrealEdViewportToolbar.cpp
--- a/Engine/Source/Editor/UnrealEd/Private/ViewportToolbar/UnrealEdViewportToolbar.cpp
+++ b/Engine/Source/Editor/UnrealEd/Private/ViewportToolbar/UnrealEdViewportToolbar.cpp
@@ -1774,17 +1774,6 @@ FToolMenuEntry CreateCameraSubmenu(TWeakPtr<SEditorViewport> InViewport)
 
 static FFormatNamedArguments GetScreenPercentageFormatArguments(const FEditorViewportClient& ViewportClient)
 {
-	static auto CVarEditorViewportDefaultScreenPercentageRealTimeMode =
-		IConsoleManager::Get().FindConsoleVariable(TEXT("r.Editor.Viewport.ScreenPercentageMode.RealTime"));
-	static auto CVarEditorViewportDefaultScreenPercentageMobileMode =
-		IConsoleManager::Get().FindConsoleVariable(TEXT("r.Editor.Viewport.ScreenPercentageMode.Mobile"));
-	static auto CVarEditorViewportDefaultScreenPercentageVRMode =
-		IConsoleManager::Get().FindConsoleVariable(TEXT("r.Editor.Viewport.ScreenPercentageMode.VR"));
-	static auto CVarEditorViewportDefaultScreenPercentagePathTracerMode =
-		IConsoleManager::Get().FindConsoleVariable(TEXT("r.Editor.Viewport.ScreenPercentageMode.PathTracer"));
-	static auto CVarEditorViewportDefaultScreenPercentageMode =
-		IConsoleManager::Get().FindConsoleVariable(TEXT("r.Editor.Viewport.ScreenPercentageMode.NonRealTime"));
-
 	const UEditorPerformanceProjectSettings* EditorProjectSettings = GetDefault<UEditorPerformanceProjectSettings>();
 	const UEditorPerformanceSettings* EditorUserSettings = GetDefault<UEditorPerformanceSettings>();
 	const FEngineShowFlags& EngineShowFlags = ViewportClient.EngineShowFlags;
@@ -1812,31 +1801,31 @@ static FFormatNamedArguments GetScreenPercentageFormatArguments(const FEditorVie
 	{
 		ProjectSetting = EditorProjectSettings->PathTracerScreenPercentageMode;
 		UserPreference = EditorUserSettings->PathTracerScreenPercentageMode;
-		CVarDefaultScreenPercentage = CVarEditorViewportDefaultScreenPercentagePathTracerMode;
+		CVarDefaultScreenPercentage = IConsoleManager::Get().FindConsoleVariable(TEXT("r.Editor.Viewport.ScreenPercentageMode.PathTracer"));
 	}
 	else if (ViewportRenderingMode == EViewStatusForScreenPercentage::VR)
 	{
 		ProjectSetting = EditorProjectSettings->VRScreenPercentageMode;
 		UserPreference = EditorUserSettings->VRScreenPercentageMode;
-		CVarDefaultScreenPercentage = CVarEditorViewportDefaultScreenPercentageVRMode;
+		CVarDefaultScreenPercentage = IConsoleManager::Get().FindConsoleVariable(TEXT("r.Editor.Viewport.ScreenPercentageMode.VR"));
 	}
 	else if (ViewportRenderingMode == EViewStatusForScreenPercentage::Mobile)
 	{
 		ProjectSetting = EditorProjectSettings->MobileScreenPercentageMode;
 		UserPreference = EditorUserSettings->MobileScreenPercentageMode;
-		CVarDefaultScreenPercentage = CVarEditorViewportDefaultScreenPercentageMobileMode;
+		CVarDefaultScreenPercentage = IConsoleManager::Get().FindConsoleVariable(TEXT("r.Editor.Viewport.ScreenPercentageMode.Mobile"));
 	}
 	else if (ViewportRenderingMode == EViewStatusForScreenPercentage::Desktop)
 	{
 		ProjectSetting = EditorProjectSettings->RealtimeScreenPercentageMode;
 		UserPreference = EditorUserSettings->RealtimeScreenPercentageMode;
-		CVarDefaultScreenPercentage = CVarEditorViewportDefaultScreenPercentageRealTimeMode;
+		CVarDefaultScreenPercentage = IConsoleManager::Get().FindConsoleVariable(TEXT("r.Editor.Viewport.ScreenPercentageMode.RealTime"));
 	}
 	else if (ViewportRenderingMode == EViewStatusForScreenPercentage::NonRealtime)
 	{
 		ProjectSetting = EditorProjectSettings->NonRealtimeScreenPercentageMode;
 		UserPreference = EditorUserSettings->NonRealtimeScreenPercentageMode;
-		CVarDefaultScreenPercentage = CVarEditorViewportDefaultScreenPercentageMode;
+		CVarDefaultScreenPercentage = IConsoleManager::Get().FindConsoleVariable(TEXT("r.Editor.Viewport.ScreenPercentageMode.NonRealTime"));
 	}
 	else
 	{
