[shaders] add compile stats logging to the end of the CookGlobalShaders commandlet (with a manual call to PrintStats, same as is done in the CompileShaders cmdlet - using the existing shader stats tickable object would be "nicer" as it would automatically handle printing the stats intermittently, but requires a lot more boilerplate so sticking with the simple thing)

diff --git a/Engine/Source/Editor/UnrealEd/Private/Commandlets/CookGlobalShaderCommandlet.cpp b/Engine/Source/Editor/UnrealEd/Private/Commandlets/CookGlobalShaderCommandlet.cpp
--- a/Engine/Source/Editor/UnrealEd/Private/Commandlets/CookGlobalShaderCommandlet.cpp
+++ b/Engine/Source/Editor/UnrealEd/Private/Commandlets/CookGlobalShaderCommandlet.cpp
@@ -364,6 +364,8 @@ int32 UCookGlobalShadersCommandlet::Main(const FString& Params)
 		CookGlobalShaders();
 	}
 
+	GShaderCompilingManager->PrintStats();
+
 	UE_LOG(LogCookGlobalShaders, Log, TEXT("Complete"));
 
 	DeviceHelper = nullptr;
