[Materials] Add static switch dependency analysis to DumpMaterialShaderType commandlet.

- Add material dependency walker function to analyze static switch parameter dependency to texture nodes;
  This can later also be used to automatically convert static switches to dynamic branches in the material translator.
- Refactor commandlet to write HTML pages in dedicated class, to handle CSS style output for better readability.
- Log analysis progress in time intervals since large projects with thousands of materials can take a long time (several minutes to hours) to finish the analysis.

#rnx

diff --git a/Engine/Source/Editor/UnrealEd/Private/Commandlets/DumpMaterialShaderTypes.cpp b/Engine/Source/Editor/UnrealEd/Private/Commandlets/DumpMaterialShaderTypes.cpp
--- a/Engine/Source/Editor/UnrealEd/Private/Commandlets/DumpMaterialShaderTypes.cpp
+++ b/Engine/Source/Editor/UnrealEd/Private/Commandlets/DumpMaterialShaderTypes.cpp
@@ -13,6 +13,8 @@
 #include "Materials/Material.h"
 #include "Materials/MaterialInstance.h"
 #include "Materials/MaterialInstanceConstant.h"
+#include "Materials/MaterialExpressionParameter.h"
+#include "Materials/MaterialDependencyWalker.h"
 #include "MaterialShared.h"
 #include "CollectionManagerTypes.h"
 #include "ICollectionManager.h"
@@ -274,6 +276,105 @@ private:
 	FString OutputFileName;
 };
 
+static const TCHAR* kCssBgStyleSolidRed = TEXT("background-color:red;");
+static const TCHAR* kCssBgStyleSolidGray = TEXT("background-color:silver;");
+static const TCHAR* kCssBgStyleDashedRed = TEXT("background: repeating-linear-gradient(45deg,red 0px,red 4px,transparent 4px,transparent 8px)");
+static const TCHAR* kCssBgStyleDashedGray = TEXT("background: repeating-linear-gradient(45deg,silver 0px,silver 4px,transparent 4px,transparent 8px)");
+static const TCHAR* kCssBgStyleDashedLiteGray = TEXT("background: repeating-linear-gradient(45deg,whitesmoke 0px,whitesmoke 4px,transparent 4px,transparent 8px)");
+
+// Helper class to emit HTML page with pre-defined CSS styles to improve readability of the analysis output.
+class FHtmlPageWriter
+{
+public:
+	FHtmlPageWriter(FString&& BaseName) :
+		BaseName(MoveTemp(BaseName))
+	{
+	}
+
+	~FHtmlPageWriter()
+	{
+		CloseDocument();
+	}
+
+	void OpenNewDocument()
+	{
+		CloseDocument();
+		Output = TUniquePtr<FShaderStatsGatheringContext>(new FShaderStatsGatheringContext(FString::Printf(TEXT("%s-Part_%d.html"), *BaseName, PartCounter)));
+		WriteHtmlHeader();
+		++PartCounter;
+		LineCounter = 0;
+	}
+
+	void WriteLine(const FString& Line)
+	{
+		if (!Output)
+		{
+			OpenNewDocument();
+		}
+		Output->Log(Line);
+		++LineCounter;
+	}
+
+	int32 NumLines() const
+	{
+		return LineCounter;
+	}
+
+private:
+	void CloseDocument()
+	{
+		if (Output)
+		{
+			WriteHtmlFooter();
+			Output.Reset();
+		}
+	}
+
+	void WriteHtmlHeader()
+	{
+		// Write HTML header
+		Output->Log(TEXT("<!DOCTYPE html>"));
+		Output->Log(TEXT("<html>"));
+		Output->Log(TEXT("<head>"));
+		Output->Log(TEXT("\t<title>StaticSwitchOptimizer</title>"));
+		Output->Log(TEXT("\t<style>"));
+		Output->Log(TEXT("\t\ttable {border: 1px solid black; font-family: monospace;}"));
+		Output->Log(TEXT("\t\tth {padding-right: 5px; padding-left: 5px; padding-top: 2px; padding-bottom: 2px;}"));
+		Output->Log(TEXT("\t</style>"));
+		Output->Log(TEXT("</head>"));
+
+		// Write HTML legend table
+		Output->Log(TEXT("<body>"));
+
+		Output->Log(TEXT("<table>"));
+		Output->Log(TEXT("\t<tr><th>"));
+		Output->Log(TEXT("\t\t<h3>LEGEND</h3>"));
+		Output->Log(TEXT("\t</th></tr>"));
+		Output->Log(TEXT("\t<tr><td>"));
+		Output->Log(TEXT("\t\t<table>"));
+		Output->Log(FString::Printf(TEXT("\t\t\t<tr><th>Unique static switch</th><th style=\"%s\">Gray background</th></tr>"), kCssBgStyleSolidGray));
+		Output->Log(FString::Printf(TEXT("\t\t\t<tr><th>Varying static switch ON</th><th style=\"%s\">Red background</th></tr>"), kCssBgStyleSolidRed));
+		Output->Log(FString::Printf(TEXT("\t\t\t<tr><th>Trivial graph dependency (No texture input)</th><th style=\"%s\">Dashed background</th></tr>"), kCssBgStyleDashedGray));
+		Output->Log(TEXT("\t\t</table>"));
+		Output->Log(TEXT("\t</th></tr>"));
+		Output->Log(TEXT("</table>"));
+		Output->Log(TEXT("<br></br>"));
+	}
+
+	void WriteHtmlFooter()
+	{
+		Output->Log(TEXT("</body>"));
+		Output->Log(TEXT("</html>"));
+	}
+
+private:
+	FString BaseName;
+	TUniquePtr<FShaderStatsGatheringContext> Output;
+	int32 PartCounter = 0;
+	int32 LineCounter = 0;
+};
+
+
 UDumpMaterialShaderTypesCommandlet::UDumpMaterialShaderTypesCommandlet(const FObjectInitializer& ObjectInitializer)
 	: Super(ObjectInitializer)
 {
@@ -396,25 +497,10 @@ static int ProcessMaterials(const ITargetPlatform* TargetPlatform, const EShader
 	return TotalShaders;
 }
 
-static void ProcessSwitchOptimizer(const ITargetPlatform* TargetPlatform, const EShaderPlatform ShaderPlatform, const TArray<FAssetData>& MaterialList, const TArray<FAssetData>& MaterialInstanceList)
+static void ProcessSwitchOptimizer(const ITargetPlatform* TargetPlatform, const EShaderPlatform ShaderPlatform, const TArray<FAssetData>& MaterialList, const TArray<FAssetData>& MaterialInstanceList, const FString& TimeNow)
 {
 	TRACE_CPUPROFILER_EVENT_SCOPE(ProcessSwitchOptimizer);
 
-	const FString TimeNow = FDateTime::Now().ToString();
-	FString FileName = FString::Printf(TEXT("%s-StaticSwitches-%s-%s-%s.html"), FApp::GetProjectName(), *TargetPlatform->PlatformName(), *LexToString(ShaderPlatform), *TimeNow);
-	FShaderStatsGatheringContext Output(FileName);
-
-	Output.Log(TEXT("<!DOCTYPE html>"));
-	Output.Log(TEXT("<html>"));
-	Output.Log(TEXT("<head>"));
-	Output.Log(TEXT("\t<title>StaticSwitchOptimizer</title>"));
-	Output.Log(TEXT("\t<style>"));
-	Output.Log(TEXT("\t\ttable {border: 1px solid black; font-family: monospace;}"));
-	Output.Log(TEXT("\t\tth {padding-right: 5px; padding-left: 5px; padding-top: 2px; padding-bottom: 2px;}"));
-	Output.Log(TEXT("\t</style>"));
-	Output.Log(TEXT("</head>"));
-	Output.Log(TEXT("<body>"));
-
 	struct FOuterKeyFuncs
 	{
 		static inline bool Matches(const FMaterialShaderMapId& A, const FMaterialShaderMapId& B)
@@ -459,8 +545,49 @@ static void ProcessSwitchOptimizer(const ITargetPlatform* TargetPlatform, const
 	using ShaderIdGroupSet = Experimental::TRobinHoodHashMap<FMaterialShaderMapId, StaticSwitchGroupSet, FOuterKeyFuncs>;
 	ShaderIdGroupSet ShaderMapHashMap;
 
-	for (const FAssetData& AssetData : MaterialList)
+	using FMaterialExpressionArray = TArray<UMaterialExpression*>;
+	using FShaderMapIdToMaterialExpressionMap = TMap<FMaterialShaderMapId, FMaterialExpressionArray>;
+	FShaderMapIdToMaterialExpressionMap MaterialExpressionMap;
+
+	auto AppendExpressionsToNameMap = [&MaterialExpressionMap](const FMaterialShaderMapId& ShaderMapId, UMaterial* Material) -> void
+		{
+			FMaterialExpressionArray& ExpressionArrayForShaderMap = MaterialExpressionMap.FindOrAdd(ShaderMapId);
+			Material->GetAllExpressionsInMaterialAndFunctionsOfType<UMaterialExpression>(ExpressionArrayForShaderMap);
+		};
+
+	auto FindMaterialExpressionByName = [&MaterialExpressionMap](const FMaterialShaderMapId& ShaderMapId, const FName& Name) -> UMaterialExpression*
+		{
+			// Try to find expression for specified shadermap
+			if (const FMaterialExpressionArray* ExpressionsForShaderMap = MaterialExpressionMap.Find(ShaderMapId))
+			{
+				for (UMaterialExpression* Expression : *ExpressionsForShaderMap)
+				{
+					if (Expression->GetParameterName() == Name)
+					{
+						return Expression;
+					}
+				}
+			}
+			return nullptr;
+		};
+
+	// Log the analysis progress ever 5 seconds as some projects can take a long time (minutes to hours) to analyze
+	constexpr int64 ProgressUpdateIntervalInSeconds = 5;
+	int64 LastProgressUpdateTimestamp = FDateTime::Now().ToUnixTimestamp();
+
+	auto LogProgressInInterval = [&LastProgressUpdateTimestamp](const TCHAR* Info, int32 Progress, int32 Max) -> void
+		{
+			const int64 CurrentTimestamp = FDateTime::Now().ToUnixTimestamp();
+			if (CurrentTimestamp - LastProgressUpdateTimestamp > ProgressUpdateIntervalInSeconds)
+			{
+				LastProgressUpdateTimestamp = CurrentTimestamp;
+				UE_LOG(LogDumpMaterialShaderTypesCommandlet, Display, TEXT("%s%d/%d (%0.2f%%)"), Info, Progress, Max, 100.0*static_cast<double>(Progress)/static_cast<double>(Max));
+			}
+		};
+
+	for (int32 MaterialListIndex = 0; MaterialListIndex < MaterialList.Num(); ++MaterialListIndex)
 	{
+		const FAssetData& AssetData = MaterialList[MaterialListIndex];
 		if (UMaterial* Material = Cast<UMaterial>(AssetData.GetAsset()))
 		{
 			TArray<FMaterialResource*> ResourcesToCache;
@@ -479,13 +606,19 @@ static void ProcessSwitchOptimizer(const ITargetPlatform* TargetPlatform, const
 				{
 					InnerValue.FindOrAdd(Material, TotalNumShaders);
 				}
+
+				// Store mapping from shadermap ID to expression array
+				AppendExpressionsToNameMap(ShaderMapId, Material);
 			}
 			FMaterial::DeferredDeleteArray(ResourcesToCache);
+
+			LogProgressInInterval(TEXT("Building shader maps for materials in progress: "), MaterialListIndex + 1, MaterialList.Num());
 		}
 	}
 
-	for (const FAssetData& AssetData : MaterialInstanceList)
+	for (int32 MaterialInstanceIndex = 0; MaterialInstanceIndex < MaterialInstanceList.Num(); ++MaterialInstanceIndex)
 	{
+		const FAssetData& AssetData = MaterialInstanceList[MaterialInstanceIndex];
 		if (UMaterialInstance* MaterialInstance = Cast<UMaterialInstance>(AssetData.GetAsset()))
 		{
 			TArray<FMaterialResource*> ResourcesToCache;
@@ -505,8 +638,13 @@ static void ProcessSwitchOptimizer(const ITargetPlatform* TargetPlatform, const
 				{
 					InnerValue.FindOrAdd(MaterialInstance, TotalNumShaders);
 				}
+
+				// Store mapping from shadermap ID to expression array
+				AppendExpressionsToNameMap(ShaderMapId, BaseMaterial);
 			}
 			FMaterial::DeferredDeleteArray(ResourcesToCache);
+
+			LogProgressInInterval(TEXT("Building shader maps for material instances in progress: "), MaterialInstanceIndex + 1, MaterialInstanceList.Num());
 		}
 	}
 
@@ -515,7 +653,8 @@ static void ProcessSwitchOptimizer(const ITargetPlatform* TargetPlatform, const
 
 	struct FStaticSwitchMetaData
 	{
-		bool bIsVarying;
+		bool bIsVarying = false;
+		bool bHasTrivialDependency = false; // Trivial dependency with only a limited number of material expressions and no texture dependency
 		TArray<bool> PermutationVector; // Row-vector in the permutation matrix
 	};
 
@@ -531,8 +670,11 @@ static void ProcessSwitchOptimizer(const ITargetPlatform* TargetPlatform, const
 	
 	TArray<FInnerFilteredType> FilteredHashMap;
 
+	int32 ShaderMapIndex = 0;
 	for(const ShaderIdAndSwitchGroup& OuterElement : ShaderMapHashMap)
 	{
+		LogProgressInInterval(TEXT("Analyzing static switches in progress: "), ++ShaderMapIndex, ShaderMapHashMap.Num());
+
 		if(OuterElement.Value.Num() > 1)
 		{
 			ShaderIdGroupArray InnerMap;
@@ -558,6 +700,20 @@ static void ProcessSwitchOptimizer(const ITargetPlatform* TargetPlatform, const
 					{
 						MetaData->bIsVarying = true;
 					}
+					if (UMaterialExpression* ExpressionForStaticSwitch = FindMaterialExpressionByName(OuterElement.Key, Key.ParameterInfo.Name))
+					{
+						constexpr int32 MaxDependencyWalkDepth = 16;
+						FMaterialDependencySearchMetadata DependencyMetaData;
+						const int32 WalkDepth = WalkMaterialDependencyGraph(ExpressionForStaticSwitch, MaxDependencyWalkDepth, MDSF_TextureDependencyOnly, DependencyMetaData);
+						if (WalkDepth != INDEX_NONE && !DependencyMetaData.bHasTextureInput)
+						{
+							MetaData->bHasTrivialDependency = true;
+						}
+					}
+					else
+					{
+						UE_LOG(LogDumpMaterialShaderTypesCommandlet, Error, TEXT("Failed to find static switch parameter \"%s\""), *Key.ParameterInfo.Name.ToString());
+					}
 				}
 				InnerMap.Emplace(OuterElement.Key, Inner);
 			}
@@ -601,10 +757,13 @@ static void ProcessSwitchOptimizer(const ITargetPlatform* TargetPlatform, const
 		return NumA > NumB;
 	});
 
-	for(int32 i = 0; i < FilteredHashMap.Num(); i++)
+	const FString BaseHtmlDocumentName = FString::Printf(TEXT("%s-StaticSwitches-%s-%s-%s"), FApp::GetProjectName(), *TargetPlatform->PlatformName(), *LexToString(ShaderPlatform), *TimeNow);
+	FHtmlPageWriter PageWriter(FPaths::Combine(*TimeNow, *TargetPlatform->PlatformName(), *LexToString(ShaderPlatform), BaseHtmlDocumentName));
+
+	for (FInnerFilteredType& InnerFilteredType : FilteredHashMap)
 	{
-		const ShaderIdGroupArray& InnerMap = FilteredHashMap[i].ShaderIdGroups;
-		SwitchPermutationVectorType& StaticSwitches = FilteredHashMap[i].StaticSwitches;
+		const ShaderIdGroupArray& InnerMap = InnerFilteredType.ShaderIdGroups;
+		SwitchPermutationVectorType& StaticSwitches = InnerFilteredType.StaticSwitches;
 		const UMaterialInterface* Parent = (*(*InnerMap.begin()).Value.Value.begin()).Key->GetMaterial();
 		const FMaterialShaderMapId& ShaderId = (*InnerMap.begin()).Key;
 
@@ -635,54 +794,71 @@ static void ProcessSwitchOptimizer(const ITargetPlatform* TargetPlatform, const
 			}
 		}
 
-		Output.Log(TEXT("<table>"));
-		Output.Log(FString::Printf(TEXT("<tr><th><h3>Candidate %s (Shaders: %d)</h3></th></tr>"), *Parent->GetOuter()->GetFName().ToString(), NumShaders));
-		Output.Log(TEXT("<tr><td><table>"));
+		// Open a new HTML document after the previous one reached the maximum size. Otherwise, it's hard to browse such large HTML documents.
+		constexpr int32 MaxLinesPerHtmlPage = 10000;
+		if (PageWriter.NumLines() > MaxLinesPerHtmlPage)
+		{
+			PageWriter.OpenNewDocument();
+		}
+
+		PageWriter.WriteLine(TEXT("<table>"));
+		PageWriter.WriteLine(FString::Printf(TEXT("<tr><th><h3>Candidate %s (Shaders: %d)</h3></th></tr>"), *Parent->GetOuter()->GetFName().ToString(), NumShaders));
+		PageWriter.WriteLine(TEXT("<tr><td><table>"));
 
 		// Row for captions
-		Output.Log(TEXT("\t<tr>"));
-		Output.Log(FString::Printf(TEXT("\t\t<th style=\"background-color:gray;\">%d Static Switch(es)</th>"), NumStaticSwitchesTotal));
-		Output.Log(FString::Printf(TEXT("\t\t<th style=\"background-color:gray;\" colspan=\"%d\">%d Permutation(s)</th>"), NumStaticSwitchPermutations, NumStaticSwitchPermutations));
-		Output.Log(TEXT("\t</tr>"));
+		PageWriter.WriteLine(TEXT("\t<tr>"));
+		PageWriter.WriteLine(FString::Printf(TEXT("\t\t<th style=\"background-color:gray;\">%d Static Switch(es)</th>"), NumStaticSwitchesTotal));
+		PageWriter.WriteLine(FString::Printf(TEXT("\t\t<th style=\"background-color:gray;\" colspan=\"%d\">%d Permutation(s)</th>"), NumStaticSwitchPermutations, NumStaticSwitchPermutations));
+		PageWriter.WriteLine(TEXT("\t</tr>"));
 
 		// Row for each static switch parameter that is included in at least one permutation
 		for (const TPair<const FName, FStaticSwitchMetaData>& Param : StaticSwitches)
 		{
-			Output.Log(TEXT("\t<tr>"));
-			Output.Log(FString::Printf(TEXT("\t\t<th>%s</th>"), *Param.Key.ToString()));
+			PageWriter.WriteLine(TEXT("\t<tr>"));
+			PageWriter.WriteLine(FString::Printf(TEXT("\t\t<th>%s</th>"), *Param.Key.ToString()));
 
 			if (Param.Value.bIsVarying)
 			{
 				for (bool bSwitchEnabled : Param.Value.PermutationVector)
 				{
-					Output.Log(
-						bSwitchEnabled
-							? TEXT("\t\t<th style=\"background-color:red;\">1</th>")
-							: TEXT("\t\t<th>0</th>")
-					);
+					if (Param.Value.bHasTrivialDependency)
+					{
+						PageWriter.WriteLine(
+							FString::Printf(
+								TEXT("\t\t<th style=\"%s\">%s</th>"),
+								bSwitchEnabled ? kCssBgStyleDashedRed : kCssBgStyleDashedLiteGray,
+								bSwitchEnabled ? TEXT("1") : TEXT("0"))
+						);
+					}
+					else
+					{
+						PageWriter.WriteLine(
+							bSwitchEnabled
+								? FString::Printf(TEXT("\t\t<th style=\"%s\">1</th>"), kCssBgStyleSolidRed)
+								: TEXT("\t\t<th>0</th>")
+						);
+					}
 				}
 			}
 			else
 			{
 				for (bool bSwitchEnabled : Param.Value.PermutationVector)
 				{
-					Output.Log(
-						bSwitchEnabled
-							? TEXT("\t\t<th style=\"background-color:silver;\">1</th>")
-							: TEXT("\t\t<th style=\"background-color:silver;\">0</th>")
+					PageWriter.WriteLine(
+						FString::Printf(
+							TEXT("\t\t<th style=\"%s\">%s</th>"),
+							Param.Value.bHasTrivialDependency ? kCssBgStyleDashedGray : kCssBgStyleSolidGray,
+							bSwitchEnabled ? TEXT("1") : TEXT("0"))
 					);
 				}
 			}
-			Output.Log(TEXT("\t</tr>"));
+			PageWriter.WriteLine(TEXT("\t</tr>"));
 		}
 
-		Output.Log(TEXT("</table></td></tr>"));
-		Output.Log(TEXT("</table>"));
-		Output.Log(TEXT("<br></br>"));
+		PageWriter.WriteLine(TEXT("</table></td></tr>"));
+		PageWriter.WriteLine(TEXT("</table>"));
+		PageWriter.WriteLine(TEXT("<br></br>"));
 	}
-
-	Output.Log(TEXT("</body>"));
-	Output.Log(TEXT("</html>"));
 }
 
 static int ProcessMaterialInstances(const ITargetPlatform* TargetPlatform, const EShaderPlatform ShaderPlatform, FShaderStatsGatheringContext& Output, const TArray<FAssetData>& MaterialInstanceList)
@@ -991,6 +1167,8 @@ PRAGMA_ENABLE_DEPRECATION_WARNINGS
 		}
 	}
 
+	const FString TimeNow = FDateTime::Now().ToString();
+
 	for (int32 Index = 0; Index < Platforms.Num(); Index++)
 	{
 		TArray<FName> DesiredShaderFormats;
@@ -1003,7 +1181,7 @@ PRAGMA_ENABLE_DEPRECATION_WARNINGS
 			UE_LOG(LogDumpMaterialShaderTypesCommandlet, Display, TEXT("Dumping material shader types for '%s' - '%s'..."), *Platforms[Index]->PlatformName(), *LexToString(ShaderPlatform));
 			if(bStaticSwitches)
 			{
-				ProcessSwitchOptimizer(Platforms[Index], ShaderPlatform, MaterialList, MaterialInstanceList);
+				ProcessSwitchOptimizer(Platforms[Index], ShaderPlatform, MaterialList, MaterialInstanceList, TimeNow);
 			}
 			else
 			{
diff --git a/Engine/Source/Runtime/Engine/Private/Materials/MaterialDependencyWalker.cpp b/Engine/Source/Runtime/Engine/Private/Materials/MaterialDependencyWalker.cpp
new file mode 100644
--- /dev/null
+++ b/Engine/Source/Runtime/Engine/Private/Materials/MaterialDependencyWalker.cpp
@@ -0,0 +1,54 @@
+// Copyright Epic Games, Inc. All Rights Reserved.
+
+/*=============================================================================
+	MaterialExpressionDependencyWalker.cpp - Material expressions implementation.
+=============================================================================*/
+
+#if WITH_EDITOR
+
+#include "Materials/MaterialDependencyWalker.h"
+#include "Materials/MaterialExpressionTextureSample.h"
+
+
+static bool WalkMaterialDependencyGraphInternal(const UMaterialExpression* Expression, int32 CurrentDepth, int32 MaxDepth, int32& OutDepth, uint32 SearchFlags, FMaterialDependencySearchMetadata& OutMetaData)
+{
+	// Check if end of search has been reached
+	const int32 NextDepth = CurrentDepth + 1;
+	if (Expression == nullptr || NextDepth > MaxDepth)
+	{
+		OutDepth = INDEX_NONE;
+		return false;
+	}
+	OutDepth = FMath::Max(OutDepth, NextDepth);
+
+	// Check for meta data we are searching for and continue walking the graph
+	if (const auto* TextureExpression = Cast<UMaterialExpressionTextureBase>(Expression))
+	{
+		OutMetaData.bHasTextureInput = true;
+		if ((SearchFlags & MDSF_TextureDependencyOnly) != 0)
+		{
+			return false;
+		}
+	}
+
+	// Continue to walk the dependency graph along the input expressions
+	for (int32 InputIndex = 0; const FExpressionInput* Input = Expression->GetInput(InputIndex); ++InputIndex)
+	{
+		if (!WalkMaterialDependencyGraphInternal(Input->Expression, NextDepth, MaxDepth, OutDepth, SearchFlags, OutMetaData))
+		{
+			return false;
+		}
+	}
+
+	return true;
+}
+
+ENGINE_API int32 WalkMaterialDependencyGraph(const UMaterialExpression* Expression, int32 MaxDepth, uint32 SearchFlags, FMaterialDependencySearchMetadata& OutMetaData)
+{
+	int32 MaxTraversalDepth = INDEX_NONE;
+	(void)WalkMaterialDependencyGraphInternal(Expression, 0, MaxDepth, MaxTraversalDepth, SearchFlags, OutMetaData);
+	return MaxTraversalDepth;
+}
+
+
+#endif
diff --git a/Engine/Source/Runtime/Engine/Public/Materials/MaterialDependencyWalker.h b/Engine/Source/Runtime/Engine/Public/Materials/MaterialDependencyWalker.h
new file mode 100644
--- /dev/null
+++ b/Engine/Source/Runtime/Engine/Public/Materials/MaterialDependencyWalker.h
@@ -0,0 +1,32 @@
+// Copyright Epic Games, Inc. All Rights Reserved.
+#pragma once
+
+#if WITH_EDITOR
+
+#include "CoreMinimal.h"
+
+class UMaterialExpression;
+
+struct FMaterialDependencySearchMetadata
+{
+	FMaterialDependencySearchMetadata() :
+		bHasTextureInput(0)
+	{
+	}
+
+	uint32 bHasTextureInput : 1;
+};
+
+enum EMaterialDependencySearchFlags
+{
+	MDSF_TextureDependencyOnly = (1 << 0),
+};
+
+
+// Walks the dependency graph (i.e. all expression inputs) of the specified material expression.
+// Returns the maxmium depth the search went or INDEX_NONE if the limit has reached before finding anything.
+// @param SearchFlags can be a bitwise OR combination of EMaterialDependencySearchFlags.
+ENGINE_API int32 WalkMaterialDependencyGraph(const UMaterialExpression* Expression, int32 MaxDepth, uint32 SearchFlags, FMaterialDependencySearchMetadata& OutMetaData);
+
+
+#endif // WITH_EDITOR
