Prevent editor context clients created after a UActorEditorContextSubsystem::PushContext() from being considered by the matching UActorEditorContextSubsystem::PopContext() call

diff --git a/Engine/Source/Editor/UnrealEd/Private/Subsystems/ActorEditorContextSubsystem.cpp b/Engine/Source/Editor/UnrealEd/Private/Subsystems/ActorEditorContextSubsystem.cpp
--- a/Engine/Source/Editor/UnrealEd/Private/Subsystems/ActorEditorContextSubsystem.cpp
+++ b/Engine/Source/Editor/UnrealEd/Private/Subsystems/ActorEditorContextSubsystem.cpp
@@ -73,6 +73,11 @@ void UActorEditorContextSubsystem::UnregisterClient(IActorEditorContextClient* C
 		if (Clients.Remove(Client))
 		{
 			Client->GetOnActorEditorContextClientChanged().RemoveAll(this);
+
+			for (TArray<IActorEditorContextClient*>& PushedClients : PushedContextsStack)
+			{
+				PushedClients.Remove(Client);
+			}
 		}
 	}
 }
@@ -145,6 +150,9 @@ void UActorEditorContextSubsystem::PushContext(bool bDuplicateContext)
 	{
 		Client->OnExecuteActorEditorContextAction(World, bDuplicateContext ? EActorEditorContextAction::PushDuplicateContext : EActorEditorContextAction::PushContext);
 	}
+
+	PushedContextsStack.Push(Clients);
+
 	ActorEditorContextSubsystemChanged.Broadcast();
 }
 
@@ -156,7 +164,7 @@ void UActorEditorContextSubsystem::PopContext()
 		return;
 	}
 
-	for (IActorEditorContextClient* Client : Clients)
+	for (IActorEditorContextClient* Client : PushedContextsStack.Pop(EAllowShrinking::No))
 	{
 		Client->OnExecuteActorEditorContextAction(World, EActorEditorContextAction::PopContext);
 	}
diff --git a/Engine/Source/Editor/UnrealEd/Public/Subsystems/ActorEditorContextSubsystem.h b/Engine/Source/Editor/UnrealEd/Public/Subsystems/ActorEditorContextSubsystem.h
--- a/Engine/Source/Editor/UnrealEd/Public/Subsystems/ActorEditorContextSubsystem.h
+++ b/Engine/Source/Editor/UnrealEd/Public/Subsystems/ActorEditorContextSubsystem.h
@@ -48,5 +48,6 @@ private:
 
 	FOnActorEditorContextSubsystemChanged ActorEditorContextSubsystemChanged;
 	TArray<IActorEditorContextClient*> Clients;
+	TArray<TArray<IActorEditorContextClient*>> PushedContextsStack;
 	bool bIsApplyEnabled = true;
 };
