MPCook: CookPackageSplitters: Fix crash when reassigning after retraction of generated packages due to spuriously clearing the package's reachability when putting it in the stalled state on the worker.
#rnx

diff --git a/Engine/Source/Editor/UnrealEd/Private/Cooker/CookWorkerClient.cpp b/Engine/Source/Editor/UnrealEd/Private/Cooker/CookWorkerClient.cpp
--- a/Engine/Source/Editor/UnrealEd/Private/Cooker/CookWorkerClient.cpp
+++ b/Engine/Source/Editor/UnrealEd/Private/Cooker/CookWorkerClient.cpp
@@ -895,17 +895,22 @@ void FCookWorkerClient::AssignPackages(FAssignPackagesMessage& Message)
 					&NeedCookPlatformsBuffer);
 			if (PackageData.IsInProgress())
 			{
-				// If already in progress but there are new platforms requested, we don't currently handle that;
-				// the director should have retracted it first
+				// If already in progress but there are new platforms requested, demote the package back to Load
 				for (const ITargetPlatform* TargetPlatform : NeedCookPlatforms)
 				{
 					check(TargetPlatform != CookerLoadingPlatformKey);
-					checkf(PackageData.FindOrAddPlatformData(TargetPlatform).IsReachable(),
-						TEXT("CookWorker received AssignPackage for package %s which is already in progress, with new platform %s. Adding new platforms to an inprogress package is not yet supported"),
-						*PackageData.GetPackageName().ToString(), *TargetPlatform->PlatformName());
+					if (!PackageData.FindOrAddPlatformData(TargetPlatform).IsReachable())
+					{
+						if (PackageData.IsInStateProperty(EPackageStateProperty::Saving))
+						{
+							UE_LOG(LogCook, Display,
+								TEXT("Package %s is in the save state, but the CookDirector updated the requested platforms to include the new platform %s. Restarting the package's save."),
+								*PackageData.GetPackageName().ToString(), *TargetPlatform->PlatformName());
+							PackageData.SendToState(EPackageState::LoadReady, ESendFlags::QueueAddAndRemove, EStateChangeReason::DirectorRequest);
+						}
+					}
 				}
-				// If no new platforms, just allow the package to continue in its progress. If it was in a stalled-by-retraction state,
-				// return it to active.
+				// Allow the package to continue in its progress. If it was in a stalled-by-retraction state, return it to active.
 				PackageData.UnStall(ESendFlags::QueueAddAndRemove);
 				continue;
 			}
@@ -1043,7 +1048,6 @@ void FCookWorkerClient::HandleRetractionMessage(FMPCollectorClientMessageContext
 	{
 		FPackageData* PackageData = COTFS.PackageDatas->FindPackageDataByPackageName(PackageName);
 		check(PackageData);
-		PackageData->ResetReachable();
 		TRefCountPtr<FGenerationHelper> GenerationHelper = PackageData->GetGenerationHelper();
 		if (!GenerationHelper)
 		{
@@ -1064,6 +1068,7 @@ void FCookWorkerClient::HandleRetractionMessage(FMPCollectorClientMessageContext
 		{
 			COTFS.DemoteToIdle(*PackageData, ESendFlags::QueueAddAndRemove,
 				ESuppressCookReason::RetractedByCookDirector);
+			PackageData->ResetReachable();
 		}
 	}
 
