MPCook: Fix long delays in handling retraction messages (CookWorker 0 has not responded to a RetractionRequest message for 302.6 seconds) when the director sends multiple message-containing packets to a CookWorker within a single poll period on  the CookWorker.
#rnx

diff --git a/Engine/Source/Editor/UnrealEd/Private/Cooker/CookWorkerClient.cpp b/Engine/Source/Editor/UnrealEd/Private/Cooker/CookWorkerClient.cpp
--- a/Engine/Source/Editor/UnrealEd/Private/Cooker/CookWorkerClient.cpp
+++ b/Engine/Source/Editor/UnrealEd/Private/Cooker/CookWorkerClient.cpp
@@ -660,15 +660,25 @@ void FCookWorkerClient::PumpReceiveMessages()
 {
 	using namespace UE::CompactBinaryTCP;
 	TArray<FMarshalledMessage> Messages;
-	EConnectionStatus SocketStatus = TryReadPacket(ServerSocket, ReceiveBuffer, Messages);
-	if (SocketStatus != EConnectionStatus::Okay && SocketStatus != EConnectionStatus::Incomplete)
+
+	// Read a packet at a time (with 1 or more messages per packet) until we fail to read any messages
+	for (;;)
 	{
-		UE_LOG(LogCook, Error,
-			TEXT("CookWorkerClient failed to read from Director. We will abort the CookAsCookWorker commandlet."));
-		SendToState(EConnectStatus::LostConnection);
-		return;
+		Messages.Reset();
+		EConnectionStatus SocketStatus = TryReadPacket(ServerSocket, ReceiveBuffer, Messages);
+		if (SocketStatus != EConnectionStatus::Okay && SocketStatus != EConnectionStatus::Incomplete)
+		{
+			UE_LOG(LogCook, Error,
+				TEXT("CookWorkerClient failed to read from Director. We will abort the CookAsCookWorker commandlet."));
+			SendToState(EConnectStatus::LostConnection);
+			return;
+		}
+		if (Messages.IsEmpty())
+		{
+			break;
+		}
+		HandleReceiveMessages(MoveTemp(Messages));
 	}
-	HandleReceiveMessages(MoveTemp(Messages));
 }
 
 void FCookWorkerClient::HandleReceiveMessages(TArray<UE::CompactBinaryTCP::FMarshalledMessage>&& Messages, FName OptionalPackageName /*= NAME_None*/)
