[BUGFIX] Interchange - Adding fbx file to monitored folder triggered legacy import

- Uses checks to see if Interchange is enabled and imports via Interchange if possible

diff --git a/Engine/Source/Editor/UnrealEd/Private/AutoReimport/ContentDirectoryMonitor.cpp b/Engine/Source/Editor/UnrealEd/Private/AutoReimport/ContentDirectoryMonitor.cpp
--- a/Engine/Source/Editor/UnrealEd/Private/AutoReimport/ContentDirectoryMonitor.cpp
+++ b/Engine/Source/Editor/UnrealEd/Private/AutoReimport/ContentDirectoryMonitor.cpp
@@ -9,6 +9,7 @@
 #include "Factories/SceneImportFactory.h"
 #include "EditorFramework/AssetImportData.h"
 #include "Editor.h"
+#include "Misc/MessageDialog.h"
 
 #include "AutoReimport/AutoReimportUtilities.h"
 
@@ -21,6 +22,9 @@
 #include "AutoReimport/AssetSourceFilenameCache.h"
 #include "Misc/NamePermissionList.h"
 
+#include "InterchangeManager.h"
+#include "InterchangeSceneImportAsset.h"
+
 #define LOCTEXT_NAMESPACE "ContentDirectoryMonitor"
 
 bool IsAssetDirty(UObject* Asset)
@@ -187,11 +191,11 @@ void FContentDirectoryMonitor::IterateUnprocessedChanges(TFunctionRef<bool(const
 int32 FContentDirectoryMonitor::StartProcessing()
 {
 	// We only process things that haven't changed for a given threshold
-	auto& FileManager = IFileManager::Get();
+	IFileManager& FileManager = IFileManager::Get();
 	const FDateTime Threshold = FDateTime::UtcNow() - FTimespan::FromSeconds(GetDefault<UEditorLoadingSavingSettings>()->AutoReimportThreshold);
 
 	// Get all the changes that have happend beyond our import threshold
-	auto OutstandingChanges = Cache.FilterOutstandingChanges([this, Threshold](const DirectoryWatcher::FUpdateCacheTransaction& Transaction, const FDateTime& TimeOfChange){
+	TArray<DirectoryWatcher::FUpdateCacheTransaction> OutstandingChanges = Cache.FilterOutstandingChanges([this, Threshold](const DirectoryWatcher::FUpdateCacheTransaction& Transaction, const FDateTime& TimeOfChange){
 		return TimeOfChange <= Threshold && ShouldConsiderChange(Transaction);
 	});
 
@@ -200,8 +204,8 @@ int32 FContentDirectoryMonitor::StartProcessing()
 		return 0;
 	}
 
-	const auto* Settings = GetDefault<UEditorLoadingSavingSettings>();
-	for (auto& Transaction : OutstandingChanges)
+	const UEditorLoadingSavingSettings* Settings = GetDefault<UEditorLoadingSavingSettings>();
+	for (DirectoryWatcher::FUpdateCacheTransaction& Transaction : OutstandingChanges)
 	{
 		switch(Transaction.Action)
 		{
@@ -246,7 +250,7 @@ UObject* AttemptImport(UClass* InFactoryType, UPackage* Package, FName InName, b
 		Factory->AddToRoot();
 		if (Factory->ConfigureProperties())
 		{
-			if (auto* SupportedClass = Factory->ResolveSupportedClass())
+			if (UClass* SupportedClass = Factory->ResolveSupportedClass())
 			{
 				Asset = Factory->ImportObject(SupportedClass, Package, InName, RF_Public | RF_Standalone, FullFilename, nullptr, bCancelled);
 			}
@@ -259,10 +263,13 @@ UObject* AttemptImport(UClass* InFactoryType, UPackage* Package, FName InName, b
 
 void FContentDirectoryMonitor::ProcessAdditions(const DirectoryWatcher::FTimeLimit& TimeLimit, TArray<TObjectPtr<UPackage>>& OutPackagesToSave, const TMap<FString, TArray<UFactory*>>& InFactoriesByExtension, FReimportFeedbackContext& Context)
 {
+	UInterchangeManager& InterchangeManager = UInterchangeManager::GetInterchangeManager();
+	const bool bUseInterchangeFramework = UInterchangeManager::IsInterchangeImportEnabled();
+
 	bool bCancelled = false;
 	for (int32 Index = 0; Index < AddedFiles.Num(); ++Index)
 	{
-		auto& Addition = AddedFiles[Index];
+		DirectoryWatcher::FUpdateCacheTransaction& Addition = AddedFiles[Index];
 
 		if (bCancelled)
 		{
@@ -278,7 +285,7 @@ void FContentDirectoryMonitor::ProcessAdditions(const DirectoryWatcher::FTimeLim
 		FString PackagePath = UPackageTools::SanitizePackageName(MountedContentPath / FPaths::GetPath(Addition.Filename.Get()) / NewAssetName);
 
 		// Don't create assets for new files if assets already exist for the filename
-		auto ExistingReferences = Utils::FindAssetsPertainingToFile(*Registry, FullFilename);
+		TArray<FAssetData> ExistingReferences = Utils::FindAssetsPertainingToFile(*Registry, FullFilename);
 		if (ExistingReferences.Num() != 0)
 		{
 			// Treat this as a modified file that will attempt to reimport it (if applicable). We don't update the progress for this item until it is processed by ProcessModifications
@@ -289,118 +296,165 @@ void FContentDirectoryMonitor::ProcessAdditions(const DirectoryWatcher::FTimeLim
 		// Move the progress on now that we know we're going to process the file
 		Context.MainTask->EnterProgressFrame();
 
-		if (FPackageName::DoesPackageExist(*PackagePath))
+		// See if the asset can be imported using the Interchange before using the provided factory information
+		bool bImportHandledThroughInterchange = false;
+		if (bUseInterchangeFramework)
 		{
-			// Package already exists, so try and import over the top of it, if it doesn't already have a source file path
-			TArray<FAssetData> Assets;
-			if (Registry->GetAssetsByPackageName(*PackagePath, Assets) && Assets.Num() == 1)
+			bool bSceneImport = false;
+			UE::Interchange::FScopedSourceData ScopedSourceData(FullFilename);
+
+			if (InterchangeManager.CanTranslateSourceData(ScopedSourceData.GetSourceData()))
 			{
-				if (UObject* ExistingAsset = Assets[0].GetAsset())
+				FImportAssetParameters ImportAssetParameters;
+				ImportAssetParameters.bIsAutomated = false;
+				ImportAssetParameters.bFollowRedirectors = false;
+				ImportAssetParameters.ReimportAsset = nullptr;
+				ImportAssetParameters.bReplaceExisting = false;
+				ImportAssetParameters.DestinationName = NewAssetName;
+
+				// Check if this is the Scene import or Asset Import
 				{
-					// We're only eligible for reimport if the existing asset doesn't reference a source file already
-					const bool bEligibleForReimport = !Utils::ExtractSourceFilePaths(ExistingAsset).ContainsByPredicate([&](const FString& In){
-						return !In.IsEmpty() && In == FullFilename;
-					});
+					UE::Interchange::FScopedTranslator ScopedTranslator(ScopedSourceData.GetSourceData());
+					bSceneImport = ScopedTranslator.GetTranslator()->GetTranslatorType() == EInterchangeTranslatorType::Scenes;
+				}
 
-					if (bEligibleForReimport)
+				if (bSceneImport)
+				{
+					if (!InterchangeManager.ImportScene(FPaths::GetPath(PackagePath), ScopedSourceData.GetSourceData(), ImportAssetParameters))
 					{
-						ReimportAssetWithNewSource(ExistingAsset, FullFilename, OutPackagesToSave, Context);
+						// Add the message to the message log letting the user know that the last scene import didn't work.
+						Context.GetMessageLog().Message(EMessageSeverity::Info, FText::Format(LOCTEXT("Info_FailedToImportAsset", "Failed to import file {0}."), FText::FromString(FullFilename)));
 					}
 				}
+				else
+				{
+					if (!InterchangeManager.ImportAsset(FPaths::GetPath(PackagePath), ScopedSourceData.GetSourceData(), ImportAssetParameters))
+					{
+						// Add the message to the message log letting the user know that the last scene import didn't work.
+						Context.GetMessageLog().Message(EMessageSeverity::Info, FText::Format(LOCTEXT("Info_FailedToImportAsset", "Failed to import file {0}."), FText::FromString(FullFilename)));
+					}
+				}
+
+				bImportHandledThroughInterchange = true;
 			}
 		}
-		else
+		
+
+		if (!bImportHandledThroughInterchange)
 		{
-			UPackage* NewPackage = CreatePackage( *PackagePath);
-			if ( !ensure(NewPackage) )
+			if (FPackageName::DoesPackageExist(*PackagePath))
 			{
-				Context.AddMessage(EMessageSeverity::Error, FText::Format(LOCTEXT("Error_FailedToCreateAsset", "Failed to create new asset ({0}) for file ({1})."), FText::FromString(NewAssetName), FText::FromString(FullFilename)));
+				// Package already exists, so try and import over the top of it, if it doesn't already have a source file path
+				TArray<FAssetData> Assets;
+				if (Registry->GetAssetsByPackageName(*PackagePath, Assets) && Assets.Num() == 1)
+				{
+					if (UObject* ExistingAsset = Assets[0].GetAsset())
+					{
+						// We're only eligible for reimport if the existing asset doesn't reference a source file already
+						const bool bEligibleForReimport = !Utils::ExtractSourceFilePaths(ExistingAsset).ContainsByPredicate([&](const FString& In) {
+							return !In.IsEmpty() && In == FullFilename;
+							});
+
+						if (bEligibleForReimport)
+						{
+							ReimportAssetWithNewSource(ExistingAsset, FullFilename, OutPackagesToSave, Context);
+						}
+					}
+				}
 			}
 			else
 			{
-				Context.AddMessage(EMessageSeverity::Info, FText::Format(LOCTEXT("Info_CreatingNewAsset", "Importing new asset {0}."), FText::FromString(PackagePath)));
-
-				// Make sure the destination package is loaded
-				NewPackage->FullyLoad();
-				
-				UObject* NewAsset = nullptr;
-
-				// Find a relevant factory for this file
-				// @todo import: gmp: show dialog in case of multiple matching factories
-				const FString Ext = FPaths::GetExtension(Addition.Filename.Get(), false);
-				auto* Factories = InFactoriesByExtension.Find(Ext);
-				if (Factories && Factories->Num() != 0)
+				UPackage* NewPackage = CreatePackage(*PackagePath);
+				if (!ensure(NewPackage))
+				{
+					Context.AddMessage(EMessageSeverity::Error, FText::Format(LOCTEXT("Error_FailedToCreateAsset", "Failed to create new asset ({0}) for file ({1})."), FText::FromString(NewAssetName), FText::FromString(FullFilename)));
+				}
+				else
 				{
-					//Make sure all the scene factory are put at the end of the array. We give priority to asset factory before scene factory
-					TArray<UFactory*> SortFactories;
-					TArray<UFactory*> SceneFactories;
-					for (UFactory *Factory : *Factories)
+					Context.AddMessage(EMessageSeverity::Info, FText::Format(LOCTEXT("Info_CreatingNewAsset", "Importing new asset {0}."), FText::FromString(PackagePath)));
+
+					// Make sure the destination package is loaded
+					NewPackage->FullyLoad();
+
+					UObject* NewAsset = nullptr;
+
+					// Find a relevant factory for this file
+					// @todo import: gmp: show dialog in case of multiple matching factories
+					const FString Ext = FPaths::GetExtension(Addition.Filename.Get(), false);
+					const TArray<UFactory*>* Factories = InFactoriesByExtension.Find(Ext);
+					if (Factories && Factories->Num() != 0)
 					{
-						if (Factory->IsA(USceneImportFactory::StaticClass()))
+						//Make sure all the scene factory are put at the end of the array. We give priority to asset factory before scene factory
+						TArray<UFactory*> SortFactories;
+						TArray<UFactory*> SceneFactories;
+						for (UFactory* Factory : *Factories)
 						{
-							SceneFactories.Add(Factory);
+							if (Factory->IsA(USceneImportFactory::StaticClass()))
+							{
+								SceneFactories.Add(Factory);
+							}
+							else
+							{
+								SortFactories.Add(Factory);
+							}
 						}
-						else
+						if (SceneFactories.Num() > 0)
 						{
-							SortFactories.Add(Factory);
+							SortFactories.Append(SceneFactories);
 						}
-					}
-					if (SceneFactories.Num() > 0)
-					{
-						SortFactories.Append(SceneFactories);
-					}
-					// Prefer a factory if it explicitly can import. UFactory::FactoryCanImport returns false by default, even if the factory supports the extension, so we can't use it directly.
-					UFactory* const * PreferredFactory = SortFactories.FindByPredicate([&](UFactory* F){ return F->FactoryCanImport(FullFilename); });
-					if (PreferredFactory)
-					{
-						NewAsset = AttemptImport((*PreferredFactory)->GetClass(), NewPackage, *NewAssetName, bCancelled, FullFilename);
-					}
-					// If there was no preferred factory, just try them all until one succeeds
-					else for (UFactory* Factory : SortFactories)
-					{
-						NewAsset = AttemptImport(Factory->GetClass(), NewPackage, *NewAssetName, bCancelled, FullFilename);
-
-						if (bCancelled || NewAsset)
+						// Prefer a factory if it explicitly can import. UFactory::FactoryCanImport returns false by default, even if the factory supports the extension, so we can't use it directly.
+						UFactory* const* PreferredFactory = SortFactories.FindByPredicate([&](UFactory* F) { return F->FactoryCanImport(FullFilename); });
+						if (PreferredFactory)
 						{
-							break;
+							NewAsset = AttemptImport((*PreferredFactory)->GetClass(), NewPackage, *NewAssetName, bCancelled, FullFilename);
+						}
+						// If there was no preferred factory, just try them all until one succeeds
+						else for (UFactory* Factory : SortFactories)
+						{
+							NewAsset = AttemptImport(Factory->GetClass(), NewPackage, *NewAssetName, bCancelled, FullFilename);
+
+							if (bCancelled || NewAsset)
+							{
+								break;
+							}
 						}
 					}
-				}
 
-				// Verify if the package still exists after the import (it may have been cleaned up by the factory if the import was canceled).
-				NewPackage = FindObject<UPackage>(nullptr, *PackagePath);
-				if (NewPackage)
-				{
-					// If we didn't create an asset and the package was not cleaned up, unload and delete the package we just created
-					if (!NewAsset)
+					// Verify if the package still exists after the import (it may have been cleaned up by the factory if the import was canceled).
+					NewPackage = FindObject<UPackage>(nullptr, *PackagePath);
+					if (NewPackage)
 					{
-						TArray<UPackage*> Packages;
-						Packages.Add(NewPackage);
+						// If we didn't create an asset and the package was not cleaned up, unload and delete the package we just created
+						if (!NewAsset)
+						{
+							TArray<UPackage*> Packages;
+							Packages.Add(NewPackage);
 
-						TGuardValue<bool> SuppressSlowTaskMessages(Context.bSuppressSlowTaskMessages, true);
+							TGuardValue<bool> SuppressSlowTaskMessages(Context.bSuppressSlowTaskMessages, true);
 
-						FText ErrorMessage;
-						if (!UPackageTools::UnloadPackages(Packages, ErrorMessage))
-						{
-							Context.AddMessage(EMessageSeverity::Error, FText::Format(LOCTEXT("Error_UnloadingPackage", "There was an error unloading a package: {0}."), ErrorMessage));
+							FText ErrorMessage;
+							if (!UPackageTools::UnloadPackages(Packages, ErrorMessage))
+							{
+								Context.AddMessage(EMessageSeverity::Error, FText::Format(LOCTEXT("Error_UnloadingPackage", "There was an error unloading a package: {0}."), ErrorMessage));
+							}
+
+							// Just add the message to the message log rather than add it to the UI
+							// Factories may opt not to import the file, so we let them report errors if they do
+							Context.GetMessageLog().Message(EMessageSeverity::Info, FText::Format(LOCTEXT("Info_FailedToImportAsset", "Failed to import file {0}."), FText::FromString(FullFilename)));
 						}
+						else if (!bCancelled)
+						{
+							FAssetRegistryModule::AssetCreated(NewAsset);
+							GEditor->BroadcastObjectReimported(NewAsset);
 
-						// Just add the message to the message log rather than add it to the UI
-						// Factories may opt not to import the file, so we let them report errors if they do
-						Context.GetMessageLog().Message(EMessageSeverity::Info, FText::Format(LOCTEXT("Info_FailedToImportAsset", "Failed to import file {0}."), FText::FromString(FullFilename)));
+							OutPackagesToSave.Add(NewPackage);
+						}
 					}
-					else if (!bCancelled)
-					{
-						FAssetRegistryModule::AssetCreated(NewAsset);
-						GEditor->BroadcastObjectReimported(NewAsset);
 
-						OutPackagesToSave.Add(NewPackage);
-					}
+					// Refresh the supported class.  Some factories (e.g. FBX) only resolve their type after reading the file
+					// ImportAssetType = Factory->ResolveSupportedClass();
+					// @todo: analytics?
 				}
-
-				// Refresh the supported class.  Some factories (e.g. FBX) only resolve their type after reading the file
-				// ImportAssetType = Factory->ResolveSupportedClass();
-				// @todo: analytics?
 			}
 		}
 
@@ -420,20 +474,20 @@ void FContentDirectoryMonitor::ProcessAdditions(const DirectoryWatcher::FTimeLim
 
 void FContentDirectoryMonitor::ProcessModifications(const DirectoryWatcher::FTimeLimit& TimeLimit, TArray<TObjectPtr<UPackage>>& OutPackagesToSave, FReimportFeedbackContext& Context)
 {
-	auto* ReimportManager = FReimportManager::Instance();
+	FReimportManager* ReimportManager = FReimportManager::Instance();
 
 	for (int32 Index = 0; Index < ModifiedFiles.Num(); ++Index)
 	{
 		Context.MainTask->EnterProgressFrame();
 
-		auto& Change = ModifiedFiles[Index];
+		DirectoryWatcher::FUpdateCacheTransaction& Change = ModifiedFiles[Index];
 		const FString FullFilename = Cache.GetDirectory() + Change.Filename.Get();
 
 		// Move the asset before reimporting it. We always reimport moved assets to ensure that their import path is up to date
 		if (Change.Action == DirectoryWatcher::EFileAction::Moved)
 		{
 			const FString OldFilename = Cache.GetDirectory() + Change.MovedFromFilename.Get();
-			const auto Assets = Utils::FindAssetsPertainingToFile(*Registry, OldFilename);
+			const TArray<FAssetData> Assets = Utils::FindAssetsPertainingToFile(*Registry, OldFilename);
 
 			if (Assets.Num() == 1)
 			{
@@ -498,13 +552,45 @@ void FContentDirectoryMonitor::ProcessModifications(const DirectoryWatcher::FTim
 		}
 		else
 		{
+			FAssetToolsModule& AssetToolsModule = FModuleManager::Get().LoadModuleChecked<FAssetToolsModule>("AssetTools");
+
+			TArray<FString> NoReimportObjectPaths;
+			TArray<UObject*> ReimportObjects;
+
 			// Modifications or additions are treated the same by this point
-			for (const auto& AssetData : Utils::FindAssetsPertainingToFile(*Registry, FullFilename))
+			for (const FAssetData& AssetData : Utils::FindAssetsPertainingToFile(*Registry, FullFilename))
+			{
+				if (UClass* AssetClass = AssetData.GetClass())
+				{
+					TSharedPtr<IAssetTypeActions> AssetTypeActions = AssetToolsModule.Get().GetAssetTypeActionsForClass(AssetClass).Pin();
+					if (!AssetTypeActions.IsValid() || !AssetTypeActions->IsImportedAsset())
+					{
+						NoReimportObjectPaths.Add(AssetData.GetObjectPathString());
+					}
+					else if (UObject* Asset = AssetData.GetAsset())
+					{
+						ReimportObjects.Add(Asset);
+					}
+				}
+			}
+
+			if (NoReimportObjectPaths.Num() > 0)
 			{
-				if (UObject* Asset = AssetData.GetAsset())
+				if (!GIsAutomationTesting && !FApp::IsUnattended() && !FApp::IsGame())
 				{
-					ReimportAsset(Asset, FullFilename, OutPackagesToSave, Context);
+					FText DialogMessage = LOCTEXT("ReimportInfo_NoReimportAvailable", "Some of the assets would need to be reimported manually.");
+					FMessageDialog::Open(EAppMsgType::Ok, DialogMessage);
 				}
+
+				for (const FString& ObjectPathString : NoReimportObjectPaths)
+				{
+					Context.AddMessage(EMessageSeverity::Warning, FText::Format(LOCTEXT("ReimportWarning_NoReimport", "Reimport for asset {0} is unavailable. Please import them manually."), FText::FromString(ObjectPathString)));
+				}
+			}
+
+			for(UObject* ReimportObject: ReimportObjects)
+			{
+				ReimportAsset(ReimportObject, FullFilename, OutPackagesToSave, Context);
 			}
 		}
 
@@ -552,9 +638,9 @@ void FContentDirectoryMonitor::ReimportAsset(UObject* Asset, const FString& Full
 
 void FContentDirectoryMonitor::ExtractAssetsToDelete(TArray<FAssetData>& OutAssetsToDelete)
 {
-	for (auto& Deletion : DeletedFiles)
+	for (DirectoryWatcher::FUpdateCacheTransaction& Deletion : DeletedFiles)
 	{
-		for (const auto& AssetData : Utils::FindAssetsPertainingToFile(*Registry, Cache.GetDirectory() + Deletion.Filename.Get()))
+		for (const FAssetData& AssetData : Utils::FindAssetsPertainingToFile(*Registry, Cache.GetDirectory() + Deletion.Filename.Get()))
 		{
 			OutAssetsToDelete.Add(AssetData);
 		}
@@ -568,25 +654,26 @@ void FContentDirectoryMonitor::ExtractAssetsToDelete(TArray<FAssetData>& OutAsse
 
 void FContentDirectoryMonitor::Abort()
 {
-	for (auto& Add : AddedFiles)
+	using namespace DirectoryWatcher;
+	for (FUpdateCacheTransaction& Add : AddedFiles)
 	{
 		Cache.CompleteTransaction(MoveTemp(Add));
 	}
 	AddedFiles.Empty();
 
-	for (auto& Mod : ModifiedFiles)
+	for (FUpdateCacheTransaction& Mod : ModifiedFiles)
 	{
 		Cache.CompleteTransaction(MoveTemp(Mod));
 	}
 	ModifiedFiles.Empty();
 
-	for (auto& Del : DeletedFiles)
+	for (FUpdateCacheTransaction& Del : DeletedFiles)
 	{
 		Cache.CompleteTransaction(MoveTemp(Del));
 	}
 	DeletedFiles.Empty();
 
-	for (auto& Change : Cache.GetOutstandingChanges())
+	for (FUpdateCacheTransaction& Change : Cache.GetOutstandingChanges())
 	{
 		Cache.CompleteTransaction(MoveTemp(Change));
 	}
diff --git a/Engine/Source/Runtime/Interchange/Engine/Private/Tasks/InterchangeTaskCompletion.cpp b/Engine/Source/Runtime/Interchange/Engine/Private/Tasks/InterchangeTaskCompletion.cpp
--- a/Engine/Source/Runtime/Interchange/Engine/Private/Tasks/InterchangeTaskCompletion.cpp
+++ b/Engine/Source/Runtime/Interchange/Engine/Private/Tasks/InterchangeTaskCompletion.cpp
@@ -20,6 +20,22 @@
 #include "UObject/ObjectMacros.h"
 #include "UObject/WeakObjectPtrTemplates.h"
 
+
+namespace UE::Interchange::Private::ObjectDeletionUtils
+{
+	static void PurgeObject(UObject* Object)
+	{
+		if (Object)
+		{
+			Object->ClearFlags(RF_Standalone | RF_Public | RF_Transactional);
+			Object->ClearInternalFlags(EInternalObjectFlags::Async);
+			Object->SetFlags(RF_Transient);
+			Object->MarkAsGarbage();
+			Object->UObject::Rename(nullptr, GetTransientPackage(), REN_NonTransactional | REN_DontCreateRedirectors);
+		}
+	}
+}
+
 void UE::Interchange::FTaskPreCompletion_GameThread::Execute()
 {
 	TRACE_CPUPROFILER_EVENT_SCOPE(UE::Interchange::FTaskPreCompletion_GameThread::DoTask)
@@ -227,15 +243,10 @@ void UE::Interchange::FTaskCompletion_GameThread::Execute()
 			{
 				for (const FImportAsyncHelper::FImportedObjectInfo& AssetInfo : AssetInfos)
 				{
-					UObject* Asset = AssetInfo.ImportedObject;
-					if (Asset)
-					{
-						//Make any created asset go away
-						Asset->ClearFlags(RF_Standalone | RF_Public | RF_Transactional);
-						Asset->ClearInternalFlags(EInternalObjectFlags::Async);
-						Asset->SetFlags(RF_Transient);
-						Asset->MarkAsGarbage();
-					}
+					using namespace UE::Interchange::Private::ObjectDeletionUtils;
+
+					//Make any created asset go away
+					PurgeObject(AssetInfo.ImportedObject);
 				}
 			});
 
