Reset property bag associations after replacing references during a forced asset deletion.

diff --git a/Engine/Source/Editor/UnrealEd/Private/ObjectTools.cpp b/Engine/Source/Editor/UnrealEd/Private/ObjectTools.cpp
--- a/Engine/Source/Editor/UnrealEd/Private/ObjectTools.cpp
+++ b/Engine/Source/Editor/UnrealEd/Private/ObjectTools.cpp
@@ -100,6 +100,7 @@
 #include "HAL/PlatformApplicationMisc.h"
 #include "Kismet2/BlueprintEditorUtils.h"
 #include "Subsystems/AssetEditorSubsystem.h"
+#include "UObject/PropertyBagRepository.h"
 #include "UObject/ReferencerFinder.h"
 #include "Containers/Set.h"
 #include "UObject/StrongObjectPtr.h"
@@ -1370,6 +1371,9 @@ namespace ObjectTools
 			}
 		}
 
+		// Reset property bag associations after replacing references.
+		UE::FPropertyBagRepository::Get().ReassociateObjects(ReplacementMap);
+
 		// Now alter the referencing objects the change has completed via PostEditChange,
 		// this is done in a separate loop to prevent reading of data that we want to overwrite
 		int32 NumObjsPostEdited = 0;
