Add new way for plugins to create ChunkDataGenerators.
New chunk data generators can now be created by any system not just systems within the cooker.
Adding CookInfo for which cultures are being cooked.

diff --git a/Engine/Source/Editor/UnrealEd/Classes/CookOnTheSide/CookOnTheFlyServer.h b/Engine/Source/Editor/UnrealEd/Classes/CookOnTheSide/CookOnTheFlyServer.h
--- a/Engine/Source/Editor/UnrealEd/Classes/CookOnTheSide/CookOnTheFlyServer.h
+++ b/Engine/Source/Editor/UnrealEd/Classes/CookOnTheSide/CookOnTheFlyServer.h
@@ -642,6 +642,7 @@ public:
 	UNREALED_API virtual void RegisterCollector(UE::Cook::IMPCollector* Collector,
 		UE::Cook::EProcessType ProcessType = UE::Cook::EProcessType::AllMPCook) override;
 	UNREALED_API virtual void UnregisterCollector(UE::Cook::IMPCollector* Collector) override;
+	UNREALED_API virtual void GetCulturesToCook(TArray<FString>& OutCulturesToCook) const override;
 
 
 	/** Dumps cooking stats to the log. Run from the exec command "Cook stats". */
diff --git a/Engine/Source/Editor/UnrealEd/Private/Commandlets/IChunkDataGenerator.cpp b/Engine/Source/Editor/UnrealEd/Private/Commandlets/IChunkDataGenerator.cpp
new file mode 100644
--- /dev/null
+++ b/Engine/Source/Editor/UnrealEd/Private/Commandlets/IChunkDataGenerator.cpp
@@ -0,0 +1,17 @@
+// Copyright Epic Games, Inc. All Rights Reserved.
+
+#pragma once
+
+#include "Commandlets/IChunkDataGenerator.h"
+
+TArray<TFunction<TSharedRef<IChunkDataGenerator>(const UE::Cook::ICookInfo&)>> IChunkDataGenerator::GeneratorsDataGeneratorFactories;
+
+void IChunkDataGenerator::AddChunkDataGeneratorFactory(const TFunction<TSharedRef<IChunkDataGenerator>(const UE::Cook::ICookInfo&)>& InChunkGeneratorFactory)
+{
+	GeneratorsDataGeneratorFactories.Add(InChunkGeneratorFactory);
+}
+
+const TArray<TFunction<TSharedRef<IChunkDataGenerator>(const UE::Cook::ICookInfo&)>>& IChunkDataGenerator::GetChunkDataGeneratorFactories()
+{
+	return GeneratorsDataGeneratorFactories;
+}
\ No newline at end of file
diff --git a/Engine/Source/Editor/UnrealEd/Private/CookOnTheFlyServer.cpp b/Engine/Source/Editor/UnrealEd/Private/CookOnTheFlyServer.cpp
--- a/Engine/Source/Editor/UnrealEd/Private/CookOnTheFlyServer.cpp
+++ b/Engine/Source/Editor/UnrealEd/Private/CookOnTheFlyServer.cpp
@@ -7667,6 +7667,14 @@ void UCookOnTheFlyServer::UnregisterCollector(UE::Cook::IMPCollector* Collector)
 }
 
 
+void UCookOnTheFlyServer::GetCulturesToCook(TArray<FString>& OutCulturesToCook) const
+{
+	if (CookByTheBookOptions)
+	{
+		OutCulturesToCook.Append(CookByTheBookOptions->AllCulturesToCook);
+	}
+}
+
 void UCookOnTheFlyServer::DumpStats()
 {
 	UE_LOG(LogCook, Display, TEXT("IntStats:"));
@@ -10394,6 +10402,14 @@ void UCookOnTheFlyServer::CookByTheBookFinishedInternal()
 				RegisterShaderChunkDataGenerator();
 			}
 
+			for (auto ChunkGeneratorFactory : IChunkDataGenerator::GetChunkDataGeneratorFactories())
+			{
+				for (const ITargetPlatform* TargetPlatform : PlatformManager->GetSessionPlatforms())
+				{
+					FAssetRegistryGenerator& RegistryGenerator = *(PlatformManager->GetPlatformData(TargetPlatform)->RegistryGenerator);
+					RegistryGenerator.RegisterChunkDataGenerator(ChunkGeneratorFactory(*this));
+				}
+			}
 			// if we are cooking DLC, the DevelopmentAR isn't needed - it's used when making DLC against shipping, so there's no need to make it
 			// again, as we don't make DLC against DLC (but allow an override just in case)
 			bool bSaveDevelopmentAssetRegistry = !FParse::Param(FCommandLine::Get(), TEXT("NoSaveDevAR"));
diff --git a/Engine/Source/Editor/UnrealEd/Public/Commandlets/IChunkDataGenerator.h b/Engine/Source/Editor/UnrealEd/Public/Commandlets/IChunkDataGenerator.h
--- a/Engine/Source/Editor/UnrealEd/Public/Commandlets/IChunkDataGenerator.h
+++ b/Engine/Source/Editor/UnrealEd/Public/Commandlets/IChunkDataGenerator.h
@@ -3,6 +3,7 @@
 #pragma once
 
 #include "CoreMinimal.h"
+#include "UObject/ICookInfo.h"
 
 class FSandboxPlatformFile;
 class ITargetPlatform;
@@ -30,4 +31,9 @@ public:
 		const ITargetPlatform* TargetPlatform, FSandboxPlatformFile* InSandboxFile, TArray<FString>& OutChunkFilenames)
 	{
 	}
+
+	UNREALED_API static void AddChunkDataGeneratorFactory(const TFunction<TSharedRef<IChunkDataGenerator>(const UE::Cook::ICookInfo&)>& InChunkGeneratorFactory);
+	UNREALED_API static const TArray<TFunction<TSharedRef<IChunkDataGenerator>(const UE::Cook::ICookInfo&)>>& GetChunkDataGeneratorFactories();
+private:
+	static TArray<TFunction<TSharedRef<IChunkDataGenerator>(const UE::Cook::ICookInfo&)>> GeneratorsDataGeneratorFactories;
 };
diff --git a/Engine/Source/Runtime/CoreUObject/Public/UObject/ICookInfo.h b/Engine/Source/Runtime/CoreUObject/Public/UObject/ICookInfo.h
--- a/Engine/Source/Runtime/CoreUObject/Public/UObject/ICookInfo.h
+++ b/Engine/Source/Runtime/CoreUObject/Public/UObject/ICookInfo.h
@@ -133,6 +133,12 @@ public:
 	 * delete the Collector if he caller does not have their own TRefCountPtr to it.
 	 */
 	virtual void UnregisterCollector(IMPCollector* Collector) = 0;
+
+
+	/**
+	 * Gets the current cultures that are being cooked.
+	 */
+	virtual void GetCulturesToCook(TArray<FString>& OutCulturesToCook) const = 0;
 };
 
 DECLARE_MULTICAST_DELEGATE_OneParam(FCookInfoEvent, ICookInfo&);
