CIS fix

diff --git a/Engine/Source/Editor/UnrealEd/Private/Commandlets/IChunkDataGenerator.cpp b/Engine/Source/Editor/UnrealEd/Private/Commandlets/IChunkDataGenerator.cpp
--- a/Engine/Source/Editor/UnrealEd/Private/Commandlets/IChunkDataGenerator.cpp
+++ b/Engine/Source/Editor/UnrealEd/Private/Commandlets/IChunkDataGenerator.cpp
@@ -1,7 +1,5 @@
 // Copyright Epic Games, Inc. All Rights Reserved.
 
-#pragma once
-
 #include "Commandlets/IChunkDataGenerator.h"
 
 TArray<TFunction<TSharedRef<IChunkDataGenerator>(const UE::Cook::ICookInfo&)>> IChunkDataGenerator::GeneratorsDataGeneratorFactories;
