Block on asset registry for ODSC server: the first request on the client happens at boot time anyway, and would block until the asset registry is ready
Skip most calls to AddCookOnTheFlyPlatformFromGameThread. This is done to allow launching a single ODSC server instance for multiple platforms without impacting the server start time too much

#rnx

diff --git a/Engine/Source/Editor/UnrealEd/Private/CookOnTheFlyServer.cpp b/Engine/Source/Editor/UnrealEd/Private/CookOnTheFlyServer.cpp
--- a/Engine/Source/Editor/UnrealEd/Private/CookOnTheFlyServer.cpp
+++ b/Engine/Source/Editor/UnrealEd/Private/CookOnTheFlyServer.cpp
@@ -723,6 +723,11 @@ bool UCookOnTheFlyServer::StartCookOnTheFly(FCookOnTheFlyStartupOptions InCookOn
 		CookOnTheFlyRequestManager = MakeNetworkFileCookOnTheFlyRequestManager(*CookOnTheFlyServerInterface, CookOnTheFlyNetworkServer.ToSharedRef());
 	}
 
+	if (bRunningAsShaderServer)
+	{
+		BlockOnAssetRegistry(TConstArrayView<FString>());
+	}
+
 	if (CookOnTheFlyNetworkServer->Start())
 	{
 		TArray<TSharedPtr<FInternetAddr>> ListenAddresses;
@@ -759,7 +764,7 @@ void UCookOnTheFlyServer::AddCookOnTheFlyPlatformFromGameThread(ITargetPlatform*
 {
 	UE::Cook::FPlatformData* PlatformData = PlatformManager->GetPlatformData(TargetPlatform);
 	check(PlatformData != nullptr); // should have been checked by the caller
-	if (PlatformData->bIsSandboxInitialized)
+	if (PlatformData->bIsSandboxInitialized || bRunningAsShaderServer)
 	{
 		return;
 	}
