World Partition - When running a builder on a non-WP map, log a warning rather than stop all processing.
* This is especially important when running on a set of maps using a collection, as a single invalid entry in the collection would prevent all further processing and would skip submitting any prior changes to source control

diff --git a/Engine/Source/Editor/UnrealEd/Private/WorldPartition/WorldPartitionBuilder.cpp b/Engine/Source/Editor/UnrealEd/Private/WorldPartition/WorldPartitionBuilder.cpp
--- a/Engine/Source/Editor/UnrealEd/Private/WorldPartition/WorldPartitionBuilder.cpp
+++ b/Engine/Source/Editor/UnrealEd/Private/WorldPartition/WorldPartitionBuilder.cpp
@@ -138,8 +138,7 @@ bool UWorldPartitionBuilder::RunBuilder(UWorld* World)
 		}
 		else
 		{
-			UE_LOG(LogWorldPartitionBuilder, Error, TEXT("WorldPartition builders only works on partitioned maps."));
-			bResult = false;
+			UE_LOG(LogWorldPartitionBuilder, Warning, TEXT("WorldPartition builder '%s' only works on partitioned worlds - skipping processing of '%s'."), *GetClass()->GetName(), *World->GetPathName());
 		}
 	}
 
