Fix reference outliving the lambda.

diff --git a/Engine/Source/Editor/UnrealEd/Private/GeomFitUtils.cpp b/Engine/Source/Editor/UnrealEd/Private/GeomFitUtils.cpp
--- a/Engine/Source/Editor/UnrealEd/Private/GeomFitUtils.cpp
+++ b/Engine/Source/Editor/UnrealEd/Private/GeomFitUtils.cpp
@@ -75,17 +75,16 @@ int32 GenerateKDopAsSimpleCollision(UStaticMesh* StaticMesh, const TArray<FVecto
 	TFunction<UE::Math::TVector<double>(int32)> GetPointFunc;
 	if(FStaticMeshRenderData* RenderData = StaticMesh->GetRenderData())
 	{
-		const FStaticMeshLODResources& LODResource = RenderData->LODResources[0];
-		NumVertices = LODResource.GetNumVertices();
-		GetPointFunc = [&](int32 VertIdx)
+		NumVertices = RenderData->LODResources[0].GetNumVertices();
+		GetPointFunc = [RenderData](int32 VertIdx)
 		{
-			return FVector{ LODResource.VertexBuffers.PositionVertexBuffer.VertexPosition(VertIdx) };
+			return FVector{ RenderData->LODResources[0].VertexBuffers.PositionVertexBuffer.VertexPosition(VertIdx) };
 		};
 	}
 	else if(const FMeshDescription* MeshDescription = StaticMesh->GetMeshDescription(0))
 	{		
 		NumVertices = MeshDescription->Vertices().Num();
-		GetPointFunc = [&](int32 VertIdx)
+		GetPointFunc = [MeshDescription](int32 VertIdx)
 		{
 			return FVector(MeshDescription->GetVertexPosition(VertIdx));
 		};
