[BUGFIX] Export Options Dialog showing up on automated export

diff --git a/Engine/Source/Editor/UnrealEd/Private/EditorExporters.cpp b/Engine/Source/Editor/UnrealEd/Private/EditorExporters.cpp
--- a/Engine/Source/Editor/UnrealEd/Private/EditorExporters.cpp
+++ b/Engine/Source/Editor/UnrealEd/Private/EditorExporters.cpp
@@ -2138,6 +2138,25 @@ UFbxExportOption* UExporterFBX::GetAutomatedExportOptionsFbx()
 	return nullptr;
 }
 
+namespace UE::EditorExporters::Private
+{
+	bool SetupExporterForAutomatedFBXExport(FScopedFbxExporterInstance& ScopedExporterInstance, UExporterFBX& CurrentExporter)
+	{
+		if (CurrentExporter.ExportTask && CurrentExporter.ExportTask->bAutomated)
+		{
+			if (UFbxExportOption* AutomatedExportOptions = CurrentExporter.GetAutomatedExportOptionsFbx())
+			{
+				ScopedExporterInstance.GetExporter()->SetExportOptionsOverride(AutomatedExportOptions);
+			}
+
+			CurrentExporter.SetShowExportOption(false);
+			return true;
+		}
+
+		return false;
+	}
+}
+
 /*------------------------------------------------------------------------------
 UStaticMeshExporterFBX implementation.
 ------------------------------------------------------------------------------*/
@@ -2158,12 +2177,7 @@ bool UStaticMeshExporterFBX::ExportBinary( UObject* Object, const TCHAR* Type, F
 {
 	UStaticMesh* StaticMesh = CastChecked<UStaticMesh>( Object );
 	FScopedFbxExporterInstance ScopedExporterInstance;
-	if (UFbxExportOption* AutomatedExportOptions = GetAutomatedExportOptionsFbx())
-	{
-		ScopedExporterInstance.GetExporter()->SetExportOptionsOverride(AutomatedExportOptions);
-		SetShowExportOption(false);
-	}
-	else
+	if (!UE::EditorExporters::Private::SetupExporterForAutomatedFBXExport(ScopedExporterInstance, *this))
 	{
 		//Show the fbx export dialog options
 		bool ExportAll = GetBatchMode() && !GetShowExportOption();
@@ -2205,12 +2219,7 @@ bool USkeletalMeshExporterFBX::ExportBinary( UObject* Object, const TCHAR* Type,
 {
 	USkeletalMesh* SkeletalMesh = CastChecked<USkeletalMesh>( Object );
 	FScopedFbxExporterInstance ScopedExporterInstance;
-	if (UFbxExportOption* AutomatedExportOptions = GetAutomatedExportOptionsFbx())
-	{
-		ScopedExporterInstance.GetExporter()->SetExportOptionsOverride(AutomatedExportOptions);
-		SetShowExportOption(false);
-	}
-	else
+	if (!UE::EditorExporters::Private::SetupExporterForAutomatedFBXExport(ScopedExporterInstance, *this))
 	{
 		//Show the fbx export dialog options
 		bool ExportAll = GetBatchMode() && !GetShowExportOption();
@@ -2256,12 +2265,7 @@ bool UAnimSequenceExporterFBX::ExportBinary( UObject* Object, const TCHAR* Type,
 	if (AnimSkeleton && PreviewMesh)
 	{
 		FScopedFbxExporterInstance ScopedExporterInstance;
-		if (UFbxExportOption* AutomatedExportOptions = GetAutomatedExportOptionsFbx())
-		{
-			ScopedExporterInstance.GetExporter()->SetExportOptionsOverride(AutomatedExportOptions);
-			SetShowExportOption(false);
-		}
-		else
+		if (!UE::EditorExporters::Private::SetupExporterForAutomatedFBXExport(ScopedExporterInstance, *this))
 		{
 			//Show the fbx export dialog options
 			bool ExportAll = GetBatchMode() && !GetShowExportOption();
