Clear redirectors in the local asset that are pointing at a deleted macro graph

[RN] Fixed cause of Blueprint Macro Library Macro Graphs remaining in the Blueprint asset after said graph was deleted by the user

diff --git a/Engine/Source/Editor/UnrealEd/Private/Kismet2/BlueprintEditorUtils.cpp b/Engine/Source/Editor/UnrealEd/Private/Kismet2/BlueprintEditorUtils.cpp
--- a/Engine/Source/Editor/UnrealEd/Private/Kismet2/BlueprintEditorUtils.cpp
+++ b/Engine/Source/Editor/UnrealEd/Private/Kismet2/BlueprintEditorUtils.cpp
@@ -2518,6 +2518,22 @@ void FBlueprintEditorUtils::RemoveGraph(UBlueprint* Blueprint, class UEdGraph* G
 
 				// Clear the cache since it's indexed by graph and one of the graphs is going away
 				FBlueprintEditorUtils::ClearMacroCosmeticInfoCache(Blueprint);
+
+				// Clear redirectors to the graph - we may have created these for macro graphs.
+				// See comment related to macro graphs in FBlueprintEditorUtils::RenameGraph
+				// involving conditional assignment of REN_DontCreateRedirectors
+				TArray<UObject*> Inners;
+				GetObjectsWithOuter(Blueprint, Inners, false);
+				for(UObject* Object : Inners)
+				{
+					if(UObjectRedirector* Redirector = Cast<UObjectRedirector>(Object))
+					{
+						if(Redirector->DestinationObject == GraphToRemove)
+						{
+							Redirector->DestinationObject = nullptr;
+						}
+					}
+				}
 			}
 
 			for (FBPInterfaceDescription& CurrInterface : Blueprint->ImplementedInterfaces)
