Relocate the TmpPackaging directory so that it is not under Project\Saved\TmpPackaging (regardless of presence of DLC), but is instead in the cook Metadata folder.  This ensures that:
1. DLC cooks don't overwrite the TmpPackaging contents of a base game
2. TmpPackaging contents are picked up by zen filesystem and snapshot logic as other metadata

diff --git a/Engine/Source/Editor/UnrealEd/Private/Commandlets/AssetRegistryGenerator.cpp b/Engine/Source/Editor/UnrealEd/Private/Commandlets/AssetRegistryGenerator.cpp
--- a/Engine/Source/Editor/UnrealEd/Private/Commandlets/AssetRegistryGenerator.cpp
+++ b/Engine/Source/Editor/UnrealEd/Private/Commandlets/AssetRegistryGenerator.cpp
@@ -201,30 +201,6 @@ FAssetRegistryGenerator::~FAssetRegistryGenerator()
 {
 }
 
-bool FAssetRegistryGenerator::CleanTempPackagingDirectory(const FString& Platform) const
-{
-	FString TmpPackagingDir = GetTempPackagingDirectoryForPlatform(Platform);
-	if (IFileManager::Get().DirectoryExists(*TmpPackagingDir))
-	{
-		if (!IFileManager::Get().DeleteDirectory(*TmpPackagingDir, false, true))
-		{
-			UE_LOG(LogAssetRegistryGenerator, Error, TEXT("Failed to delete directory: %s"), *TmpPackagingDir);
-			return false;
-		}
-	}
-
-	FString ChunkListDir = FPaths::Combine(*FPaths::ProjectLogDir(), TEXT("ChunkLists"));
-	if (IFileManager::Get().DirectoryExists(*ChunkListDir))
-	{
-		if (!IFileManager::Get().DeleteDirectory(*ChunkListDir, false, true))
-		{
-			UE_LOG(LogAssetRegistryGenerator, Error, TEXT("Failed to delete directory: %s"), *ChunkListDir);
-			return false;
-		}
-	}
-	return true;
-}
-
 bool FAssetRegistryGenerator::ShouldPlatformGenerateStreamingInstallManifest(const ITargetPlatform* Platform) const
 {
 	if (Platform)
@@ -341,21 +317,21 @@ bool FAssetRegistryGenerator::GenerateStreamingInstallManifest(int64 InOverrideC
 	UE::Cook::FCookSandbox& InSandboxFile)
 {
 	const FString Platform = TargetPlatform->PlatformName();
-	FString TmpPackagingDir = GetTempPackagingDirectoryForPlatform(Platform);
+	FString ChunkManifestDir = GetChunkManifestDirectoryForPlatform(Platform, InSandboxFile);
 	if (InManifestSubDir)
 	{
-		TmpPackagingDir /= InManifestSubDir;
+		ChunkManifestDir /= InManifestSubDir;
 	}
 	int64 MaxChunkSize = InOverrideChunkSize > 0 ? InOverrideChunkSize : GetMaxChunkSizePerPlatform(TargetPlatform);
 
-	if (!IFileManager::Get().MakeDirectory(*TmpPackagingDir, true /* Tree */))
+	if (!IFileManager::Get().MakeDirectory(*ChunkManifestDir, true /* Tree */))
 	{
-		UE_LOG(LogAssetRegistryGenerator, Error, TEXT("Failed to create directory: %s"), *TmpPackagingDir);
+		UE_LOG(LogAssetRegistryGenerator, Error, TEXT("Failed to create directory: %s"), *ChunkManifestDir);
 		return false;
 	}
 	
-	FString PakChunkListFilename = TmpPackagingDir / TEXT("pakchunklist.txt");
-	FString PakChunkLayerInfoFilename = TmpPackagingDir / TEXT("pakchunklayers.txt");
+	FString PakChunkListFilename = ChunkManifestDir / TEXT("pakchunklist.txt");
+	FString PakChunkLayerInfoFilename = ChunkManifestDir / TEXT("pakchunklayers.txt");
 	// List of pak file lists
 	TUniquePtr<FArchive> PakChunkListFile(IFileManager::Get().CreateFileWriter(*PakChunkListFilename));
 	// List of disc layer for each chunk
@@ -614,7 +590,7 @@ bool FAssetRegistryGenerator::GenerateStreamingInstallManifest(int64 InOverrideC
 				? FString::Printf(TEXT("pakchunk%d_s%d.txt"), PakchunkIndex, SubChunkIndex)
 				: FString::Printf(TEXT("pakchunk%d.txt"), PakchunkIndex);
 
-			const FString PakListFilename = FString::Printf(TEXT("%s/%s"), *TmpPackagingDir, *PakChunkFilename);
+			const FString PakListFilename = FString::Printf(TEXT("%s/%s"), *ChunkManifestDir, *PakChunkFilename);
 			TUniquePtr<FArchive> PakListFile(IFileManager::Get().CreateFileWriter(*PakListFilename));
 
 			if (!PakListFile)
@@ -816,20 +792,6 @@ bool FAssetRegistryGenerator::GenerateStreamingInstallManifest(int64 InOverrideC
 
 	ChunkLayerFile->Close();
 	PakChunkListFile->Close();
-	
-	if (bSucceeded)
-	{
-		FString ChunkManifestDirectory = FPaths::ProjectDir() / TEXT("Metadata") / TEXT("ChunkManifest");
-		ChunkManifestDirectory =
-			InSandboxFile.ConvertToAbsolutePathForExternalAppForWrite(*ChunkManifestDirectory, Platform);
-
-		IPlatformFile& PlatformFile = FPlatformFileManager::Get().GetPlatformFile();
-		if (!PlatformFile.CopyDirectoryTree(*ChunkManifestDirectory, *TmpPackagingDir, true))
-		{
-			UE_LOG(LogAssetRegistryGenerator, Error, TEXT("Failed to copy chunk manifest from '%s' to '%s'"), *TmpPackagingDir, *ChunkManifestDirectory)
-			return false;
-		}
-	}
 
 	return bSucceeded;
 }
@@ -877,10 +839,29 @@ void FAssetRegistryGenerator::CalculateChunkIdsAndAssignToManifest(const FName&
 	}
 }
 
-void FAssetRegistryGenerator::CleanManifestDirectories()
+bool FAssetRegistryGenerator::CleanManifestDirectories(UE::Cook::FCookSandbox& InSandboxFile)
 {
 	LLM_SCOPE_BYTAG(Cooker_GeneratedAssetRegistry);
-	CleanTempPackagingDirectory(TargetPlatform->PlatformName());
+	FString ChunkManifestDir = GetChunkManifestDirectoryForPlatform(TargetPlatform->PlatformName(), InSandboxFile);
+	if (IFileManager::Get().DirectoryExists(*ChunkManifestDir))
+	{
+		if (!IFileManager::Get().DeleteDirectory(*ChunkManifestDir, false, true))
+		{
+			UE_LOG(LogAssetRegistryGenerator, Error, TEXT("Failed to delete directory: %s"), *ChunkManifestDir);
+			return false;
+		}
+	}
+
+	FString ChunkListDir = FPaths::Combine(*FPaths::ProjectLogDir(), TEXT("ChunkLists"));
+	if (IFileManager::Get().DirectoryExists(*ChunkListDir))
+	{
+		if (!IFileManager::Get().DeleteDirectory(*ChunkListDir, false, true))
+		{
+			UE_LOG(LogAssetRegistryGenerator, Error, TEXT("Failed to delete directory: %s"), *ChunkListDir);
+			return false;
+		}
+	}
+	return true;
 }
 
 void FAssetRegistryGenerator::SetPreviousAssetRegistry(TUniquePtr<FAssetRegistryState>&& InPreviousState)
@@ -2369,9 +2350,9 @@ void FAssetRegistryGenerator::AddPackageToChunk(FChunkPackageSet& ThisPackageSet
 	ThisPackageSet.Add(InPkgName, InSandboxFile);
 }
 
-FString FAssetRegistryGenerator::GetTempPackagingDirectoryForPlatform(const FString& Platform) const
+FString FAssetRegistryGenerator::GetChunkManifestDirectoryForPlatform(const FString& Platform, UE::Cook::FCookSandbox& InSandboxFile) const
 {
-	return FPaths::ProjectSavedDir() / TEXT("TmpPackaging") / Platform;
+	return InSandboxFile.GetSandboxDirectory(Platform) / FApp::GetProjectName() / TEXT("Metadata") / TEXT("ChunkManifest");
 }
 
 void FAssetRegistryGenerator::FixupPackageDependenciesForChunks(UE::Cook::FCookSandbox& InSandboxFile)
diff --git a/Engine/Source/Editor/UnrealEd/Private/CookOnTheFlyServer.cpp b/Engine/Source/Editor/UnrealEd/Private/CookOnTheFlyServer.cpp
--- a/Engine/Source/Editor/UnrealEd/Private/CookOnTheFlyServer.cpp
+++ b/Engine/Source/Editor/UnrealEd/Private/CookOnTheFlyServer.cpp
@@ -11321,9 +11321,10 @@ void UCookOnTheFlyServer::BeginCookSandbox(FBeginCookContext& BeginContext)
 			if (!PlatformContext.bWorkerOnSharedSandbox)
 			{
 				check(!IsCookWorkerMode());
+				checkf(SandboxFile, TEXT("Cannot begin cooking to a sandbox until after CreateSandboxFile has been called from a StartCook function."));
 				// Clean the Manifest directory even on iterative builds; it is written from scratch each time
 				// But only do this if we own the output directory
-				PlatformData->RegistryGenerator->CleanManifestDirectories();
+				PlatformData->RegistryGenerator->CleanManifestDirectories(*SandboxFile);
 			}
 
 			if (PlatformContext.bPopulateMemoryResultsFromDiskResults)
diff --git a/Engine/Source/Editor/UnrealEd/Public/Commandlets/AssetRegistryGenerator.h b/Engine/Source/Editor/UnrealEd/Public/Commandlets/AssetRegistryGenerator.h
--- a/Engine/Source/Editor/UnrealEd/Public/Commandlets/AssetRegistryGenerator.h
+++ b/Engine/Source/Editor/UnrealEd/Public/Commandlets/AssetRegistryGenerator.h
@@ -194,7 +194,7 @@ public:
 	/**
 	 * Deletes temporary manifest directories.
 	 */
-	void CleanManifestDirectories();
+	bool CleanManifestDirectories(UE::Cook::FCookSandbox& InSandboxFile);
 
 	/**
 	 * Saves all generated manifests for each target platform.
@@ -388,9 +388,9 @@ private:
 		const FString& InSandboxFile, int32 PakchunkIndex, UE::Cook::FCookSandbox& SandboxPlatformFile);
 
 	/**
-	 * Returns the path of the temporary packaging directory for the specified platform.
+	 * Returns the path of the chunk manifest directory for the specified platform.
 	 */
-	FString GetTempPackagingDirectoryForPlatform(const FString& Platform) const;
+	FString GetChunkManifestDirectoryForPlatform(const FString& Platform, UE::Cook::FCookSandbox& InSandboxFile) const;
 
 	/** Returns the config-driven max size of a chunk for the given platform, or -1 for no limit. */
 	int64 GetMaxChunkSizePerPlatform( const ITargetPlatform* Platform ) const;
@@ -408,9 +408,6 @@ private:
 	void CalculateChunkIdsAndAssignToManifest(const FName& PackageFName, const FString& PackagePathName,
 		const FString& SandboxFilename, const FString& LastLoadedMapName, UE::Cook::FCookSandbox& InSandboxFile);
 
-	/** Deletes the temporary packaging directory for the specified platform */
-	bool CleanTempPackagingDirectory(const FString& Platform) const;
-
 	/** Returns true if the specific platform desires multiple chunks suitable for streaming install */
 	bool ShouldPlatformGenerateStreamingInstallManifest(const ITargetPlatform* Platform) const;
 
diff --git a/Engine/Source/Programs/AutomationTool/BuildGraph/Tasks/CookTask.cs b/Engine/Source/Programs/AutomationTool/BuildGraph/Tasks/CookTask.cs
--- a/Engine/Source/Programs/AutomationTool/BuildGraph/Tasks/CookTask.cs
+++ b/Engine/Source/Programs/AutomationTool/BuildGraph/Tasks/CookTask.cs
@@ -131,13 +131,6 @@ namespace AutomationTool.Tasks
 						throw new AutomationException("Cooking did not produce any files in {0}", platformCookedDirectory.FullName);
 					}
 					cookedFiles.AddRange(platformCookedFiles);
-
-					DirectoryReference packagingFilesDirectory = DirectoryReference.Combine(projectFile.Directory, "Saved", "TmpPackaging", platform);
-					if (DirectoryReference.Exists(packagingFilesDirectory))
-					{
-						List<FileReference> packagingFiles = DirectoryReference.EnumerateFiles(packagingFilesDirectory, "*", System.IO.SearchOption.AllDirectories).ToList();
-						cookedFiles.AddRange(packagingFiles);
-					}
 				}
 			}
 
diff --git a/Engine/Source/Programs/AutomationTool/IOS/IOSPlatform.Automation.cs b/Engine/Source/Programs/AutomationTool/IOS/IOSPlatform.Automation.cs
--- a/Engine/Source/Programs/AutomationTool/IOS/IOSPlatform.Automation.cs
+++ b/Engine/Source/Programs/AutomationTool/IOS/IOSPlatform.Automation.cs
@@ -2101,12 +2101,12 @@ public class IOSPlatform : ApplePlatform
 
 	private static string GetChunkPakManifestListFilename(ProjectParams Params, DeploymentContext SC)
 	{
-		return CombinePaths(GetTmpPackagingPath(Params, SC), "pakchunklist.txt");
+		return CombinePaths(GetChunkManifestPath(Params, SC), "pakchunklist.txt");
 	}
 
-	private static string GetTmpPackagingPath(ProjectParams Params, DeploymentContext SC)
+	private static string GetChunkManifestPath(ProjectParams Params, DeploymentContext SC)
 	{
-		return CombinePaths(Path.GetDirectoryName(Params.RawProjectPath.FullName), "Saved", "TmpPackaging", SC.StageTargetPlatform.GetCookPlatform(SC.DedicatedServer, false));
+		return CombinePaths(SC.MetadataDir.FullName, "ChunkManifest");
 	}
 
 	private static StringBuilder AppendKeyValue(StringBuilder Text, string Key, object Value, int Level)
diff --git a/Engine/Source/Programs/AutomationTool/Scripts/CopyBuildToStagingDirectory.Automation.cs b/Engine/Source/Programs/AutomationTool/Scripts/CopyBuildToStagingDirectory.Automation.cs
--- a/Engine/Source/Programs/AutomationTool/Scripts/CopyBuildToStagingDirectory.Automation.cs
+++ b/Engine/Source/Programs/AutomationTool/Scripts/CopyBuildToStagingDirectory.Automation.cs
@@ -4482,7 +4482,7 @@ namespace AutomationScripts
 			List<ChunkDefinition> ChunkDefinitions = new List<ChunkDefinition>();
 			List<PakFileRules> PakRulesList = GetPakFileRules(Params, SC);
 
-			var TmpPackagingPath = GetTmpPackagingPath(Params, SC);
+			var ChunkManifestPath = GetChunkManifestPath(Params, SC);
 
 			// Parse and cache crypto settings from INI file
 			EncryptionAndSigning.CryptoSettings PakCryptoSettings = EncryptionAndSigning.ParseCryptoSettings(DirectoryReference.FromFile(Params.RawProjectPath), SC.StageTargetPlatform.IniPlatformType, Log.Logger);
@@ -4506,7 +4506,7 @@ namespace AutomationScripts
 					string[] ChunkOptions = ChunkList[Index].Split(' ');
 
 					// Set chunk name to string like "pakchunk0"
-					var ChunkManifestFilename = CombinePaths(TmpPackagingPath, ChunkOptions[0]);
+					var ChunkManifestFilename = CombinePaths(ChunkManifestPath, ChunkOptions[0]);
 					ChunkDefinition CD = new ChunkDefinition(Path.GetFileNameWithoutExtension(ChunkOptions[0]));
 					for (int IOption = 1; IOption < ChunkOptions.Length; ++IOption)
 					{
@@ -4838,7 +4838,7 @@ namespace AutomationScripts
 
 			CreatePaks(Params, SC, PakInputs, PakCryptoSettings, CryptoKeysCacheFilename);
 
-			String ChunkLayerFilename = CombinePaths(GetTmpPackagingPath(Params, SC), GetChunkPakLayerListName());
+			String ChunkLayerFilename = CombinePaths(GetChunkManifestPath(Params, SC), GetChunkPakLayerListName());
 			String OutputChunkLayerFilename = Path.Combine(SC.ProjectRoot.FullName, "Build", SC.FinalCookPlatform, "ChunkLayerInfo", GetChunkPakLayerListName());
 			Directory.CreateDirectory(Path.GetDirectoryName(OutputChunkLayerFilename));
 			File.Copy(ChunkLayerFilename, OutputChunkLayerFilename, true);
@@ -4888,7 +4888,7 @@ namespace AutomationScripts
 
 		private static string GetChunkPakManifestListFilename(ProjectParams Params, DeploymentContext SC)
 		{
-			return CombinePaths(GetTmpPackagingPath(Params, SC), "pakchunklist.txt");
+			return CombinePaths(GetChunkManifestPath(Params, SC), "pakchunklist.txt");
 		}
 
 		private static string GetChunkPakLayerListName()
@@ -4896,15 +4896,15 @@ namespace AutomationScripts
 			return "pakchunklayers.txt";
 		}
 
-		private static string GetTmpPackagingPath(ProjectParams Params, DeploymentContext SC)
+		private static string GetChunkManifestPath(ProjectParams Params, DeploymentContext SC)
 		{
-			string TmpPackagingPath = CombinePaths(Path.GetDirectoryName(Params.RawProjectPath.FullName), "Saved", "TmpPackaging", SC.CookPlatform);
+			string ChunkManifestPath = CombinePaths(SC.MetadataDir.FullName, "ChunkManifest");
 			if (Params.bUseExtraFlavor)
 			{
-				TmpPackagingPath = CombinePaths(TmpPackagingPath, "ExtraFlavor");
+				ChunkManifestPath = CombinePaths(ChunkManifestPath, "ExtraFlavor");
 			}
 
-			return TmpPackagingPath;
+			return ChunkManifestPath;
 		}
 
 		private static bool ShouldCreateIoStoreContainerFiles(ProjectParams Params, DeploymentContext SC)
