New commandlet to detect orphaned localized assets

diff --git a/Engine/Source/Editor/UnrealEd/Private/Commandlets/DetectOrphanedLocalizedAssetsCommandlet.cpp b/Engine/Source/Editor/UnrealEd/Private/Commandlets/DetectOrphanedLocalizedAssetsCommandlet.cpp
new file mode 100644
--- /dev/null
+++ b/Engine/Source/Editor/UnrealEd/Private/Commandlets/DetectOrphanedLocalizedAssetsCommandlet.cpp
@@ -0,0 +1,143 @@
+// Copyright Epic Games, Inc. All Rights Reserved.
+
+#include "Commandlets/DetectOrphanedLocalizedAssetsCommandlet.h"
+#include "CollectionManagerModule.h"
+#include "AssetRegistry/ARFilter.h"
+#include "AssetRegistry/AssetData.h"
+#include "AssetRegistry/IAssetRegistry.h"
+#include "ICollectionManager.h"
+#include "Misc/FileHelper.h"
+#include "Internationalization/PackageLocalizationUtil.h"
+
+DEFINE_LOG_CATEGORY_STATIC(LogDetectOrphanedLocalizedAssetsCommandlet, Log, All);
+
+/**
+ *	UDetectOrphanedLocalizedAssetsCommandlet
+ */
+UDetectOrphanedLocalizedAssetsCommandlet::UDetectOrphanedLocalizedAssetsCommandlet(const FObjectInitializer& ObjectInitializer)
+	: Super(ObjectInitializer)
+{
+}
+
+const FString UDetectOrphanedLocalizedAssetsCommandlet::UsageText
+	(
+	TEXT("DetectOrphanedLocalizedAssetsCommandlet usage...\r\n")
+	TEXT("    <GameName> DetectOrphanedLocalizedAssetsCommandlet -OutputOrphans=<path to output text file containing all orphaned assets>\r\n")
+	);
+
+int32 UDetectOrphanedLocalizedAssetsCommandlet::Main(const FString& Params)
+{
+	TArray<FString> Tokens;
+	TArray<FString> Switches;
+	TMap<FString, FString> ParamVals;
+	UCommandlet::ParseCommandLine(*Params, Tokens, Switches, ParamVals);
+	if (Switches.Contains(TEXT("help")) || Switches.Contains(TEXT("Help")))
+	{
+		UE_LOG(LogDetectOrphanedLocalizedAssetsCommandlet, Display, TEXT("%s"), *UsageText);
+		return 0;
+	}
+
+	FCollectionManagerModule& CollectionManagerModule = FCollectionManagerModule::GetModule();
+	ICollectionManager& CollectionManager = CollectionManagerModule.Get();
+	FARFilter CollectionFilter;
+	bool bSuccess = CollectionManager.GetObjectsInCollection(FName("Audit_InCook"), ECollectionShareType::CST_All, CollectionFilter.SoftObjectPaths, ECollectionRecursionFlags::SelfAndChildren);
+	
+	IAssetRegistry::Get()->SearchAllAssets(true);
+	TArray<FAssetData> AllAssets;
+	const double GetAllAssetsWithFirstPassFilterStartTime = FPlatformTime::Seconds();
+	
+	if (bSuccess)
+	{
+		IAssetRegistry::GetChecked().GetAssets(CollectionFilter, AllAssets);
+	}
+	else
+	{
+		IAssetRegistry::GetChecked().GetAllAssets(AllAssets);
+	}
+	UE_LOG(LogDetectOrphanedLocalizedAssetsCommandlet, Display, TEXT("Getting all assets from asset registry took %.2f seconds."), FPlatformTime::Seconds() - GetAllAssetsWithFirstPassFilterStartTime);
+
+	UE_LOG(LogDetectOrphanedLocalizedAssetsCommandlet, Display, TEXT("Processing %d assets."), AllAssets.Num());
+	TSet<FSoftObjectPath> LocalizedAssets;
+	LocalizedAssets.Reserve(AllAssets.Num());
+	const double AllAssetsIterationStartTime = FPlatformTime::Seconds();
+	for (const FAssetData& Asset : AllAssets)
+	{
+		if (FPackageName::IsLocalizedPackage(Asset.GetSoftObjectPath().GetLongPackageName()) && !Asset.IsRedirector())
+		{
+			LocalizedAssets.Add(Asset.GetSoftObjectPath());
+		}
+	}
+	UE_LOG(LogDetectOrphanedLocalizedAssetsCommandlet, Display, TEXT("Iterating through all assets took %.2f seconds."), FPlatformTime::Seconds() - AllAssetsIterationStartTime);
+	float LocalizedAssetsPercentage = static_cast<float>(LocalizedAssets.Num()) / static_cast<float>(AllAssets.Num()) * 100.0f;
+	UE_LOG(LogDetectOrphanedLocalizedAssetsCommandlet, Display, TEXT("Found %d localized assets out of %d assets. %.2f percent of assets are localized."), LocalizedAssets.Num(), AllAssets.Num(), LocalizedAssetsPercentage);
+
+	IAssetRegistry& AssetRegistry = IAssetRegistry::GetChecked();
+	FString SourceObjectPath;
+	FAssetData OutAssetData;
+	TArray<FSoftObjectPath> OrphanedLocalizedAssets;
+	TArray<FName> Referencers;
+	const double DetectOrphansStartTime = FPlatformTime::Seconds();
+	for (const FSoftObjectPath& LocalizedAsset: LocalizedAssets)
+	{
+		SourceObjectPath.Reset();
+
+		if (FPackageLocalizationUtil::ConvertLocalizedToSource(LocalizedAsset.GetLongPackageName(), SourceObjectPath))
+		{
+			UE::AssetRegistry::EExists Exists = AssetRegistry.TryGetAssetByObjectPath(FSoftObjectPath(SourceObjectPath), OutAssetData);
+			if (Exists == UE::AssetRegistry::EExists::Exists && !OutAssetData.IsRedirector())
+			{
+				// The source version of this asset exists and it's not a redirector. This localized asset is definitely not orphaned. Moving along 
+				continue;
+			}
+			// If the source asset doesn't exist, we still need to check and make sure that none of the referencers are a localized asset 
+			Referencers.Reset();
+			AssetRegistry.GetReferencers(LocalizedAsset.GetLongPackageFName(), Referencers);
+			if (Referencers.Num() == 0)
+			{
+				OrphanedLocalizedAssets.Add(LocalizedAsset);
+				continue;
+			}
+			bool bOrphaned = true;
+			// We go through all of the referencers and see if any of those packages are localized packages.
+			for (const FName& Reference : Referencers)
+			{
+				if (FPackageName::IsLocalizedPackage(Reference.ToString()))
+				{
+					// This is a localized asset with no source asset but there is another localized asset referencing it. This is not considered orphaned 
+					bOrphaned = false;
+					break;
+				}
+			}
+			if (bOrphaned)
+			{
+				OrphanedLocalizedAssets.Add(LocalizedAsset);
+			}
+		}
+	}
+	UE_LOG(LogDetectOrphanedLocalizedAssetsCommandlet, Display, TEXT("Detecting orphaned localized assets took %.2f seconds."), FPlatformTime::Seconds() - DetectOrphansStartTime);
+	float OrphanedPercentage = static_cast<float>(OrphanedLocalizedAssets.Num()) / static_cast<float>(LocalizedAssets.Num()) * 100.0f;
+	UE_LOG(LogDetectOrphanedLocalizedAssetsCommandlet, Display, TEXT("%d out of %d localized assets are orphaned. %.2f of all localized assets are orphaned."), OrphanedLocalizedAssets.Num(), LocalizedAssets.Num(), OrphanedPercentage);
+
+	TArray<FString> OrphanedLocalizedAssetsStrings;
+	OrphanedLocalizedAssetsStrings.Reserve(OrphanedLocalizedAssets.Num());
+	for (const FSoftObjectPath& OrphanedAsset : OrphanedLocalizedAssets)
+	{
+		OrphanedLocalizedAssetsStrings.Add(OrphanedAsset.ToString());
+	}
+
+	if (ParamVals.Contains("OutputOrphans"))
+	{
+		UE_LOG(LogDetectOrphanedLocalizedAssetsCommandlet, Display, TEXT("An output file was provided. Dumping all found orphaned assets to the file: %s"), *ParamVals["OutputOrphans"]);
+		FFileHelper::SaveStringArrayToFile(OrphanedLocalizedAssetsStrings, *ParamVals["OutputOrphans"]);
+	}
+	else
+	{
+		UE_LOG(LogDetectOrphanedLocalizedAssetsCommandlet, Display, TEXT("No output file (-OutputOrphans=<PathToOutput>) was provided. Dumping all found orphaned assets to the console."));
+		for (const FString& OrphanedLocalizedAssetString : OrphanedLocalizedAssetsStrings)
+		{
+			UE_LOG(LogDetectOrphanedLocalizedAssetsCommandlet, Display, TEXT("%s"), *OrphanedLocalizedAssetString);
+		}
+	}
+
+	return 0;
+}
\ No newline at end of file
diff --git a/Engine/Source/Editor/UnrealEd/Public/Commandlets/DetectOrphanedLocalizedAssetsCommandlet.h b/Engine/Source/Editor/UnrealEd/Public/Commandlets/DetectOrphanedLocalizedAssetsCommandlet.h
new file mode 100644
--- /dev/null
+++ b/Engine/Source/Editor/UnrealEd/Public/Commandlets/DetectOrphanedLocalizedAssetsCommandlet.h
@@ -0,0 +1,20 @@
+// Copyright Epic Games, Inc. All Rights Reserved.
+
+#pragma once 
+
+#include "CoreMinimal.h"
+#include "Commandlets/Commandlet.h"
+#include "DetectOrphanedLocalizedAssetsCommandlet.generated.h"
+
+UCLASS()
+class UDetectOrphanedLocalizedAssetsCommandlet : public UCommandlet
+{
+    GENERATED_UCLASS_BODY()
+
+public:
+	//~ Begin UCommandlet Interface
+	virtual int32 Main(const FString& Params) override;
+	//~ End UCommandlet Interface
+
+	static const FString UsageText;	 
+};
