LoadWorldPackageForEditor() - Added proper handling of redirectors

diff --git a/Engine/Source/Editor/UnrealEd/Private/EditorWorldUtils.cpp b/Engine/Source/Editor/UnrealEd/Private/EditorWorldUtils.cpp
--- a/Engine/Source/Editor/UnrealEd/Private/EditorWorldUtils.cpp
+++ b/Engine/Source/Editor/UnrealEd/Private/EditorWorldUtils.cpp
@@ -31,8 +31,7 @@ FScopedEditorWorld::FScopedEditorWorld(const TSoftObjectPtr<UWorld>& InSoftWorld
 		FString WorldObjectPathStr = WorldObjectPath.GetLongPackageName() + TEXT(".") + FPackageName::GetShortName(WorldObjectPath.GetLongPackageName());
 		WorldObjectPath = FSoftObjectPath(WorldObjectPathStr);
 	}
-	UAssetRegistryHelpers::FixupRedirectedAssetPath(WorldObjectPath);
-	
+		
 	if (UPackage* WorldPackage = LoadWorldPackageForEditor(WorldObjectPath.GetLongPackageName()))
 	{
 		if (UWorld* RuntimeWorld = UWorld::FindWorldInPackage(WorldPackage))
@@ -91,9 +90,14 @@ FScopedEditorWorld::~FScopedEditorWorld()
 
 UPackage* LoadWorldPackageForEditor(const FStringView InLongPackageName, EWorldType::Type InWorldType, uint32 InLoadFlags)
 {
-	FName WorldPackageFName(InLongPackageName);
+	FSoftObjectPath WorldPackagePath(InLongPackageName);
+	UAssetRegistryHelpers::FixupRedirectedAssetPath(WorldPackagePath);
+
+	FString LongPackageName = WorldPackagePath.GetLongPackageName();
+	FName WorldPackageFName(LongPackageName);
 	UWorld::WorldTypePreLoadMap.FindOrAdd(WorldPackageFName) = InWorldType;
-	UPackage* WorldPackage = LoadPackage(nullptr, InLongPackageName.GetData(), InLoadFlags);
+	UPackage* WorldPackage = LoadPackage(nullptr, *LongPackageName, InLoadFlags);
 	UWorld::WorldTypePreLoadMap.Remove(WorldPackageFName);
+
 	return WorldPackage;
 }
\ No newline at end of file
