WorldPartitionBuilderCommandlet: When processing a collection of maps, resolve redirectors and ensure we don't end up with duplicate entries

diff --git a/Engine/Source/Editor/UnrealEd/Classes/Commandlets/WorldPartitionBuilderCommandlet.h b/Engine/Source/Editor/UnrealEd/Classes/Commandlets/WorldPartitionBuilderCommandlet.h
--- a/Engine/Source/Editor/UnrealEd/Classes/Commandlets/WorldPartitionBuilderCommandlet.h
+++ b/Engine/Source/Editor/UnrealEd/Classes/Commandlets/WorldPartitionBuilderCommandlet.h
@@ -19,7 +19,7 @@ class UWorldPartitionBuilderCommandlet : public UCommandlet
 	//~ End UCommandlet Interface
 
 private:
-	TArray<FString> GatherMapsFromCollection(const FString& CollectionName) const;
+	TSet<FString> GatherMapsFromCollection(const FString& CollectionName) const;
 	bool RunBuilder(TSubclassOf<UWorldPartitionBuilder> InBuilderClass, const FString& InWorldPackageName);
 
 	bool OnFilesModified(const TArray<FString>& InModifiedFiles, const FString& InChangeDescription);
diff --git a/Engine/Source/Editor/UnrealEd/Private/Commandlets/WorldPartitionBuilderCommandlet.cpp b/Engine/Source/Editor/UnrealEd/Private/Commandlets/WorldPartitionBuilderCommandlet.cpp
--- a/Engine/Source/Editor/UnrealEd/Private/Commandlets/WorldPartitionBuilderCommandlet.cpp
+++ b/Engine/Source/Editor/UnrealEd/Private/Commandlets/WorldPartitionBuilderCommandlet.cpp
@@ -23,6 +23,8 @@
 #include "ISourceControlProvider.h"
 #include "SourceControlOperations.h"
 
+#include "AssetRegistry/AssetRegistryHelpers.h"
+
 DEFINE_LOG_CATEGORY_STATIC(LogWorldPartitionBuilderCommandlet, All, All);
 
 UWorldPartitionBuilderCommandlet::UWorldPartitionBuilderCommandlet(const FObjectInitializer& ObjectInitializer)
@@ -75,7 +77,7 @@ int32 UWorldPartitionBuilderCommandlet::Main(const FString& Params)
 	}
 
 	ICollectionManager& CollectionManager = FModuleManager::LoadModuleChecked<FCollectionManagerModule>("CollectionManager").Get();
-	TArray<FString> MapPackagesNames;
+	TSet<FString> MapPackagesNames;
 
 	// Parse map name or maps collection
 	FString MapLongPackageName;
@@ -140,9 +142,9 @@ int32 UWorldPartitionBuilderCommandlet::Main(const FString& Params)
 	return 0;
 }
 
-TArray<FString> UWorldPartitionBuilderCommandlet::GatherMapsFromCollection(const FString& CollectionName) const
+TSet<FString> UWorldPartitionBuilderCommandlet::GatherMapsFromCollection(const FString& CollectionName) const
 {
-	TArray<FString> MapPackagesNames;
+	TSet<FString> MapPackagesNames;
 
 	ICollectionManager& CollectionManager = FModuleManager::LoadModuleChecked<FCollectionManagerModule>("CollectionManager").Get();
 
@@ -150,8 +152,10 @@ TArray<FString> UWorldPartitionBuilderCommandlet::GatherMapsFromCollection(const
 	CollectionManager.GetAssetsInCollection(FName(CollectionName), ECollectionShareType::CST_All, AssetsPaths, ECollectionRecursionFlags::SelfAndChildren);
 
 	UE_LOG(LogWorldPartitionBuilderCommandlet, Display, TEXT("Processing collection %s (%d items)"), *CollectionName, AssetsPaths.Num());
-	for (const auto& AssetPath : AssetsPaths)
+	for (FSoftObjectPath& AssetPath : AssetsPaths)
 	{
+		UAssetRegistryHelpers::FixupRedirectedAssetPath(AssetPath);
+
 		FString PackageName = AssetPath.GetLongPackageName();
 
 		if (FEditorFileUtils::IsMapPackageAsset(PackageName))
