MPCook: Prevent dereferencing null ServerSocket when handling messages in a loop if one of the messages terminates the socket connection.
#rnx

diff --git a/Engine/Source/Editor/UnrealEd/Private/Cooker/CookWorkerClient.cpp b/Engine/Source/Editor/UnrealEd/Private/Cooker/CookWorkerClient.cpp
--- a/Engine/Source/Editor/UnrealEd/Private/Cooker/CookWorkerClient.cpp
+++ b/Engine/Source/Editor/UnrealEd/Private/Cooker/CookWorkerClient.cpp
@@ -664,6 +664,12 @@ void FCookWorkerClient::PumpReceiveMessages()
 	// Read a packet at a time (with 1 or more messages per packet) until we fail to read any messages
 	for (;;)
 	{
+		if (ServerSocket == nullptr)
+		{
+			// HandleReceiveMessage might change our connectionstatus to LostConnection and kill the ServerSocket,
+			// so we need to check for null after each time we handle messages.
+			break;
+		}
 		Messages.Reset();
 		EConnectionStatus SocketStatus = TryReadPacket(ServerSocket, ReceiveBuffer, Messages);
 		if (SocketStatus != EConnectionStatus::Okay && SocketStatus != EConnectionStatus::Incomplete)
