Fixed Replace Actor not preserving the root transform and mobility if the root component wasn't a native component.

diff --git a/Engine/Source/Editor/UnrealEd/Private/EditorEngine.cpp b/Engine/Source/Editor/UnrealEd/Private/EditorEngine.cpp
--- a/Engine/Source/Editor/UnrealEd/Private/EditorEngine.cpp
+++ b/Engine/Source/Editor/UnrealEd/Private/EditorEngine.cpp
@@ -5497,6 +5497,15 @@ void UEditorEngine::ReplaceActors(UActorFactory* Factory, const FAssetData& Asse
 		// Destroy any non-native constructed components, but make sure we grab the transform first in case it has a
 		// non-native root component. These will be reconstructed as part of the new actor when it's created/instanced.
 		const FTransform OldTransform = OldActor->ActorToWorld();
+
+		TOptional<FVector> OldRelativeScale3D;
+		TOptional<EComponentMobility::Type> OldMobility;
+		if (OldActor->GetRootComponent())
+		{
+			OldRelativeScale3D = OldActor->GetRootComponent()->GetRelativeScale3D();
+			OldMobility = OldActor->GetRootComponent()->Mobility;
+		}
+
 		OldActor->DestroyConstructedComponents();
 
 		// Unregister this actors components because we are effectively replacing it with an actor sharing the same ActorGuid.
@@ -5540,18 +5549,18 @@ void UEditorEngine::ReplaceActors(UActorFactory* Factory, const FAssetData& Asse
 			USceneComponent* const NewActorRootComponent = NewActor->GetRootComponent();
 			if(NewActorRootComponent)
 			{
-				if(!GetDefault<ULevelEditorMiscSettings>()->bReplaceRespectsScale || OldActor->GetRootComponent() == NULL )
+				if(!GetDefault<ULevelEditorMiscSettings>()->bReplaceRespectsScale || !OldRelativeScale3D.IsSet())
 				{
 					NewActorRootComponent->SetRelativeScale3D(FVector(1.0f, 1.0f, 1.0f));
 				}
 				else
 				{
-					NewActorRootComponent->SetRelativeScale3D( OldActor->GetRootComponent()->GetRelativeScale3D() );
+					NewActorRootComponent->SetRelativeScale3D( OldRelativeScale3D.GetValue() );
 				}
 
-				if (OldActor->GetRootComponent() != NULL)
+				if (OldMobility.IsSet())
 				{
-					NewActorRootComponent->SetMobility(OldActor->GetRootComponent()->Mobility);
+					NewActorRootComponent->SetMobility(OldMobility.GetValue());
 				}
 			}
 
