[SkeletalMesh] SkinWeightProfile handling fixes (Import/Remove).

#rnx

diff --git a/Engine/Source/Editor/UnrealEd/Private/SkinWeightsUtilities.cpp b/Engine/Source/Editor/UnrealEd/Private/SkinWeightsUtilities.cpp
--- a/Engine/Source/Editor/UnrealEd/Private/SkinWeightsUtilities.cpp
+++ b/Engine/Source/Editor/UnrealEd/Private/SkinWeightsUtilities.cpp
@@ -349,14 +349,20 @@ bool FSkinWeightsUtilities::ImportAlternateSkinWeight(USkeletalMesh* SkeletalMes
 				// Only add if it is an initial import and it was successful 
 				if (!bIsReimportLocal && bResult)
 				{
-					FSkinWeightProfileInfo SkeletalMeshProfile;
-					SkeletalMeshProfile.DefaultProfile = (SkeletalMesh->GetNumSkinWeightProfiles() == 0);
-					SkeletalMeshProfile.DefaultProfileFromLODIndex = TargetLODIndex;
-					SkeletalMeshProfile.Name = ProfileName;
-					SkeletalMeshProfile.PerLODSourceFiles.Add(TargetLODIndex, UAssetImportData::SanitizeImportFilename(AbsoluteFilePath, SkeletalMesh->GetOutermost()));
-					SkeletalMesh->AddSkinWeightProfile(SkeletalMeshProfile);
-
-					Profile = &SkeletalMeshProfile;
+					TArray<FSkinWeightProfileInfo>& SkinWeightProfileInfos = SkeletalMesh->GetSkinWeightProfiles();
+
+					for (size_t SkinWeightProfileIndex = 0; SkinWeightProfileIndex < SkinWeightProfileInfos.Num(); SkinWeightProfileIndex++)
+					{
+						FSkinWeightProfileInfo& SkinWeightProfileInfo = SkinWeightProfileInfos[SkinWeightProfileInfos.Num() - 1 - SkinWeightProfileIndex];
+						if (SkinWeightProfileInfo.Name == ProfileName)
+						{
+							SkinWeightProfileInfo.PerLODSourceFiles.Add(TargetLODIndex, UAssetImportData::SanitizeImportFilename(AbsoluteFilePath, SkeletalMesh->GetOutermost()));
+
+							Profile = &SkinWeightProfileInfo;
+
+							break;
+						}
+					}
 				}
 			}
 		}
@@ -429,6 +435,16 @@ bool FSkinWeightsUtilities::RemoveSkinnedWeightProfileData(USkeletalMesh* Skelet
 	SkeletalMesh->LoadLODImportedData(LODIndex, ImportDataDest);
 	PRAGMA_ENABLE_DEPRECATION_WARNINGS
 
+	for (size_t AlternateInfluenceIndex = 0; AlternateInfluenceIndex < ImportDataDest.AlternateInfluenceProfileNames.Num(); AlternateInfluenceIndex++)
+	{
+		if (ImportDataDest.AlternateInfluenceProfileNames[AlternateInfluenceIndex] == ProfileName)
+		{
+			ImportDataDest.AlternateInfluenceProfileNames.RemoveAt(AlternateInfluenceIndex);
+			ImportDataDest.AlternateInfluences.RemoveAt(AlternateInfluenceIndex);
+			break;
+		}
+	}
+
 	//If we have a LOD info we use the build settings to be sure we are rechunking the LOD with the existing options
 	IMeshUtilities::MeshBuildOptions BuildOptions;
 	if (FSkeletalMeshLODInfo* LODInfo = SkeletalMesh->GetLODInfo(LODIndex))
@@ -467,6 +483,11 @@ bool FSkinWeightsUtilities::RemoveSkinnedWeightProfileData(USkeletalMesh* Skelet
 	const bool bBuildSuccess = MeshUtilities.BuildSkeletalMesh(LODModelDest, SkeletalMesh->GetPathName(), SkeletalMesh->GetRefSkeleton(), LODInfluencesDest, LODWedgesDest, LODFacesDest, LODPointsDest, LODPointToRawMapDest, BuildOptions, &WarningMessages, &WarningNames);
 	FLODUtilities::RegenerateAllImportSkinWeightProfileData(LODModelDest, BuildOptions.BoneInfluenceLimit, BuildOptions.TargetPlatform);
 
+	//Resave the bulk data with the new or refreshed data
+	PRAGMA_DISABLE_DEPRECATION_WARNINGS
+	SkeletalMesh->SaveLODImportedData(LODIndex, ImportDataDest);
+	PRAGMA_ENABLE_DEPRECATION_WARNINGS
+
 	return bBuildSuccess;
 }
 
