Interchange:
- Add flag when a static mesh collision was imported or editor created
- Add another value in the enum of the mesh collision to only import collision from a file, otherwise in case of a reimport we'll keep the previous imported collisions

AggregateGeom:
- Add a parameter to empty only the imported elements

GeomFitUtils:
- Set to true the generated collision

diff --git a/Engine/Plugins/Interchange/Runtime/Source/Import/Private/Mesh/InterchangeStaticMeshFactory.cpp b/Engine/Plugins/Interchange/Runtime/Source/Import/Private/Mesh/InterchangeStaticMeshFactory.cpp
--- a/Engine/Plugins/Interchange/Runtime/Source/Import/Private/Mesh/InterchangeStaticMeshFactory.cpp
+++ b/Engine/Plugins/Interchange/Runtime/Source/Import/Private/Mesh/InterchangeStaticMeshFactory.cpp
@@ -432,7 +432,7 @@ UInterchangeFactoryBase::FImportAssetResult UInterchangeStaticMeshFactory::Impor
 			{
 				if (bReimport)
 				{
-					StaticMesh->GetBodySetup()->AggGeom.EmptyElements();
+					StaticMesh->GetBodySetup()->AggGeom.EmptyImportedElements();
 				}
 				bImportedCustomCollision |= ImportBoxCollision(Arguments, StaticMesh, LodDataNode);
 				bImportedCustomCollision |= ImportCapsuleCollision(Arguments, StaticMesh, LodDataNode);
diff --git a/Engine/Plugins/Interchange/Runtime/Source/Nodes/Public/InterchangeMeshDefinitions.h b/Engine/Plugins/Interchange/Runtime/Source/Nodes/Public/InterchangeMeshDefinitions.h
--- a/Engine/Plugins/Interchange/Runtime/Source/Nodes/Public/InterchangeMeshDefinitions.h
+++ b/Engine/Plugins/Interchange/Runtime/Source/Nodes/Public/InterchangeMeshDefinitions.h
@@ -34,6 +34,9 @@ enum class EInterchangeMeshCollision : uint8
 	/** Generates a new axis-aligned box collision mesh with all edges and corners beveled (26 total sides) encompassing the static mesh*/
 	Convex26DOP UMETA(DisplayName = "26DOP Simplified"),
 	
+	/** Generates only collisions retrieved from the file*/
+	Custom UMETA(DisplayName = "From File"),
+
 	None = 0xFF
 };
 
diff --git a/Engine/Source/Editor/UnrealEd/Private/GeomFitUtils.cpp b/Engine/Source/Editor/UnrealEd/Private/GeomFitUtils.cpp
--- a/Engine/Source/Editor/UnrealEd/Private/GeomFitUtils.cpp
+++ b/Engine/Source/Editor/UnrealEd/Private/GeomFitUtils.cpp
@@ -108,6 +108,7 @@ int32 GenerateKDopAsSimpleCollision(UStaticMesh* StaticMesh, const TArray<FVecto
 	bs->InvalidatePhysicsData();
 
 	FKConvexElem ConvexElem;
+	ConvexElem.bIsGenerated = true;
 	ConvexElem.VertexData = HullVertices;
 	// Note: UpdateElemBox also computes the convex hull indices
 	ConvexElem.UpdateElemBox();
@@ -160,6 +161,7 @@ int32 GenerateBoxAsSimpleCollision(UStaticMesh* StaticMesh, bool bUpdate)
 	bs->InvalidatePhysicsData();
 
 	FKBoxElem BoxElem;
+	BoxElem.bIsGenerated = true;
 	BoxElem.Center = Center;
 	BoxElem.X = Extents.X * 2.0f;
 	BoxElem.Y = Extents.Y * 2.0f;
@@ -364,6 +366,7 @@ int32 GenerateSphereAsSimpleCollision(UStaticMesh* StaticMesh, bool bUpdate)
 	bs->InvalidatePhysicsData();
 
 	FKSphereElem SphereElem;
+	SphereElem.bIsGenerated = true;
 	SphereElem.Center = bestSphere.Center;
 	SphereElem.Radius = bestSphere.W;
 	bs->AggGeom.SphereElems.Add(SphereElem);
@@ -514,6 +517,7 @@ int32 GenerateSphylAsSimpleCollision(UStaticMesh* StaticMesh, bool bUpdate)
 	bs->InvalidatePhysicsData();
 
 	FKSphylElem SphylElem;
+	SphylElem.bIsGenerated = true;
 	SphylElem.Center = sphere.Center;
 	SphylElem.Rotation = rotation;
 	SphylElem.Radius = sphere.W;
diff --git a/Engine/Source/Runtime/Engine/Classes/PhysicsEngine/AggregateGeom.h b/Engine/Source/Runtime/Engine/Classes/PhysicsEngine/AggregateGeom.h
--- a/Engine/Source/Runtime/Engine/Classes/PhysicsEngine/AggregateGeom.h
+++ b/Engine/Source/Runtime/Engine/Classes/PhysicsEngine/AggregateGeom.h
@@ -222,6 +222,28 @@ struct FKAggregateGeom
 		return INDEX_NONE;
 	}
 
+#if WITH_EDITORONLY_DATA
+	void EmptyImportedElements()
+	{
+		auto CleanUp = [](auto& Elems)
+		{
+			Elems.RemoveAllSwap([](const FKShapeElem& Elem)
+			{
+				return Elem.bIsGenerated == false;
+			});
+		};
+		CleanUp(BoxElems);
+		CleanUp(ConvexElems);
+		CleanUp(SphylElems);
+		CleanUp(SphereElems);
+		CleanUp(TaperedCapsuleElems);
+		CleanUp(LevelSetElems);
+		CleanUp(SkinnedLevelSetElems);
+
+		FreeRenderInfo();
+	}
+#endif
+
 	void EmptyElements()
 	{
 		BoxElems.Empty();
diff --git a/Engine/Source/Runtime/Engine/Classes/PhysicsEngine/ShapeElem.h b/Engine/Source/Runtime/Engine/Classes/PhysicsEngine/ShapeElem.h
--- a/Engine/Source/Runtime/Engine/Classes/PhysicsEngine/ShapeElem.h
+++ b/Engine/Source/Runtime/Engine/Classes/PhysicsEngine/ShapeElem.h
@@ -118,7 +118,7 @@ struct FKShapeElem
 
 #if WITH_EDITORONLY_DATA
 	/** True when the shape was created by the engine and was not imported. */
-	UPROPERTY()
+	UPROPERTY(Category = Shape, EditAnywhere, DisplayName="Collision Generated")
 	uint8 bIsGenerated : 1;
 #endif
 
@@ -131,6 +131,9 @@ protected:
 		Name = Other.Name;
 		bContributeToMass = Other.bContributeToMass;
 		CollisionEnabled = Other.CollisionEnabled;
+#if WITH_EDITORONLY_DATA
+		bIsGenerated = Other.bIsGenerated;
+#endif
 	}
 
 private:
