Add VCam's UAssetFilteringAndSortingFunctionLibrary functions to UEditorAssetSubsystem and UEditorAssetLibrary.
VCam's version will be deprecated with UE-219983 once Blueprint references are fixed up, see UE-219982.

diff --git a/Engine/Source/Editor/UnrealEd/Private/Subsystems/EditorAssetSubsystem.cpp b/Engine/Source/Editor/UnrealEd/Private/Subsystems/EditorAssetSubsystem.cpp
--- a/Engine/Source/Editor/UnrealEd/Private/Subsystems/EditorAssetSubsystem.cpp
+++ b/Engine/Source/Editor/UnrealEd/Private/Subsystems/EditorAssetSubsystem.cpp
@@ -32,7 +32,7 @@ using TError = TValueOrError<UE::EditorAssetUtils::ErrorTag, T>;
 
 namespace UE::EditorAssetUtils
 {
-	bool EnsureAssetsLoaded()
+	static bool EnsureAssetsLoaded()
 	{
 		FAssetRegistryModule& AssetRegistryModule = FModuleManager::LoadModuleChecked<FAssetRegistryModule>("AssetRegistry");
 		if (IAssetRegistry& AssetRegistry = AssetRegistryModule.Get(); AssetRegistry.IsLoadingAssets())
@@ -49,7 +49,7 @@ namespace UE::EditorAssetUtils
 		return true;
 	}
 
-	TValueOrError<FAssetData, FString> FindAssetDataFromAnyPath(const FString& AnyAssetPath)
+	static TValueOrError<FAssetData, FString> FindAssetDataFromAnyPath(const FString& AnyAssetPath)
 	{
 		FString FailureReason;
 		FString ObjectPath = EditorScriptingHelpers::ConvertAnyPathToSubObjectPath(AnyAssetPath, FailureReason);
@@ -78,7 +78,7 @@ namespace UE::EditorAssetUtils
 		return MakeValue(AssetData);
 	}
 
-	TValueOrError<UObject*, FString> LoadAssetFromData(const FAssetData& AssetData)
+	static TValueOrError<UObject*, FString> LoadAssetFromData(const FAssetData& AssetData)
 	{
 		if (!AssetData.IsValid())
 		{
@@ -97,7 +97,7 @@ namespace UE::EditorAssetUtils
 		return MakeValue(FoundObject);
 	}
 	
-	TValueOrError<UObject*, FString> LoadAssetFromPath(const FString& AssetPath)
+	static TValueOrError<UObject*, FString> LoadAssetFromPath(const FString& AssetPath)
 	{
 		TValueOrError<FAssetData, FString> AssetDataResult = FindAssetDataFromAnyPath(AssetPath);
 		if (AssetDataResult.HasError())
@@ -107,7 +107,7 @@ namespace UE::EditorAssetUtils
 		return LoadAssetFromData(AssetDataResult.GetValue());
 	}
 
-	TError<FString> IsARegisteredAsset(UObject* Object, bool bAllowSkipBrowsableTestForExternalObject = false)
+	static TError<FString> IsARegisteredAsset(UObject* Object, bool bAllowSkipBrowsableTestForExternalObject = false)
 	{
 		if (!IsValid(Object))
 		{
@@ -131,7 +131,7 @@ namespace UE::EditorAssetUtils
 		return MakeValue();
 	}
 
-	TError<FString> EnumerateAssetsInDirectory(const FString& AnyPathDirectoryPath, bool bRecursive, TArray<FAssetData>& OutResult, FString& OutDirectoryPath)
+	static TError<FString> EnumerateAssetsInDirectory(const FString& AnyPathDirectoryPath, bool bRecursive, TArray<FAssetData>& OutResult, FString& OutDirectoryPath)
 	{
 		OutResult.Reset();
 		OutDirectoryPath.Reset();
@@ -158,7 +158,7 @@ namespace UE::EditorAssetUtils
 		OnlyDirty
 	};
 	
-	TError<FString> EnumeratePackagesInDirectory(const FString& AnyDirectoryPath, EPackageEnumerationFilter EnumerationFilter, bool bRecursive, TArray<UPackage*>& OutResult)
+	static TError<FString> EnumeratePackagesInDirectory(const FString& AnyDirectoryPath, EPackageEnumerationFilter EnumerationFilter, bool bRecursive, TArray<UPackage*>& OutResult)
 	{
 		FString ValidDirectoryPath;
 		TArray<FAssetData> Assets;
@@ -197,7 +197,7 @@ namespace UE::EditorAssetUtils
 		return MakeValue();
 	}
 	
-	bool DeleteEmptyDirectoryFromDisk(const FString& LongPackagePath)
+	static bool DeleteEmptyDirectoryFromDisk(const FString& LongPackagePath)
 	{
 		struct FEmptyDirectoryVisitor : public IPlatformFile::FDirectoryVisitor
 		{
@@ -236,7 +236,7 @@ namespace UE::EditorAssetUtils
 		FString DestinationFilePath;
 	};
 	
-	TValueOrError<FDirectoryRenamePaths, FString> GetDirectoryRenamePaths(const FString& SourceDirectoryPath, const FString& DestinationDirectoryPath)
+	static TValueOrError<FDirectoryRenamePaths, FString> GetDirectoryRenamePaths(const FString& SourceDirectoryPath, const FString& DestinationDirectoryPath)
 	{
 		FDirectoryRenamePaths Paths;
 		FString FailureReason;
@@ -267,7 +267,7 @@ namespace UE::EditorAssetUtils
 		return MakeValue(Paths);
 	}
 
-	TError<FString> SetupDirectoryRename(const FDirectoryRenamePaths& Paths)
+	static TError<FString> SetupDirectoryRename(const FDirectoryRenamePaths& Paths)
 	{
 		// If the source directory doesn't exist on disk then it can't be operated on
 		if (!IFileManager::Get().DirectoryExists(*Paths.SourceFilePath))
@@ -294,7 +294,7 @@ namespace UE::EditorAssetUtils
 		TArray<FString> DestinationDirectoryAssetPaths;
 	};
 	
-	TValueOrError<FDirectoryRenameAssetPaths, FString> GetDirectoryRenameAssetPaths(const FDirectoryRenamePaths& Paths)
+	static TValueOrError<FDirectoryRenameAssetPaths, FString> GetDirectoryRenameAssetPaths(const FDirectoryRenamePaths& Paths)
 	{
 		FDirectoryRenameAssetPaths AssetPaths;
 		
@@ -341,6 +341,53 @@ namespace UE::EditorAssetUtils
 		}
 		return MakeValue(AssetPaths);
 	}
+	
+	static void SortAssets(
+		TArray<FAssetData>& Assets,
+		TFunctionRef<bool(const FAssetData& Left, const FAssetData& Right)> Predicate,
+		EEditorAssetSortOrder SortOrder
+	)
+	{
+		// Careful this would be undefined behaviour: TFunctionRef Variable = [](){}; 
+		auto ReversePredicate = [&Predicate](const FAssetData& Left, const FAssetData& Right) { return !Predicate(Left, Right); };
+		const TFunctionRef<bool(const FAssetData& Left, const FAssetData& Right)> ReversePredicateFunc = ReversePredicate;
+		const TFunctionRef<bool(const FAssetData& Left, const FAssetData& Right)> PredicateToUse = SortOrder == EEditorAssetSortOrder::Ascending ? Predicate : ReversePredicateFunc;
+	
+		Assets.Sort(PredicateToUse);
+	}
+	
+	template<typename TType>
+	static bool Sort(
+		TArray<FAssetData>& Assets,
+		FName MetaDataTag,
+		TFunctionRef<bool(const FString& TagType, TType& Converted)> Converter,
+		EEditorAssetSortOrder SortOrder
+		)
+	{
+		TMap<FSoftObjectPath, TType> MetaData;
+		MetaData.Reserve(Assets.Num());
+	
+		for (const FAssetData& AssetData : Assets)
+		{
+			const FAssetTagValueRef Value = AssetData.TagsAndValues.FindTag(MetaDataTag);
+			TType AssetTagValue;
+			if (Value.IsSet() && Converter(Value.GetValue(), AssetTagValue))
+			{
+				MetaData.Add(AssetData.GetSoftObjectPath(), AssetTagValue);
+			}
+			else
+			{
+				UE_LOG(LogEditorAssetSubsystem, Warning, TEXT("Not all assets have the tag '%s'"), *MetaDataTag.ToString());
+				return false;
+			}
+		}
+
+		SortAssets(Assets, [&MetaData](const FAssetData& Left, const FAssetData& Right)
+		{
+			return MetaData[Left.GetSoftObjectPath()] <= MetaData[Right.GetSoftObjectPath()]; 
+		}, SortOrder);
+		return true;
+	}
 }
 
 UEditorAssetSubsystem::UEditorAssetSubsystem()
@@ -1634,7 +1681,8 @@ void UEditorAssetSubsystem::RemoveOnExtractAssetFromFile(FOnExtractAssetFromFile
 	OnExtractAssetFromFileDynamicArray.Remove(Delegate);
 }
 
-void UEditorAssetSubsystem::CallOnExtractAssetFromFileDynamicArray(const TArray<FString>& Files,
+void UEditorAssetSubsystem::CallOnExtractAssetFromFileDynamicArray(
+	const TArray<FString>& Files,
 	TArray<FAssetData>& OutAssetDataArray)
 {
 	TArray<FAssetData> LocalArray;
@@ -1645,3 +1693,88 @@ void UEditorAssetSubsystem::CallOnExtractAssetFromFileDynamicArray(const TArray<
 		LocalArray.Reset();
 	}
 }
+
+TArray<FAssetData> UEditorAssetSubsystem::GetAllAssetsByMetaDataTags(
+	const TSet<FName>& RequiredTags,
+	const TSet<UClass*>& AllowedClasses
+	)
+{
+	TGuardValue<bool> UnattendedScriptGuard(GIsRunningUnattendedScript, true);
+
+	TArray<FAssetData> Result;
+	if (!EditorScriptingHelpers::CheckIfInEditorAndPIE() || !UE::EditorAssetUtils::EnsureAssetsLoaded())
+	{
+		return Result;
+	}
+	
+	FARFilter Filter;
+	Filter.TagsAndValues.Reserve(RequiredTags.Num());
+	for (FName RequiredTag : RequiredTags)
+	{
+		Filter.TagsAndValues.Add(RequiredTag);
+	}
+	Filter.ClassPaths.Reserve(RequiredTags.Num());
+	for (UClass* AllowedClass : AllowedClasses)
+	{
+		Filter.ClassPaths.Add(FTopLevelAssetPath(AllowedClass));
+	}
+
+	IAssetRegistry::Get()->GetAssets(Filter, Result);
+	return Result;
+}
+
+void UEditorAssetSubsystem::SortByPredicate(
+	TArray<FAssetData>& Assets,
+	FEditorAssetSortingPredicate SortingPredicate,
+	EEditorAssetSortOrder SortOrder
+	)
+{
+	if (SortingPredicate.IsBound())
+	{
+		UE::EditorAssetUtils::SortAssets(
+			Assets,
+			[&SortingPredicate](const FAssetData& Left, const FAssetData& Right)
+			{
+				return SortingPredicate.Execute(Left, Right);
+			}, SortOrder);
+	}
+}
+
+void UEditorAssetSubsystem::SortByName(TArray<FAssetData>& Assets, EEditorAssetSortOrder SortOrder)
+{
+	UE::EditorAssetUtils::SortAssets(Assets,
+		[](const FAssetData& Left, const FAssetData& Right)
+		{
+			return Left.AssetName.LexicalLess(Right.AssetName);
+		}, SortOrder);
+}
+
+bool UEditorAssetSubsystem::SortByMetaData(
+	TArray<FAssetData>& Assets,
+	FName MetaDataTag,
+	EEditorAssetMetaDataSortType MetaDataType,
+	EEditorAssetSortOrder SortOrder
+	)
+{
+	using namespace UE::EditorAssetUtils;
+	switch (MetaDataType)
+	{
+	case EEditorAssetMetaDataSortType::String:
+		return Sort<FString>(Assets, MetaDataTag, [](const FString& String, FString& Result){ Result = String; return true; }, SortOrder);
+	case EEditorAssetMetaDataSortType::Numeric:
+		return Sort<double>(Assets, MetaDataTag, [](const FString& String, double& Result)
+		{
+			if (String.IsNumeric())
+			{
+				Result = FCString::Atod(*String);
+				return true;
+			}
+			return false;
+		}, SortOrder);
+	case EEditorAssetMetaDataSortType::DateTime:
+		return Sort<FDateTime>(Assets, MetaDataTag, [](const FString& String, FDateTime& Result){ return FDateTime::Parse(String, Result); }, SortOrder);
+	default:
+		checkNoEntry();
+		return false;
+	}
+}
diff --git a/Engine/Source/Editor/UnrealEd/Public/Subsystems/EditorAssetSubsystem.h b/Engine/Source/Editor/UnrealEd/Public/Subsystems/EditorAssetSubsystem.h
--- a/Engine/Source/Editor/UnrealEd/Public/Subsystems/EditorAssetSubsystem.h
+++ b/Engine/Source/Editor/UnrealEd/Public/Subsystems/EditorAssetSubsystem.h
@@ -8,6 +8,23 @@
 
 #include "EditorAssetSubsystem.generated.h"
 
+UENUM(BlueprintType, DisplayName = "Sort Order")
+enum class EEditorAssetSortOrder : uint8
+{
+	Ascending,
+	Descending
+};
+
+UENUM(BlueprintType, DisplayName = "Meta Data Sort Type")
+enum class EEditorAssetMetaDataSortType : uint8
+{
+	String,
+	Numeric,
+	DateTime
+};
+
+DECLARE_DYNAMIC_DELEGATE_RetVal_TwoParams(bool, FEditorAssetSortingPredicate, const FAssetData&, Left, const FAssetData&, Right);
+
 /**
 * UEditorAssetSubsystem
 * Subsystem for exposing asset related utilities to scripts.
@@ -394,14 +411,14 @@ public:
 	 * for example from a drag and drop operation.
 	 */
 	UFUNCTION(BlueprintCallable, Category = "Editor Scripting | Asset")
-		UNREALED_API void AddOnExtractAssetFromFile(FOnExtractAssetFromFileDynamic Delegate);
+	UNREALED_API void AddOnExtractAssetFromFile(FOnExtractAssetFromFileDynamic Delegate);
 
 	/**
 	 * Call this to remove a callback added with AddOnExtractAssetFromFile.
 	 */
 	UFUNCTION(BlueprintCallable, Category = "Editor Scripting | Asset")
-		UNREALED_API void RemoveOnExtractAssetFromFile(FOnExtractAssetFromFileDynamic Delegate);
-
+	UNREALED_API void RemoveOnExtractAssetFromFile(FOnExtractAssetFromFileDynamic Delegate);
+	
 	/**
 	 * Get the delegate for extracting an asset from a file,
 	 * for example from a drag and drop operation.
@@ -410,6 +427,47 @@ public:
 	 * Broadcasting this will also call anything added with AddOnExtractAssetFromFile.
 	 */
 	FOnExtractAssetFromFile& GetOnExtractAssetFromFile() { return OnExtractAssetFromFile; }
+	
+	/**
+	 * Gets all assets which have the given tags.
+	 * 
+	 * @params RequiredTags The tags the assets should have
+	 * @params AllowedClasses The class types the contained assets should have
+	 */
+	UFUNCTION(BlueprintCallable, Category = "Editor Scripting | Asset")
+	UNREALED_API TArray<FAssetData> GetAllAssetsByMetaDataTags(const TSet<FName>& RequiredTags, const TSet<UClass*>& AllowedClasses);
+
+	/**
+	 * Sorts the assets based on a custom Blueprint delegate.
+	 * 
+	 * @param Assets The assets to sort
+	 * @param SortingPredicate Implements a Left <= Right relation
+	 * @param SortOrder Whether to sort ascending or descending
+	 */
+	UFUNCTION(BlueprintCallable, Category = "Editor Scripting | Asset")
+	UNREALED_API void SortByPredicate(UPARAM(Ref) TArray<FAssetData>& Assets, FEditorAssetSortingPredicate SortingPredicate, EEditorAssetSortOrder SortOrder);
+
+	/**
+	 * Sorts the assets by their asset name.
+	 *
+	 * @param Assets The assets to sort
+	 * @param SortOrder Whether to sort ascending or descending
+	 */
+	UFUNCTION(BlueprintCallable, Category = "Editor Scripting | Asset")
+	UNREALED_API void SortByName(UPARAM(Ref) TArray<FAssetData>& Assets, EEditorAssetSortOrder SortOrder);
+	
+	/**
+	 * Sorts the assets based on their meta data's type.
+	 * Supported types: FString, int, float, FDateTime.
+	 * 
+	 * @param Assets The assets to sort
+	 * @param MetaDataTag The on which the sort is based
+	 * @param MetaDataType The meta data type of MetaDataTag
+	 * @param SortOrder Whether to sort ascending or descending
+	 * @return Whether the data was sorted, e.g. false if not all assets have the MetaDataTag.
+	 */
+	UFUNCTION(BlueprintCallable, Category = "Editor Scripting | Asset")
+	UNREALED_API bool SortByMetaData(UPARAM(Ref) TArray<FAssetData>& Assets, FName MetaDataTag, EEditorAssetMetaDataSortType MetaDataType, EEditorAssetSortOrder SortOrder);
 
 private:
 
