Only defer construction scripts if static mesh compilation is actually pending. Also, run deferred construction scripts before the PIE world is duplicated, as otherwise we will duplicate actors that have not had their construction scripts run.
[RN] Fixed broken logic related to construction script deferral due to static mesh compilation

diff --git a/Engine/Source/Editor/UnrealEd/Private/PlayLevel.cpp b/Engine/Source/Editor/UnrealEd/Private/PlayLevel.cpp
--- a/Engine/Source/Editor/UnrealEd/Private/PlayLevel.cpp
+++ b/Engine/Source/Editor/UnrealEd/Private/PlayLevel.cpp
@@ -2,6 +2,7 @@
 
 #include "PlayLevel.h"
 #include "CoreMinimal.h"
+#include "AssetCompilingManager.h"
 #include "AssetRegistry/AssetRegistryModule.h"
 #include "AssetRegistry/IAssetRegistry.h"
 #include "Misc/MessageDialog.h"
@@ -2341,6 +2342,9 @@ UWorld* UEditorEngine::CreatePIEWorldByDuplication(FWorldContext &WorldContext,
 
 		// NULL GWorld before various PostLoad functions are called, this makes it easier to debug invalid GWorld accesses
 		GWorld = NULL;
+		
+		// Gives a chance to any assets being used for PIE/game to complete
+		FAssetCompilingManager::Get().ProcessAsyncTasks();
 
 		// Duplicate the editor world to create the PIE world
 		NewPIEWorld = UWorld::GetDuplicatedWorldForPIE(InWorld, PlayWorldPackage, WorldContext.PIEInstance);
diff --git a/Engine/Source/Runtime/Engine/Private/Level.cpp b/Engine/Source/Runtime/Engine/Private/Level.cpp
--- a/Engine/Source/Runtime/Engine/Private/Level.cpp
+++ b/Engine/Source/Runtime/Engine/Private/Level.cpp
@@ -1966,7 +1966,7 @@ bool ULevel::HasStaticMeshCompilationPending()
 			{
 				if (Component->IsCompiling())
 				{
-					bHasStaticMeshCompilationPending = false;
+					bHasStaticMeshCompilationPending = true;
 					return false;
 				}
 			}
@@ -1992,7 +1992,7 @@ bool ULevel::DeferRunningConstructionScripts(AActor* InActor)
 			bCachedHasStaticMeshCompilationPending = HasStaticMeshCompilationPending();
 		}
 
-		if (bCachedHasStaticMeshCompilationPending)
+		if (bCachedHasStaticMeshCompilationPending.GetValue())
 		{
 			FActorDeferredScriptManager::Get().AddActor(InActor);
 			return true;
