[Viewport Toolbar] Create the generic UnrealEd "Camera" submenu and add to the new Blueprint and Static Mesh viewport toolbars

Also updates Level Editor toolbar to use the generic UnrealEd Camera.

diff --git a/Engine/Source/Editor/Kismet/Private/SSCSEditorViewport.cpp b/Engine/Source/Editor/Kismet/Private/SSCSEditorViewport.cpp
--- a/Engine/Source/Editor/Kismet/Private/SSCSEditorViewport.cpp
+++ b/Engine/Source/Editor/Kismet/Private/SSCSEditorViewport.cpp
@@ -56,6 +56,107 @@ struct FGeometry;
 
 #define LOCTEXT_NAMESPACE "SSCSEditorViewportToolBar"
 
+namespace UE::SCSEditor::Private
+{
+
+void CreateCameraSpeedMenu(UToolMenu* InMenu, TWeakPtr<FSCSEditorViewportClient> InWeakViewportClient)
+{
+	FToolMenuSection& PositioningSection =
+		InMenu->FindOrAddSection("Positioning", LOCTEXT("PositioningLabel", "Positioning"));
+
+	PositioningSection.AddSubMenu(
+		"CameraSpeed",
+		LOCTEXT("CameraSpeedSubMenu", "Camera Speed"),
+		LOCTEXT("CameraSpeedSubMenu_ToolTip", "Camera Speed related actions"),
+		FNewToolMenuDelegate::CreateLambda(
+			[InWeakViewportClient](UToolMenu* InMenu)
+			{
+				// It seems like the speed settings can only be set via the mouse wheel for BP editor.
+				// The other slider values (see legacy toolbar) are not actually changing speed.
+
+				// Taken from legacy toolbar values
+				constexpr int32 MinSpeed = 1;
+				constexpr int32 MaxSpeed = 8;
+
+				FToolMenuEntry CameraSpeedSlider = UnrealEd::CreateNumericEntry(
+					"CameraSpeed",
+					LOCTEXT("CameraSpeedLabel", "Camera Speed"),
+					LOCTEXT("CameraSpeedTooltip", "Camera Speed"),
+					FCanExecuteAction(),
+					UnrealEd::FNumericEntryExecuteActionDelegate::CreateLambda(
+						[InWeakViewportClient](float InValue)
+						{
+							if (TSharedPtr<FSCSEditorViewportClient> LevelViewport = InWeakViewportClient.Pin())
+							{
+								// TODO: properly implement
+								//LevelViewport->SetCameraSpeedSetting(InValue);
+							}
+						}
+					),
+					TAttribute<float>::CreateLambda(
+						[InWeakViewportClient]() -> float
+						{
+							// TODO: properly implement
+							/*if (TSharedPtr<FSCSEditorViewportClient> LevelViewport = InWeakViewportClient.Pin())
+							{
+								return LevelViewport->GetCameraSpeedSetting();
+							}*/
+
+							return 1;
+						}
+					),
+					MinSpeed,
+					MaxSpeed,
+					0
+				);
+
+				InMenu->AddMenuEntry("CameraSpeed", CameraSpeedSlider);
+			}
+		),
+		false,
+		FSlateIcon(FSlateIcon(FAppStyle::GetAppStyleSetName(), "LevelViewport.ToggleActorPilotCameraView"))
+	);
+}
+
+void ExtendCameraSubmenu(FName InCameraOptionsSubmenuName, TSharedPtr<FSCSEditorViewportClient> InViewportClient)
+{
+	UToolMenu* const Submenu = UToolMenus::Get()->ExtendMenu(InCameraOptionsSubmenuName);
+
+	TWeakPtr<FSCSEditorViewportClient> WeakViewportClient = InViewportClient;
+	Submenu->AddDynamicSection(
+		"EditorCameraExtensionDynamicSection",
+		FNewToolMenuDelegate::CreateLambda(
+			[WeakViewportClient](UToolMenu* InDynamicMenu)
+			{
+				UUnrealEdViewportToolbarContext* const EditorViewportContext =
+					InDynamicMenu->FindContext<UUnrealEdViewportToolbarContext>();
+				if (!EditorViewportContext)
+				{
+					return;
+				}
+
+				const TSharedPtr<SEditorViewport> EditorViewport = EditorViewportContext->Viewport.Pin();
+				if (!EditorViewport)
+				{
+					return;
+				}
+
+				FToolMenuSection& PositioningSection =
+					InDynamicMenu->FindOrAddSection("Positioning", LOCTEXT("PositioningLabel", "Positioning"));
+
+				// Camera Speed Submenu
+				{
+					CreateCameraSpeedMenu(InDynamicMenu, WeakViewportClient);
+				}
+
+				PositioningSection.AddMenuEntry(FBlueprintEditorCommands::Get().ResetCamera);
+			}
+		)
+	);
+}
+
+} // namespace UE::SCSEditor::Private
+
 /*-----------------------------------------------------------------------------
    SSCSEditorViewportToolBar
 -----------------------------------------------------------------------------*/
@@ -372,7 +473,33 @@ TSharedPtr<SWidget> SSCSEditorViewport::BuildViewportToolbar()
 			FToolMenuSection& RightSection = ViewportToolbarMenu->FindOrAddSection("Right");
 			RightSection.Alignment = EToolMenuSectionAlign::Last;
 
-			// TODO: Add Camera menu here.
+			// Add the "Camera" submenu.
+			{
+				const FName GrandParentSubmenuName = "UnrealEd.ViewportToolbar.Camera";
+				const FName ParentSubmenuName = "SCSEditor.ViewportToolbar.Camera";
+				const FName SubmenuName = "SCSEditor.ViewportToolbar.CameraOptions";
+
+				// Create our grandparent menu.
+				if (!UToolMenus::Get()->IsMenuRegistered(GrandParentSubmenuName))
+				{
+					UToolMenus::Get()->RegisterMenu(GrandParentSubmenuName);
+				}
+
+				// Create our parent menu.
+				if (!UToolMenus::Get()->IsMenuRegistered(ParentSubmenuName))
+				{
+					UToolMenus::Get()->RegisterMenu(ParentSubmenuName, GrandParentSubmenuName);
+				}
+
+				// Create our menu.
+				UToolMenus::Get()->RegisterMenu(SubmenuName, ParentSubmenuName);
+
+				UE::SCSEditor::Private::ExtendCameraSubmenu(SubmenuName, ViewportClient);
+
+				FToolMenuEntry CameraSubmenu = UE::UnrealEd::CreateViewportToolbarCameraSubmenu();
+				CameraSubmenu.InsertPosition.Position = EToolMenuInsertType::First;
+				RightSection.AddEntry(CameraSubmenu);
+			}
 
 			// TODO: Filter this menu with IsViewModeSupportedDelegate (see further down in this file) and remove the
 			// "Exposure" section.
diff --git a/Engine/Source/Editor/LevelEditor/Private/SLevelViewport.cpp b/Engine/Source/Editor/LevelEditor/Private/SLevelViewport.cpp
--- a/Engine/Source/Editor/LevelEditor/Private/SLevelViewport.cpp
+++ b/Engine/Source/Editor/LevelEditor/Private/SLevelViewport.cpp
@@ -2013,7 +2013,31 @@ TSharedPtr<SWidget> SLevelViewport::BuildViewportToolbar()
 
 				// Add the "Camera" submenu.
 				{
-					FToolMenuEntry CameraSubmenu = UE::LevelEditor::CreateLevelEditorViewportToolbarCameraSubmenu();
+
+					const FName GrandParentSubmenuName = "UnrealEd.ViewportToolbar.Camera";
+					const FName ParentSubmenuName = "LevelEditor.ViewportToolbar.Camera";
+					const FName SubmenuName = "LevelEditor.ViewportToolbar.CameraOptions";
+
+					// Create our grandparent menu.
+					if (!UToolMenus::Get()->IsMenuRegistered(GrandParentSubmenuName))
+					{
+						UToolMenus::Get()->RegisterMenu(GrandParentSubmenuName);
+					}
+
+					// Create our parent menu.
+					if (!UToolMenus::Get()->IsMenuRegistered(ParentSubmenuName))
+					{
+						UToolMenus::Get()->RegisterMenu(ParentSubmenuName, GrandParentSubmenuName);
+					}
+
+					// Create our menu.
+					UToolMenus::Get()->RegisterMenu(SubmenuName, ParentSubmenuName);
+
+					// Extending using both Level Editor specific and UnrealEd generic entries
+					UE::LevelEditor::ExtendCameraSubmenu(SubmenuName);
+					UE::UnrealEd::ExtendCameraSubmenu(SubmenuName);
+
+					FToolMenuEntry CameraSubmenu = UE::UnrealEd::CreateViewportToolbarCameraSubmenu();
 					CameraSubmenu.InsertPosition.Position = EToolMenuInsertType::First;
 					RightSection.AddEntry(CameraSubmenu);
 				}
diff --git a/Engine/Source/Editor/LevelEditor/Private/ViewportToolbar/LevelEditorViewportToolbarSections.cpp b/Engine/Source/Editor/LevelEditor/Private/ViewportToolbar/LevelEditorViewportToolbarSections.cpp
--- a/Engine/Source/Editor/LevelEditor/Private/ViewportToolbar/LevelEditorViewportToolbarSections.cpp
+++ b/Engine/Source/Editor/LevelEditor/Private/ViewportToolbar/LevelEditorViewportToolbarSections.cpp
@@ -1999,53 +1999,23 @@ void CreateCameraSpeedMenu(UToolMenu* InMenu)
 	Section.AddEntry(CreateCameraSpeedScalarSlider(LevelViewport));
 }
 
-FToolMenuEntry CreateLevelEditorViewportToolbarCameraSubmenu()
+void ExtendCameraSubmenu(FName InCameraOptionsSubmenuName)
 {
-	return FToolMenuEntry::InitSubMenu(
-		"Camera",
-		LOCTEXT("CameraSubmenuLabel", "Camera"),
-		LOCTEXT("CameraSubmenuTooltip", "Viewport-related Camera settings"),
+	UToolMenu* const Submenu = UToolMenus::Get()->ExtendMenu(InCameraOptionsSubmenuName);
+
+	Submenu->AddDynamicSection(
+		"LevelEditorCameraExtensionDynamicSection",
 		FNewToolMenuDelegate::CreateLambda(
-			[](UToolMenu* Submenu) -> void
+			[](UToolMenu* InDynamicMenu)
 			{
-				TWeakPtr<::SLevelViewport> LevelViewportWeak = Submenu->FindContext<ULevelViewportContext>()->LevelViewport;
+				TWeakPtr<::SLevelViewport> LevelViewportWeak = InDynamicMenu->FindContext<ULevelViewportContext>()->LevelViewport;
 
 				// TODO:
 				// add Select Active Camera
 
-				// Perspective Section
-				{
-					FToolMenuSection& PerspectiveSection =
-						Submenu->FindOrAddSection("Perspective", LOCTEXT("PerspectiveLabel", "Perspective"));
-
-					PerspectiveSection.AddMenuEntry(FEditorViewportCommands::Get().Perspective);
-
-					// TODO: show separator based on actual list of cameras
-					// PerspectiveSection.AddSeparator("PerspectiveSeparator");
-					// Camera list
-					// Camera types
-				}
-
-				// Orthographic Section
-				{
-					FToolMenuSection& OrthographicSection =
-						Submenu->FindOrAddSection("Orthographic", LOCTEXT("OrthographicLabel", "Orthographic"));
-					OrthographicSection.AddMenuEntry(FEditorViewportCommands::Get().Top);
-					OrthographicSection.AddMenuEntry(FEditorViewportCommands::Get().Bottom);
-					OrthographicSection.AddMenuEntry(FEditorViewportCommands::Get().Left);
-					OrthographicSection.AddMenuEntry(FEditorViewportCommands::Get().Right);
-					OrthographicSection.AddMenuEntry(FEditorViewportCommands::Get().Front);
-					OrthographicSection.AddMenuEntry(FEditorViewportCommands::Get().Back);
-
-					OrthographicSection.AddSeparator("PerspectiveSeparator");
-
-					OrthographicSection.AddEntry(CreateFOVMenu(LevelViewportWeak));
-					OrthographicSection.AddEntry(CreateFarViewPlaneMenu(LevelViewportWeak));
-				}
-
 				// Create Section
 				{
-					FToolMenuSection& CreateSection = Submenu->FindOrAddSection("Create", LOCTEXT("CreateLabel", "Create"));
+					FToolMenuSection& CreateSection = InDynamicMenu->FindOrAddSection("Create", LOCTEXT("CreateLabel", "Create"));
 
 					CreateSection.AddSubMenu(
 						"CreateCamera",
@@ -2079,7 +2049,7 @@ FToolMenuEntry CreateLevelEditorViewportToolbarCameraSubmenu()
 				// Positioning Section
 				{
 					FToolMenuSection& PositioningSection =
-						Submenu->FindOrAddSection("Positioning", LOCTEXT("PositioningLabel", "Positioning"));
+						InDynamicMenu->FindOrAddSection("Positioning", LOCTEXT("PositioningLabel", "Positioning"));
 
 					// Camera Speed Submenu
 					{
@@ -2147,7 +2117,7 @@ FToolMenuEntry CreateLevelEditorViewportToolbarCameraSubmenu()
 				// Options Section
 				{
 					FToolMenuSection& OptionsSection =
-						Submenu->FindOrAddSection("CameraOptions", LOCTEXT("OptionsLabel", "Options"));
+						InDynamicMenu->FindOrAddSection("CameraOptions", LOCTEXT("OptionsLabel", "Options"));
 					// add Cinematic Viewport
 					// add Allow Cinematic Control
 					// add Game View
@@ -2165,7 +2135,7 @@ FToolMenuEntry CreateLevelEditorViewportToolbarCameraSubmenu()
 					// This additional options section is used to force certain elements to appear after extensions
 					{
 						FToolMenuSection& AdditionalOptions =
-							Submenu->FindOrAddSection("AdditionalOptions", LOCTEXT("AdditionalOptionsLabel", ""));
+							InDynamicMenu->FindOrAddSection("AdditionalOptions", LOCTEXT("AdditionalOptionsLabel", ""));
 						AdditionalOptions.AddSeparator("AdditionalOptionsSeparator");
 
 						FToolMenuEntry HighResolutionScreenshot =
diff --git a/Engine/Source/Editor/LevelEditor/Private/ViewportToolbar/LevelEditorViewportToolbarSections.h b/Engine/Source/Editor/LevelEditor/Private/ViewportToolbar/LevelEditorViewportToolbarSections.h
--- a/Engine/Source/Editor/LevelEditor/Private/ViewportToolbar/LevelEditorViewportToolbarSections.h
+++ b/Engine/Source/Editor/LevelEditor/Private/ViewportToolbar/LevelEditorViewportToolbarSections.h
@@ -47,6 +47,5 @@ void GenerateViewportLayoutsMenu(UToolMenu* InMenu);
 TSharedRef<SWidget> BuildVolumeControlCustomWidget();
 FToolMenuEntry CreateLevelEditorViewportToolbarSettingsSubmenu();
 
-FToolMenuEntry CreateLevelEditorViewportToolbarCameraSubmenu();
-
+void ExtendCameraSubmenu(FName InCameraOptionsSubmenuName);
 } // namespace UE::LevelEditor
diff --git a/Engine/Source/Editor/MaterialEditor/Private/SMaterialEditorViewport.cpp b/Engine/Source/Editor/MaterialEditor/Private/SMaterialEditorViewport.cpp
--- a/Engine/Source/Editor/MaterialEditor/Private/SMaterialEditorViewport.cpp
+++ b/Engine/Source/Editor/MaterialEditor/Private/SMaterialEditorViewport.cpp
@@ -963,7 +963,28 @@ TSharedPtr<SWidget> SMaterialEditor3DPreviewViewport::BuildViewportToolbar()
 
 				// Add the "Camera" submenu.
 				{
-					FToolMenuEntry CameraSubmenu = UE::UnrealEd::CreateCameraSubmenu();
+					// Create our grandparent menu.
+					if (!UToolMenus::Get()->IsMenuRegistered("UnrealEd.ViewportToolbar.Camera"))
+					{
+						UToolMenus::Get()->RegisterMenu("UnrealEd.ViewportToolbar.Camera");
+					}
+
+					// Create our parent menu.
+					if (!UToolMenus::Get()->IsMenuRegistered("MaterialEditor.ViewportToolbar.Camera"))
+					{
+						UToolMenus::Get()->RegisterMenu(
+							"MaterialEditor.ViewportToolbar.Camera", "UnrealEd.ViewportToolbar.Camera"
+						);
+					}
+
+					// Create our menu.
+					UToolMenus::Get()->RegisterMenu(
+						"MaterialEditor.ViewportToolbar.CameraOptions", "MaterialEditor.ViewportToolbar.Camera"
+					);
+
+					UE::UnrealEd::ExtendCameraSubmenu("MaterialEditor.ViewportToolbar.CameraOptions");
+
+					FToolMenuEntry CameraSubmenu = UE::UnrealEd::CreateViewportToolbarCameraSubmenu();
 					CameraSubmenu.InsertPosition.Position = EToolMenuInsertType::First;
 					RightSection.AddEntry(CameraSubmenu);
 				}
diff --git a/Engine/Source/Editor/StaticMeshEditor/Private/SStaticMeshEditorViewport.cpp b/Engine/Source/Editor/StaticMeshEditor/Private/SStaticMeshEditorViewport.cpp
--- a/Engine/Source/Editor/StaticMeshEditor/Private/SStaticMeshEditorViewport.cpp
+++ b/Engine/Source/Editor/StaticMeshEditor/Private/SStaticMeshEditorViewport.cpp
@@ -819,6 +819,34 @@ TSharedPtr<SWidget> SStaticMeshEditorViewport::BuildViewportToolbar()
 			FToolMenuSection& RightSection = ViewportToolbarMenu->FindOrAddSection("Right");
 			RightSection.Alignment = EToolMenuSectionAlign::Last;
 
+			// Add the "Camera" submenu.
+			{
+				const FName GrandParentSubmenuName = "UnrealEd.ViewportToolbar.Camera";
+				const FName ParentSubmenuName = "StaticMeshEditor.ViewportToolbar.Camera";
+				const FName SubmenuName = "StaticMeshEditor.ViewportToolbar.CameraOptions";
+
+				// Create our grandparent menu.
+				if (!UToolMenus::Get()->IsMenuRegistered(GrandParentSubmenuName))
+				{
+					UToolMenus::Get()->RegisterMenu(GrandParentSubmenuName);
+				}
+
+				// Create our parent menu.
+				if (!UToolMenus::Get()->IsMenuRegistered(ParentSubmenuName))
+				{
+					UToolMenus::Get()->RegisterMenu(ParentSubmenuName, GrandParentSubmenuName);
+				}
+
+				// Create our menu.
+				UToolMenus::Get()->RegisterMenu(SubmenuName, ParentSubmenuName);
+
+				UE::UnrealEd::ExtendCameraSubmenu(SubmenuName);
+
+				FToolMenuEntry CameraSubmenu = UE::UnrealEd::CreateViewportToolbarCameraSubmenu();
+				CameraSubmenu.InsertPosition.Position = EToolMenuInsertType::First;
+				RightSection.AddEntry(CameraSubmenu);
+			}
+
 			// Add the "View Modes" sub menu.
 			{
 				// Stay backward-compatible with the old viewport toolbar.
diff --git a/Engine/Source/Editor/UnrealEd/Private/ViewportToolbar/UnrealEdViewportToolbar.cpp b/Engine/Source/Editor/UnrealEd/Private/ViewportToolbar/UnrealEdViewportToolbar.cpp
--- a/Engine/Source/Editor/UnrealEd/Private/ViewportToolbar/UnrealEdViewportToolbar.cpp
+++ b/Engine/Source/Editor/UnrealEd/Private/ViewportToolbar/UnrealEdViewportToolbar.cpp
@@ -23,6 +23,7 @@
 #include "ToolMenu.h"
 #include "ToolMenuEntry.h"
 #include "ToolMenuSection.h"
+#include "ToolMenus.h"
 #include "ViewportToolbar/UnrealEdViewportToolbarContext.h"
 #include "Widgets/Input/SNumericEntryBox.h"
 #include "Widgets/Input/SSpinBox.h"
@@ -1663,7 +1664,7 @@ TSharedRef<SWidget> CreateFOVMenuWidget(const TSharedRef<SEditorViewport>& InVie
 	constexpr float FOVMin = 5.0f;
 	constexpr float FOVMax = 170.0f;
 
-	TSharedPtr<FEditorViewportClient> ViewportClient = InViewport->GetViewportClient();
+	TWeakPtr<FEditorViewportClient> ViewportClientWeak = InViewport->GetViewportClient();
 
 	return
 		// clang-format off
@@ -1683,15 +1684,22 @@ TSharedRef<SWidget> CreateFOVMenuWidget(const TSharedRef<SEditorViewport>& InVie
 					.Font( FAppStyle::GetFontStyle(TEXT("MenuItem.Font")))
 					.MinValue(FOVMin)
 					.MaxValue(FOVMax)
-					.Value_Lambda([ViewportClient]()
+					.Value_Lambda([ViewportClientWeak]()
 					{
-						return ViewportClient->ViewFOV;
+						if (TSharedPtr<FEditorViewportClient> ViewportClient = ViewportClientWeak.Pin())
+						{
+							return ViewportClient->ViewFOV;
+						}
+						return 90.0f;
 					})
-					.OnValueChanged_Lambda([ViewportClient](float InNewValue)
+					.OnValueChanged_Lambda([ViewportClientWeak](float InNewValue)
 					{
-						ViewportClient->FOVAngle = InNewValue;
-						ViewportClient->ViewFOV = InNewValue;
-						ViewportClient->Invalidate();
+						if (TSharedPtr<FEditorViewportClient> ViewportClient = ViewportClientWeak.Pin())
+						{
+							ViewportClient->FOVAngle = InNewValue;
+							ViewportClient->ViewFOV = InNewValue;
+							ViewportClient->Invalidate();
+						}
 					})
 				]
 			]
@@ -1701,7 +1709,7 @@ TSharedRef<SWidget> CreateFOVMenuWidget(const TSharedRef<SEditorViewport>& InVie
 
 TSharedRef<SWidget> CreateFarViewPlaneMenuWidget(const TSharedRef<SEditorViewport>& InViewport)
 {
-	TSharedPtr<FEditorViewportClient> ViewportClient = InViewport->GetViewportClient();
+	TWeakPtr<FEditorViewportClient> ViewportClientWeak = InViewport->GetViewportClient();
 
 	return
 		// clang-format off
@@ -1722,14 +1730,22 @@ TSharedRef<SWidget> CreateFarViewPlaneMenuWidget(const TSharedRef<SEditorViewpor
 					.MinValue(0.0f)
 					.MaxValue(100000.0f)
 					.Font(FAppStyle::GetFontStyle(TEXT("MenuItem.Font")))
-					.Value_Lambda([ViewportClient]()
+					.Value_Lambda([ViewportClientWeak]()
 					{
-						return ViewportClient->GetFarClipPlaneOverride();
+						if (TSharedPtr<FEditorViewportClient> ViewportClient = ViewportClientWeak.Pin())
+						{
+							return ViewportClient->GetFarClipPlaneOverride();
+						}
+
+						return 100000.0f;
 					})
-					.OnValueChanged_Lambda([ViewportClient](float InNewValue)
+					.OnValueChanged_Lambda([ViewportClientWeak](float InNewValue)
 					{
-						ViewportClient->OverrideFarClipPlane(InNewValue);
-						ViewportClient->Invalidate();
+						if (TSharedPtr<FEditorViewportClient> ViewportClient = ViewportClientWeak.Pin())
+						{
+							ViewportClient->OverrideFarClipPlane(InNewValue);
+							ViewportClient->Invalidate();
+						}
 					})
 				]
 			]
@@ -1737,18 +1753,61 @@ TSharedRef<SWidget> CreateFarViewPlaneMenuWidget(const TSharedRef<SEditorViewpor
 	// clang-format on
 }
 
-FToolMenuEntry CreateCameraSubmenu()
+FToolMenuEntry CreateViewportToolbarCameraSubmenu()
 {
-	return FToolMenuEntry::InitSubMenu(
-		"Camera",
-		LOCTEXT("CameraSubmenuLabel", "Camera"),
-		LOCTEXT("CameraSubmenuTooltip", "Camera options"),
+	return FToolMenuEntry::InitDynamicEntry(
+		"DynamicCameraOptions",
+		FNewToolMenuSectionDelegate::CreateLambda(
+			[](FToolMenuSection& InDynamicSection) -> void
+			{
+				InDynamicSection.AddSubMenu(
+					"CameraOptions",
+					LOCTEXT("CameraSubmenuLabel", "Camera"),
+					LOCTEXT("CameraSubmenuTooltip", "Camera options"),
+					FNewToolMenuDelegate::CreateLambda(
+						[](UToolMenu* Submenu) -> void
+						{
+							PopulateCameraMenu(Submenu);
+						}
+					)
+				);
+			}
+		)
+	);
+}
+
+void PopulateCameraMenu(UToolMenu* InMenu)
+{
+	UUnrealEdViewportToolbarContext* const EditorViewportContext = InMenu->FindContext<UUnrealEdViewportToolbarContext>();
+	if (!EditorViewportContext)
+	{
+		return;
+	}
+
+	const TSharedPtr<SEditorViewport> EditorViewport = EditorViewportContext->Viewport.Pin();
+	if (!EditorViewport)
+	{
+		return;
+	}
+
+	FToolMenuSection& UnnamedSection = InMenu->FindOrAddSection("", LOCTEXT("UnnamedLabel", ""));
+
+	UnnamedSection.AddEntry(FToolMenuEntry::InitWidget(
+		"CameraMenuItems", UE::UnrealEd::CreateCameraMenuWidget(EditorViewport.ToSharedRef()), FText(), true
+	));
+}
+
+void ExtendCameraSubmenu(FName InCameraOptionsSubmenuName)
+{
+	UToolMenu* const Submenu = UToolMenus::Get()->ExtendMenu(InCameraOptionsSubmenuName);
+
+	Submenu->AddDynamicSection(
+		"EditorCameraExtensionDynamicSection",
 		FNewToolMenuDelegate::CreateLambda(
-			[](UToolMenu* Submenu) -> void
+			[](UToolMenu* InDynamicMenu)
 			{
 				UUnrealEdViewportToolbarContext* const EditorViewportContext =
-					Submenu->FindContext<UUnrealEdViewportToolbarContext>();
-
+					InDynamicMenu->FindContext<UUnrealEdViewportToolbarContext>();
 				if (!EditorViewportContext)
 				{
 					return;
@@ -1760,12 +1819,10 @@ FToolMenuEntry CreateCameraSubmenu()
 					return;
 				}
 
-				FToolMenuSection& UnnamedSection = Submenu->FindOrAddSection("", LOCTEXT("UnnamedLabel", ""));
-
-				UnnamedSection.AddEntry(FToolMenuEntry::InitWidget(
-					"CameraMenuItems", UE::UnrealEd::CreateCameraMenuWidget(EditorViewport.ToSharedRef()), FText(), true
-				));
+				FToolMenuInsert InsertPosition("CameraOptions", EToolMenuInsertType::After);
 
+				FToolMenuSection& UnnamedSection =
+					InDynamicMenu->FindOrAddSection("", LOCTEXT("UnnamedLabel", ""), InsertPosition);
 				UnnamedSection.AddSeparator("CameraSubmenuSeparator");
 
 				UnnamedSection.AddEntry(FToolMenuEntry::InitWidget(
diff --git a/Engine/Source/Editor/UnrealEd/Public/ViewportToolbar/UnrealEdViewportToolbar.h b/Engine/Source/Editor/UnrealEd/Public/ViewportToolbar/UnrealEdViewportToolbar.h
--- a/Engine/Source/Editor/UnrealEd/Public/ViewportToolbar/UnrealEdViewportToolbar.h
+++ b/Engine/Source/Editor/UnrealEd/Public/ViewportToolbar/UnrealEdViewportToolbar.h
@@ -115,7 +115,14 @@ UNREALED_API FToolMenuEntry CreateNumericEntry(
 	int32 InMaxFractionalDigits = 2
 );
 
-UNREALED_API FToolMenuEntry CreateCameraSubmenu();
+UNREALED_API FToolMenuEntry CreateViewportToolbarCameraSubmenu();
+
+UNREALED_API void PopulateCameraMenu(UToolMenu* InMenu);
+
+/**
+ * Adds Field of View and Far View Plane entries to the specified Camera Submenu
+ */
+UNREALED_API void ExtendCameraSubmenu(FName InCameraOptionsSubmenuName);
 
 UNREALED_API FToolMenuEntry CreatePerformanceAndScalabilitySubmenu();
 
