Post Process Material:  Fix for uninitialized property warning for FMaterialEditorPostProcessOverrides structure.  Structure is only present in UMaterialEditorInstanceConstant, which fully initialized the structure in its constructor, but this still triggers validation warnings.  Moved initialization to the struct declaration.

#rnx

diff --git a/Engine/Source/Editor/UnrealEd/Classes/MaterialEditor/MaterialEditorInstanceConstant.h b/Engine/Source/Editor/UnrealEd/Classes/MaterialEditor/MaterialEditorInstanceConstant.h
--- a/Engine/Source/Editor/UnrealEd/Classes/MaterialEditor/MaterialEditorInstanceConstant.h
+++ b/Engine/Source/Editor/UnrealEd/Classes/MaterialEditor/MaterialEditorInstanceConstant.h
@@ -237,19 +237,19 @@ struct FMaterialEditorPostProcessOverrides
 
 	// Tracks if this is a material where post process overrides can be applied (MaterialDomain == MD_PostProcess, BlendableLocation != BL_ReplacingTonemapper)
 	UPROPERTY()
-	bool bIsOverrideable;
+	bool bIsOverrideable = false;
 
 	UPROPERTY(EditAnywhere, Category = PostProcessOverrideValue)
-	bool bOverrideBlendableLocation;
+	bool bOverrideBlendableLocation = false;
 
 	UPROPERTY(EditAnywhere, Category = PostProcessOverrideValue)
-	bool bOverrideBlendablePriority;
+	bool bOverrideBlendablePriority = false;
 
 	UPROPERTY(EditAnywhere, Category = PostProcessOverrideValue, meta = (DisplayName = "Blendable Location"), meta = (InvalidEnumValues = "BL_ReplacingTonemapper"))
-	TEnumAsByte<EBlendableLocation> BlendableLocationOverride;
+	TEnumAsByte<EBlendableLocation> BlendableLocationOverride { BL_SceneColorAfterTonemapping };
 
 	UPROPERTY(EditAnywhere, Category = PostProcessOverrideValue, meta = (DisplayName = "Blendable Priority"))
-	int32 BlendablePriorityOverride;
+	int32 BlendablePriorityOverride = 0;
 
 	/** Overrides for user scene texture inputs */
 	UPROPERTY(EditAnywhere, editfixedsize, Category = PostProcessOverrideValue)
diff --git a/Engine/Source/Editor/UnrealEd/Private/PreviewMaterial.cpp b/Engine/Source/Editor/UnrealEd/Private/PreviewMaterial.cpp
--- a/Engine/Source/Editor/UnrealEd/Private/PreviewMaterial.cpp
+++ b/Engine/Source/Editor/UnrealEd/Private/PreviewMaterial.cpp
@@ -477,12 +477,6 @@ UMaterialEditorInstanceConstant::UMaterialEditorInstanceConstant(const FObjectIn
 
 	// Default to override with nothing on MIC (don't inherit parent setting).
 	bNaniteOverride = true;
-
-	PostProcessOverrides.bIsOverrideable = false;
-	PostProcessOverrides.bOverrideBlendableLocation = false;
-	PostProcessOverrides.bOverrideBlendablePriority = false;
-	PostProcessOverrides.BlendableLocationOverride = BL_SceneColorAfterTonemapping;
-	PostProcessOverrides.BlendablePriorityOverride = 0;
 }
 
 void UMaterialEditorInstanceConstant::PostEditChangeProperty(FPropertyChangedEvent& PropertyChangedEvent)
