Removing the usage of "Friendly Names" from the Functions and (Event Graphs/Overrides Menu).  This reverts a previous change which introduced consistency to the Friendly Names feature in those windows.  The reasoning is that there is a long-term goal of removing Friendly Names entirely from the Editor.  It interferes with the users being able to properly name their Functions/Events and have that name consistent through the interface.  I have kept the Friendly Names when *calling* a function in Blueprints as well as the Blueprints Action list.  The reasoning behind that is the nodes are very hard to decipher when zoomed out on the graph unless there are spaces between the names.

diff --git a/Engine/Source/Editor/BlueprintGraph/Classes/K2Node_CallFunction.h b/Engine/Source/Editor/BlueprintGraph/Classes/K2Node_CallFunction.h
--- a/Engine/Source/Editor/BlueprintGraph/Classes/K2Node_CallFunction.h
+++ b/Engine/Source/Editor/BlueprintGraph/Classes/K2Node_CallFunction.h
@@ -181,7 +181,7 @@ public:
 	virtual void PostParameterPinCreated(UEdGraphPin *Pin) {}
 
 	UE_DEPRECATED(5.5, "Moved to ObjectTools::GetUserFacingFunctionName.")
-	static FText GetUserFacingFunctionName(const UFunction* Function);
+	static FText GetUserFacingFunctionName(const UFunction* Function, ENodeTitleType::Type NodeTitleType = ENodeTitleType::EditableTitle);
 
 	/** Set up a pins tooltip from a function's tooltip */
 	static void GeneratePinTooltipFromFunction(UEdGraphPin& Pin, const UFunction* Function);
diff --git a/Engine/Source/Editor/BlueprintGraph/Private/BlueprintFunctionNodeSpawner.cpp b/Engine/Source/Editor/BlueprintGraph/Private/BlueprintFunctionNodeSpawner.cpp
--- a/Engine/Source/Editor/BlueprintGraph/Private/BlueprintFunctionNodeSpawner.cpp
+++ b/Engine/Source/Editor/BlueprintGraph/Private/BlueprintFunctionNodeSpawner.cpp
@@ -303,7 +303,8 @@ UBlueprintFunctionNodeSpawner* UBlueprintFunctionNodeSpawner::Create(TSubclassOf
 	}
 	else
 	{
-		MenuSignature.MenuName = ObjectTools::GetUserFacingFunctionName(Function);
+		constexpr bool bAllowFriendlyNames = true;
+		MenuSignature.MenuName = ObjectTools::GetUserFacingFunctionName(Function, bAllowFriendlyNames);
 		MenuSignature.Category = UK2Node_CallFunction::GetDefaultCategoryForFunction(Function, FText::GetEmpty());
 		MenuSignature.Tooltip = FText::FromString(ObjectTools::GetDefaultTooltipForFunction(Function));
 		// add at least one character, so that PrimeDefaultUiSpec() doesn't attempt to query the template node
diff --git a/Engine/Source/Editor/BlueprintGraph/Private/EdGraphSchema_K2.cpp b/Engine/Source/Editor/BlueprintGraph/Private/EdGraphSchema_K2.cpp
--- a/Engine/Source/Editor/BlueprintGraph/Private/EdGraphSchema_K2.cpp
+++ b/Engine/Source/Editor/BlueprintGraph/Private/EdGraphSchema_K2.cpp
@@ -5683,7 +5683,6 @@ void UEdGraphSchema_K2::GetGraphDisplayInformation(const UEdGraph& Graph, /*out*
 		if (GraphType == GT_Function && Function)
 		{
 			DisplayInfo.DisplayName = GetFriendlySignatureName(Function);
-			DisplayInfo.Tooltip = FText::FormatOrdered(LOCTEXT("Tooltip_With_RealFunctionName", "{0}\nFunction Name is {1}"), DisplayInfo.Tooltip, DisplayInfo.PlainName);
 		}
 		else
 		{
diff --git a/Engine/Source/Editor/BlueprintGraph/Private/K2Node_CallFunction.cpp b/Engine/Source/Editor/BlueprintGraph/Private/K2Node_CallFunction.cpp
--- a/Engine/Source/Editor/BlueprintGraph/Private/K2Node_CallFunction.cpp
+++ b/Engine/Source/Editor/BlueprintGraph/Private/K2Node_CallFunction.cpp
@@ -610,8 +610,9 @@ FText UK2Node_CallFunction::GetNodeTitle(ENodeTitleType::Type TitleType) const
 
 	if (UFunction* Function = GetTargetFunction())
 	{
+		const bool bAllowFriendlyNames = (TitleType == ENodeTitleType::FullTitle);
 		RPCString = UK2Node_Event::GetLocalizedNetString(Function->FunctionFlags, true);
-		FunctionName = ObjectTools::GetUserFacingFunctionName(Function);
+		FunctionName = ObjectTools::GetUserFacingFunctionName(Function, bAllowFriendlyNames);
 		ContextString = GetFunctionContextString();
 	}
 	else
@@ -1797,9 +1798,10 @@ void UK2Node_CallFunction::GeneratePinTooltipFromFunction(UEdGraphPin& Pin, cons
 	GetDefault<UEdGraphSchema_K2>()->ConstructBasicPinTooltip(Pin, FText::FromString(Pin.PinToolTip), Pin.PinToolTip);
 }
 
-FText UK2Node_CallFunction::GetUserFacingFunctionName(const UFunction* Function)
+FText UK2Node_CallFunction::GetUserFacingFunctionName(const UFunction* Function, ENodeTitleType::Type NodeTitleType /*= ENodeTitleType::EditableTitle*/)
 {
-	return ObjectTools::GetUserFacingFunctionName(Function);
+	const bool bAllowFriendlyNames = NodeTitleType == ENodeTitleType::FullTitle;
+	return ObjectTools::GetUserFacingFunctionName(Function, bAllowFriendlyNames);
 }
 
 FString UK2Node_CallFunction::GetDefaultTooltipForFunction(const UFunction* Function)
diff --git a/Engine/Source/Editor/UnrealEd/Private/ObjectTools.cpp b/Engine/Source/Editor/UnrealEd/Private/ObjectTools.cpp
--- a/Engine/Source/Editor/UnrealEd/Private/ObjectTools.cpp
+++ b/Engine/Source/Editor/UnrealEd/Private/ObjectTools.cpp
@@ -5035,17 +5035,27 @@ namespace ObjectTools
 		}
 	}
 
-	FText GetUserFacingFunctionName(const UFunction* Function)
+	FText GetUserFacingFunctionName(const UFunction* Function, bool bAllowFriendlyNames)
 	{
 		FText ReturnDisplayName;
 
 		if (Function != nullptr)
 		{
-			if (GEditor && GetDefault<UEditorStyleSettings>()->bShowFriendlyNames)
+			static const FName NAME_DisplayName { TEXT("DisplayName") };
+
+			// Functions do not use friendly names because they can be manually input by a user in the editor (and it would otherwise not adhere to their name)
+			// There is a long-term goal of removing friendly names from the Engine.  However, we keep them in the case of FullTitle as it helps the nodes
+			// be decipherable in a zoomed-out view.
+			if (GEditor && bAllowFriendlyNames && GetDefault<UEditorStyleSettings>()->bShowFriendlyNames)
 			{
 				ReturnDisplayName = Function->GetDisplayNameText();
 			}
-			else
+			else if (const FString* OverrideDisplayName = Function->FindMetaData(NAME_DisplayName))
+			{
+				ReturnDisplayName = FText::FromString(*OverrideDisplayName);
+			}
+
+			if (ReturnDisplayName.IsEmpty())
 			{
 				// Previous (and similar) code paths would go through FField::GetMetaDataText(DisplayName) which attempts localization
 				// However, we do not localize function names (and we've explicitly requested non-friendly names), so just show us the real name
diff --git a/Engine/Source/Editor/UnrealEd/Public/ObjectTools.h b/Engine/Source/Editor/UnrealEd/Public/ObjectTools.h
--- a/Engine/Source/Editor/UnrealEd/Public/ObjectTools.h
+++ b/Engine/Source/Editor/UnrealEd/Public/ObjectTools.h
@@ -678,10 +678,11 @@ namespace ObjectTools
 
 	/**
 	 * Gets the user-facing name for the given UFunction.
-	 * @param Function	the function to find the name for
+	 * @param Function				the function to find the name for
+	 * @param bAllowFriendlyNames	set to true if you will allow a FriendlyName (name with spaces which does not match the Function's Name)
 	 * @return the name as localized text
 	 */
-	UNREALED_API FText GetUserFacingFunctionName(const UFunction* Function);
+	UNREALED_API FText GetUserFacingFunctionName(const UFunction* Function, bool bAllowFriendlyNames = false);
 
 
 	/**
