Fix for ThumbnailManager singleton shutdown crash, due to member SharedThumbnailPool elements trying to access the singleton in the destructor

diff --git a/Engine/Source/Editor/UnrealEd/Classes/ThumbnailRendering/ThumbnailManager.h b/Engine/Source/Editor/UnrealEd/Classes/ThumbnailRendering/ThumbnailManager.h
--- a/Engine/Source/Editor/UnrealEd/Classes/ThumbnailRendering/ThumbnailManager.h
+++ b/Engine/Source/Editor/UnrealEd/Classes/ThumbnailRendering/ThumbnailManager.h
@@ -228,6 +228,8 @@ protected:
 	 */
 	void InitializeRenderTypeArray(TArray<struct FThumbnailRenderingInfo>& ThumbnailRendererTypes);
 
+	virtual void BeginDestroy() override;
+
 private:
 	/** Initialize the checkerboard texture for texture thumbnails */
 	void SetupCheckerboardTexture();
diff --git a/Engine/Source/Editor/UnrealEd/Private/ThumbnailManager.cpp b/Engine/Source/Editor/UnrealEd/Private/ThumbnailManager.cpp
--- a/Engine/Source/Editor/UnrealEd/Private/ThumbnailManager.cpp
+++ b/Engine/Source/Editor/UnrealEd/Private/ThumbnailManager.cpp
@@ -122,6 +122,13 @@ void UThumbnailManager::InitializeRenderTypeArray(TArray<FThumbnailRenderingInfo
 	}
 }
 
+void UThumbnailManager::BeginDestroy()
+{
+	Super::BeginDestroy();
+
+	SharedThumbnailPool.Reset();
+}
+
 FThumbnailRenderingInfo* UThumbnailManager::GetRenderingInfo(UObject* Object)
 {
 	// If something may have been GCed, empty the map so we don't crash
