[Viewport Toolbar] Sync with design and implement "Settings" submenu for Static Mesh editor's new viewport toolbar

Currently choosing to implement the menu items creation differently with respect to the existing legacy menu, which is implemented as a ComboBox, which also looks different from the current style of the new toolbar.

Maintaining both items (one a Widget, the other one a list of entries) with the same code allowing for them to look different based on their context doesn't seem too convenient at this stage.

New implementation is generating checkbox items based on a list provided by the IPreviewProfileController being used (similarly to what other menus do).

diff --git a/Engine/Source/Editor/StaticMeshEditor/Private/SStaticMeshEditorViewport.cpp b/Engine/Source/Editor/StaticMeshEditor/Private/SStaticMeshEditorViewport.cpp
--- a/Engine/Source/Editor/StaticMeshEditor/Private/SStaticMeshEditorViewport.cpp
+++ b/Engine/Source/Editor/StaticMeshEditor/Private/SStaticMeshEditorViewport.cpp
@@ -20,6 +20,7 @@
 #include "Engine/StaticMeshSocket.h"
 #include "SEditorViewportToolBarMenu.h"
 #include "Editor.h"
+#include "PreviewProfileController.h"
 #include "ToolMenus.h"
 #include "ViewportToolbar/UnrealEdViewportToolbar.h"
 #include "Widgets/Text/SRichTextBlock.h"
@@ -865,11 +866,23 @@ TSharedPtr<SWidget> SStaticMeshEditorViewport::BuildViewportToolbar()
 				FToolMenuEntry ViewModesSubmenu = UE::UnrealEd::CreateViewportToolbarViewModesSubmenu();
 				ViewModesSubmenu.InsertPosition.Position = EToolMenuInsertType::Last;
 				RightSection.AddEntry(ViewModesSubmenu);
+			}
 
+			// Add the "Performance and Scalability" sub menu.
+			{
 				FToolMenuEntry PerformanceAndScalabilitySubmenu = UE::UnrealEd::CreatePerformanceAndScalabilitySubmenu();
 				PerformanceAndScalabilitySubmenu.InsertPosition.Position = EToolMenuInsertType::Last;
 				RightSection.AddEntry(PerformanceAndScalabilitySubmenu);
 			}
+
+			// Add the "Preview Profile" sub menu.
+			{
+				PreviewProfileController = MakeShared<FPreviewProfileController>();
+				FToolMenuEntry PreviewProfileSubmenu =
+					UE::UnrealEd::CreateViewportToolbarAssetViewerProfileSubmenu(PreviewProfileController);
+				PreviewProfileSubmenu.InsertPosition.Position = EToolMenuInsertType::Last;
+				RightSection.AddEntry(PreviewProfileSubmenu);
+			}
 		}
 	}
 
diff --git a/Engine/Source/Editor/StaticMeshEditor/Private/SStaticMeshEditorViewport.h b/Engine/Source/Editor/StaticMeshEditor/Private/SStaticMeshEditorViewport.h
--- a/Engine/Source/Editor/StaticMeshEditor/Private/SStaticMeshEditorViewport.h
+++ b/Engine/Source/Editor/StaticMeshEditor/Private/SStaticMeshEditorViewport.h
@@ -15,6 +15,7 @@
 #include "SAssetEditorViewport.h"
 #include "Styling/StyleColors.h"
 
+class FPreviewProfileController;
 class IStaticMeshEditor;
 class SVerticalBox;
 class UStaticMesh;
@@ -211,6 +212,8 @@ private:
 	/** Editor viewport client */
 	TSharedPtr<class FStaticMeshEditorViewportClient> EditorViewportClient;
 
+	TSharedPtr<FPreviewProfileController> PreviewProfileController;
+
 	/** Static mesh being edited */
 	TObjectPtr<UStaticMesh> StaticMesh;
 
diff --git a/Engine/Source/Editor/UnrealEd/Private/ViewportToolbar/UnrealEdViewportToolbar.cpp b/Engine/Source/Editor/UnrealEd/Private/ViewportToolbar/UnrealEdViewportToolbar.cpp
--- a/Engine/Source/Editor/UnrealEd/Private/ViewportToolbar/UnrealEdViewportToolbar.cpp
+++ b/Engine/Source/Editor/UnrealEd/Private/ViewportToolbar/UnrealEdViewportToolbar.cpp
@@ -15,6 +15,7 @@
 #include "PreviewProfileController.h"
 #include "RayTracingDebugVisualizationMenuCommands.h"
 #include "SAssetEditorViewport.h"
+#include "SCommonEditorViewportToolbarBase.h"
 #include "SEditorViewport.h"
 #include "Settings/EditorProjectSettings.h"
 #include "Settings/LevelEditorViewportSettings.h"
@@ -1776,6 +1777,97 @@ FToolMenuEntry CreateViewportToolbarCameraSubmenu()
 	);
 }
 
+FToolMenuEntry CreateViewportToolbarAssetViewerProfileSubmenu(const TSharedPtr<IPreviewProfileController>& InPreviewProfileController
+)
+{
+	TWeakPtr<IPreviewProfileController> PreviewProfileControllerWeak = InPreviewProfileController;
+
+	return FToolMenuEntry::InitSubMenu(
+		"AssetViewerProfile",
+		TAttribute<FText>::CreateLambda(
+			[PreviewProfileControllerWeak]()
+			{
+				if (TSharedPtr<IPreviewProfileController> PreviewProfileController = PreviewProfileControllerWeak.Pin())
+				{
+					return FText::FromString(PreviewProfileController->GetActiveProfile());
+				}
+
+				return LOCTEXT("AssetViewerDefaultProfileLabel", "Profile");
+			}
+		),
+		LOCTEXT("AssetViewerProfileSelectionSectionTooltip", "Select the Preview Scene Profile for this viewport."),
+		FNewToolMenuDelegate::CreateLambda(
+			[PreviewProfileControllerWeak](UToolMenu* Submenu) -> void
+			{
+				TSharedPtr<IPreviewProfileController> PreviewProfileController = PreviewProfileControllerWeak.Pin();
+
+				if (!PreviewProfileController)
+				{
+					return;
+				}
+
+				FToolMenuSection& UnnamedSection = Submenu->FindOrAddSection("", LOCTEXT("UnnamedLabel", ""));
+
+				constexpr bool bInShouldCloseWindowAfterMenuSelection = true;
+				FMenuBuilder PreviewProfilesSelectionMenuBuilder(bInShouldCloseWindowAfterMenuSelection, nullptr);
+				PreviewProfilesSelectionMenuBuilder.BeginSection(
+					"AssetViewerProfileSelectionSection",
+					LOCTEXT("AssetViewerProfileSelectionSectionLabel", "Preview Scene Profiles")
+				);
+
+				int32 CurrProfileIndex = 0;
+				const TArray<FString>& PreviewProfiles = PreviewProfileController->GetPreviewProfiles(CurrProfileIndex);
+
+				for (int32 ProfileIndex = 0; ProfileIndex < PreviewProfiles.Num(); ProfileIndex++)
+				{
+					const FString& ProfileName = PreviewProfiles[ProfileIndex];
+					PreviewProfilesSelectionMenuBuilder.AddMenuEntry(
+						FText::FromString(ProfileName),
+						LOCTEXT("", ""),
+						FSlateIcon(),
+						FUIAction(
+							FExecuteAction::CreateLambda(
+								[PreviewProfileControllerWeak, ProfileIndex, PreviewProfiles]()
+								{
+									if (TSharedPtr<IPreviewProfileController> PreviewProfileController =
+											PreviewProfileControllerWeak.Pin())
+									{
+										PreviewProfileController->SetActiveProfile(PreviewProfiles[ProfileIndex]);
+									}
+								}
+							),
+							FCanExecuteAction(),
+							FIsActionChecked::CreateLambda(
+								[PreviewProfileControllerWeak, ProfileIndex]()
+								{
+									if (TSharedPtr<IPreviewProfileController> PreviewProfileController =
+											PreviewProfileControllerWeak.Pin())
+									{
+										int32 CurrentlySelectedProfileIndex;
+										PreviewProfileController->GetPreviewProfiles(CurrentlySelectedProfileIndex);
+
+										return ProfileIndex == CurrentlySelectedProfileIndex;
+									}
+
+									return false;
+								}
+							)
+						),
+						NAME_None,
+						EUserInterfaceActionType::RadioButton
+					);
+				}
+
+				PreviewProfilesSelectionMenuBuilder.EndSection();
+
+				UnnamedSection.AddEntry(FToolMenuEntry::InitWidget(
+					"AssetViewerProfile", PreviewProfilesSelectionMenuBuilder.MakeWidget(), FText(), true
+				));
+			}
+		)
+	);
+}
+
 void PopulateCameraMenu(UToolMenu* InMenu)
 {
 	UUnrealEdViewportToolbarContext* const EditorViewportContext = InMenu->FindContext<UUnrealEdViewportToolbarContext>();
diff --git a/Engine/Source/Editor/UnrealEd/Public/ViewportToolbar/UnrealEdViewportToolbar.h b/Engine/Source/Editor/UnrealEd/Public/ViewportToolbar/UnrealEdViewportToolbar.h
--- a/Engine/Source/Editor/UnrealEd/Public/ViewportToolbar/UnrealEdViewportToolbar.h
+++ b/Engine/Source/Editor/UnrealEd/Public/ViewportToolbar/UnrealEdViewportToolbar.h
@@ -11,6 +11,7 @@
 class IPreviewProfileController;
 class FEditorViewportClient;
 class FText;
+class IPreviewProfileController;
 class SEditorViewport;
 class UToolMenu;
 struct FNewToolMenuChoice;
@@ -117,6 +118,8 @@ UNREALED_API FToolMenuEntry CreateNumericEntry(
 
 UNREALED_API FToolMenuEntry CreateViewportToolbarCameraSubmenu();
 
+UNREALED_API FToolMenuEntry CreateViewportToolbarAssetViewerProfileSubmenu(const TSharedPtr<IPreviewProfileController>& InPreviewProfileController);
+
 UNREALED_API void PopulateCameraMenu(UToolMenu* InMenu);
 
 /**
