Fix issue when re-importing a static mesh and one or more socket have been rename.

#rnx

diff --git a/Engine/Source/Editor/UnrealEd/Private/SSocketManager.cpp b/Engine/Source/Editor/UnrealEd/Private/SSocketManager.cpp
--- a/Engine/Source/Editor/UnrealEd/Private/SSocketManager.cpp
+++ b/Engine/Source/Editor/UnrealEd/Private/SSocketManager.cpp
@@ -604,8 +604,23 @@ void SSocketManager::RefreshSocketList()
 		// This is done so that an undo on a socket property doesn't cause the selected
 		// socket to be de-selected, thus hiding the socket properties on the detail view.
 		// NB: Also force a rebuild if the underlying StaticMesh has been changed.
-		if (StaticMesh->Sockets.Num() != SocketList.Num() || !bIsSameStaticMesh)
+		bool bRefreshList = StaticMesh->Sockets.Num() != SocketList.Num() || !bIsSameStaticMesh;
+		if (!bRefreshList)
 		{
+			//Make sure list is in sync (UObject are the sames)
+			for (int32 SocketIndex = 0; SocketIndex < StaticMesh->Sockets.Num(); ++SocketIndex)
+			{
+				if (StaticMesh->Sockets[SocketIndex].Get() != SocketList[SocketIndex].Get()->Socket)
+				{
+					bRefreshList = true;
+					break;
+				}
+			}
+		}
+
+		if (bRefreshList)
+		{
+			SocketListView->ClearSelection();
 			SocketList.Empty();
 			for (int32 i = 0; i < StaticMesh->Sockets.Num(); i++)
 			{
