Move asset sorting functions SortByName and SortByPredicate from UnrealEd to AssetRegistry

diff --git a/Engine/Source/Editor/UnrealEd/Private/Subsystems/EditorAssetSubsystem.cpp b/Engine/Source/Editor/UnrealEd/Private/Subsystems/EditorAssetSubsystem.cpp
--- a/Engine/Source/Editor/UnrealEd/Private/Subsystems/EditorAssetSubsystem.cpp
+++ b/Engine/Source/Editor/UnrealEd/Private/Subsystems/EditorAssetSubsystem.cpp
@@ -32,362 +32,368 @@ using TError = TValueOrError<UE::EditorAssetUtils::ErrorTag, T>;
 
 namespace UE::EditorAssetUtils
 {
-	static bool EnsureAssetsLoaded()
+static bool EnsureAssetsLoaded()
+{
+	FAssetRegistryModule& AssetRegistryModule = FModuleManager::LoadModuleChecked<FAssetRegistryModule>("AssetRegistry");
+	if (IAssetRegistry& AssetRegistry = AssetRegistryModule.Get(); AssetRegistry.IsLoadingAssets())
 	{
-		FAssetRegistryModule& AssetRegistryModule = FModuleManager::LoadModuleChecked<FAssetRegistryModule>("AssetRegistry");
-		if (IAssetRegistry& AssetRegistry = AssetRegistryModule.Get(); AssetRegistry.IsLoadingAssets())
+		if (AssetRegistry.IsSearchAsync() && AssetRegistry.IsSearchAllAssets())
 		{
-			if (AssetRegistry.IsSearchAsync() && AssetRegistry.IsSearchAllAssets())
-			{
-				AssetRegistry.WaitForCompletion();
-			}
-			else
-			{
-				AssetRegistry.SearchAllAssets(true /* bSynchronousSearch */);
-			}
+			AssetRegistry.WaitForCompletion();
 		}
-		return true;
+		else
+		{
+			AssetRegistry.SearchAllAssets(true /* bSynchronousSearch */);
+		}
+	}
+	return true;
+}
+
+static TValueOrError<FAssetData, FString> FindAssetDataFromAnyPath(const FString& AnyAssetPath)
+{
+	FString FailureReason;
+	FString ObjectPath = EditorScriptingHelpers::ConvertAnyPathToSubObjectPath(AnyAssetPath, FailureReason);
+	if (ObjectPath.IsEmpty())
+	{
+		return MakeError(FailureReason);
 	}
 
-	static TValueOrError<FAssetData, FString> FindAssetDataFromAnyPath(const FString& AnyAssetPath)
+	FAssetRegistryModule& AssetRegistryModule = FModuleManager::LoadModuleChecked<FAssetRegistryModule>("AssetRegistry");
+	FAssetData AssetData = AssetRegistryModule.Get().GetAssetByObjectPath(FSoftObjectPath(ObjectPath));
+	if (!AssetData.IsValid())
 	{
-		FString FailureReason;
-		FString ObjectPath = EditorScriptingHelpers::ConvertAnyPathToSubObjectPath(AnyAssetPath, FailureReason);
+		ObjectPath = EditorScriptingHelpers::ConvertAnyPathToObjectPath(AnyAssetPath, FailureReason);
 		if (ObjectPath.IsEmpty())
 		{
 			return MakeError(FailureReason);
 		}
 
-		FAssetRegistryModule& AssetRegistryModule = FModuleManager::LoadModuleChecked<FAssetRegistryModule>("AssetRegistry");
-		FAssetData AssetData = AssetRegistryModule.Get().GetAssetByObjectPath(FSoftObjectPath(ObjectPath));
+		AssetData = AssetRegistryModule.Get().GetAssetByObjectPath(FSoftObjectPath(ObjectPath));
 		if (!AssetData.IsValid())
 		{
-			ObjectPath = EditorScriptingHelpers::ConvertAnyPathToObjectPath(AnyAssetPath, FailureReason);
-			if (ObjectPath.IsEmpty())
-			{
-				return MakeError(FailureReason);
-			}
-
-			AssetData = AssetRegistryModule.Get().GetAssetByObjectPath(FSoftObjectPath(ObjectPath));
-			if (!AssetData.IsValid())
-			{
-				return MakeError(FString::Printf(TEXT("The AssetData '%s' could not be found in the Asset Registry."), *ObjectPath));
-			}
+			return MakeError(FString::Printf(TEXT("The AssetData '%s' could not be found in the Asset Registry."), *ObjectPath));
 		}
-		
-		return MakeValue(AssetData);
 	}
+	
+	return MakeValue(AssetData);
+}
 
-	static TValueOrError<UObject*, FString> LoadAssetFromData(const FAssetData& AssetData)
+static TValueOrError<UObject*, FString> LoadAssetFromData(const FAssetData& AssetData)
+{
+	if (!AssetData.IsValid())
 	{
-		if (!AssetData.IsValid())
-		{
-			return MakeError("Asset Data is not valid.");
-		}
+		return MakeError("Asset Data is not valid.");
+	}
 
-		UObject* FoundObject = AssetData.GetAsset();
-		if (!IsValid(FoundObject))
-		{
-			return MakeError(FString::Printf(TEXT("The asset '%s' exists but was not able to be loaded."), *AssetData.GetObjectPathString()));
-		}
-		else if (!FoundObject->IsAsset())
-		{
-			return MakeError(FString::Printf(TEXT("'%s' is not a valid asset."), *AssetData.GetObjectPathString()));
-		}
-		return MakeValue(FoundObject);
+	UObject* FoundObject = AssetData.GetAsset();
+	if (!IsValid(FoundObject))
+	{
+		return MakeError(FString::Printf(TEXT("The asset '%s' exists but was not able to be loaded."), *AssetData.GetObjectPathString()));
 	}
-	
-	static TValueOrError<UObject*, FString> LoadAssetFromPath(const FString& AssetPath)
+	else if (!FoundObject->IsAsset())
 	{
-		TValueOrError<FAssetData, FString> AssetDataResult = FindAssetDataFromAnyPath(AssetPath);
-		if (AssetDataResult.HasError())
-		{
-			return MakeError(AssetDataResult.StealError());
-		}
-		return LoadAssetFromData(AssetDataResult.GetValue());
+		return MakeError(FString::Printf(TEXT("'%s' is not a valid asset."), *AssetData.GetObjectPathString()));
 	}
+	return MakeValue(FoundObject);
+}
 
-	static TError<FString> IsARegisteredAsset(UObject* Object, bool bAllowSkipBrowsableTestForExternalObject = false)
+static TValueOrError<UObject*, FString> LoadAssetFromPath(const FString& AssetPath)
+{
+	TValueOrError<FAssetData, FString> AssetDataResult = FindAssetDataFromAnyPath(AssetPath);
+	if (AssetDataResult.HasError())
 	{
-		if (!IsValid(Object))
-		{
-			return MakeError(TEXT("The Asset is not valid."));
-		}
-
-		const bool bCanSkipIsBrowsable = bAllowSkipBrowsableTestForExternalObject && Object->IsPackageExternal();
-		if (!bCanSkipIsBrowsable && !ObjectTools::IsObjectBrowsable(Object))
-		{
-			return MakeError(FString::Printf(TEXT("The object '%s' is not an asset."), *Object->GetName()));
-		}
+		return MakeError(AssetDataResult.StealError());
+	}
+	return LoadAssetFromData(AssetDataResult.GetValue());
+}
 
-		FSoftObjectPath ObjectPath = FSoftObjectPath(Object);
-		FAssetRegistryModule& AssetRegistryModule = FModuleManager::LoadModuleChecked<FAssetRegistryModule>("AssetRegistry");
-		FAssetData AssetData = AssetRegistryModule.Get().GetAssetByObjectPath(ObjectPath);
-		if (!AssetData.IsValid())
-		{
-			return MakeError(FString::Printf(TEXT("The AssetData '%s' could not be found in the Asset Registry."), *Object->GetPathName()));
-		}
+static TError<FString> IsARegisteredAsset(UObject* Object, bool bAllowSkipBrowsableTestForExternalObject = false)
+{
+	if (!IsValid(Object))
+	{
+		return MakeError(TEXT("The Asset is not valid."));
+	}
 
-		return MakeValue();
+	const bool bCanSkipIsBrowsable = bAllowSkipBrowsableTestForExternalObject && Object->IsPackageExternal();
+	if (!bCanSkipIsBrowsable && !ObjectTools::IsObjectBrowsable(Object))
+	{
+		return MakeError(FString::Printf(TEXT("The object '%s' is not an asset."), *Object->GetName()));
 	}
 
-	static TError<FString> EnumerateAssetsInDirectory(const FString& AnyPathDirectoryPath, bool bRecursive, TArray<FAssetData>& OutResult, FString& OutDirectoryPath)
+	FSoftObjectPath ObjectPath = FSoftObjectPath(Object);
+	FAssetRegistryModule& AssetRegistryModule = FModuleManager::LoadModuleChecked<FAssetRegistryModule>("AssetRegistry");
+	FAssetData AssetData = AssetRegistryModule.Get().GetAssetByObjectPath(ObjectPath);
+	if (!AssetData.IsValid())
 	{
-		OutResult.Reset();
-		OutDirectoryPath.Reset();
+		return MakeError(FString::Printf(TEXT("The AssetData '%s' could not be found in the Asset Registry."), *Object->GetPathName()));
+	}
 
-		FString FailureReason;
-		OutDirectoryPath = EditorScriptingHelpers::ConvertAnyPathToLongPackagePath(AnyPathDirectoryPath, FailureReason);
-		if (OutDirectoryPath.IsEmpty())
-		{
-			return MakeError(FailureReason);
-		}
+	return MakeValue();
+}
 
-		FAssetRegistryModule& AssetRegistryModule = FModuleManager::LoadModuleChecked<FAssetRegistryModule>("AssetRegistry");
-		if (!AssetRegistryModule.Get().GetAssetsByPath(*OutDirectoryPath, OutResult, bRecursive))
-		{
-			return MakeError(FString::Printf(TEXT("Could not get assets from path '%s'"), *OutDirectoryPath));
-		}
-		
-		return MakeValue();
+static TError<FString> EnumerateAssetsInDirectory(const FString& AnyPathDirectoryPath, bool bRecursive, TArray<FAssetData>& OutResult, FString& OutDirectoryPath)
+{
+	OutResult.Reset();
+	OutDirectoryPath.Reset();
+
+	FString FailureReason;
+	OutDirectoryPath = EditorScriptingHelpers::ConvertAnyPathToLongPackagePath(AnyPathDirectoryPath, FailureReason);
+	if (OutDirectoryPath.IsEmpty())
+	{
+		return MakeError(FailureReason);
 	}
 
-	enum class EPackageEnumerationFilter
+	FAssetRegistryModule& AssetRegistryModule = FModuleManager::LoadModuleChecked<FAssetRegistryModule>("AssetRegistry");
+	if (!AssetRegistryModule.Get().GetAssetsByPath(*OutDirectoryPath, OutResult, bRecursive))
 	{
-		NoFilter,
-		OnlyDirty
-	};
+		return MakeError(FString::Printf(TEXT("Could not get assets from path '%s'"), *OutDirectoryPath));
+	}
 	
-	static TError<FString> EnumeratePackagesInDirectory(const FString& AnyDirectoryPath, EPackageEnumerationFilter EnumerationFilter, bool bRecursive, TArray<UPackage*>& OutResult)
+	return MakeValue();
+}
+
+enum class EPackageEnumerationFilter
+{
+	NoFilter,
+	OnlyDirty
+};
+
+static TError<FString> EnumeratePackagesInDirectory(const FString& AnyDirectoryPath, EPackageEnumerationFilter EnumerationFilter, bool bRecursive, TArray<UPackage*>& OutResult)
+{
+	FString ValidDirectoryPath;
+	TArray<FAssetData> Assets;
+	if (TError<FString> Result = EnumerateAssetsInDirectory(AnyDirectoryPath, bRecursive, Assets, ValidDirectoryPath); Result.HasError())
 	{
-		FString ValidDirectoryPath;
-		TArray<FAssetData> Assets;
-		if (TError<FString> Result = EnumerateAssetsInDirectory(AnyDirectoryPath, bRecursive, Assets, ValidDirectoryPath); Result.HasError())
-		{
-			return Result;
-		}
+		return Result;
+	}
 
-		if (EnumerationFilter == EPackageEnumerationFilter::OnlyDirty)
-		{
-			for (const FAssetData& AssetData : Assets)
-			{
-				// Can't be dirty if not loaded
-				if (AssetData.IsAssetLoaded())
-				{
-					if (UPackage* Package = AssetData.GetPackage(); Package && Package->IsDirty())
-					{
-						Package->FullyLoad();
-						OutResult.AddUnique(Package);
-					}
-				}
-			}
-		}
-		else
+	if (EnumerationFilter == EPackageEnumerationFilter::OnlyDirty)
+	{
+		for (const FAssetData& AssetData : Assets)
 		{
-			// load all assets
-			for (const FAssetData& AssetData : Assets)
+			// Can't be dirty if not loaded
+			if (AssetData.IsAssetLoaded())
 			{
-				if (UPackage* Package = AssetData.GetPackage())
+				if (UPackage* Package = AssetData.GetPackage(); Package && Package->IsDirty())
 				{
 					Package->FullyLoad();
 					OutResult.AddUnique(Package);
 				}
 			}
 		}
-		return MakeValue();
 	}
-	
-	static bool DeleteEmptyDirectoryFromDisk(const FString& LongPackagePath)
+	else
 	{
-		struct FEmptyDirectoryVisitor : public IPlatformFile::FDirectoryVisitor
-		{
-			bool bIsEmpty = true;
-
-			virtual bool Visit(const TCHAR*, bool bIsDirectory) override
-			{
-				if (!bIsDirectory)
-				{
-					bIsEmpty = false;
-					return false; // abort searching
-				}
-				return true; // continue searching
-			}
-		};
-		
-		if (FString PathToDeleteOnDisk = UPackageTools::PackageNameToFilename(LongPackagePath); !PathToDeleteOnDisk.IsEmpty())
+		// load all assets
+		for (const FAssetData& AssetData : Assets)
 		{
-			// Look for files on disk in case the directory contains things not tracked by the asset registry
-			FEmptyDirectoryVisitor EmptyDirectoryVisitor;
-			IFileManager::Get().IterateDirectoryRecursively(*PathToDeleteOnDisk, EmptyDirectoryVisitor);
-
-			if (EmptyDirectoryVisitor.bIsEmpty)
+			if (UPackage* Package = AssetData.GetPackage())
 			{
-				return IFileManager::Get().DeleteDirectory(*PathToDeleteOnDisk, false, true);
+				Package->FullyLoad();
+				OutResult.AddUnique(Package);
 			}
 		}
-		return false;
 	}
+	return MakeValue();
+}
 
-	struct FDirectoryRenamePaths
+static bool DeleteEmptyDirectoryFromDisk(const FString& LongPackagePath)
+{
+	struct FEmptyDirectoryVisitor : public IPlatformFile::FDirectoryVisitor
 	{
-		FString SourceDirectoryPath;
-		FString SourceFilePath;
-		FString DestinationDirectoryPath;
-		FString DestinationFilePath;
+		bool bIsEmpty = true;
+
+		virtual bool Visit(const TCHAR*, bool bIsDirectory) override
+		{
+			if (!bIsDirectory)
+			{
+				bIsEmpty = false;
+				return false; // abort searching
+			}
+			return true; // continue searching
+		}
 	};
 	
-	static TValueOrError<FDirectoryRenamePaths, FString> GetDirectoryRenamePaths(const FString& SourceDirectoryPath, const FString& DestinationDirectoryPath)
+	if (FString PathToDeleteOnDisk = UPackageTools::PackageNameToFilename(LongPackagePath); !PathToDeleteOnDisk.IsEmpty())
 	{
-		FDirectoryRenamePaths Paths;
-		FString FailureReason;
-		Paths.SourceDirectoryPath = EditorScriptingHelpers::ConvertAnyPathToLongPackagePath(SourceDirectoryPath, FailureReason);
-		if (Paths.SourceDirectoryPath.IsEmpty())
-		{
-			return MakeError(FString::Printf(TEXT("Failed to convert the source path. %s"), *FailureReason));
-		}
+		// Look for files on disk in case the directory contains things not tracked by the asset registry
+		FEmptyDirectoryVisitor EmptyDirectoryVisitor;
+		IFileManager::Get().IterateDirectoryRecursively(*PathToDeleteOnDisk, EmptyDirectoryVisitor);
 
-		Paths.SourceFilePath = FPaths::ConvertRelativePathToFull(UPackageTools::PackageNameToFilename(Paths.SourceDirectoryPath));
-		if (Paths.SourceFilePath.IsEmpty())
+		if (EmptyDirectoryVisitor.bIsEmpty)
 		{
-			return MakeError(FString::Printf(TEXT("Failed to convert the source path '%s' to a full path."), *Paths.SourceDirectoryPath));
+			return IFileManager::Get().DeleteDirectory(*PathToDeleteOnDisk, false, true);
 		}
+	}
+	return false;
+}
 
-		Paths.DestinationDirectoryPath = EditorScriptingHelpers::ConvertAnyPathToLongPackagePath(DestinationDirectoryPath, FailureReason);
-		if (Paths.DestinationDirectoryPath.IsEmpty())
-		{
-			return MakeError(FString::Printf(TEXT("Failed to convert the destination path. %s"), *FailureReason));
-		}
+struct FDirectoryRenamePaths
+{
+	FString SourceDirectoryPath;
+	FString SourceFilePath;
+	FString DestinationDirectoryPath;
+	FString DestinationFilePath;
+};
 
-		Paths.DestinationFilePath = FPaths::ConvertRelativePathToFull(UPackageTools::PackageNameToFilename(Paths.DestinationDirectoryPath));
-		if (Paths.DestinationFilePath.IsEmpty())
-		{
-			return MakeError(FString::Printf(TEXT("Failed to convert the destination path '%s' to a full path."), *Paths.DestinationDirectoryPath));
-		}
-		
-		return MakeValue(Paths);
+static TValueOrError<FDirectoryRenamePaths, FString> GetDirectoryRenamePaths(const FString& SourceDirectoryPath, const FString& DestinationDirectoryPath)
+{
+	FDirectoryRenamePaths Paths;
+	FString FailureReason;
+	Paths.SourceDirectoryPath = EditorScriptingHelpers::ConvertAnyPathToLongPackagePath(SourceDirectoryPath, FailureReason);
+	if (Paths.SourceDirectoryPath.IsEmpty())
+	{
+		return MakeError(FString::Printf(TEXT("Failed to convert the source path. %s"), *FailureReason));
 	}
 
-	static TError<FString> SetupDirectoryRename(const FDirectoryRenamePaths& Paths)
+	Paths.SourceFilePath = FPaths::ConvertRelativePathToFull(UPackageTools::PackageNameToFilename(Paths.SourceDirectoryPath));
+	if (Paths.SourceFilePath.IsEmpty())
 	{
-		// If the source directory doesn't exist on disk then it can't be operated on
-		if (!IFileManager::Get().DirectoryExists(*Paths.SourceFilePath))
-		{
-			return MakeError(FString::Printf(TEXT("The source directory '%s' does not exist on disk."), *Paths.SourceFilePath));
-		}
+		return MakeError(FString::Printf(TEXT("Failed to convert the source path '%s' to a full path."), *Paths.SourceDirectoryPath));
+	}
+
+	Paths.DestinationDirectoryPath = EditorScriptingHelpers::ConvertAnyPathToLongPackagePath(DestinationDirectoryPath, FailureReason);
+	if (Paths.DestinationDirectoryPath.IsEmpty())
+	{
+		return MakeError(FString::Printf(TEXT("Failed to convert the destination path. %s"), *FailureReason));
+	}
+
+	Paths.DestinationFilePath = FPaths::ConvertRelativePathToFull(UPackageTools::PackageNameToFilename(Paths.DestinationDirectoryPath));
+	if (Paths.DestinationFilePath.IsEmpty())
+	{
+		return MakeError(FString::Printf(TEXT("Failed to convert the destination path '%s' to a full path."), *Paths.DestinationDirectoryPath));
+	}
+	
+	return MakeValue(Paths);
+}
 
-		// Create the destination directory if it doesn't already exist
-		if (!IFileManager::Get().DirectoryExists(*Paths.DestinationFilePath))
+static TError<FString> SetupDirectoryRename(const FDirectoryRenamePaths& Paths)
+{
+	// If the source directory doesn't exist on disk then it can't be operated on
+	if (!IFileManager::Get().DirectoryExists(*Paths.SourceFilePath))
+	{
+		return MakeError(FString::Printf(TEXT("The source directory '%s' does not exist on disk."), *Paths.SourceFilePath));
+	}
+
+	// Create the destination directory if it doesn't already exist
+	if (!IFileManager::Get().DirectoryExists(*Paths.DestinationFilePath))
+	{
+		const bool bTree = true;
+		if (!IFileManager::Get().MakeDirectory(*Paths.DestinationFilePath, bTree))
 		{
-			const bool bTree = true;
-			if (!IFileManager::Get().MakeDirectory(*Paths.DestinationFilePath, bTree))
-			{
-				return MakeError(FString::Printf(TEXT("The destination directory '%s' could not be created."), *Paths.DestinationFilePath));
-			}
+			return MakeError(FString::Printf(TEXT("The destination directory '%s' could not be created."), *Paths.DestinationFilePath));
 		}
-
-		return MakeValue();
 	}
+
+	return MakeValue();
+}
+
+struct FDirectoryRenameAssetPaths
+{
+	TArray<FAssetData> SourceDirectoryAssetDatas;
+	TArray<FString> DestinationDirectoryAssetPaths;
+};
+
+static TValueOrError<FDirectoryRenameAssetPaths, FString> GetDirectoryRenameAssetPaths(const FDirectoryRenamePaths& Paths)
+{
+	FDirectoryRenameAssetPaths AssetPaths;
 	
-	struct FDirectoryRenameAssetPaths
+	// Load all assets the directory contains
+	// Because we want to rename a directory, we can't rename any files that can't be deleted
+	FString OutPath;
+	if (TError<FString> Result = EnumerateAssetsInDirectory(Paths.SourceDirectoryPath, true, AssetPaths.SourceDirectoryAssetDatas, OutPath); Result.HasError())
 	{
-		TArray<FAssetData> SourceDirectoryAssetDatas;
-		TArray<FString> DestinationDirectoryAssetPaths;
-	};
+		return MakeError(Result.StealError());
+	}
 	
-	static TValueOrError<FDirectoryRenameAssetPaths, FString> GetDirectoryRenameAssetPaths(const FDirectoryRenamePaths& Paths)
+	for (const FAssetData& AssetData : AssetPaths.SourceDirectoryAssetDatas)
 	{
-		FDirectoryRenameAssetPaths AssetPaths;
-		
-		// Load all assets the directory contains
-		// Because we want to rename a directory, we can't rename any files that can't be deleted
-		FString OutPath;
-		if (TError<FString> Result = EnumerateAssetsInDirectory(Paths.SourceDirectoryPath, true, AssetPaths.SourceDirectoryAssetDatas, OutPath); Result.HasError())
+		FString PackageName = AssetData.PackageName.ToString();
+		FString LongPackagePath = FPackageName::GetLongPackagePath(PackageName);
+
+		// Remove source from the object name
+		LongPackagePath.MidInline(Paths.SourceDirectoryPath.Len(), MAX_int32, EAllowShrinking::No);
+
+		// Create AssetPath /Game/MyFolder/MyAsset.MyAsset
+		FString NewAssetPackageName;
+		if (LongPackagePath.IsEmpty())
 		{
-			return MakeError(Result.StealError());
+			NewAssetPackageName = FString::Printf(TEXT("%s/%s.%s"), *Paths.DestinationDirectoryPath, *AssetData.AssetName.ToString(), *AssetData.AssetName.ToString());
 		}
-		
-		for (const FAssetData& AssetData : AssetPaths.SourceDirectoryAssetDatas)
+		else
 		{
-			FString PackageName = AssetData.PackageName.ToString();
-			FString LongPackagePath = FPackageName::GetLongPackagePath(PackageName);
-
-			// Remove source from the object name
-			LongPackagePath.MidInline(Paths.SourceDirectoryPath.Len(), MAX_int32, EAllowShrinking::No);
+			NewAssetPackageName = FString::Printf(TEXT("%s%s/%s.%s"), *Paths.DestinationDirectoryPath, *LongPackagePath, *AssetData.AssetName.ToString(), *AssetData.AssetName.ToString());
+		}
 
-			// Create AssetPath /Game/MyFolder/MyAsset.MyAsset
-			FString NewAssetPackageName;
-			if (LongPackagePath.IsEmpty())
-			{
-				NewAssetPackageName = FString::Printf(TEXT("%s/%s.%s"), *Paths.DestinationDirectoryPath, *AssetData.AssetName.ToString(), *AssetData.AssetName.ToString());
-			}
-			else
-			{
-				NewAssetPackageName = FString::Printf(TEXT("%s%s/%s.%s"), *Paths.DestinationDirectoryPath, *LongPackagePath, *AssetData.AssetName.ToString(), *AssetData.AssetName.ToString());
-			}
+		FString FailureReason;
+		if (!EditorScriptingHelpers::IsAValidPathForCreateNewAsset(NewAssetPackageName, FailureReason))
+		{
+			return MakeError(FString::Printf(TEXT("Failed to validate the destination for asset '%s'. %s"), *AssetData.AssetName.ToString(), *FailureReason));
+		}
+		
+		if (FPackageName::DoesPackageExist(NewAssetPackageName))
+		{
+			return MakeError(FString::Printf(TEXT("Failed to validate the destination for asset '%s'. There's already an asset at the destination."), *NewAssetPackageName));
+		}
 
-			FString FailureReason;
-			if (!EditorScriptingHelpers::IsAValidPathForCreateNewAsset(NewAssetPackageName, FailureReason))
-			{
-				return MakeError(FString::Printf(TEXT("Failed to validate the destination for asset '%s'. %s"), *AssetData.AssetName.ToString(), *FailureReason));
-			}
-			
-			if (FPackageName::DoesPackageExist(NewAssetPackageName))
-			{
-				return MakeError(FString::Printf(TEXT("Failed to validate the destination for asset '%s'. There's already an asset at the destination."), *NewAssetPackageName));
-			}
+		// Keep AssetPath /Game/MyFolder
+		AssetPaths.DestinationDirectoryAssetPaths.Add(NewAssetPackageName);
+	}
+	return MakeValue(AssetPaths);
+}
 
-			// Keep AssetPath /Game/MyFolder
-			AssetPaths.DestinationDirectoryAssetPaths.Add(NewAssetPackageName);
-		}
-		return MakeValue(AssetPaths);
+template<typename TLessThan>
+static void SortAssets(
+	TArray<FAssetData>& Assets,
+	TLessThan&& Predicate,
+	EEditorAssetSortOrder SortOrder
+)
+{
+	switch (SortOrder)
+	{
+	case EEditorAssetSortOrder::Ascending: Assets.Sort(Predicate); break;
+	case EEditorAssetSortOrder::Descending:
+		Assets.Sort([&Predicate](const FAssetData& Left, const FAssetData& Right)
+		{
+			return Predicate(Right, Left);
+		});
+		break;
+	default: checkNoEntry(); break;
 	}
-	
-	static void SortAssets(
-		TArray<FAssetData>& Assets,
-		TFunctionRef<bool(const FAssetData& Left, const FAssetData& Right)> Predicate,
-		EEditorAssetSortOrder SortOrder
+}
+
+template<typename TType>
+static bool Sort(
+	TArray<FAssetData>& Assets,
+	FName MetaDataTag,
+	TFunctionRef<bool(const FString& TagType, TType& Converted)> Converter,
+	EEditorAssetSortOrder SortOrder
 	)
+{
+	TMap<FSoftObjectPath, TType> MetaData;
+	MetaData.Reserve(Assets.Num());
+
+	for (const FAssetData& AssetData : Assets)
 	{
-		// Careful this would be undefined behaviour: TFunctionRef Variable = [](){}; 
-		auto ReversePredicate = [&Predicate](const FAssetData& Left, const FAssetData& Right) { return !Predicate(Left, Right); };
-		const TFunctionRef<bool(const FAssetData& Left, const FAssetData& Right)> ReversePredicateFunc = ReversePredicate;
-		const TFunctionRef<bool(const FAssetData& Left, const FAssetData& Right)> PredicateToUse = SortOrder == EEditorAssetSortOrder::Ascending ? Predicate : ReversePredicateFunc;
-	
-		Assets.Sort(PredicateToUse);
-	}
-	
-	template<typename TType>
-	static bool Sort(
-		TArray<FAssetData>& Assets,
-		FName MetaDataTag,
-		TFunctionRef<bool(const FString& TagType, TType& Converted)> Converter,
-		EEditorAssetSortOrder SortOrder
-		)
-	{
-		TMap<FSoftObjectPath, TType> MetaData;
-		MetaData.Reserve(Assets.Num());
-	
-		for (const FAssetData& AssetData : Assets)
+		const FAssetTagValueRef Value = AssetData.TagsAndValues.FindTag(MetaDataTag);
+		TType AssetTagValue;
+		if (Value.IsSet() && Converter(Value.GetValue(), AssetTagValue))
 		{
-			const FAssetTagValueRef Value = AssetData.TagsAndValues.FindTag(MetaDataTag);
-			TType AssetTagValue;
-			if (Value.IsSet() && Converter(Value.GetValue(), AssetTagValue))
-			{
-				MetaData.Add(AssetData.GetSoftObjectPath(), AssetTagValue);
-			}
-			else
-			{
-				UE_LOG(LogEditorAssetSubsystem, Warning, TEXT("Not all assets have the tag '%s'"), *MetaDataTag.ToString());
-				return false;
-			}
+			MetaData.Add(AssetData.GetSoftObjectPath(), AssetTagValue);
 		}
-
-		SortAssets(Assets, [&MetaData](const FAssetData& Left, const FAssetData& Right)
+		else
 		{
-			return MetaData[Left.GetSoftObjectPath()] <= MetaData[Right.GetSoftObjectPath()]; 
-		}, SortOrder);
-		return true;
+			UE_LOG(LogEditorAssetSubsystem, Warning, TEXT("Not all assets have the tag '%s'"), *MetaDataTag.ToString());
+			return false;
+		}
 	}
+
+	SortAssets(Assets, [&MetaData](const FAssetData& Left, const FAssetData& Right)
+	{
+		return MetaData[Left.GetSoftObjectPath()] <= MetaData[Right.GetSoftObjectPath()]; 
+	}, SortOrder);
+	return true;
+}
 }
 
 UEditorAssetSubsystem::UEditorAssetSubsystem()
@@ -1723,32 +1729,6 @@ TArray<FAssetData> UEditorAssetSubsystem::GetAllAssetsByMetaDataTags(
 	return Result;
 }
 
-void UEditorAssetSubsystem::SortByPredicate(
-	TArray<FAssetData>& Assets,
-	FEditorAssetSortingPredicate SortingPredicate,
-	EEditorAssetSortOrder SortOrder
-	)
-{
-	if (SortingPredicate.IsBound())
-	{
-		UE::EditorAssetUtils::SortAssets(
-			Assets,
-			[&SortingPredicate](const FAssetData& Left, const FAssetData& Right)
-			{
-				return SortingPredicate.Execute(Left, Right);
-			}, SortOrder);
-	}
-}
-
-void UEditorAssetSubsystem::SortByName(TArray<FAssetData>& Assets, EEditorAssetSortOrder SortOrder)
-{
-	UE::EditorAssetUtils::SortAssets(Assets,
-		[](const FAssetData& Left, const FAssetData& Right)
-		{
-			return Left.AssetName.LexicalLess(Right.AssetName);
-		}, SortOrder);
-}
-
 bool UEditorAssetSubsystem::SortByMetaData(
 	TArray<FAssetData>& Assets,
 	FName MetaDataTag,
diff --git a/Engine/Source/Editor/UnrealEd/Public/Subsystems/EditorAssetSubsystem.h b/Engine/Source/Editor/UnrealEd/Public/Subsystems/EditorAssetSubsystem.h
--- a/Engine/Source/Editor/UnrealEd/Public/Subsystems/EditorAssetSubsystem.h
+++ b/Engine/Source/Editor/UnrealEd/Public/Subsystems/EditorAssetSubsystem.h
@@ -5,7 +5,6 @@
 #include "CoreMinimal.h"
 #include "AssetRegistry/AssetData.h"
 #include "EditorSubsystem.h"
-
 #include "EditorAssetSubsystem.generated.h"
 
 UENUM(BlueprintType, DisplayName = "Sort Order")
@@ -23,8 +22,6 @@ enum class EEditorAssetMetaDataSortType : uint8
 	DateTime
 };
 
-DECLARE_DYNAMIC_DELEGATE_RetVal_TwoParams(bool, FEditorAssetSortingPredicate, const FAssetData&, Left, const FAssetData&, Right);
-
 /**
 * UEditorAssetSubsystem
 * Subsystem for exposing asset related utilities to scripts.
@@ -436,25 +433,6 @@ public:
 	 */
 	UFUNCTION(BlueprintCallable, Category = "Editor Scripting | Asset")
 	UNREALED_API TArray<FAssetData> GetAllAssetsByMetaDataTags(const TSet<FName>& RequiredTags, const TSet<UClass*>& AllowedClasses);
-
-	/**
-	 * Sorts the assets based on a custom Blueprint delegate.
-	 * 
-	 * @param Assets The assets to sort
-	 * @param SortingPredicate Implements a Left <= Right relation
-	 * @param SortOrder Whether to sort ascending or descending
-	 */
-	UFUNCTION(BlueprintCallable, Category = "Editor Scripting | Asset")
-	UNREALED_API void SortByPredicate(UPARAM(Ref) TArray<FAssetData>& Assets, FEditorAssetSortingPredicate SortingPredicate, EEditorAssetSortOrder SortOrder);
-
-	/**
-	 * Sorts the assets by their asset name.
-	 *
-	 * @param Assets The assets to sort
-	 * @param SortOrder Whether to sort ascending or descending
-	 */
-	UFUNCTION(BlueprintCallable, Category = "Editor Scripting | Asset")
-	UNREALED_API void SortByName(UPARAM(Ref) TArray<FAssetData>& Assets, EEditorAssetSortOrder SortOrder);
 	
 	/**
 	 * Sorts the assets based on their meta data's type.
diff --git a/Engine/Source/Runtime/AssetRegistry/Private/AssetRegistryHelpers.cpp b/Engine/Source/Runtime/AssetRegistry/Private/AssetRegistryHelpers.cpp
--- a/Engine/Source/Runtime/AssetRegistry/Private/AssetRegistryHelpers.cpp
+++ b/Engine/Source/Runtime/AssetRegistry/Private/AssetRegistryHelpers.cpp
@@ -14,6 +14,26 @@
 #include "Misc/RedirectCollector.h"
 #endif
 
+namespace UE::AssetRegistry
+{
+
+template<typename TLessThan>
+static void SortAssets(TArray<FAssetData>& Assets, TLessThan&& Predicate, EAssetRegistrySortOrder SortOrder)
+{
+	switch (SortOrder)
+	{
+	case EAssetRegistrySortOrder::Ascending: Assets.Sort(Predicate); break;
+	case EAssetRegistrySortOrder::Descending:
+		Assets.Sort([&Predicate](const FAssetData& Left, const FAssetData& Right)
+		{
+			return Predicate(Right, Left);
+		});
+		break;
+	default: checkNoEntry(); break;
+	}
+}
+}
+
 TScriptInterface<IAssetRegistry> UAssetRegistryHelpers::GetAssetRegistry()
 {
 	return &UAssetRegistryImpl::Get();
@@ -118,12 +138,46 @@ UClass* UAssetRegistryHelpers::FindAssetNativeClass(const FAssetData& AssetData)
 	return AssetClass;
 }
 
-void UAssetRegistryHelpers::FindReferencersOfAssetOfClass(UObject* AssetInstance, TConstArrayView<UClass*> InMatchClasses, TArray<FAssetData>& OutAssetDatas)
+void UAssetRegistryHelpers::SortByPredicate(
+	TArray<FAssetData>& Assets,
+	FSortingPredicate SortingPredicate,
+	EAssetRegistrySortOrder SortOrder
+	)
+{
+	if (SortingPredicate.IsBound())
+	{
+		UE::AssetRegistry::SortAssets(
+			Assets,
+			[&SortingPredicate](const FAssetData& Left, const FAssetData& Right)
+			{
+				return SortingPredicate.Execute(Left, Right);
+			}, SortOrder);
+	}
+}
+
+void UAssetRegistryHelpers::SortByAssetName(TArray<FAssetData>& Assets, EAssetRegistrySortOrder SortOrder)
+{
+	UE::AssetRegistry::SortAssets(Assets,
+		[](const FAssetData& Left, const FAssetData& Right)
+		{
+			return Left.AssetName.LexicalLess(Right.AssetName);
+		}, SortOrder);
+}
+
+void UAssetRegistryHelpers::FindReferencersOfAssetOfClass(
+	UObject* AssetInstance,
+	TConstArrayView<UClass*> InMatchClasses,
+	TArray<FAssetData>& OutAssetDatas
+	)
 {
 	FindReferencersOfAssetOfClass(AssetInstance->GetOutermost()->GetFName(), InMatchClasses, OutAssetDatas);
 }
 
-void UAssetRegistryHelpers::FindReferencersOfAssetOfClass(const FAssetIdentifier& InAssetIdentifier, TConstArrayView<UClass*> InMatchClasses, TArray<FAssetData>& OutAssetDatas)
+void UAssetRegistryHelpers::FindReferencersOfAssetOfClass(
+	const FAssetIdentifier& InAssetIdentifier,
+	TConstArrayView<UClass*> InMatchClasses,
+	TArray<FAssetData>& OutAssetDatas
+	)
 {
 	// If the asset registry is still loading assets, we cant check for referencers, so we must open the rename dialog
 	const IAssetRegistry& AssetRegistry = IAssetRegistry::GetChecked();
@@ -169,7 +223,11 @@ void UAssetRegistryHelpers::GetBlueprintAssets(const FARFilter& InFilter, TArray
 	TSet<FTopLevelAssetPath> BlueprintParentClassPaths;
 	if (Filter.bRecursiveClasses)
 	{
-		AssetRegistry.GetDerivedClassNames(BlueprintParentClassPathRoots, TSet<FTopLevelAssetPath>(), BlueprintParentClassPaths);
+		AssetRegistry.GetDerivedClassNames(
+			BlueprintParentClassPathRoots, 
+			TSet<FTopLevelAssetPath>(),
+			 BlueprintParentClassPaths
+			 );
 	}
 	else
 	{
diff --git a/Engine/Source/Runtime/AssetRegistry/Public/AssetRegistry/AssetRegistryHelpers.h b/Engine/Source/Runtime/AssetRegistry/Public/AssetRegistry/AssetRegistryHelpers.h
--- a/Engine/Source/Runtime/AssetRegistry/Public/AssetRegistry/AssetRegistryHelpers.h
+++ b/Engine/Source/Runtime/AssetRegistry/Public/AssetRegistry/AssetRegistryHelpers.h
@@ -8,6 +8,7 @@
 #include "Containers/Array.h"
 #include "Containers/Set.h"
 #include "Containers/UnrealString.h"
+#include "Delegates/Delegate.h"
 #include "UObject/NameTypes.h"
 #include "UObject/Object.h"
 #include "UObject/ObjectMacros.h"
@@ -39,6 +40,13 @@ struct FTagAndValue
 	FString Value;
 };
 
+UENUM(BlueprintType, DisplayName = "Sort Order")
+enum class EAssetRegistrySortOrder : uint8
+{
+	Ascending,
+	Descending
+};
+
 UCLASS(transient)
 class UAssetRegistryHelpers : public UObject
 {
@@ -114,6 +122,26 @@ public:
 	UFUNCTION(BlueprintPure, Category = "Asset Registry", meta=(ScriptMethod))
 	static ASSETREGISTRY_API UClass* FindAssetNativeClass(const FAssetData& AssetData);
 
+	DECLARE_DYNAMIC_DELEGATE_RetVal_TwoParams(bool, FSortingPredicate, const FAssetData&, Left, const FAssetData&, Right);
+	/**
+	 * Sorts the assets based on a custom Blueprint delegate.
+	 * 
+	 * @param Assets The assets to sort
+	 * @param SortingPredicate Implements a Left <= Right relation
+	 * @param SortOrder Whether to sort ascending or descending
+	 */
+	UFUNCTION(BlueprintCallable, Category = "Asset Registry")
+	static ASSETREGISTRY_API void SortByPredicate(UPARAM(Ref) TArray<FAssetData>& Assets, FSortingPredicate SortingPredicate, EAssetRegistrySortOrder SortOrder);
+
+	/**
+	 * Sorts the assets by their asset name.
+	 *
+	 * @param Assets The assets to sort
+	 * @param SortOrder Whether to sort ascending or descending
+	 */
+	UFUNCTION(BlueprintCallable, Category = "Asset Registry")
+	static ASSETREGISTRY_API void SortByAssetName(UPARAM(Ref) TArray<FAssetData>& Assets, EAssetRegistrySortOrder SortOrder);
+	
 	/**
 	 * Finds references of the provided asset that are of the a class contained in the InMatchClasses set.
 	 */
