IncrementalCook: Write a CookInProgress flag to CookSettings.txt, clear it at the end of cook, and use it to invalidate the previous cook if the cook crashed before completion.
This is a temporary solution until we harden shader code against corrupt iterative shader libraries, and add the restriction for licensees to likewise harden their reading of iterative cook artifacts.
#rnx

diff --git a/Engine/Source/Editor/UnrealEd/Classes/CookOnTheSide/CookOnTheFlyServer.h b/Engine/Source/Editor/UnrealEd/Classes/CookOnTheSide/CookOnTheFlyServer.h
--- a/Engine/Source/Editor/UnrealEd/Classes/CookOnTheSide/CookOnTheFlyServer.h
+++ b/Engine/Source/Editor/UnrealEd/Classes/CookOnTheSide/CookOnTheFlyServer.h
@@ -1230,6 +1230,8 @@ private:
 		const ITargetPlatform* TargetPlatform) const;
 	/** Save the CurrentCookSettings into the output directory. */
 	void SaveCookSettings(const TMap<FName, FString>& CurrentCookSettings, const ITargetPlatform* TargetPlatform);
+	/** Load the CookSettings written at beginning of cook and rewrite them with CookInProgress flag removed. */
+	void ClearCookInProgressFlagFromCookSettings(const ITargetPlatform* TargetPlatform) const;
 	/**
 	 * Populate a map suitable for saving as an ini with the current value of the Cook settings that need to be
 	 * tested for compatibility.
diff --git a/Engine/Source/Editor/UnrealEd/Private/CookOnTheFlyServer.cpp b/Engine/Source/Editor/UnrealEd/Private/CookOnTheFlyServer.cpp
--- a/Engine/Source/Editor/UnrealEd/Private/CookOnTheFlyServer.cpp
+++ b/Engine/Source/Editor/UnrealEd/Private/CookOnTheFlyServer.cpp
@@ -8176,6 +8176,7 @@ bool UCookOnTheFlyServer::GetCookedIniVersionStrings(const ITargetPlatform* Targ
 }
 
 static const TCHAR* TEXT_CookSettings(TEXT("CookSettings"));
+static const TCHAR* TEXT_CookInProgress(TEXT("CookInProgress"));
 static FName ExecutableHashName(TEXT("ExecutableHash"));
 static FName ExecutableHashInvalidModuleName(TEXT("ExecutableHashInvalidModule"));
 
@@ -8252,13 +8253,14 @@ FString UCookOnTheFlyServer::GetCookSettingsFileName(const ITargetPlatform* Targ
 bool UCookOnTheFlyServer::ArePreviousCookSettingsCompatible(const TMap<FName, FString>& CurrentCookSettings, const ITargetPlatform* TargetPlatform) const
 {
 	FConfigFile ConfigFile;
-	ConfigFile.Read(GetCookSettingsFileName(TargetPlatform));
+	FString CookSettingsFileName = GetCookSettingsFileName(TargetPlatform);
+	ConfigFile.Read(CookSettingsFileName);
 
 	const FConfigSection* CookSettings = ConfigFile.FindSection(TEXT_CookSettings);
 	if (CookSettings == nullptr)
 	{
 		UE_LOG(LogCook, Display, TEXT("Cook invalidated for platform %s because CookSettings file %s is invalid. Clearing previously cooked packages."),
-			*TargetPlatform->PlatformName(), *GetCookSettingsFileName(TargetPlatform));
+			*TargetPlatform->PlatformName(), *CookSettingsFileName);
 		return false;
 	}
 
@@ -8283,6 +8285,18 @@ bool UCookOnTheFlyServer::ArePreviousCookSettingsCompatible(const TMap<FName, FS
 		}
 	}
 
+	if (GIsBuildMachine)
+	{
+		bool bCookInProgress;
+		if (ConfigFile.GetBool(TEXT_CookSettings, TEXT_CookInProgress, bCookInProgress) && bCookInProgress)
+		{
+			UE_LOG(LogCook, Display, TEXT("Cook invalidated for platform %s because the previous cook crashed (or otherwise did not report completion).")
+				TEXT(" CookSettings file %s still has [%s]:%s=true. Clearing previously cooked packages."),
+				*TargetPlatform->PlatformName(), *CookSettingsFileName, TEXT_CookSettings, TEXT_CookInProgress);
+			return false;
+		}
+	}
+
 	if (!bIterativeIgnoreIni && !bHybridIterativeEnabled && IniSettingsOutOfDate(TargetPlatform))
 	{
 		UE_LOG(LogCook, Display, TEXT("Cook invalidated for platform %s because ini settings have changed. Clearing previously cooked packages."),
@@ -8327,10 +8341,21 @@ void UCookOnTheFlyServer::SaveCookSettings(const TMap<FName, FString>& CurrentCo
 	{
 		ConfigFile.AddToSection(TEXT_CookSettings, CurrentSetting.Key, CurrentSetting.Value);
 	}
+	ConfigFile.AddToSection(TEXT_CookSettings, TEXT_CookInProgress, TEXT("true"));
 	ConfigFile.Dirty = true; // Writing to a section does not set the dirty flag, so set it manually to make Write work
 	ConfigFile.Write(GetCookSettingsFileName(TargetPlatform));
 }
 
+void UCookOnTheFlyServer::ClearCookInProgressFlagFromCookSettings(const ITargetPlatform* TargetPlatform) const
+{
+	FConfigFile ConfigFile;
+	FString Filename = GetCookSettingsFileName(TargetPlatform);
+	ConfigFile.Read(Filename);
+	ConfigFile.RemoveKeyFromSection(TEXT_CookSettings, TEXT_CookInProgress);
+	ConfigFile.Dirty = true; // Writing to a section does not set the dirty flag, so set it manually to make Write work
+	ConfigFile.Write(Filename);
+}
+
 bool UCookOnTheFlyServer::IniSettingsOutOfDate(const ITargetPlatform* TargetPlatform) const
 {
 #if !UE_WITH_CONFIG_TRACKING
@@ -10316,6 +10341,10 @@ void UCookOnTheFlyServer::CookByTheBookFinished()
 		CookByTheBookFinishedInternal();
 	}
 
+	for (const ITargetPlatform* TargetPlatform : PlatformManager->GetSessionPlatforms())
+	{
+		ClearCookInProgressFlagFromCookSettings(TargetPlatform);
+	}
 	ShutdownCookSession();
 	BroadcastCookByTheBookFinished();
 	UE_LOG(LogCook, Display, TEXT("Done!"));
