[Viewport Toolbar] Reduce the number of menu entries in the Blueprint editor's "View Modes" menu

- Adds filtering to UE::UnrealEd::PopulateViewModesMenu for those elements which were missing it

- Adds bool to Toolbar Context Object to specify if certain sections need to be shown or hidden.

- Updates Level editor and BP Editor

diff --git a/Engine/Source/Editor/Kismet/Private/SSCSEditorViewport.cpp b/Engine/Source/Editor/Kismet/Private/SSCSEditorViewport.cpp
--- a/Engine/Source/Editor/Kismet/Private/SSCSEditorViewport.cpp
+++ b/Engine/Source/Editor/Kismet/Private/SSCSEditorViewport.cpp
@@ -155,6 +155,25 @@ void ExtendCameraSubmenu(FName InCameraOptionsSubmenuName, TSharedPtr<FSCSEditor
 	);
 }
 
+bool IsViewModeSupported(EViewModeIndex InViewModeIndex)
+{
+	switch (InViewModeIndex)
+	{
+	case VMI_Unlit:
+	case VMI_Lit:
+	case VMI_BrushWireframe:
+	case VMI_CollisionVisibility:
+		return true;
+	default:
+		return false;
+	}
+}
+
+bool DoesViewModeMenuShowSection(UE::UnrealEd::EHidableViewModeMenuSections)
+{
+	return false;
+}
+
 } // namespace UE::SCSEditor::Private
 
 /*-----------------------------------------------------------------------------
@@ -575,23 +594,12 @@ TSharedPtr<SWidget> SSCSEditorViewport::BuildViewportToolbar()
 			ContextObject->Viewport = SharedThis(this);
 
 			// Setup the callback to filter available view modes
-			ContextObject->IsViewModeSupported = UE::UnrealEd::IsViewModeSupportedDelegate::CreateLambda(
-				[](EViewModeIndex ViewModeIndex) -> bool
-				{
-					// This code is taken from SViewportToolBar::IsViewModeSupported
-					// SSCSEditorViewportToolBar does not override it, so we just take it as-is
-					// TODO: maybe create a private function for it, or move IsViewModeSupported to SEditorViewport
+			ContextObject->IsViewModeSupported =
+				UE::UnrealEd::IsViewModeSupportedDelegate::CreateStatic(&UE::SCSEditor::Private::IsViewModeSupported);
 
-					switch (ViewModeIndex)
-					{
-					case VMI_PrimitiveDistanceAccuracy:
-					case VMI_MaterialTextureScaleAccuracy:
-					case VMI_RequiredTextureResolution:
-						return false;
-					default:
-						return true;
-					}
-				}
+			// Setup the callback to hide/show specific sections
+			ContextObject->DoesViewModeMenuShowSection = UE::UnrealEd::DoesViewModeMenuShowSectionDelegate::CreateStatic(
+				&UE::SCSEditor::Private::DoesViewModeMenuShowSection
 			);
 
 			ViewportToolbarContext.AddObject(ContextObject);
diff --git a/Engine/Source/Editor/UnrealEd/Private/SEditorViewportViewMenu.cpp b/Engine/Source/Editor/UnrealEd/Private/SEditorViewportViewMenu.cpp
--- a/Engine/Source/Editor/UnrealEd/Private/SEditorViewportViewMenu.cpp
+++ b/Engine/Source/Editor/UnrealEd/Private/SEditorViewportViewMenu.cpp
@@ -98,7 +98,14 @@ void SEditorViewportViewMenu::FillViewMenu(UToolMenu* Menu) const
 			return true;
 		});
 
-	UE::UnrealEd::PopulateViewModesMenu(Menu, Viewport.Pin().ToSharedRef(), IsViewModeSupported);
+	// Add the UnrealEd viewport toolbar context.
+	{
+		UUnrealEdViewportToolbarContext* const ContextObject = NewObject<UUnrealEdViewportToolbarContext>();
+		ContextObject->Viewport = Viewport;
+		Menu->Context.AddObject(ContextObject);
+	}
+
+	UE::UnrealEd::PopulateViewModesMenu(Menu);
 }
 
 #undef LOCTEXT_NAMESPACE
diff --git a/Engine/Source/Editor/UnrealEd/Private/ViewportToolbar/UnrealEdViewportToolbar.cpp b/Engine/Source/Editor/UnrealEd/Private/ViewportToolbar/UnrealEdViewportToolbar.cpp
--- a/Engine/Source/Editor/UnrealEd/Private/ViewportToolbar/UnrealEdViewportToolbar.cpp
+++ b/Engine/Source/Editor/UnrealEd/Private/ViewportToolbar/UnrealEdViewportToolbar.cpp
@@ -515,6 +515,19 @@ FToolMenuEntry CreateScaleSnapCheckboxMenu()
 	);
 }
 
+bool IsViewModeSupported(EViewModeIndex InViewModeIndex)
+{
+	switch (InViewModeIndex)
+	{
+	case VMI_PrimitiveDistanceAccuracy:
+	case VMI_MaterialTextureScaleAccuracy:
+	case VMI_RequiredTextureResolution:
+		return false;
+	default:
+		return true;
+	}
+}
+
 }
 
 static FAutoConsoleVariableRef CVarToolMenusViewportToolbars(
@@ -1036,31 +1049,82 @@ FText GetViewModesSubmenuLabel(TWeakPtr<SEditorViewport> InViewport)
 	return Label;
 }
 
-void PopulateViewModesMenu(
-	UToolMenu* InMenu, TSharedRef<SEditorViewport> InViewport, IsViewModeSupportedDelegate InIsViewModeSupported)
+void PopulateViewModesMenu(UToolMenu* InMenu)
 {
+	UUnrealEdViewportToolbarContext* const Context = InMenu->FindContext<UUnrealEdViewportToolbarContext>();
+	if (!Context)
+	{
+		return;
+	}
+
+	const TSharedPtr<SEditorViewport> EditorViewport = Context->Viewport.Pin();
+
+	if (!EditorViewport)
+	{
+		return;
+	}
+
 	const FEditorViewportCommands& BaseViewportActions = FEditorViewportCommands::Get();
 
+	IsViewModeSupportedDelegate IsViewModeSupported = Context->IsViewModeSupported;
+
+	struct ViewModesSubmenu
+	{
+		static void AddModeIfSupported(
+			const IsViewModeSupportedDelegate& InIsViewModeSupported,
+			FToolMenuSection& InMenuSection,
+			const TSharedPtr<FUICommandInfo>& InModeCommandInfo,
+			EViewModeIndex InViewModeIndex,
+			const TAttribute<FText>& InToolTipOverride = TAttribute<FText>(),
+			const TAttribute<FSlateIcon>& InIconOverride = TAttribute<FSlateIcon>()
+		)
+		{
+			if (!InIsViewModeSupported.IsBound() || InIsViewModeSupported.Execute(InViewModeIndex))
+			{
+				InMenuSection.AddMenuEntry(
+					InModeCommandInfo, UViewModeUtils::GetViewModeDisplayName(InViewModeIndex), InToolTipOverride, InIconOverride
+				);
+			}
+		}
+
+		static bool IsMenuSectionAvailable(const UUnrealEdViewportToolbarContext* InContext, EHidableViewModeMenuSections InMenuSection)
+		{
+			if (!InContext->DoesViewModeMenuShowSection.IsBound())
+			{
+				return true;
+			}
+
+			return InContext->DoesViewModeMenuShowSection.Execute(InMenuSection);
+		}
+	};
+
 	// View modes
 	{
 		FToolMenuSection& Section = InMenu->AddSection("ViewMode", LOCTEXT("ViewModeHeader", "View Mode"));
 		{
-			Section.AddMenuEntry(BaseViewportActions.LitMode, UViewModeUtils::GetViewModeDisplayName(VMI_Lit));
-			Section.AddMenuEntry(BaseViewportActions.UnlitMode, UViewModeUtils::GetViewModeDisplayName(VMI_Unlit));
-			Section.AddMenuEntry(
-				BaseViewportActions.WireframeMode, UViewModeUtils::GetViewModeDisplayName(VMI_BrushWireframe));
-			Section.AddMenuEntry(
-				BaseViewportActions.LitWireframeMode, UViewModeUtils::GetViewModeDisplayName(VMI_Lit_Wireframe));
-			Section.AddMenuEntry(
-				BaseViewportActions.DetailLightingMode, UViewModeUtils::GetViewModeDisplayName(VMI_Lit_DetailLighting));
-			Section.AddMenuEntry(
-				BaseViewportActions.LightingOnlyMode, UViewModeUtils::GetViewModeDisplayName(VMI_LightingOnly));
-			Section.AddMenuEntry(BaseViewportActions.ReflectionOverrideMode,
-				UViewModeUtils::GetViewModeDisplayName(VMI_ReflectionOverride));
-			Section.AddMenuEntry(
-				BaseViewportActions.CollisionPawn, UViewModeUtils::GetViewModeDisplayName(VMI_CollisionPawn));
-			Section.AddMenuEntry(BaseViewportActions.CollisionVisibility,
-				UViewModeUtils::GetViewModeDisplayName(VMI_CollisionVisibility));
+			ViewModesSubmenu::AddModeIfSupported(IsViewModeSupported, Section, BaseViewportActions.UnlitMode, VMI_Unlit);
+			ViewModesSubmenu::AddModeIfSupported(IsViewModeSupported, Section, BaseViewportActions.LitMode, VMI_Lit);
+			ViewModesSubmenu::AddModeIfSupported(
+				IsViewModeSupported, Section, BaseViewportActions.WireframeMode, VMI_BrushWireframe
+			);
+			ViewModesSubmenu::AddModeIfSupported(
+				IsViewModeSupported, Section, BaseViewportActions.LitWireframeMode, VMI_Lit_Wireframe
+			);
+			ViewModesSubmenu::AddModeIfSupported(
+				IsViewModeSupported, Section, BaseViewportActions.DetailLightingMode, VMI_Lit_DetailLighting
+			);
+			ViewModesSubmenu::AddModeIfSupported(
+				IsViewModeSupported, Section, BaseViewportActions.LightingOnlyMode, VMI_LightingOnly
+			);
+			ViewModesSubmenu::AddModeIfSupported(
+				IsViewModeSupported, Section, BaseViewportActions.ReflectionOverrideMode, VMI_ReflectionOverride
+			);
+			ViewModesSubmenu::AddModeIfSupported(
+				IsViewModeSupported, Section, BaseViewportActions.CollisionPawn, VMI_CollisionPawn
+			);
+			ViewModesSubmenu::AddModeIfSupported(
+				IsViewModeSupported, Section, BaseViewportActions.CollisionVisibility, VMI_CollisionVisibility
+			);
 		}
 
 		if (IsRayTracingEnabled())
@@ -1070,8 +1134,9 @@ void PopulateViewModesMenu(
 			const bool bPathTracingEnabled = PathTracingCvar && PathTracingCvar->GetValueOnAnyThread() != 0;
 			if (bPathTracingSupported && bPathTracingEnabled)
 			{
-				Section.AddMenuEntry(
-					BaseViewportActions.PathTracingMode, UViewModeUtils::GetViewModeDisplayName(VMI_PathTracing));
+				ViewModesSubmenu::AddModeIfSupported(
+					IsViewModeSupported, Section, BaseViewportActions.PathTracingMode, VMI_PathTracing
+				);
 			}
 		}
 
@@ -1081,7 +1146,7 @@ void PopulateViewModesMenu(
 			{
 				static void BuildOptimizationMenu(UToolMenu* Menu, IsViewModeSupportedDelegate IsViewModeSupported)
 				{
-					const FEditorViewportCommands& BaseViewportCommands = FEditorViewportCommands::Get();
+					const FEditorViewportCommands& BaseViewportActions = FEditorViewportCommands::Get();
 
 					UWorld* World = GWorld;
 					const ERHIFeatureLevel::Type FeatureLevel = (IsInGameThread() && World)
@@ -1090,113 +1155,171 @@ void PopulateViewModesMenu(
 
 					{
 						FToolMenuSection& Section = Menu->AddSection(
-							"OptimizationViewmodes", LOCTEXT("OptimizationSubMenuHeader", "Optimization Viewmodes"));
+							"OptimizationViewmodes", LOCTEXT("OptimizationSubMenuHeader", "Optimization Viewmodes")
+						);
 						if (FeatureLevel >= ERHIFeatureLevel::SM5)
 						{
-							Section.AddMenuEntry(BaseViewportCommands.LightComplexityMode,
-								UViewModeUtils::GetViewModeDisplayName(VMI_LightComplexity));
+							ViewModesSubmenu::AddModeIfSupported(
+								IsViewModeSupported, Section, BaseViewportActions.LightComplexityMode, VMI_LightComplexity
+							);
+
 							if (IsStaticLightingAllowed())
 							{
-								Section.AddMenuEntry(BaseViewportCommands.LightmapDensityMode,
-									UViewModeUtils::GetViewModeDisplayName(VMI_LightmapDensity));
+								ViewModesSubmenu::AddModeIfSupported(
+									IsViewModeSupported, Section, BaseViewportActions.LightmapDensityMode, VMI_LightmapDensity
+								);
 							}
-							Section.AddMenuEntry(BaseViewportCommands.StationaryLightOverlapMode,
-								UViewModeUtils::GetViewModeDisplayName(VMI_StationaryLightOverlap));
+
+							ViewModesSubmenu::AddModeIfSupported(
+								IsViewModeSupported, Section, BaseViewportActions.StationaryLightOverlapMode, VMI_StationaryLightOverlap
+							);
 						}
 
-						Section.AddMenuEntry(BaseViewportCommands.ShaderComplexityMode,
-							UViewModeUtils::GetViewModeDisplayName(VMI_ShaderComplexity));
+						ViewModesSubmenu::AddModeIfSupported(
+							IsViewModeSupported, Section, BaseViewportActions.ShaderComplexityMode, VMI_ShaderComplexity
+						);
 
 						if (AllowDebugViewShaderMode(
-								DVSM_ShaderComplexityContainedQuadOverhead, GMaxRHIShaderPlatform, FeatureLevel))
+								DVSM_ShaderComplexityContainedQuadOverhead, GMaxRHIShaderPlatform, FeatureLevel
+							))
 						{
-							Section.AddMenuEntry(BaseViewportCommands.ShaderComplexityWithQuadOverdrawMode,
-								UViewModeUtils::GetViewModeDisplayName(VMI_ShaderComplexityWithQuadOverdraw));
+							ViewModesSubmenu::AddModeIfSupported(
+								IsViewModeSupported,
+								Section,
+								BaseViewportActions.ShaderComplexityWithQuadOverdrawMode,
+								VMI_ShaderComplexityWithQuadOverdraw
+							);
 						}
+
 						if (AllowDebugViewShaderMode(DVSM_QuadComplexity, GMaxRHIShaderPlatform, FeatureLevel))
 						{
-							Section.AddMenuEntry(BaseViewportCommands.QuadOverdrawMode,
-								UViewModeUtils::GetViewModeDisplayName(VMI_QuadOverdraw));
+							ViewModesSubmenu::AddModeIfSupported(
+								IsViewModeSupported, Section, BaseViewportActions.QuadOverdrawMode, VMI_QuadOverdraw
+							);
 						}
+
 						if (AllowDebugViewShaderMode(DVSM_LWCComplexity, GMaxRHIShaderPlatform, FeatureLevel))
 						{
-							Section.AddMenuEntry(BaseViewportCommands.VisualizeLWCComplexity,
-								UViewModeUtils::GetViewModeDisplayName(VMI_LWCComplexity), TAttribute<FText>(),
-								FSlateIcon(FAppStyle::GetAppStyleSetName(), "EditorViewport.LWCComplexityMode"));
+							ViewModesSubmenu::AddModeIfSupported(
+								IsViewModeSupported,
+								Section,
+								BaseViewportActions.VisualizeLWCComplexity,
+								VMI_LWCComplexity,
+								TAttribute<FText>(),
+								FSlateIcon(FAppStyle::GetAppStyleSetName(), "EditorViewport.LWCComplexityMode")
+							);
 						}
 					}
 
 					{
 						FToolMenuSection& Section = Menu->AddSection(
-							"TextureStreaming", LOCTEXT("TextureStreamingHeader", "Texture Streaming Accuracy"));
+							"TextureStreaming", LOCTEXT("TextureStreamingHeader", "Texture Streaming Accuracy")
+						);
 
-						if (AllowDebugViewShaderMode(DVSM_PrimitiveDistanceAccuracy, GMaxRHIShaderPlatform, FeatureLevel)
-							&& (!IsViewModeSupported.IsBound()
-								|| IsViewModeSupported.Execute(VMI_PrimitiveDistanceAccuracy)))
+						if (AllowDebugViewShaderMode(DVSM_PrimitiveDistanceAccuracy, GMaxRHIShaderPlatform, FeatureLevel))
 						{
-							Section.AddMenuEntry(BaseViewportCommands.TexStreamAccPrimitiveDistanceMode,
-								UViewModeUtils::GetViewModeDisplayName(VMI_PrimitiveDistanceAccuracy));
+							ViewModesSubmenu::AddModeIfSupported(
+								IsViewModeSupported, Section, BaseViewportActions.TexStreamAccPrimitiveDistanceMode, VMI_PrimitiveDistanceAccuracy
+							);
 						}
-						if (AllowDebugViewShaderMode(DVSM_MeshUVDensityAccuracy, GMaxRHIShaderPlatform, FeatureLevel)
-							&& (!IsViewModeSupported.IsBound() || IsViewModeSupported.Execute(VMI_MeshUVDensityAccuracy)))
+
+						if (AllowDebugViewShaderMode(DVSM_MeshUVDensityAccuracy, GMaxRHIShaderPlatform, FeatureLevel))
 						{
-							Section.AddMenuEntry(BaseViewportCommands.TexStreamAccMeshUVDensityMode,
-								UViewModeUtils::GetViewModeDisplayName(VMI_MeshUVDensityAccuracy));
+							ViewModesSubmenu::AddModeIfSupported(
+								IsViewModeSupported, Section, BaseViewportActions.TexStreamAccMeshUVDensityMode, VMI_MeshUVDensityAccuracy
+							);
 						}
+
 						// TexCoordScale accuracy viewmode requires shaders that are only built in the
 						// TextureStreamingBuild, which requires the new metrics to be enabled.
 						if (AllowDebugViewShaderMode(DVSM_MaterialTextureScaleAccuracy, GMaxRHIShaderPlatform, FeatureLevel)
-							&& CVarStreamingUseNewMetrics.GetValueOnAnyThread() != 0
-							&& (!IsViewModeSupported.IsBound()
-								|| IsViewModeSupported.Execute(VMI_MaterialTextureScaleAccuracy)))
+							&& CVarStreamingUseNewMetrics.GetValueOnAnyThread() != 0)
 						{
-							Section.AddMenuEntry(BaseViewportCommands.TexStreamAccMaterialTextureScaleMode,
-								UViewModeUtils::GetViewModeDisplayName(VMI_MaterialTextureScaleAccuracy));
+							ViewModesSubmenu::AddModeIfSupported(
+								IsViewModeSupported,
+								Section,
+								BaseViewportActions.TexStreamAccMaterialTextureScaleMode,
+								VMI_MaterialTextureScaleAccuracy
+							);
 						}
-						if (AllowDebugViewShaderMode(DVSM_RequiredTextureResolution, GMaxRHIShaderPlatform, FeatureLevel)
-							&& (!IsViewModeSupported.IsBound()
-								|| IsViewModeSupported.Execute(VMI_RequiredTextureResolution)))
+
+						if (AllowDebugViewShaderMode(DVSM_RequiredTextureResolution, GMaxRHIShaderPlatform, FeatureLevel))
 						{
-							Section.AddMenuEntry(BaseViewportCommands.RequiredTextureResolutionMode,
-								UViewModeUtils::GetViewModeDisplayName(VMI_RequiredTextureResolution));
+							ViewModesSubmenu::AddModeIfSupported(
+								IsViewModeSupported, Section, BaseViewportActions.RequiredTextureResolutionMode, VMI_RequiredTextureResolution
+							);
 						}
-						if (AllowDebugViewShaderMode(DVSM_RequiredTextureResolution, GMaxRHIShaderPlatform, FeatureLevel)
-							&& (!IsViewModeSupported.IsBound()
-								|| IsViewModeSupported.Execute(VMI_VirtualTexturePendingMips)))
+
+						if (AllowDebugViewShaderMode(DVSM_RequiredTextureResolution, GMaxRHIShaderPlatform, FeatureLevel))
 						{
-							Section.AddMenuEntry(BaseViewportCommands.VirtualTexturePendingMipsMode,
-								UViewModeUtils::GetViewModeDisplayName(VMI_VirtualTexturePendingMips));
+							ViewModesSubmenu::AddModeIfSupported(
+								IsViewModeSupported, Section, BaseViewportActions.VirtualTexturePendingMipsMode, VMI_VirtualTexturePendingMips
+							);
 						}
 					}
 				}
+
+				static bool ViewModesShouldShowOptimizationEntries(const IsViewModeSupportedDelegate& InIsViewModeSupported)
+				{
+					if (!InIsViewModeSupported.IsBound())
+					{
+						return true;
+					}
+
+					return InIsViewModeSupported.Execute(VMI_LightComplexity)
+						|| InIsViewModeSupported.Execute(VMI_LightmapDensity)
+						|| InIsViewModeSupported.Execute(VMI_StationaryLightOverlap)
+						|| InIsViewModeSupported.Execute(VMI_ShaderComplexity)
+						|| InIsViewModeSupported.Execute(VMI_ShaderComplexityWithQuadOverdraw)
+						|| InIsViewModeSupported.Execute(VMI_QuadOverdraw)
+						|| InIsViewModeSupported.Execute(VMI_PrimitiveDistanceAccuracy)
+						|| InIsViewModeSupported.Execute(VMI_MeshUVDensityAccuracy)
+						|| InIsViewModeSupported.Execute(VMI_MaterialTextureScaleAccuracy)
+						|| InIsViewModeSupported.Execute(VMI_RequiredTextureResolution)
+						|| InIsViewModeSupported.Execute(VMI_VirtualTexturePendingMips);
+				}
 			};
 
-			Section.AddSubMenu("OptimizationSubMenu", LOCTEXT("OptimizationSubMenu", "Optimization Viewmodes"),
-				LOCTEXT("Optimization_ToolTip", "Select optimization visualizer"),
-				FNewToolMenuDelegate::CreateStatic(&Local::BuildOptimizationMenu, InIsViewModeSupported),
-				FUIAction(FExecuteAction(), FCanExecuteAction(),
-					FIsActionChecked::CreateLambda([Viewport = InViewport.ToWeakPtr()]() {
-						const TSharedRef<SEditorViewport> ViewportRef = Viewport.Pin().ToSharedRef();
-						const TSharedPtr<FEditorViewportClient> ViewportClient = ViewportRef->GetViewportClient();
-						check(ViewportClient.IsValid());
-						const EViewModeIndex ViewMode = ViewportClient->GetViewMode();
-						return (
-							// Texture Streaming Accuracy
-							ViewMode == VMI_LightComplexity || ViewMode == VMI_LightmapDensity
-							|| ViewMode == VMI_StationaryLightOverlap || ViewMode == VMI_ShaderComplexity
-							|| ViewMode == VMI_ShaderComplexityWithQuadOverdraw
-							|| ViewMode == VMI_QuadOverdraw
-							// Texture Streaming Accuracy
-							|| ViewMode == VMI_PrimitiveDistanceAccuracy || ViewMode == VMI_MeshUVDensityAccuracy
-							|| ViewMode == VMI_MaterialTextureScaleAccuracy || ViewMode == VMI_RequiredTextureResolution
-							|| ViewMode == VMI_VirtualTexturePendingMips);
-					})),
-				EUserInterfaceActionType::RadioButton,
-				/* bInOpenSubMenuOnClick = */ false,
-				FSlateIcon(FAppStyle::GetAppStyleSetName(), "EditorViewport.QuadOverdrawMode"));
+			if (Local::ViewModesShouldShowOptimizationEntries(IsViewModeSupported))
+			{
+				Section.AddSubMenu(
+					"OptimizationSubMenu",
+					LOCTEXT("OptimizationSubMenu", "Optimization Viewmodes"),
+					LOCTEXT("Optimization_ToolTip", "Select optimization visualizer"),
+					FNewToolMenuDelegate::CreateStatic(&Local::BuildOptimizationMenu, IsViewModeSupported),
+					FUIAction(
+						FExecuteAction(),
+						FCanExecuteAction(),
+						FIsActionChecked::CreateLambda(
+							[Viewport = EditorViewport.ToWeakPtr()]()
+							{
+								const TSharedRef<SEditorViewport> ViewportRef = Viewport.Pin().ToSharedRef();
+								const TSharedPtr<FEditorViewportClient> ViewportClient = ViewportRef->GetViewportClient();
+								check(ViewportClient.IsValid());
+								const EViewModeIndex ViewMode = ViewportClient->GetViewMode();
+								return (
+									// Texture Streaming Accuracy
+									ViewMode == VMI_LightComplexity || ViewMode == VMI_LightmapDensity
+									|| ViewMode == VMI_StationaryLightOverlap || ViewMode == VMI_ShaderComplexity
+									|| ViewMode == VMI_ShaderComplexityWithQuadOverdraw
+									|| ViewMode == VMI_QuadOverdraw
+									// Texture Streaming Accuracy
+									|| ViewMode == VMI_PrimitiveDistanceAccuracy || ViewMode == VMI_MeshUVDensityAccuracy
+									|| ViewMode == VMI_MaterialTextureScaleAccuracy
+									|| ViewMode == VMI_RequiredTextureResolution || ViewMode == VMI_VirtualTexturePendingMips
+								);
+							}
+						)
+					),
+					EUserInterfaceActionType::RadioButton,
+					/* bInOpenSubMenuOnClick = */ false,
+					FSlateIcon(FAppStyle::GetAppStyleSetName(), "EditorViewport.QuadOverdrawMode")
+				);
+			}
 		}
 
-		if (IsRayTracingEnabled())
+		if (IsRayTracingEnabled()
+			&& ViewModesSubmenu::IsMenuSectionAvailable(Context, EHidableViewModeMenuSections::RayTracingDebug))
 		{
 			struct Local
 			{
@@ -1208,87 +1331,136 @@ void PopulateViewModesMenu(
 				}
 			};
 
-			Section.AddSubMenu("RayTracingDebugSubMenu", LOCTEXT("RayTracingDebugSubMenu", "Ray Tracing Debug"),
+			Section.AddSubMenu(
+				"RayTracingDebugSubMenu",
+				LOCTEXT("RayTracingDebugSubMenu", "Ray Tracing Debug"),
 				LOCTEXT("RayTracing_ToolTip", "Select ray tracing buffer visualization view modes"),
 				FNewMenuDelegate::CreateStatic(&Local::BuildRayTracingDebugMenu), //, ParentToolBar)
-				false, FSlateIcon(FAppStyle::GetAppStyleSetName(), "EditorViewport.RayTracingDebugMode"));
+				false,
+				FSlateIcon(FAppStyle::GetAppStyleSetName(), "EditorViewport.RayTracingDebugMode")
+			);
 		}
 
 		{
 			struct Local
 			{
-				static void BuildLODMenu(UToolMenu* Menu)
+				static void BuildLODMenu(UToolMenu* Menu, IsViewModeSupportedDelegate IsViewModeSupported)
 				{
 					{
 						FToolMenuSection& Section = Menu->AddSection(
-							"LevelViewportLODColoration", LOCTEXT("LODModesHeader", "Level of Detail Coloration"));
-						Section.AddMenuEntry(FEditorViewportCommands::Get().LODColorationMode,
-							UViewModeUtils::GetViewModeDisplayName(VMI_LODColoration));
-						Section.AddMenuEntry(FEditorViewportCommands::Get().HLODColorationMode,
-							UViewModeUtils::GetViewModeDisplayName(VMI_HLODColoration));
+							"LevelViewportLODColoration", LOCTEXT("LODModesHeader", "Level of Detail Coloration")
+						);
+
+						ViewModesSubmenu::AddModeIfSupported(
+							IsViewModeSupported, Section, FEditorViewportCommands::Get().LODColorationMode, VMI_LODColoration
+						);
+
+						ViewModesSubmenu::AddModeIfSupported(
+							IsViewModeSupported, Section, FEditorViewportCommands::Get().HLODColorationMode, VMI_HLODColoration
+						);
 					}
 				}
 			};
 
-			Section.AddSubMenu("VisualizeGroupedLOD",
-				LOCTEXT("VisualizeGroupedLODDisplayName", "Level of Detail Coloration"),
-				LOCTEXT("GroupedLODMenu_ToolTip", "Select a mode for LOD Coloration"),
-				FNewToolMenuDelegate::CreateStatic(&Local::BuildLODMenu),
-				FUIAction(FExecuteAction(), FCanExecuteAction(),
-					FIsActionChecked::CreateLambda([WeakViewport = InViewport.ToWeakPtr()]() {
-						const TSharedRef<SEditorViewport> ViewportRef = WeakViewport.Pin().ToSharedRef();
-						const TSharedPtr<FEditorViewportClient> ViewportClient = ViewportRef->GetViewportClient();
-						check(ViewportClient.IsValid());
-						const EViewModeIndex ViewMode = ViewportClient->GetViewMode();
-						return (ViewMode == VMI_LODColoration || ViewMode == VMI_HLODColoration);
-					})),
-				EUserInterfaceActionType::RadioButton,
-				/* bInOpenSubMenuOnClick = */ false,
-				FSlateIcon(FAppStyle::GetAppStyleSetName(), "EditorViewport.GroupLODColorationMode"));
+			if (!IsViewModeSupported.IsBound()
+				|| (IsViewModeSupported.Execute(VMI_LODColoration) || IsViewModeSupported.Execute(VMI_HLODColoration)))
+			{
+				Section.AddSubMenu(
+					"VisualizeGroupedLOD",
+					LOCTEXT("VisualizeGroupedLODDisplayName", "Level of Detail Coloration"),
+					LOCTEXT("GroupedLODMenu_ToolTip", "Select a mode for LOD Coloration"),
+					FNewToolMenuDelegate::CreateStatic(&Local::BuildLODMenu, IsViewModeSupported),
+					FUIAction(
+						FExecuteAction(),
+						FCanExecuteAction(),
+						FIsActionChecked::CreateLambda(
+							[WeakViewport = EditorViewport.ToWeakPtr()]()
+							{
+								const TSharedRef<SEditorViewport> ViewportRef = WeakViewport.Pin().ToSharedRef();
+								const TSharedPtr<FEditorViewportClient> ViewportClient = ViewportRef->GetViewportClient();
+								check(ViewportClient.IsValid());
+								const EViewModeIndex ViewMode = ViewportClient->GetViewMode();
+								return (ViewMode == VMI_LODColoration || ViewMode == VMI_HLODColoration);
+							}
+						)
+					),
+					EUserInterfaceActionType::RadioButton,
+					/* bInOpenSubMenuOnClick = */ false,
+					FSlateIcon(FAppStyle::GetAppStyleSetName(), "EditorViewport.GroupLODColorationMode")
+				);
+			}
 		}
 
-		if (GEnableGPUSkinCache)
+		if (GEnableGPUSkinCache
+			&& ViewModesSubmenu::IsMenuSectionAvailable(Context, EHidableViewModeMenuSections::GPUSkinCache))
 		{
-			Section.AddSubMenu("VisualizeGPUSkinCacheViewMode",
+			Section.AddSubMenu(
+				"VisualizeGPUSkinCacheViewMode",
 				LOCTEXT("VisualizeGPUSkinCacheViewModeDisplayName", "GPU Skin Cache"),
 				LOCTEXT("GPUSkinCacheVisualizationMenu_ToolTip", "Select a mode for GPU Skin Cache visualization."),
 				FNewMenuDelegate::CreateStatic(&FGPUSkinCacheVisualizationMenuCommands::BuildVisualisationSubMenu),
-				FUIAction(FExecuteAction(), FCanExecuteAction(),
-					FIsActionChecked::CreateLambda([WeakViewport = InViewport.ToWeakPtr()]() {
-						const TSharedRef<SEditorViewport> ViewportRef = WeakViewport.Pin().ToSharedRef();
-						const TSharedPtr<FEditorViewportClient> ViewportClient = ViewportRef->GetViewportClient();
-						check(ViewportClient.IsValid());
-						return ViewportClient->IsViewModeEnabled(VMI_VisualizeGPUSkinCache);
-					})),
+				FUIAction(
+					FExecuteAction(),
+					FCanExecuteAction(),
+					FIsActionChecked::CreateLambda(
+						[WeakViewport = EditorViewport.ToWeakPtr()]()
+						{
+							const TSharedRef<SEditorViewport> ViewportRef = WeakViewport.Pin().ToSharedRef();
+							const TSharedPtr<FEditorViewportClient> ViewportClient = ViewportRef->GetViewportClient();
+							check(ViewportClient.IsValid());
+							return ViewportClient->IsViewModeEnabled(VMI_VisualizeGPUSkinCache);
+						}
+					)
+				),
 				EUserInterfaceActionType::RadioButton,
 				/* bInOpenSubMenuOnClick = */ false,
-				FSlateIcon(FAppStyle::GetAppStyleSetName(), "EditorViewport.VisualizeGPUSkinCacheMode"));
+				FSlateIcon(FAppStyle::GetAppStyleSetName(), "EditorViewport.VisualizeGPUSkinCacheMode")
+			);
 		}
 	}
 
 	// Auto Exposure
+	if (ViewModesSubmenu::IsMenuSectionAvailable(Context, EHidableViewModeMenuSections::Exposure))
 	{
 		const FEditorViewportCommands& BaseViewportCommands = FEditorViewportCommands::Get();
 
-		TSharedRef<SWidget> FixedEV100Menu = InViewport->BuildFixedEV100Menu();
-		TSharedPtr<FEditorViewportClient> EditorViewPostClient = InViewport->GetViewportClient();
-		const bool bIsLevelEditor = EditorViewPostClient.IsValid() && EditorViewPostClient->IsLevelEditorClient();
+		TSharedRef<SWidget> FixedEV100Menu = EditorViewport->BuildFixedEV100Menu();
+		TSharedPtr<FEditorViewportClient> EditorViewportClient = EditorViewport->GetViewportClient();
+		const bool bIsLevelEditor = EditorViewportClient.IsValid() && EditorViewportClient->IsLevelEditorClient();
 
 		FToolMenuSection& Section = InMenu->AddSection("Exposure", LOCTEXT("ExposureHeader", "Exposure"));
 		Section.AddMenuEntry(
-			bIsLevelEditor ? BaseViewportCommands.ToggleInGameExposure : BaseViewportCommands.ToggleAutoExposure);
+			bIsLevelEditor ? BaseViewportCommands.ToggleInGameExposure : BaseViewportCommands.ToggleAutoExposure
+		);
 		Section.AddEntry(FToolMenuEntry::InitWidget("FixedEV100", FixedEV100Menu, LOCTEXT("FixedEV100", "EV100")));
 	}
 
+	// TODO: would be nice to make this appear/disappear based on current mode
 	// Wireframe Opacity
+	if (!IsViewModeSupported.IsBound()
+		|| (IsViewModeSupported.Execute(VMI_Wireframe) || IsViewModeSupported.Execute(VMI_BrushWireframe)
+			|| IsViewModeSupported.Execute(VMI_Lit_Wireframe)))
 	{
-		TSharedRef<SWidget> WireOpacityMenu = InViewport->BuildWireframeMenu();
+		TSharedRef<SWidget> WireOpacityMenu = EditorViewport->BuildWireframeMenu();
 		FToolMenuSection& Section = InMenu->AddSection("Wireframe", LOCTEXT("WireframeHeader", "Wireframe"));
 		Section.AddEntry(
-			FToolMenuEntry::InitWidget("WireframeOpacity", WireOpacityMenu, LOCTEXT("WireframeOpacity", "Opacity")));
+			FToolMenuEntry::InitWidget("WireframeOpacity", WireOpacityMenu, LOCTEXT("WireframeOpacity", "Opacity"))
+		);
 	}
 }
 
+UUnrealEdViewportToolbarContext* CreateViewportToolbarDefaultContext(const TWeakPtr<SEditorViewport>& InViewport)
+{
+	UUnrealEdViewportToolbarContext* const ContextObject = NewObject<UUnrealEdViewportToolbarContext>();
+	ContextObject->Viewport = InViewport;
+
+	// Hook up our toolbar's filter for supported view modes.
+	ContextObject->IsViewModeSupported =
+		UE::UnrealEd::IsViewModeSupportedDelegate::CreateStatic(&UE::UnrealEd::Private::IsViewModeSupported);
+
+	return ContextObject;
+}
+
 FToolMenuEntry CreateViewportToolbarViewModesSubmenu()
 {
 	// This has to be a dynamic entry for the ViewModes submenu's label to be able to access the context.
@@ -1317,17 +1489,7 @@ FToolMenuEntry CreateViewportToolbarViewModesSubmenu()
 					FNewToolMenuDelegate::CreateLambda(
 						[](UToolMenu* Submenu) -> void
 						{
-							UUnrealEdViewportToolbarContext* const Context =
-								Submenu->FindContext<UUnrealEdViewportToolbarContext>();
-							if (!Context)
-							{
-								return;
-							}
-
-							if (const TSharedPtr<SEditorViewport> Viewport = Context->Viewport.Pin())
-							{
-								PopulateViewModesMenu(Submenu, Viewport.ToSharedRef(), Context->IsViewModeSupported);
-							}
+							PopulateViewModesMenu(Submenu);
 						}
 					)
 				);
diff --git a/Engine/Source/Editor/UnrealEd/Public/ViewportToolbar/UnrealEdViewportToolbar.h b/Engine/Source/Editor/UnrealEd/Public/ViewportToolbar/UnrealEdViewportToolbar.h
--- a/Engine/Source/Editor/UnrealEd/Public/ViewportToolbar/UnrealEdViewportToolbar.h
+++ b/Engine/Source/Editor/UnrealEd/Public/ViewportToolbar/UnrealEdViewportToolbar.h
@@ -21,6 +21,14 @@ enum ERotationGridMode : int;
 namespace UE::UnrealEd
 {
 
+/** Lists View Mode Menu Sections which can be shown/hidden based on specific menu requirements */
+enum EHidableViewModeMenuSections : uint8
+{
+	Exposure = 0,
+	GPUSkinCache = 1,
+	RayTracingDebug = 2
+};
+
 /** The value of this function is controlled by the CVAR "ToolMenusViewportToolbars". */
 UNREALED_API bool ShowOldViewportToolbars();
 
@@ -38,13 +46,13 @@ UNREALED_API FText GetViewModesSubmenuLabel(TWeakPtr<SEditorViewport> InViewport
 /**
  * Populate a given UToolMenu with entries for a View Modes viewport toolbar submenu.
  *
- * @param InMenu The menu to poulate with entries.
- * @param InViewport The viewport associated with this viewport toolbar.
- * @param InIsViewModeSupported Optional delegate to filter which view modes are added to the list.
+ * @param InMenu The menu to populate with entries.
  */
-UNREALED_API void PopulateViewModesMenu(UToolMenu* InMenu,
-	TSharedRef<SEditorViewport> InViewport,
-	IsViewModeSupportedDelegate InIsViewModeSupported = IsViewModeSupportedDelegate());
+UNREALED_API void PopulateViewModesMenu(UToolMenu* InMenu);
+
+/** Create a Viewport Toolbar Context with common values (many Asset Editors have the same settings) */
+UNREALED_API UUnrealEdViewportToolbarContext* CreateViewportToolbarDefaultContext(const TWeakPtr<SEditorViewport>& InViewport
+);
 
 UNREALED_API FToolMenuEntry CreateViewportToolbarViewModesSubmenu();
 
diff --git a/Engine/Source/Editor/UnrealEd/Public/ViewportToolbar/UnrealEdViewportToolbarContext.h b/Engine/Source/Editor/UnrealEd/Public/ViewportToolbar/UnrealEdViewportToolbarContext.h
--- a/Engine/Source/Editor/UnrealEd/Public/ViewportToolbar/UnrealEdViewportToolbarContext.h
+++ b/Engine/Source/Editor/UnrealEd/Public/ViewportToolbar/UnrealEdViewportToolbarContext.h
@@ -13,6 +13,9 @@ namespace UE::UnrealEd
 
 DECLARE_DELEGATE_RetVal_OneParam(bool, IsViewModeSupportedDelegate, EViewModeIndex);
 
+enum EHidableViewModeMenuSections : uint8;
+DECLARE_DELEGATE_RetVal_OneParam(bool, DoesViewModeMenuShowSectionDelegate, EHidableViewModeMenuSections);
+
 }
 
 class SEditorViewport;
@@ -25,4 +28,5 @@ class UNREALED_API UUnrealEdViewportToolbarContext : public UObject
 public:
 	TWeakPtr<SEditorViewport> Viewport;
 	UE::UnrealEd::IsViewModeSupportedDelegate IsViewModeSupported;
+	UE::UnrealEd::DoesViewModeMenuShowSectionDelegate DoesViewModeMenuShowSection;
 };
