Build fix MacArm64

diff --git a/Engine/Source/Editor/UnrealEd/Classes/Editor/AssetReferenceFilter.h b/Engine/Source/Editor/UnrealEd/Classes/Editor/AssetReferenceFilter.h
--- a/Engine/Source/Editor/UnrealEd/Classes/Editor/AssetReferenceFilter.h
+++ b/Engine/Source/Editor/UnrealEd/Classes/Editor/AssetReferenceFilter.h
@@ -25,5 +25,5 @@ protected:
 	UNREALED_API static bool IsCrossPluginReferenceAllowed(const FAssetData& ReferencingAssetData, const FAssetData& ReferencedAssetData);
 
 private:
-	static FOnIsCrossPluginReferenceAllowed OnIsCrossPluginReferenceAllowedDelegate;
+	UNREALED_API static FOnIsCrossPluginReferenceAllowed OnIsCrossPluginReferenceAllowedDelegate;
 };
\ No newline at end of file
