Add SetDirtyFlag to script

#rnx

diff --git a/Engine/Source/Editor/UnrealEd/Private/Subsystems/EditorAssetSubsystem.cpp b/Engine/Source/Editor/UnrealEd/Private/Subsystems/EditorAssetSubsystem.cpp
--- a/Engine/Source/Editor/UnrealEd/Private/Subsystems/EditorAssetSubsystem.cpp
+++ b/Engine/Source/Editor/UnrealEd/Private/Subsystems/EditorAssetSubsystem.cpp
@@ -1145,6 +1145,27 @@ bool UEditorAssetSubsystem::RenameDirectory(const FString& SourceDirectoryPath,
 	return true;
 }
 
+bool UEditorAssetSubsystem::SetDirtyFlag(UObject* Object, const bool bDirtyState)
+{
+	if (!IsValid(Object) || !Object->IsAsset())
+	{
+		return false;
+	}
+	UPackage* Pkg = Object->GetOutermost();
+	if (!Pkg || Pkg == GetTransientPackage() || Pkg->ContainsMap())
+	{
+		return false;
+	}
+	bool bOldDirtyValue = Pkg->IsDirty();
+	if (bOldDirtyValue != bDirtyState)
+	{
+		//If we need to change the state, return true if the state has change to the one specified
+		Pkg->SetDirtyFlag(bDirtyState);
+		return bDirtyState == Pkg->IsDirty();
+	}
+	return true;
+}
+
 // Checkout operations
  
  namespace UE::EditorAssetUtils
diff --git a/Engine/Source/Editor/UnrealEd/Public/Subsystems/EditorAssetSubsystem.h b/Engine/Source/Editor/UnrealEd/Public/Subsystems/EditorAssetSubsystem.h
--- a/Engine/Source/Editor/UnrealEd/Public/Subsystems/EditorAssetSubsystem.h
+++ b/Engine/Source/Editor/UnrealEd/Public/Subsystems/EditorAssetSubsystem.h
@@ -224,6 +224,15 @@ public:
 	UFUNCTION(BlueprintCallable, Category = "Editor Scripting | Asset")
 	UNREALED_API bool RenameDirectory(const FString& SourceDirectoryPath, const FString& DestinationDirectoryPath);
 	
+	/**
+	 * Set the package dirty flag for an asset
+	 * @param	Object			Object we want to set the package dirty state.
+	 * @param	bDirtyState		The dirty state, true mean the asset package need to be save.
+	 * @return	True if the operation succeeds.
+	 */
+	UFUNCTION(BlueprintCallable, Category = "Editor Scripting | Asset")
+	UNREALED_API bool SetDirtyFlag(UObject* Object, const bool bDirtyState);
+
 	/**
 	 * Checkout the asset corresponding to an object.
 	 * @param	AssetToCheckout		Asset to checkout.
