Do not try to run blueprint actor reconstruction on worlds that have not yet initialized, this can cause reentrant behavior that crashes

diff --git a/Engine/Source/Editor/UnrealEd/Private/Kismet2/KismetReinstanceUtilities.cpp b/Engine/Source/Editor/UnrealEd/Private/Kismet2/KismetReinstanceUtilities.cpp
--- a/Engine/Source/Editor/UnrealEd/Private/Kismet2/KismetReinstanceUtilities.cpp
+++ b/Engine/Source/Editor/UnrealEd/Private/Kismet2/KismetReinstanceUtilities.cpp
@@ -1529,34 +1529,36 @@ void FActorReplacementHelper::Finalize(const TMap<UObject*, UObject*>& OldToNewI
 	// because this is an editor context it's important to use this execution guard
 	FEditorScriptExecutionGuard ScriptGuard;
 
-	// run the construction script, which will use the properties we just copied over
-	// @TODO: This code is similar to AActor::RerunConstructionScripts and ideally could use shared code for restoring state
-
-	bool bCanReRun = UBlueprint::IsBlueprintHierarchyErrorFree(NewActor->GetClass());
-	if (NewActor->CurrentTransactionAnnotation.IsValid() && bCanReRun)
-	{
-		NewActor->CurrentTransactionAnnotation->ActorTransactionAnnotationData.ComponentInstanceData.FindAndReplaceInstances(OldToNewInstanceMap);
-		NewActor->RerunConstructionScripts();
-	}
-	else if (CachedActorData.IsValid())
+	// Only run construction if this world was constructed in the first place, it could be halfway through loading
+	UWorld* World = NewActor->GetWorld();
+	if (World && World->IsInitialized())
 	{
-		CachedActorData->ActorTransactionAnnotationData.ComponentInstanceData.FindAndReplaceInstances(OldToNewInstanceMap);
-		const bool bErrorFree = NewActor->ExecuteConstruction(TargetWorldTransform, nullptr, &CachedActorData->ActorTransactionAnnotationData.ComponentInstanceData);
-		if (!bErrorFree)
+		// run the construction script, which will use the properties we just copied over
+		// @TODO: This code is similar to AActor::RerunConstructionScripts and ideally could use shared code for restoring state
+
+		bool bCanReRun = UBlueprint::IsBlueprintHierarchyErrorFree(NewActor->GetClass());
+		if (NewActor->CurrentTransactionAnnotation.IsValid() && bCanReRun)
 		{
-			// Save off the cached actor data for once the blueprint has been fixed so we can reapply it
-			NewActor->CurrentTransactionAnnotation = CachedActorData;
+			NewActor->CurrentTransactionAnnotation->ActorTransactionAnnotationData.ComponentInstanceData.FindAndReplaceInstances(OldToNewInstanceMap);
+			NewActor->RerunConstructionScripts();
+		}
+		else if (CachedActorData.IsValid())
+		{
+			CachedActorData->ActorTransactionAnnotationData.ComponentInstanceData.FindAndReplaceInstances(OldToNewInstanceMap);
+			const bool bErrorFree = NewActor->ExecuteConstruction(TargetWorldTransform, nullptr, &CachedActorData->ActorTransactionAnnotationData.ComponentInstanceData);
+			if (!bErrorFree)
+			{
+				// Save off the cached actor data for once the blueprint has been fixed so we can reapply it
+				NewActor->CurrentTransactionAnnotation = CachedActorData;
+			}
+		}
+		else
+		{
+			FComponentInstanceDataCache DummyComponentData;
+			NewActor->ExecuteConstruction(TargetWorldTransform, nullptr, &DummyComponentData);
 		}
-	}
-	else
-	{
-		FComponentInstanceDataCache DummyComponentData;
-		NewActor->ExecuteConstruction(TargetWorldTransform, nullptr, &DummyComponentData);
-	}	
 
-	// Try to restore gameplay initialization state
-	if (UWorld* World = NewActor->GetWorld())
-	{
+		// Try to restore gameplay initialization state
 		// This is unsafe to call from a loading stack but that should never happen for an actor that was fully initialized
 		// @TODO: If there is a need for this case, it must be deferred until later in the frame
 		if (World->IsGameWorld() && bHasInitialized && ensure(!FUObjectThreadContext::Get().IsRoutingPostLoad))
@@ -1575,13 +1577,13 @@ void FActorReplacementHelper::Finalize(const TMap<UObject*, UObject*>& OldToNewI
 				NewActor->DispatchBeginPlay(false);
 			}
 		}
-	}
 
-	// Restore editor visibility
-	if (bWasHiddenEdLevel)
-	{
-		NewActor->bHiddenEdLevel = true;
-		NewActor->MarkComponentsRenderStateDirty();
+		// Restore editor visibility
+		if (bWasHiddenEdLevel)
+		{
+			NewActor->bHiddenEdLevel = true;
+			NewActor->MarkComponentsRenderStateDirty();
+		}
 	}
 
 	TMap<UObject*, UObject*> ConstructedComponentReplacementMap;
