SimpleAssetEditor: Add support for opening in read only mode
Asset Definition: Add support for opening the simple asset editor in read only mode if GetAssetOpenSupport() allows it

diff --git a/Engine/Source/Editor/AssetDefinition/Public/AssetDefinition.h b/Engine/Source/Editor/AssetDefinition/Public/AssetDefinition.h
--- a/Engine/Source/Editor/AssetDefinition/Public/AssetDefinition.h
+++ b/Engine/Source/Editor/AssetDefinition/Public/AssetDefinition.h
@@ -248,7 +248,13 @@ struct FAssetDiffArgs
 
 struct FAssetOpenSupportArgs
 {
-	EAssetOpenMethod OpenMethod = EAssetOpenMethod::Edit;
+	FAssetOpenSupportArgs(EAssetOpenMethod InOpenMethod = EAssetOpenMethod::Edit)
+		: OpenMethod(InOpenMethod)
+	{
+		
+	}
+	
+	EAssetOpenMethod OpenMethod;
 };
 
 
diff --git a/Engine/Source/Editor/UnrealEd/Private/AssetDefinitionDefault.cpp b/Engine/Source/Editor/UnrealEd/Private/AssetDefinitionDefault.cpp
--- a/Engine/Source/Editor/UnrealEd/Private/AssetDefinitionDefault.cpp
+++ b/Engine/Source/Editor/UnrealEd/Private/AssetDefinitionDefault.cpp
@@ -12,7 +12,7 @@
 
 EAssetCommandResult UAssetDefinitionDefault::OpenAssets(const FAssetOpenArgs& OpenArgs) const
 {
-	if (OpenArgs.OpenMethod == EAssetOpenMethod::Edit)
+	if (GetAssetOpenSupport(FAssetOpenSupportArgs(OpenArgs.OpenMethod)).IsSupported)
 	{
 		FSimpleAssetEditor::CreateEditor(EToolkitMode::Standalone, OpenArgs.ToolkitHost, OpenArgs.LoadObjects<UObject>());
 		return EAssetCommandResult::Handled;
diff --git a/Engine/Source/Editor/UnrealEd/Private/Toolkits/SimpleAssetEditor.cpp b/Engine/Source/Editor/UnrealEd/Private/Toolkits/SimpleAssetEditor.cpp
--- a/Engine/Source/Editor/UnrealEd/Private/Toolkits/SimpleAssetEditor.cpp
+++ b/Engine/Source/Editor/UnrealEd/Private/Toolkits/SimpleAssetEditor.cpp
@@ -27,7 +27,8 @@ void FSimpleAssetEditor::RegisterTabSpawners(const TSharedRef<class FTabManager>
 	InTabManager->RegisterTabSpawner( PropertiesTabId, FOnSpawnTab::CreateSP(this, &FSimpleAssetEditor::SpawnPropertiesTab) )
 		.SetDisplayName( LOCTEXT("PropertiesTab", "Details") )
 		.SetGroup(WorkspaceMenuCategory.ToSharedRef())
-		.SetIcon(FSlateIcon(FAppStyle::GetAppStyleSetName(), "LevelEditor.Tabs.Details"));
+		.SetIcon(FSlateIcon(FAppStyle::GetAppStyleSetName(), "LevelEditor.Tabs.Details"))
+		.SetReadOnlyBehavior(ETabReadOnlyBehavior::Custom);
 }
 
 void FSimpleAssetEditor::UnregisterTabSpawners(const TSharedRef<class FTabManager>& InTabManager)
@@ -58,6 +59,7 @@ void FSimpleAssetEditor::InitEditor( const EToolkitMode::Type Mode, const TShare
 	FDetailsViewArgs DetailsViewArgs;
 	DetailsViewArgs.NameAreaSettings = FDetailsViewArgs::HideNameArea;
 	DetailsView = PropertyEditorModule.CreateDetailView( DetailsViewArgs );
+
 	const TSharedRef<FTabManager::FLayout> StandaloneDefaultLayout = FTabManager::NewLayout( "Standalone_SimpleAssetEditor_Layout_v4" )
 	->AddArea
 	(
@@ -98,6 +100,8 @@ void FSimpleAssetEditor::InitEditor( const EToolkitMode::Type Mode, const TShare
 
 	if( DetailsView.IsValid() )
 	{
+		DetailsView->SetIsPropertyEditingEnabledDelegate(FIsPropertyEditingEnabled::CreateSP(this, &FSimpleAssetEditor::IsPropertyEditingEnabled));
+
 		// Make sure details window is pointing to our object
 		DetailsView->SetObjects( ObjectsToEditInDetailsView );
 	}
@@ -228,7 +232,7 @@ void FSimpleAssetEditor::SetPropertyVisibilityDelegate(FIsPropertyVisible InVisi
 
 void FSimpleAssetEditor::SetPropertyEditingEnabledDelegate(FIsPropertyEditingEnabled InPropertyEditingDelegate)
 {
-	DetailsView->SetIsPropertyEditingEnabledDelegate(InPropertyEditingDelegate);
+	IsPropertyEditingEnabledDelegate = InPropertyEditingDelegate;
 	DetailsView->ForceRefresh();
 }
 
@@ -321,6 +325,22 @@ FReply FSimpleAssetEditor::OnEditParentClassClicked(TObjectPtr<UObject> EditClas
 	return FReply::Handled();
 }
 
+bool FSimpleAssetEditor::IsPropertyEditingEnabled() const
+{
+	bool bIsPropertyEditingEnabled = true;
+
+	// First check any user provided delegate
+	if (IsPropertyEditingEnabledDelegate.IsBound())
+	{
+		bIsPropertyEditingEnabled &= IsPropertyEditingEnabledDelegate.Execute();
+	}
+
+	// Also make sure we are opened in Edit mode
+	bIsPropertyEditingEnabled &= (GetOpenMethod() == EAssetOpenMethod::Edit);
+
+	return bIsPropertyEditingEnabled;
+}
+
 void FSimpleAssetEditor::PostRegenerateMenusAndToolbars()
 {
 	// Find the common denominator class of the assets we're editing
diff --git a/Engine/Source/Editor/UnrealEd/Public/Toolkits/SimpleAssetEditor.h b/Engine/Source/Editor/UnrealEd/Public/Toolkits/SimpleAssetEditor.h
--- a/Engine/Source/Editor/UnrealEd/Public/Toolkits/SimpleAssetEditor.h
+++ b/Engine/Source/Editor/UnrealEd/Public/Toolkits/SimpleAssetEditor.h
@@ -58,6 +58,11 @@ protected:
 	/** Handler for "Edit parent class" button */
 	FReply OnEditParentClassClicked(TObjectPtr<UObject> EditClass) const;
 
+	/** Check if the details panel has editing enabled, checks both the asset editor open method to see if we are opened in edit mode and
+	 *  the delegate a user can supply using SetPropertyEditingEnabledDelegate
+	 */
+	bool IsPropertyEditingEnabled() const;
+
 private:
 	/** Create the properties tab and its content */
 	UNREALED_API TSharedRef<SDockTab> SpawnPropertiesTab( const FSpawnTabArgs& Args );
@@ -80,6 +85,9 @@ private:
 	/** The objects open within this editor */
 	TArray<UObject*> EditingObjects;
 
+	/** User provided delegate to check if property editing is enabled in the details panel */
+	FIsPropertyEditingEnabled IsPropertyEditingEnabledDelegate;
+
 public:
 	/** The name given to all instances of this type of editor */
 	static UNREALED_API const FName ToolkitFName;
