Fixed Concert failing to send an update for an object that was persistent and has since become transient

This should be replicated as the original object would have been subject to synchronization, and so the transient state change should also be synchronized, especially as it may also be paired with something important like a delete

#rnx

diff --git a/Engine/Plugins/Developer/Concert/ConcertSync/ConcertSyncClient/Source/ConcertSyncClient/Private/ConcertClientTransactionBridge.cpp b/Engine/Plugins/Developer/Concert/ConcertSync/ConcertSyncClient/Source/ConcertSyncClient/Private/ConcertClientTransactionBridge.cpp
--- a/Engine/Plugins/Developer/Concert/ConcertSync/ConcertSyncClient/Source/ConcertSyncClient/Private/ConcertClientTransactionBridge.cpp
+++ b/Engine/Plugins/Developer/Concert/ConcertSync/ConcertSyncClient/Source/ConcertSyncClient/Private/ConcertClientTransactionBridge.cpp
@@ -156,13 +156,33 @@ ETransactionFilterResult ApplyTransactionFilters(const TMap<FName, FOnFilterTran
 		return true;
 	};
 
+	auto WasObjectPersistent = [&IsObjectPersistent](const EObjectFlags OriginalFlags, const FName OriginalOuterPathName)
+	{
+		if (!EnumHasAnyFlags(OriginalFlags, RF_Transient))
+		{
+			if (UObject* OriginalOuter = FSoftObjectPath(FNameBuilder(OriginalOuterPathName).ToView()).ResolveObject();
+				OriginalOuter && IsObjectPersistent(OriginalOuter))
+			{
+				return true;
+			}
+		}
+		return false;
+	};
+
 	const ETransactionFilterResult FilterResult = ConcertClientTransactionBridgeUtil::ApplyCustomFilter(CustomFilters, InObject, InChangedPackage, InTransactionEvent);
 	if (FilterResult != ETransactionFilterResult::UseDefault)
 	{
 		return FilterResult;
 	}
-	// Ignore transient packages and objects, compiled in package are not considered Multi-user content.
-	if (!InChangedPackage || InChangedPackage == GetTransientPackage() || InChangedPackage->HasAnyPackageFlags(PKG_CompiledIn) || !IsObjectPersistent(InObject))
+
+	// Ignore compiled in package, as they are not considered Multi-user content.
+	if (!InChangedPackage || InChangedPackage->HasAnyPackageFlags(PKG_CompiledIn))
+	{
+		return ETransactionFilterResult::ExcludeObject;
+	}
+
+	// Ignore transient packages and objects, unless the object was previously persistent
+	if ((InChangedPackage == GetTransientPackage() || !IsObjectPersistent(InObject)) && !WasObjectPersistent(InTransactionEvent.GetOriginalObjectFlags(), InTransactionEvent.GetOriginalObjectOuterPathName()))
 	{
 		return ETransactionFilterResult::ExcludeObject;
 	}
@@ -278,7 +298,7 @@ struct FEditorTransactionNotification
 					DeltaChange.ChangedProperties.Add(PropertyData.PropertyName);
 				}
 				TransactionObjectEvent = FTransactionObjectEvent(TransactionContext.TransactionId, TransactionContext.OperationId, ETransactionObjectEventType::UndoRedo, ETransactionObjectChangeCreatedBy::TransactionRecord,
-					FTransactionObjectChange{ InObjectUpdate.ObjectId.ToTransactionObjectId(), MoveTemp(DeltaChange) }, InTransactionAnnotation);
+					FTransactionObjectChange{ InObjectUpdate.ObjectId.ToTransactionObjectId(), (EObjectFlags)InObjectUpdate.ObjectId.ObjectPersistentFlags, MoveTemp(DeltaChange) }, InTransactionAnnotation);
 			}
 			GUnrealEd->HandleObjectTransacted(InTransactionObject, TransactionObjectEvent);
 		}
diff --git a/Engine/Plugins/Developer/Concert/ConcertUI/ConcertSharedSlate/Source/ConcertSharedSlate/Private/Session/Activity/SConcertSessionActivities.cpp b/Engine/Plugins/Developer/Concert/ConcertUI/ConcertSharedSlate/Source/ConcertSharedSlate/Private/Session/Activity/SConcertSessionActivities.cpp
--- a/Engine/Plugins/Developer/Concert/ConcertUI/ConcertSharedSlate/Source/ConcertSharedSlate/Private/Session/Activity/SConcertSessionActivities.cpp
+++ b/Engine/Plugins/Developer/Concert/ConcertUI/ConcertSharedSlate/Source/ConcertSharedSlate/Private/Session/Activity/SConcertSessionActivities.cpp
@@ -865,7 +865,7 @@ void SConcertSessionActivities::DisplayTransactionDetails(const FConcertSessionA
 
 		FTransactionObjectId TransactionObjectId = ExportedObject.ObjectId.ToTransactionObjectId();
 		TSharedPtr<FTransactionObjectEvent> Event = MakeShared<FTransactionObjectEvent>(InTransaction.TransactionId, InTransaction.OperationId, ETransactionObjectEventType::Finalized, ETransactionObjectChangeCreatedBy::TransactionRecord, 
-			FTransactionObjectChange{ TransactionObjectId, MoveTemp(DeltaChange) }, nullptr);
+			FTransactionObjectChange{ TransactionObjectId, (EObjectFlags)ExportedObject.ObjectId.ObjectPersistentFlags, MoveTemp(DeltaChange) }, nullptr);
 
 		TransactionDiff.DiffMap.Emplace(TransactionObjectId.ObjectPathName, MoveTemp(Event));
 	}
diff --git a/Engine/Source/Editor/UnrealEd/Classes/Editor/Transactor.h b/Engine/Source/Editor/UnrealEd/Classes/Editor/Transactor.h
--- a/Engine/Source/Editor/UnrealEd/Classes/Editor/Transactor.h
+++ b/Engine/Source/Editor/UnrealEd/Classes/Editor/Transactor.h
@@ -54,6 +54,7 @@ protected:
 			void SetObject(const UObject* InObject)
 			{
 				ObjectId.SetObject(InObject);
+				ObjectFlags = InObject->GetFlags();
 				ObjectAnnotation = InObject->FindOrCreateTransactionAnnotation();
 			}
 
@@ -61,6 +62,7 @@ protected:
 			{
 				UE::Transaction::FSerializedObject::Reset();
 				ObjectId.Reset();
+				ObjectFlags = RF_NoFlags;
 				ObjectAnnotation.Reset();
 				PendingKillChange = EPendingKillChange::None;
 			}
@@ -69,6 +71,7 @@ protected:
 			{
 				UE::Transaction::FSerializedObject::Swap(Other);
 				ObjectId.Swap(Other.ObjectId);
+				Exchange(ObjectFlags, Other.ObjectFlags);
 				Exchange(ObjectAnnotation, Other.ObjectAnnotation);
 				Exchange(PendingKillChange, Other.PendingKillChange);
 			}
@@ -76,6 +79,9 @@ protected:
 			/** ID of the object when it was serialized */
 			FTransactionObjectId ObjectId;
 
+			/** The flags of the object when it was serialized */
+			EObjectFlags ObjectFlags = RF_NoFlags;
+
 			/** Annotation data for the object stored externally */
 			TSharedPtr<ITransactionObjectAnnotation> ObjectAnnotation;
 
diff --git a/Engine/Source/Editor/UnrealEd/Private/EditorTransaction.cpp b/Engine/Source/Editor/UnrealEd/Private/EditorTransaction.cpp
--- a/Engine/Source/Editor/UnrealEd/Private/EditorTransaction.cpp
+++ b/Engine/Source/Editor/UnrealEd/Private/EditorTransaction.cpp
@@ -498,7 +498,7 @@ void FTransaction::FObjectRecord::Snapshot( FTransaction* Owner, UE::Transaction
 				ObjectAnnotationSnapshot->ComputeAdditionalObjectChanges(InitialObjectTransactionAnnotation.Get(), AdditionalObjectChanges);
 			}
 
-			CurrentObject->PostTransacted(FTransactionObjectEvent(Owner->GetId(), Owner->GetOperationId(), ETransactionObjectEventType::Snapshot, ETransactionObjectChangeCreatedBy::TransactionRecord, FTransactionObjectChange{ InitialDiffableObject.ObjectInfo, SnapshotDeltaChange }, ObjectAnnotationSnapshot));
+			CurrentObject->PostTransacted(FTransactionObjectEvent(Owner->GetId(), Owner->GetOperationId(), ETransactionObjectEventType::Snapshot, ETransactionObjectChangeCreatedBy::TransactionRecord, FTransactionObjectChange{ InitialDiffableObject.ObjectInfo, InitialDiffableObject.ObjectFlags, SnapshotDeltaChange }, ObjectAnnotationSnapshot));
 			for (const TTuple<UObject*, FTransactionObjectChange>& AdditionalObjectChangePair : AdditionalObjectChanges)
 			{
 				AdditionalObjectChangePair.Key->PostTransacted(FTransactionObjectEvent(Owner->GetId(), Owner->GetOperationId(), ETransactionObjectEventType::Snapshot, ETransactionObjectChangeCreatedBy::TransactionAnnotation, AdditionalObjectChangePair.Value, nullptr));
@@ -987,7 +987,7 @@ void FTransaction::Apply()
 				ChangedObjectTransactionAnnotation->ComputeAdditionalObjectChanges(InitialSerializedObject.ObjectAnnotation.Get(), AdditionalObjectChanges);
 			}
 
-			ChangedObject->PostTransacted(FTransactionObjectEvent(Id, OperationId, ETransactionObjectEventType::UndoRedo, ETransactionObjectChangeCreatedBy::TransactionRecord, FTransactionObjectChange{ InitialSerializedObject.ObjectId, MoveTemp(PrimaryDeltaChange) }, ChangedObjectTransactionAnnotation));
+			ChangedObject->PostTransacted(FTransactionObjectEvent(Id, OperationId, ETransactionObjectEventType::UndoRedo, ETransactionObjectChangeCreatedBy::TransactionRecord, FTransactionObjectChange{ InitialSerializedObject.ObjectId, InitialSerializedObject.ObjectFlags, MoveTemp(PrimaryDeltaChange) }, ChangedObjectTransactionAnnotation));
 			for (const TTuple<UObject*, FTransactionObjectChange>& AdditionalObjectChangePair : AdditionalObjectChanges)
 			{
 				AdditionalObjectChangePair.Key->PostTransacted(FTransactionObjectEvent(Id, OperationId, ETransactionObjectEventType::UndoRedo, ETransactionObjectChangeCreatedBy::TransactionAnnotation, AdditionalObjectChangePair.Value, nullptr));
@@ -1069,7 +1069,7 @@ void FTransaction::Finalize()
 				ChangedObjectTransactionAnnotation->ComputeAdditionalObjectChanges(ChangedObjectIt.Value.AnnotationSnapshot.Get(), AdditionalObjectChanges);
 			}
 
-			ChangedObject->PostTransacted(FTransactionObjectEvent(Id, OperationId, ETransactionObjectEventType::Finalized, ETransactionObjectChangeCreatedBy::TransactionRecord, FTransactionObjectChange{ InitialSerializedObject.ObjectId, DeltaChange }, ChangedObjectTransactionAnnotation));
+			ChangedObject->PostTransacted(FTransactionObjectEvent(Id, OperationId, ETransactionObjectEventType::Finalized, ETransactionObjectChangeCreatedBy::TransactionRecord, FTransactionObjectChange{ InitialSerializedObject.ObjectId, InitialSerializedObject.ObjectFlags, DeltaChange }, ChangedObjectTransactionAnnotation));
 			for (const TTuple<UObject*, FTransactionObjectChange>& AdditionalObjectChangePair : AdditionalObjectChanges)
 			{
 				AdditionalObjectChangePair.Key->PostTransacted(FTransactionObjectEvent(Id, OperationId, ETransactionObjectEventType::Finalized, ETransactionObjectChangeCreatedBy::TransactionAnnotation, AdditionalObjectChangePair.Value, nullptr));
@@ -1127,7 +1127,7 @@ FTransactionDiff FTransaction::GenerateDiff() const
 
 					TransactionDiff.DiffMap.Emplace(
 						FName(*TransactedObject->GetPathName()), 
-						MakeShared<FTransactionObjectEvent>(this->GetId(), OperationGuid, ETransactionObjectEventType::Finalized, ETransactionObjectChangeCreatedBy::TransactionRecord, FTransactionObjectChange{ ObjectRecord.SerializedObject.ObjectId, ObjectRecord.DeltaChange }, ObjectRecord.SerializedObject.ObjectAnnotation)
+						MakeShared<FTransactionObjectEvent>(this->GetId(), OperationGuid, ETransactionObjectEventType::Finalized, ETransactionObjectChangeCreatedBy::TransactionRecord, FTransactionObjectChange{ ObjectRecord.SerializedObject.ObjectId, ObjectRecord.SerializedObject.ObjectFlags, ObjectRecord.DeltaChange }, ObjectRecord.SerializedObject.ObjectAnnotation)
 						);
 
 					if (ObjectRecord.SerializedObjectFlip.ObjectAnnotation && ObjectRecord.SerializedObject.ObjectAnnotation)
diff --git a/Engine/Source/Runtime/CoreUObject/Public/Misc/TransactionObjectEvent.h b/Engine/Source/Runtime/CoreUObject/Public/Misc/TransactionObjectEvent.h
--- a/Engine/Source/Runtime/CoreUObject/Public/Misc/TransactionObjectEvent.h
+++ b/Engine/Source/Runtime/CoreUObject/Public/Misc/TransactionObjectEvent.h
@@ -5,6 +5,7 @@
 #include "CoreTypes.h"
 #include "Misc/Guid.h"
 #include "UObject/NameTypes.h"
+#include "UObject/ObjectMacros.h"
 #include "Templates/SharedPointer.h"
 
 class ITransactionObjectAnnotation;
@@ -142,6 +143,9 @@ struct FTransactionObjectChange
 	/** Original ID of the object that changed */
 	FTransactionObjectId OriginalId;
 
+	/** Original flags of the object that changed */
+	EObjectFlags OriginalObjectFlags = RF_NoFlags;
+
 	/** Information about how the object changed */
 	FTransactionObjectDeltaChange DeltaChange;
 };
@@ -190,14 +194,6 @@ public:
 		check(OperationId.IsValid());
 	}
 
-	UE_DEPRECATED(5.1, "Use the constructor that takes a FTransactionObjectChange.")
-	FTransactionObjectEvent(const FGuid& InTransactionId, const FGuid& InOperationId, const ETransactionObjectEventType InEventType, const FTransactionObjectDeltaChange& InDeltaChange, const TSharedPtr<ITransactionObjectAnnotation>& InAnnotation
-		, const FName InOriginalObjectPackageName, const FName InOriginalObjectName, const FName InOriginalObjectPathName, const FName InOriginalObjectOuterPathName, const FName InOriginalObjectExternalPackageName, const FName InOriginalObjectClassPathName)
-		: FTransactionObjectEvent(InTransactionId, InOperationId, InEventType, ETransactionObjectChangeCreatedBy::TransactionRecord, 
-			FTransactionObjectChange{ FTransactionObjectId(InOriginalObjectPackageName, InOriginalObjectName, InOriginalObjectPathName, InOriginalObjectOuterPathName, InOriginalObjectExternalPackageName, InOriginalObjectClassPathName), InDeltaChange }, InAnnotation)
-	{
-	}
-
 	/** The unique identifier of the transaction this event belongs to */
 	const FGuid& GetTransactionId() const
 	{
@@ -240,6 +236,12 @@ public:
 		return ObjectChange.OriginalId;
 	}
 
+	/** Get the original flags of this object */
+	EObjectFlags GetOriginalObjectFlags() const
+	{
+		return ObjectChange.OriginalObjectFlags;
+	}
+
 	/** Get the original package name of this object */
 	FName GetOriginalObjectPackageName() const
 	{
diff --git a/Engine/Source/Runtime/Engine/Private/ActorEditor.cpp b/Engine/Source/Runtime/Engine/Private/ActorEditor.cpp
--- a/Engine/Source/Runtime/Engine/Private/ActorEditor.cpp
+++ b/Engine/Source/Runtime/Engine/Private/ActorEditor.cpp
@@ -629,7 +629,7 @@ void FActorTransactionAnnotation::ComputeAdditionalObjectChanges(const ITransact
 				}
 				else
 				{
-					OutAdditionalObjectChanges.Add(CurrentComponent, FTransactionObjectChange{ OldDiffableComponent.ObjectInfo, MoveTemp(ComponentDeltaChange) });
+					OutAdditionalObjectChanges.Add(CurrentComponent, FTransactionObjectChange{ OldDiffableComponent.ObjectInfo, OldDiffableComponent.ObjectFlags, MoveTemp(ComponentDeltaChange) });
 				}
 			}
 		};
diff --git a/Engine/Source/Runtime/Engine/Public/TransactionCommon.h b/Engine/Source/Runtime/Engine/Public/TransactionCommon.h
--- a/Engine/Source/Runtime/Engine/Public/TransactionCommon.h
+++ b/Engine/Source/Runtime/Engine/Public/TransactionCommon.h
@@ -254,12 +254,14 @@ public:
 	{
 		ObjectInfo.SetObject(InObject);
 		ObjectArchetype = FPersistentObjectRef(InObject->GetArchetype());
+		ObjectFlags = InObject->GetFlags();
 	}
 
 	void Reset()
 	{
 		ObjectInfo.Reset();
 		ObjectArchetype = FPersistentObjectRef();
+		ObjectFlags = RF_NoFlags;
 		SerializedData.Reset();
 		SerializedTaggedData.Reset();
 	}
@@ -268,6 +270,7 @@ public:
 	{
 		ObjectInfo.Swap(Other.ObjectInfo);
 		Exchange(ObjectArchetype, Other.ObjectArchetype);
+		Exchange(ObjectFlags, Other.ObjectFlags);
 		Exchange(SerializedData, Other.SerializedData);
 		Exchange(SerializedTaggedData, Other.SerializedTaggedData);
 	}
@@ -278,6 +281,9 @@ public:
 	/** The archetype of the object when it was serialized */
 	FPersistentObjectRef ObjectArchetype;
 
+	/** The flags of the object when it was serialized */
+	EObjectFlags ObjectFlags = RF_NoFlags;
+
 	/** The serialized data for the diffable object */
 	FSerializedObjectData SerializedData;
 
