[Spline Visualizer] Fix an issue with spline component auto selection, where the clicked on point would not get the focus on the first click when the actor is selected already.

diff --git a/Engine/Source/Editor/UnrealEd/Private/ComponentVisualizerManager.cpp b/Engine/Source/Editor/UnrealEd/Private/ComponentVisualizerManager.cpp
--- a/Engine/Source/Editor/UnrealEd/Private/ComponentVisualizerManager.cpp
+++ b/Engine/Source/Editor/UnrealEd/Private/ComponentVisualizerManager.cpp
@@ -88,8 +88,7 @@ bool FComponentVisualizerManager::HandleProxyForComponentVis(FEditorViewportClie
 					{
 						if (!ElementSelectionSet->IsElementSelected(SelectedComponentHandle, FTypedElementIsSelectedOptions()))
 						{
-							TArray<FTypedElementHandle> TmpArray = { SelectedComponentHandle };
-							ElementSelectionSet->SetSelection(TmpArray, FTypedElementSelectionOptions());
+							ElementSelectionSet->SelectElement(SelectedComponentHandle, FTypedElementSelectionOptions());
 						}
 					}
 
diff --git a/Engine/Source/Editor/UnrealEd/Private/EditorSelectUtils.cpp b/Engine/Source/Editor/UnrealEd/Private/EditorSelectUtils.cpp
--- a/Engine/Source/Editor/UnrealEd/Private/EditorSelectUtils.cpp
+++ b/Engine/Source/Editor/UnrealEd/Private/EditorSelectUtils.cpp
@@ -379,6 +379,12 @@ void UUnrealEdEngine::OnEditorElementSelectionChanged(const UTypedElementSelecti
 				bDidSetVizThisTick = true;
 				break;
 			}
+
+			if (ComponentVisManager.IsActive() && VisualizerForSelection.ComponentVisualizer.Visualizer->GetEditedComponent() == ComponentVisManager.GetActiveComponentVis()->GetEditedComponent())
+			{
+				bDidSetVizThisTick = true;
+				break;
+			}
 		}
 	}
 
