Preload StructUtilsEditor before running the asset localization gather, as this module defines a gather override for FInstancedStruct

This fixes a regression introduced when FInstancedStruct moved to CoreUObject (34509884), that caused its localization data to no longer be gathered

#rnx

diff --git a/Engine/Source/Editor/UnrealEd/Private/Commandlets/GatherTextFromAssetsCommandlet.cpp b/Engine/Source/Editor/UnrealEd/Private/Commandlets/GatherTextFromAssetsCommandlet.cpp
--- a/Engine/Source/Editor/UnrealEd/Private/Commandlets/GatherTextFromAssetsCommandlet.cpp
+++ b/Engine/Source/Editor/UnrealEd/Private/Commandlets/GatherTextFromAssetsCommandlet.cpp
@@ -905,6 +905,13 @@ bool UGatherTextFromAssetsCommandlet::ParseCommandLineHelper(const FString& InCo
 			return false;
 		}
 	}
+
+	// Preload modules that register custom gather callbacks
+	if (!FModuleManager::LoadModulePtr<IModuleInterface>("StructUtilsEditor"))
+	{
+		UE_LOG(LogGatherTextFromAssetsCommandlet, Warning, TEXT("Failed to preload StructUtilsEditor. FInstancedStruct gathering may be incomplete."));
+	}
+
 	return true;
 }
 
