Improve WorldPartition actors / OFPA material support for ODSC

Since the path from the material is not enough to guess what package it belongs to, add the current FPrimitiveSceneInfo to retrieve the actor path. In order to not change a lot of functions' signatures, rely on FODSCPrimitiveSceneInfoScope that sets up a thread local variable called ODSCManagerPrivate::CurrentPrimitiveSceneInfo
When we have a new request, append the actor path in addition to the material path, separated by ':::'
On the server, this allows to retrieve the package path since we now have the actor path
Some paths still require some patching, like those starting with /MyMountPoint/CB/ which point to ContentBundle from plugins. The logic is done in FODSCClientDataAccess::TryFindInContentBundle
Add odsc.logmissedmaterials to see on the client the list of materials which failed to compiled on the ODSC server

#rnx

diff --git a/Engine/Source/Editor/UnrealEd/Private/Cooker/OnDemandShaderCompilation.cpp b/Engine/Source/Editor/UnrealEd/Private/Cooker/OnDemandShaderCompilation.cpp
--- a/Engine/Source/Editor/UnrealEd/Private/Cooker/OnDemandShaderCompilation.cpp
+++ b/Engine/Source/Editor/UnrealEd/Private/Cooker/OnDemandShaderCompilation.cpp
@@ -10,6 +10,7 @@
 #include "UObject/CoreRedirects.h"
 #include "PackageTools.h"
 #include "WorldPartition/WorldPartitionHelpers.h"
+#include "WorldPartition/ContentBundle/ContentBundlePaths.h"
 
 
 int32 GODSCShaderMapsLifetime = 25;
@@ -40,6 +41,12 @@ static FAutoConsoleVariableRef CVarODSCExludedClasses(
 namespace UE::Cook
 {
 
+class FODSCClientDataAccess
+{
+public:
+	static FODSCClientData::FWorldPartitionAssets* TryFindInContentBundle(FSoftObjectPath& AssetSoftPath);
+};
+
 void FODSCClientData::OnClientConnected(const void* ConnectionPtr)
 {
 }
@@ -132,17 +139,54 @@ void FODSCClientData::CleanupWorldPartitionAssets()
 		PackagesToUnloadArray.Add(Package);
 	}
 
-	UPackageTools::UnloadPackages(PackagesToUnloadArray);
+	FText OutErrorMessage;
+	// bUnloadDirtyPackages=true because some systems (UPCGGraphInstance::RefreshParameters for example) mark the package dirty 
+	// and prevent the unloading from happening
+	UPackageTools::UnloadPackages(PackagesToUnloadArray, OutErrorMessage, true /*bUnloadDirtyPackages*/);
+
+	if (!OutErrorMessage.IsEmpty())
+	{
+		UE_LOG(LogShaders, Error, TEXT("UPackageTools::UnloadPackages: %s"), *OutErrorMessage.ToString());
+	}
 }
 
-bool ExtractMaterialPath(FSoftObjectPath& MaterialPath, const FString& MaterialKey)
+bool ExtractMaterialPath(FSoftObjectPath& MaterialPath, FSoftObjectPath& ActorPath, const FString& MaterialKey)
 {
-	FString MaterialPathString = MaterialKey;
+	FString MaterialPathString;
+	FString ActorPathString;
+	int32 ActorSeparatorIndex = MaterialKey.Find(":::");
+	if (ActorSeparatorIndex != INDEX_NONE)
+	{
+		MaterialPathString = MaterialKey.Left(ActorSeparatorIndex);
+		ActorPathString = MaterialKey.Right(MaterialKey.Len() - ActorSeparatorIndex - 3);
+	}
+	else
+	{
+		MaterialPathString = MaterialKey;
+	}
+
+	bool bValidActorPath = false;
+	if (!ActorPathString.IsEmpty())
+	{
+		if (!FWorldPartitionHelpers::ConvertRuntimePathToEditorPath(ActorPathString, ActorPath))
+		{
+			ActorPathString.ReplaceInline(TEXT("/_Generated_/"), TEXT("/"));
+			ActorPath = ActorPathString;
+		}
+		else
+		{
+			bValidActorPath = true;
+		}
+	}
+
 	if (!FWorldPartitionHelpers::ConvertRuntimePathToEditorPath(MaterialPathString, MaterialPath))
 	{
 		MaterialPathString.ReplaceInline(TEXT("/_Generated_/"), TEXT("/"));
 		MaterialPath = MaterialPathString;
-		return false;
+		if (!bValidActorPath)
+		{
+			return false;
+		}
 	}
 	return true;
 }
@@ -172,13 +216,14 @@ UMaterialInterface* FODSCClientData::FindMaterial(const FString& InMaterialKey)
 	};
 
 	FSoftObjectPath MaterialPath;
-	bool bIsWorldPartitionPath = ExtractMaterialPath(MaterialPath, InMaterialKey);
+	FSoftObjectPath ActorPath;
+	bool bIsWorldPartitionPath = ExtractMaterialPath(MaterialPath, ActorPath, InMaterialKey);
 
 	UMaterialInterface* MaterialInterface = nullptr;
 
 	if (bIsWorldPartitionPath)
 	{
-		MaterialInterface = TryFindWorldPartitionMaterial(MaterialPath);
+		MaterialInterface = TryFindWorldPartitionMaterial(MaterialPath, ActorPath);
 		if (MaterialInterface)
 		{
 			return MaterialInterface;
@@ -196,65 +241,169 @@ UMaterialInterface* FODSCClientData::FindMaterial(const FString& InMaterialKey)
 	return MaterialInterface;
 }
 
-FString GetWorldPartitionActorPath(const FString& MaterialEditorFullPath)
+FSoftObjectPath GetWorldPartitionActorPath(const FSoftObjectPath& InMaterialSoftPath)
 {
-	FString PathToProxy = MaterialEditorFullPath;
+	FString PathToProxy = InMaterialSoftPath.ToString();
 	// Remove the landscape prefix since Landscape MIC are embedded in the package of their proxy
 	FString LandscapeMaterialInstanceConstantClassName = ULandscapeMaterialInstanceConstant::StaticClass()->GetFName().ToString();
 	int32 SubObjectNameChopIndex = PathToProxy.Find(FString(".") + LandscapeMaterialInstanceConstantClassName, ESearchCase::IgnoreCase, ESearchDir::FromEnd);
 	if (SubObjectNameChopIndex != INDEX_NONE)
 	{
 		PathToProxy.LeftChopInline(PathToProxy.Len() - SubObjectNameChopIndex);
+		return FSoftObjectPath(PathToProxy);
 	}
 
-	return PathToProxy;
+	return FSoftObjectPath();
 }
 
-UMaterialInterface* FODSCClientData::TryFindWorldPartitionMaterial(const FSoftObjectPath& MaterialSoftPath)
+FODSCClientData::FWorldPartitionAssets* FODSCClientDataAccess::TryFindInContentBundle(FSoftObjectPath& AssetSoftPath)
 {
-	FString AssetPath = MaterialSoftPath.GetAssetPath().GetPackageName().ToString();
+	TStringBuilder<256> ActualMountPointPackageName;
+	TStringBuilder<256> ActualMountPointFilePath;
+	TStringBuilder<256> ActualRelPath;
 
-	bool bAlreadySeenPath;
-	ScannedWorldPartitionPaths.FindOrAdd(AssetPath, &bAlreadySeenPath);
-	if (!bAlreadySeenPath)
+	FPackageName::TryGetMountPointForPath(AssetSoftPath.ToString(), ActualMountPointPackageName, ActualMountPointFilePath, ActualRelPath);
+
+	FString ActualRelPathString(ActualRelPath);
+	
+	FODSCClientData::FWorldPartitionAssets* DynamicMaterialData = nullptr;
+	// WP actors with a path like /MyMountPoint/CB/ have their paths actually remapped to /Game/
+	if (ActualRelPathString.StartsWith(TEXT("CB/")))
+	{
+		FString ActualRelPathStringCopy(TEXT("/Game"));
+		ActualRelPathStringCopy += ActualRelPathString.Right(ActualRelPathString.Len() - 2);
+		DynamicMaterialData = FODSCClientData::WorldPartitionAssets.Find(ActualRelPathStringCopy);
+		if (DynamicMaterialData)
+		{
+			AssetSoftPath.SetPath(ActualRelPathStringCopy);
+		}
+	}
+
+	return DynamicMaterialData;
+}
+
+
+UMaterialInterface* FODSCClientData::TryFindWorldPartitionMaterial(const FSoftObjectPath& InMaterialSoftPath, const FSoftObjectPath& InActorSoftPath)
+{
+	FSoftObjectPath ActorSoftPath(InActorSoftPath);
+	FSoftObjectPath MaterialSoftPath(InMaterialSoftPath);
+
+	ScanWorldPartitionAssets(MaterialSoftPath.GetAssetPath().GetPackageName().ToString());
+	ScanWorldPartitionAssets(ActorSoftPath.GetAssetPath().GetPackageName().ToString());
+
+	FWorldPartitionAssets* DynamicMaterialData = WorldPartitionAssets.Find(*MaterialSoftPath.ToString());
+	// Landscape sometimes issues requests without actors. Try to reconstruct the path 
+	if (DynamicMaterialData == nullptr && !ActorSoftPath.IsValid())
+	{
+		ActorSoftPath = GetWorldPartitionActorPath(MaterialSoftPath);
+	}
+
+	if (DynamicMaterialData == nullptr && ActorSoftPath.IsValid())
+	{
+		DynamicMaterialData = WorldPartitionAssets.Find(ActorSoftPath.ToString());
+	}
+
+	if (DynamicMaterialData == nullptr && MaterialSoftPath.IsValid())
 	{
-		ScanWorldPartitionAssets(AssetPath);
+		DynamicMaterialData = FODSCClientDataAccess::TryFindInContentBundle(MaterialSoftPath);
 	}
 
-	FString MaterialEditorFullPath = *MaterialSoftPath.ToString();
-	FString PathToProxy = GetWorldPartitionActorPath(MaterialEditorFullPath);
+	if (DynamicMaterialData == nullptr && ActorSoftPath.IsValid())
+	{
+		DynamicMaterialData = FODSCClientDataAccess::TryFindInContentBundle(ActorSoftPath);
+	}
 
-	FWorldPartitionAssets* DynamicMaterialData = WorldPartitionAssets.Find(PathToProxy);
 	if (DynamicMaterialData == nullptr)
 	{
 		return nullptr;
 	}
 
-	if (DynamicMaterialData->PackagePtr == nullptr)
+	if (!DynamicMaterialData->PackagePtr)
 	{
 		TRACE_CPUPROFILER_EVENT_SCOPE_TEXT(*FString::Printf(TEXT("FODSCClientData_LoadPackage %s"), *DynamicMaterialData->PackageName));
 		DynamicMaterialData->PackagePtr = LoadPackage(nullptr, *DynamicMaterialData->PackageName, LOAD_None);
 	}
 
-	return FindObject<UMaterialInterface>(nullptr, *MaterialEditorFullPath);
+	TArray<FString> MaterialPathsToTry;
+	MaterialPathsToTry.Add(MaterialSoftPath.ToString());
+
+	{
+		// When the provided material path doesn't work, try to replace the base path by the package's
+		FTopLevelAssetPath MaterialTopPath;
+		MaterialTopPath.TrySetPath(FName(DynamicMaterialData->PackageName), MaterialSoftPath.GetAssetPath().GetAssetName());
+		FSoftObjectPath MaterialSoftPathCopy(MaterialTopPath, MaterialSoftPath.GetSubPathString());
+		MaterialPathsToTry.Add(MaterialSoftPathCopy.ToString());
+	}
+
+	UMaterialInterface* MaterialInterface = nullptr;
+	for (const FString& MaterialPathToTry : MaterialPathsToTry)
+	{
+		MaterialInterface = FindObject<UMaterialInterface>(nullptr, *MaterialPathToTry);
+		if (MaterialInterface)
+		{
+			return MaterialInterface;
+		}
+	}
+
+	for (const FString& MaterialPathToTry : MaterialPathsToTry)
+	{
+		MaterialInterface = LoadObject<UMaterialInterface>(nullptr, *MaterialPathToTry);
+		if (MaterialInterface)
+		{
+			return MaterialInterface;
+		}
+	}
+
+	return MaterialInterface;
+
 }
 
-void FODSCClientData::ScanWorldPartitionAssets(const FString& AssetPath)
+void FODSCClientData::ScanWorldPartitionAssets(const FString& InAssetPath)
 {
-	if (AssetPath.IsEmpty())
+	if (InAssetPath.IsEmpty())
+	{
+		return;
+	}
+
+	bool bAlreadySeenPath;
+	ScannedWorldPartitionPaths.FindOrAdd(InAssetPath, &bAlreadySeenPath);
+	if (bAlreadySeenPath)
 	{
 		return;
 	}
 
+	FString AssetPath = InAssetPath;
+
 	TArray<FString> PathsToScan;
 	PathsToScan.Add(AssetPath);
 	PathsToScan.Append(ULevel::GetExternalObjectsPaths(AssetPath));
 
+	{
+		TStringBuilder<256> ActualMountPointPackageName;
+		TStringBuilder<256> ActualMountPointFilePath;
+		TStringBuilder<256> ActualRelPath;
+
+		FPackageName::TryGetMountPointForPath(InAssetPath, ActualMountPointPackageName, ActualMountPointFilePath, ActualRelPath);
+
+		// If we have /MyOtherMountPoint/CB/ as a base path, try scanning the external folders' content bundle as well
+		FString ActualMountPointPackageNameString(ActualMountPointPackageName);
+		if (!ActualMountPointPackageNameString.StartsWith(TEXT("/Game/")))
+		{
+			FString ActualRelPathStr(ActualRelPath);
+			if (ActualRelPathStr.StartsWith(TEXT("CB/")))
+			{
+				PathsToScan.Add(ActualMountPointPackageNameString + FPackagePath::GetExternalActorsFolderName() + TEXT("/ContentBundle/"));
+				PathsToScan.Add(ActualMountPointPackageNameString + FPackagePath::GetExternalObjectsFolderName() + TEXT("/ContentBundle/"));
+			}
+		}
+	}
+
+
 	// Do a synchronous scan of the level external actors path.					
 	IAssetRegistry& AssetRegistry = IAssetRegistry::GetChecked();
 	{
 		TRACE_CPUPROFILER_EVENT_SCOPE(FODSCClientData_ScanSynchronous);
-		AssetRegistry.ScanSynchronous(PathsToScan, TArray<FString>());
+		AssetRegistry.ScanSynchronous(PathsToScan, TArray<FString>(), UE::AssetRegistry::EScanFlags::IgnoreInvalidPathWarning);
 	}
 
 	FARFilter Filter;
diff --git a/Engine/Source/Editor/UnrealEd/Private/Cooker/OnDemandShaderCompilation.h b/Engine/Source/Editor/UnrealEd/Private/Cooker/OnDemandShaderCompilation.h
--- a/Engine/Source/Editor/UnrealEd/Private/Cooker/OnDemandShaderCompilation.h
+++ b/Engine/Source/Editor/UnrealEd/Private/Cooker/OnDemandShaderCompilation.h
@@ -22,12 +22,13 @@ public:
 
 private:
 	
-	static UMaterialInterface* TryFindWorldPartitionMaterial(const FSoftObjectPath& MaterialSoftPath);
+	friend class FODSCClientDataAccess;
+	static UMaterialInterface* TryFindWorldPartitionMaterial(const FSoftObjectPath& MaterialSoftPath, const FSoftObjectPath& ActorSoftPath);
 
 	static void ScanWorldPartitionAssets(const FString& AssetPath);
 	static void SetupClassExclusionList();
 
-	void CleanupWorldPartitionAssets();
+	static void CleanupWorldPartitionAssets();
 	struct FWorldPartitionAssets
 	{
 		FString PackageName;
diff --git a/Engine/Source/Runtime/Engine/Private/Materials/MaterialShader.cpp b/Engine/Source/Runtime/Engine/Private/Materials/MaterialShader.cpp
--- a/Engine/Source/Runtime/Engine/Private/Materials/MaterialShader.cpp
+++ b/Engine/Source/Runtime/Engine/Private/Materials/MaterialShader.cpp
@@ -2095,6 +2095,12 @@ void FMaterialShaderMap::LoadForRemoteRecompile(FArchive& Ar, EShaderPlatform Sh
 		FMaterialShaderMapData& MaterialShaderMapData = MaterialShaderMapDataArray.Emplace_GetRef();
 		Ar << MaterialShaderMapData.MaterialName;
 
+		int32 ActorSeparator = MaterialShaderMapData.MaterialName.Find(TEXT(":::"));
+		if (ActorSeparator != INDEX_NONE)
+		{
+			MaterialShaderMapData.MaterialName.LeftInline(ActorSeparator);
+		}
+
 		int32 NumShaderMaps = 0;
 		Ar << NumShaderMaps;
 
diff --git a/Engine/Source/Runtime/Engine/Private/ODSC/ODSCManager.cpp b/Engine/Source/Runtime/Engine/Private/ODSC/ODSCManager.cpp
--- a/Engine/Source/Runtime/Engine/Private/ODSC/ODSCManager.cpp
+++ b/Engine/Source/Runtime/Engine/Private/ODSC/ODSCManager.cpp
@@ -16,6 +16,10 @@
 // For GetCachedScalabilityCVars
 #include "UnrealEngine.h"
 
+#include "Engine/Console.h"
+#include "Engine/GameViewportClient.h"
+#include "EngineUtils.h"
+
 DEFINE_LOG_CATEGORY(LogODSC);
 
 // FODSCManager
@@ -30,10 +34,53 @@ static TAutoConsoleVariable<int32> CVarODSCRecompileMode(
 
 FODSCManager* GODSCManager = nullptr;
 
+class FODSCManagerAccess
+{
+public:
+	static void ODSCLogMissedMaterials(const TArray<FString>& Args)
+	{
+#if WITH_ODSC
+		if (FODSCManager::IsODSCActive())
+		{
+			TArray<FString> MaterialPaths;
+			GODSCManager->Thread->RetrieveMissedMaterials(MaterialPaths);
+			UConsole* ViewportConsole = (GEngine->GameViewport != nullptr) ? GEngine->GameViewport->ViewportConsole : nullptr;
+			FConsoleOutputDevice StrOut(ViewportConsole);
+
+			for (const FString& MaterialKey : MaterialPaths)
+			{
+				StrOut.CategorizedLogf(LogODSC.GetCategoryName(), ELogVerbosity::Error, TEXT("ODSC missed material: %s"), *MaterialKey);
+			}
+		}
+#endif
+	}
+};
+
+static FAutoConsoleCommand ODSCLogMissedMaterialsCmd(
+	TEXT("odsc.logmissedmaterials"),
+	TEXT("Logs materials that were not found by the ODSC server"),
+	FConsoleCommandWithArgsDelegate::CreateStatic(&FODSCManagerAccess::ODSCLogMissedMaterials)
+);
+
 namespace ODSCManagerPrivate
 {
 #if WITH_ODSC
 	static thread_local int32 ODSCSuspendForceRecompileCount=0;
+	static thread_local FPrimitiveSceneInfo* CurrentPrimitiveSceneInfo=nullptr;
+#endif
+}
+
+void FODSCManager::SetCurrentPrimitiveSceneInfo(FPrimitiveSceneInfo* PrimitiveSceneInfo)
+{
+#if WITH_ODSC
+	ODSCManagerPrivate::CurrentPrimitiveSceneInfo = PrimitiveSceneInfo;
+#endif
+}
+
+void FODSCManager::ResetCurrentPrimitiveSceneInfo()
+{
+#if WITH_ODSC
+	ODSCManagerPrivate::CurrentPrimitiveSceneInfo = nullptr;
 #endif
 }
 
@@ -201,10 +248,13 @@ void FODSCManager::AddThreadedShaderPipelineRequest(
 	int32 PermutationId,
 	const TArray<FShaderId>& RequestShaderIds)
 {
-	if (IsHandlingRequests())
+#if WITH_ODSC
+	// Testing ODSCManagerPrivate::ODSCSuspendForceRecompileCount since a map from ODSC may actually be in use and we request PSO precaching
+	if (IsHandlingRequests() && ODSCManagerPrivate::ODSCSuspendForceRecompileCount == 0)
 	{
-		Thread->AddShaderPipelineRequest(ShaderPlatform, FeatureLevel, QualityLevel, Material, VertexFactoryName, PipelineName, ShaderTypeNames, PermutationId, RequestShaderIds);
+		Thread->AddShaderPipelineRequest(ShaderPlatform, FeatureLevel, QualityLevel, Material, ODSCManagerPrivate::CurrentPrimitiveSceneInfo, VertexFactoryName, PipelineName, ShaderTypeNames, PermutationId, RequestShaderIds);
 	}
+#endif
 }
 
 void FODSCManager::RegisterMaterialInstance(const UMaterialInstance* MaterialInstance)
diff --git a/Engine/Source/Runtime/Engine/Private/ODSC/ODSCThread.cpp b/Engine/Source/Runtime/Engine/Private/ODSC/ODSCThread.cpp
--- a/Engine/Source/Runtime/Engine/Private/ODSC/ODSCThread.cpp
+++ b/Engine/Source/Runtime/Engine/Private/ODSC/ODSCThread.cpp
@@ -8,6 +8,8 @@
 #include "MaterialShared.h"
 #include "Materials/MaterialInterface.h"
 #include "UObject/UObjectIterator.h"
+#include "PrimitiveSceneInfo.h"
+#include "Components/PrimitiveComponent.h"
 
 FODSCRequestPayload::FODSCRequestPayload(
 	EShaderPlatform InShaderPlatform,
@@ -329,6 +331,25 @@ FMaterialShaderMap* FODSCThread::FindMaterialShaderMap(const FString& MaterialNa
 	return nullptr;
 }
 
+void FODSCThread::RetrieveMissedMaterials(TArray<FString>& OutMaterialPaths) const
+{
+	FReadScopeLock ReadLock(RequestHashesRWLock);
+	for (auto Iter = RequestHashes.CreateConstIterator(); Iter; ++Iter)
+	{
+		const FODSCShaderMapData& ODSCShaderMapData = Iter.Value();
+		if (!ODSCShaderMapData.CurrentRequests.IsEmpty())
+		{
+			FString MaterialKey(Iter.Key().ToString());
+			if (ODSCShaderMapData.ActorPath.IsValid())
+			{
+				MaterialKey += ":::";
+				MaterialKey += ODSCShaderMapData.ActorPath.ToString();
+			}
+			OutMaterialPaths.Add(MaterialKey);
+		}
+	}
+}
+
 void FODSCThread::AddRequest(const TArray<FString>& MaterialsToCompile, const FString& ShaderTypesToLoad, EShaderPlatform ShaderPlatform, ERHIFeatureLevel::Type FeatureLevel, EMaterialQualityLevel::Type QualityLevel, ODSCRecompileCommand RecompileCommandType)
 {
 	PendingMaterialThreadedRequests.Enqueue(new FODSCMessageHandler(MaterialsToCompile, ShaderTypesToLoad, ShaderPlatform, FeatureLevel, QualityLevel, RecompileCommandType));
@@ -339,6 +360,7 @@ void FODSCThread::AddShaderPipelineRequest(
 	ERHIFeatureLevel::Type FeatureLevel,
 	EMaterialQualityLevel::Type QualityLevel,
 	const FMaterial* Material,
+	const FPrimitiveSceneInfo* PrimitiveSceneInfo,
 	const FString& VertexFactoryName,
 	const FString& PipelineName,
 	const TArray<FString>& ShaderTypeNames,
@@ -347,6 +369,8 @@ void FODSCThread::AddShaderPipelineRequest(
 )
 {
 	bool bShouldAddRequest = false;
+
+	FString ActorPath;
 	{
 		FWriteScopeLock WriteLock(RequestHashesRWLock);
 
@@ -367,6 +391,25 @@ void FODSCThread::AddShaderPipelineRequest(
 				bShouldAddRequest = true;
 			}
 		}
+
+		if (bShouldAddRequest)
+		{
+#if WITH_ODSC
+			if (PrimitiveSceneInfo)
+			{
+				AActor* OwningActor = PrimitiveSceneInfo->GetComponentForDebugOnly() ? PrimitiveSceneInfo->GetComponentForDebugOnly()->GetOwner() : nullptr;
+				if (OwningActor)
+				{
+					ActorPath = OwningActor->GetPathName();
+				}
+			}
+#endif
+		}
+
+		if (!ActorPath.IsEmpty())
+		{
+			ODSCShaderMapData.ActorPath = FName(ActorPath);
+		}
 	}
 
 	if (bShouldAddRequest)
@@ -374,6 +417,13 @@ void FODSCThread::AddShaderPipelineRequest(
 		SCOPED_NAMED_EVENT(AddShaderPipelineRequest_AddRequest, FColor::Emerald);
 
 		FString MaterialName = Material->GetFullPath();
+
+		if (!ActorPath.IsEmpty())
+		{
+			MaterialName += ":::";
+			MaterialName += ActorPath;
+		}
+
 		FString RequestString = (MaterialName + VertexFactoryName + PipelineName);
 		for (const auto& ShaderTypeName : ShaderTypeNames)
 		{
@@ -545,6 +595,7 @@ void FODSCThread::Process()
 		{
 			PendingRequestsPipeline.Add(NextRequest);
 		}
+		NumPendingMaterialsShaders -= NextRequest->NumPayloads();
 	}
 
 	NumPendingMaterialsShaders.store(0, std::memory_order_release);
diff --git a/Engine/Source/Runtime/Engine/Private/ODSC/ODSCThread.h b/Engine/Source/Runtime/Engine/Private/ODSC/ODSCThread.h
--- a/Engine/Source/Runtime/Engine/Private/ODSC/ODSCThread.h
+++ b/Engine/Source/Runtime/Engine/Private/ODSC/ODSCThread.h
@@ -14,6 +14,7 @@ class FEvent;
 class FRunnableThread;
 class FMaterialShaderMap;
 class FMaterialShaderMapId;
+class FPrimitiveSceneInfo;
 
 namespace UE
 {
@@ -141,6 +142,7 @@ public:
 		ERHIFeatureLevel::Type FeatureLevel,
 		EMaterialQualityLevel::Type QualityLevel,
 		const FMaterial* Material,
+		const FPrimitiveSceneInfo* PrimitiveSceneInfo,
 		const FString& VertexFactoryName,
 		const FString& PipelineName,
 		const TArray<FString>& ShaderTypeNames,
@@ -176,6 +178,8 @@ public:
 	void RegisterMaterialShaderMaps(const FString& MaterialName, const TArray<TRefCountPtr<FMaterialShaderMap>>& LoadedShaderMaps);
 	FMaterialShaderMap* FindMaterialShaderMap(const FString& MaterialName, const FMaterialShaderMapId& ShaderMapId) const;
 
+	void RetrieveMissedMaterials(TArray<FString>& OutMaterialPaths) const;
+
 protected:
 
 	//~ Begin FRunnable Interface
@@ -264,6 +268,8 @@ private:
 		TArray<TRefCountPtr<FMaterialShaderMap>> MaterialShaderMaps;
     	/** Hashes for all Pending or Completed requests.  This is so we avoid making the same request multiple times. */
 		TSet<FODSCShaderId> CurrentRequests;
+		
+		FName ActorPath;
 	};
 
     /** Requests seen for a given material name */
diff --git a/Engine/Source/Runtime/Engine/Public/ODSC/ODSCManager.h b/Engine/Source/Runtime/Engine/Public/ODSC/ODSCManager.h
--- a/Engine/Source/Runtime/Engine/Public/ODSC/ODSCManager.h
+++ b/Engine/Source/Runtime/Engine/Public/ODSC/ODSCManager.h
@@ -12,6 +12,7 @@ class FODSCThread;
 class UMaterialInstance;
 class FMaterialShaderMap;
 class FMaterialShaderMapId;
+class FPrimitiveSceneInfo;
 
 /**
  * Responsible for processing shader compile responses from the ODSC Thread.
@@ -114,8 +115,11 @@ public:
 	static void UnregisterMaterialName(const FMaterial* Material);
 	static void RegisterMaterialShaderMaps(const FString& MaterialName, const TArray<TRefCountPtr<FMaterialShaderMap>>& LoadedShaderMaps);
 	static FMaterialShaderMap* FindMaterialShaderMap(const FString& MaterialName, const FMaterialShaderMapId& ShaderMapId);
+	static void SetCurrentPrimitiveSceneInfo(FPrimitiveSceneInfo* PrimitiveSceneInfo);
+	static void ResetCurrentPrimitiveSceneInfo();
 
 private:
+	friend class FODSCManagerAccess;
 
 	ENGINE_API void OnEnginePreExit();
 	ENGINE_API void StopThread();
@@ -139,6 +143,19 @@ private:
 	FName MaterialNameToRecompile;
 };
 
+struct FODSCPrimitiveSceneInfoScope
+{
+	FODSCPrimitiveSceneInfoScope(FPrimitiveSceneInfo* PrimitiveSceneInfo)
+	{
+		FODSCManager::SetCurrentPrimitiveSceneInfo(PrimitiveSceneInfo);
+	}
+
+	~FODSCPrimitiveSceneInfoScope()
+	{
+		FODSCManager::ResetCurrentPrimitiveSceneInfo();
+	}
+};
+
 struct FODSCSuspendForceRecompileScope
 {
 	FODSCSuspendForceRecompileScope()
diff --git a/Engine/Source/Runtime/Renderer/Private/MeshDrawCommands.cpp b/Engine/Source/Runtime/Renderer/Private/MeshDrawCommands.cpp
--- a/Engine/Source/Runtime/Renderer/Private/MeshDrawCommands.cpp
+++ b/Engine/Source/Runtime/Renderer/Private/MeshDrawCommands.cpp
@@ -13,6 +13,9 @@ MeshDrawCommandSetup.cpp: Mesh draw command setup.
 #include "StaticMeshBatch.h"
 #include "SceneDefinitions.h"
 #include "MeshDrawCommandStats.h"
+#if WITH_ODSC
+#include "ODSC/ODSCManager.h"
+#endif
 
 TGlobalResource<FPrimitiveIdVertexBufferPool> GPrimitiveIdVertexBufferPool;
 
@@ -617,6 +620,9 @@ void GenerateDynamicMeshDrawCommands(
 			{
 				const FMeshBatchAndRelevance& MeshAndRelevance = DynamicMeshElements[MeshIndex];
 				const uint64 BatchElementMask = ~0ull;
+#if WITH_ODSC
+				FODSCPrimitiveSceneInfoScope ODSCPrimitiveSceneInfoScope(MeshAndRelevance.PrimitiveSceneProxy ? MeshAndRelevance.PrimitiveSceneProxy->GetPrimitiveSceneInfo() : nullptr);
+#endif
 
 				PassMeshProcessor->AddMeshBatch(*MeshAndRelevance.Mesh, BatchElementMask, MeshAndRelevance.PrimitiveSceneProxy);
 			}
@@ -637,6 +643,10 @@ void GenerateDynamicMeshDrawCommands(
 			const uint64 DefaultBatchElementMask = ~0ul;
 			const int32 StartCommandIndex = VisibleCommands.Num();
 
+#if WITH_ODSC
+			FODSCPrimitiveSceneInfoScope ODSCPrimitiveSceneInfoScope(StaticMeshBatch->PrimitiveSceneInfo);
+#endif
+
 			if (StaticMeshBatch->bViewDependentArguments)
 			{
 				FMeshBatch ViewDepenedentMeshBatch(*StaticMeshBatch);
diff --git a/Engine/Source/Runtime/Renderer/Private/PrimitiveSceneInfo.cpp b/Engine/Source/Runtime/Renderer/Private/PrimitiveSceneInfo.cpp
--- a/Engine/Source/Runtime/Renderer/Private/PrimitiveSceneInfo.cpp
+++ b/Engine/Source/Runtime/Renderer/Private/PrimitiveSceneInfo.cpp
@@ -39,6 +39,7 @@
 #include "InstanceDataSceneProxy.h"
 #include "DecalRenderingCommon.h"
 #include "RendererPrivateUtils.h"
+#include "ODSC/ODSCManager.h"
 
 extern int32 GGPUSceneInstanceClearList;
 
@@ -528,6 +529,10 @@ void FPrimitiveSceneInfo::CacheMeshDrawCommands(FScene* Scene, TArrayView<FPrimi
 
 						check(!MeshRelevance.CommandInfosMask.Get(PassType));
 
+#if WITH_ODSC
+						FODSCPrimitiveSceneInfoScope ODSCPrimitiveSceneInfoScope(SceneInfo);
+#endif
+
 						uint64 BatchElementMask = ~0ull;
 						// NOTE: AddMeshBatch calls FCachedPassMeshDrawListContext::FinalizeCommand
 						PassMeshProcessor->AddMeshBatch(Mesh, BatchElementMask, SceneInfo->Proxy);
