Cooker: Change NeverCookPackageList to contain packagenames rather than filenames. This prevents the creation of FNames for packages that we end up not using in the cook, and preventing that FName creation saves memory.
#rnx

diff --git a/Engine/Source/Editor/UnrealEd/Classes/CookOnTheSide/CookOnTheFlyServer.h b/Engine/Source/Editor/UnrealEd/Classes/CookOnTheSide/CookOnTheFlyServer.h
--- a/Engine/Source/Editor/UnrealEd/Classes/CookOnTheSide/CookOnTheFlyServer.h
+++ b/Engine/Source/Editor/UnrealEd/Classes/CookOnTheSide/CookOnTheFlyServer.h
@@ -966,10 +966,9 @@ private:
 	 * Does not include checking UAssetManager, which has to be queried later
 	 * This function is const because it is not always called and should avoid side effects
 	 */
-	TArray<FName> GetNeverCookPackageFileNames(TArrayView<const FString> ExtraNeverCookDirectories
+	TArray<FName> GetNeverCookPackageNames(TArrayView<const FString> ExtraNeverCookDirectories
 		= TArrayView<const FString>()) const;
 
-
 	/** AddFileToCook add file to cook list */
 	void AddFileToCook( TArray<FName>& InOutFilesToCook, TMap<FName, UE::Cook::FInstigator>& InOutInstigators,
 		const FString &InFilename, const UE::Cook::FInstigator& Instigator) const;
diff --git a/Engine/Source/Editor/UnrealEd/Private/CookOnTheFlyServer.cpp b/Engine/Source/Editor/UnrealEd/Private/CookOnTheFlyServer.cpp
--- a/Engine/Source/Editor/UnrealEd/Private/CookOnTheFlyServer.cpp
+++ b/Engine/Source/Editor/UnrealEd/Private/CookOnTheFlyServer.cpp
@@ -4383,7 +4383,7 @@ void UCookOnTheFlyServer::ProcessUnsolicitedPackages(TArray<FName>* OutDiscovere
 		}
 		else if (bSkipOnlyEditorOnly &&
 					(Instigator.Category == EInstigator::ForceExplorableSaveTimeSoftDependency ||
-						(PackageTracker->NeverCookPackageList.Contains(PackageData->GetFileName()) &&
+						(PackageTracker->NeverCookPackageList.Contains(PackageData->GetPackageName()) &&
 							INDEX_NONE != UE::String::FindFirst(WriteToString<256>(PackageData->GetPackageName()),
 								ULevel::GetExternalActorsFolderName(), ESearchCase::IgnoreCase))))
 		{
@@ -4545,7 +4545,7 @@ void UCookOnTheFlyServer::PumpSaves(UE::Cook::FTickStackData& StackData, uint32
 		// asset registry at the end of the cook
 		PackageTracker->UncookedEditorOnlyPackages.Remove(Package->GetFName());
 
-		if (PackageTracker->NeverCookPackageList.Contains(PackageData.GetFileName()))
+		if (PackageTracker->NeverCookPackageList.Contains(Package->GetFName()))
 		{
 			// refuse to save this package, it's clearly one of the undesirables
 			DemoteToIdle(PackageData, ESendFlags::QueueAdd, ESuppressCookReason::NeverCook);
@@ -9271,11 +9271,10 @@ void UCookOnTheFlyServer::CollectFilesToCook(TArray<FName>& FilesInPath, TMap<FN
 
 		for (FName NeverCookPackage : PackagesToNeverCook)
 		{
-			const FName StandardPackageFilename = PackageDatas->GetFileNameByFlexName(NeverCookPackage);
-
-			if (!StandardPackageFilename.IsNone())
+			FName PackageName;
+			if (PackageDatas->TryGetNamesByFlexName(NeverCookPackage, &PackageName, nullptr, true /* bRequireExists */))
 			{
-				PackageTracker->NeverCookPackageList.Add(StandardPackageFilename);
+				PackageTracker->NeverCookPackageList.Add(PackageName);
 			}
 		}
 	}
@@ -9329,11 +9328,10 @@ void UCookOnTheFlyServer::CollectFilesToCook(TArray<FName>& FilesInPath, TMap<FN
 
 		for (FName NeverCookPackage : PackagesToNeverCook)
 		{
-			FName StandardPackageFilename = PackageDatas->GetFileNameByFlexName(NeverCookPackage);
-
-			if (!StandardPackageFilename.IsNone())
+			FName PackageName;
+			if (PackageDatas->TryGetNamesByFlexName(NeverCookPackage, &PackageName, nullptr, true /* bRequireExists */))
 			{
-				PackageTracker->NeverCookPackageList.Add(StandardPackageFilename);
+				PackageTracker->NeverCookPackageList.Add(PackageName);
 			}
 		}
 	}
@@ -11158,7 +11156,7 @@ void FBeginCookConfigSettings::LoadNeverCookLocal(FBeginCookContext& BeginContex
 	{
 		ExtraNeverCookDirectories = BeginContext.StartupOptions->NeverCookDirectories;
 	}
-	for (FName NeverCookPackage : BeginContext.COTFS.GetNeverCookPackageFileNames(ExtraNeverCookDirectories))
+	for (FName NeverCookPackage : BeginContext.COTFS.GetNeverCookPackageNames(ExtraNeverCookDirectories))
 	{
 		NeverCookPackageList.Add(NeverCookPackage);
 	}
@@ -11185,9 +11183,9 @@ void UCookOnTheFlyServer::SetNeverCookPackageConfigSettings(FBeginCookContext& B
 {
 	UE::Cook::FThreadSafeSet<FName>& NeverCookPackageList = PackageTracker->NeverCookPackageList;
 	NeverCookPackageList.Empty();
-	for (FName FileName : Settings.NeverCookPackageList)
+	for (FName PackageName : Settings.NeverCookPackageList)
 	{
-		NeverCookPackageList.Add(FileName);
+		NeverCookPackageList.Add(PackageName);
 	}
 	PackageTracker->PlatformSpecificNeverCookPackages = MoveTemp(Settings.PlatformSpecificNeverCookPackages);
 }
@@ -12565,8 +12563,7 @@ void UCookOnTheFlyServer::UnregisterCookByTheBookDelegates()
 	}
 }
 
-
-TArray<FName> UCookOnTheFlyServer::GetNeverCookPackageFileNames(TArrayView<const FString> ExtraNeverCookDirectories) const
+TArray<FName> UCookOnTheFlyServer::GetNeverCookPackageNames(TArrayView<const FString> ExtraNeverCookDirectories) const
 {
 	TArray<FString> NeverCookDirectories(ExtraNeverCookDirectories);
 
@@ -12616,7 +12613,7 @@ TArray<FName> UCookOnTheFlyServer::GetNeverCookPackageFileNames(TArrayView<const
 		NeverCookDirectories.Add(MoveTemp(FullExternalActorsPath));
 	}
 
-	TArray<FString> NeverCookPackagesPaths;
+	TArray<FName> NeverCookPackages;
 	if (AssetRegistry->IsSearchAllAssets() && !AssetRegistry->IsLoadingAssets())
 	{
 		TDirectoryTree<int32> NeverCookDirectoryTree;
@@ -12631,25 +12628,12 @@ TArray<FName> UCookOnTheFlyServer::GetNeverCookPackageFileNames(TArrayView<const
 
 		FString PackageNameStr;
 		AssetRegistry->EnumerateAllPackages(
-			[&NeverCookPackagesPaths, &NeverCookDirectoryTree, &PackageNameStr](FName PackageName, const FAssetPackageData& PackageData)
+			[&NeverCookPackages, &NeverCookDirectoryTree, &PackageNameStr](FName PackageName, const FAssetPackageData& PackageData)
 			{
 				PackageName.ToString(PackageNameStr);
 				if (NeverCookDirectoryTree.ContainsPathOrParent(PackageNameStr))
 				{
-					FString LocalFileName;
-					if (PackageData.Extension != EPackageExtension::Unspecified && PackageData.Extension != EPackageExtension::Custom)
-					{
-						FString Extension = LexToString(PackageData.Extension);
-						FPackageName::TryConvertLongPackageNameToFilename(PackageNameStr, LocalFileName, Extension);
-					}
-					else
-					{
-						FPackageName::DoesPackageExist(PackageNameStr, &LocalFileName);
-					}
-					if (!LocalFileName.IsEmpty())
-					{
-						NeverCookPackagesPaths.Add(LocalFileName);
-					}
+					NeverCookPackages.Add(PackageName);
 				}
 			});
 	}
@@ -12671,16 +12655,21 @@ TArray<FName> UCookOnTheFlyServer::GetNeverCookPackageFileNames(TArrayView<const
 
 		if (bUseDirectoryScanFallback)
 		{
-			FPackageName::FindPackagesInDirectories(NeverCookPackagesPaths, NeverCookDirectories);
+			TArray<FString> ResultFilePathsToNeverCook;
+			FPackageName::FindPackagesInDirectories(ResultFilePathsToNeverCook, NeverCookDirectories);
+			NeverCookPackages.Reserve(ResultFilePathsToNeverCook.Num());
+			FString PackageName;
+			for (FString& FilePath : ResultFilePathsToNeverCook)
+			{
+				if (FPackageName::TryConvertFilenameToLongPackageName(FilePath, PackageName))
+				{
+					NeverCookPackages.Add(FName(PackageName));
+				}
+			}
 		}
 	}
 
-	TArray<FName> NeverCookNormalizedFileNames;
-	for (const FString& NeverCookPackagePath : NeverCookPackagesPaths)
-	{
-		NeverCookNormalizedFileNames.Add(UE::Cook::FPackageDatas::GetStandardFileName(NeverCookPackagePath));
-	}
-	return NeverCookNormalizedFileNames;
+	return NeverCookPackages;
 }
 
 bool UCookOnTheFlyServer::RecompileChangedShaders(const TArray<const ITargetPlatform*>& TargetPlatforms)
diff --git a/Engine/Source/Editor/UnrealEd/Private/Cooker/CookGenerationHelper.cpp b/Engine/Source/Editor/UnrealEd/Private/Cooker/CookGenerationHelper.cpp
--- a/Engine/Source/Editor/UnrealEd/Private/Cooker/CookGenerationHelper.cpp
+++ b/Engine/Source/Editor/UnrealEd/Private/Cooker/CookGenerationHelper.cpp
@@ -1351,7 +1351,7 @@ void FGenerationHelper::FetchExternalActorDependencies()
 			{
 				return true;
 			}
-			bool bIsNeverCook = NeverCookPackageList.Contains(PackageData->GetFileName());
+			bool bIsNeverCook = NeverCookPackageList.Contains(PackageData->GetPackageName());
 			return !bIsNeverCook;
 		});
 	ExternalActorDependencies.Shrink();
diff --git a/Engine/Source/Editor/UnrealEd/Private/Cooker/CookPackageData.cpp b/Engine/Source/Editor/UnrealEd/Private/Cooker/CookPackageData.cpp
--- a/Engine/Source/Editor/UnrealEd/Private/Cooker/CookPackageData.cpp
+++ b/Engine/Source/Editor/UnrealEd/Private/Cooker/CookPackageData.cpp
@@ -2385,14 +2385,29 @@ FName FPackageDatas::GetFileNameByPackageName(FName PackageName, bool bRequireEx
 	return PackageData ? PackageData->GetFileName() : NAME_None;
 }
 
-FName FPackageDatas::GetFileNameByFlexName(FName PackageOrFileName, bool bRequireExists, bool bCreateAsMap)
+bool FPackageDatas::TryGetNamesByFlexName(FName PackageOrFileName, FName* OutPackageName, FName* OutFileName,
+	bool bRequireExists, bool bCreateAsMap)
 {
 	FString Buffer = PackageOrFileName.ToString();
 	if (!FPackageName::TryConvertFilenameToLongPackageName(Buffer, Buffer))
 	{
-		return NAME_None;
+		return false;
+	}
+	FName PackageName = FName(Buffer);
+	FName FileName = GetFileNameByPackageName(PackageName, bRequireExists, bCreateAsMap);
+	if (FileName.IsNone())
+	{
+		return false;
 	}
-	return GetFileNameByPackageName(FName(Buffer), bRequireExists, bCreateAsMap);
+	if (OutPackageName)
+	{
+		*OutPackageName = PackageName;
+	}
+	if (OutFileName)
+	{
+		*OutFileName = FileName;
+	}
+	return true;
 }
 
 FName FPackageDatas::LookupFileNameOnDisk(FName PackageName, bool bRequireExists, bool bCreateAsMap)
diff --git a/Engine/Source/Editor/UnrealEd/Private/Cooker/CookPackageData.h b/Engine/Source/Editor/UnrealEd/Private/Cooker/CookPackageData.h
--- a/Engine/Source/Editor/UnrealEd/Private/Cooker/CookPackageData.h
+++ b/Engine/Source/Editor/UnrealEd/Private/Cooker/CookPackageData.h
@@ -1254,7 +1254,8 @@ public:
 	 *                     the extension is set to .umap, if false, it is set to .uasset.
 	 * @return Local WorkspaceDomain path in FPaths::MakeStandardFilename form, or NAME_None if it does not exist.
 	 */
-	FName GetFileNameByFlexName(FName PackageOrFileName, bool bRequireExists = true, bool bCreateAsMap = false);
+	bool TryGetNamesByFlexName(FName PackageOrFileName, FName* OutPackageName = nullptr, FName* OutFileName = nullptr,
+		bool bRequireExists = true, bool bCreateAsMap = false);
 
 	/**
 	 * Uncached; reads the AssetRegistry and disk to find the filename for the given PackageName.
diff --git a/Engine/Source/Editor/UnrealEd/Private/Cooker/CookRequestCluster.cpp b/Engine/Source/Editor/UnrealEd/Private/Cooker/CookRequestCluster.cpp
--- a/Engine/Source/Editor/UnrealEd/Private/Cooker/CookRequestCluster.cpp
+++ b/Engine/Source/Editor/UnrealEd/Private/Cooker/CookRequestCluster.cpp
@@ -2189,7 +2189,7 @@ void FRequestCluster::IsRequestCookable(const ITargetPlatform* Platform, FName P
 	};
 
 	FName FileName = PackageData.GetFileName();
-	if (InCOTFS.PackageTracker->NeverCookPackageList.Contains(FileName))
+	if (InCOTFS.PackageTracker->NeverCookPackageList.Contains(PackageName))
 	{
 		if (INDEX_NONE != UE::String::FindFirst(NameBuffer, ULevel::GetExternalActorsFolderName(), 
 			ESearchCase::IgnoreCase))
diff --git a/Engine/Source/Runtime/Engine/Classes/Engine/World.h b/Engine/Source/Runtime/Engine/Classes/Engine/World.h
--- a/Engine/Source/Runtime/Engine/Classes/Engine/World.h
+++ b/Engine/Source/Runtime/Engine/Classes/Engine/World.h
@@ -3772,7 +3772,7 @@ public:
 	 *
 	 * @param bForce	If true, load the levels even is a commandlet
 	 */
-	void LoadSecondaryLevels(bool bForce = false, TSet<FName>* FilenamesToSkip = NULL);
+	void LoadSecondaryLevels(bool bForce = false, TSet<FName>* PackageNamesToSkip = nullptr);
 
 	/** Utility for returning the ULevelStreaming object for a particular sub-level, specified by package name */
 	ULevelStreaming* GetLevelStreamingForPackageName(FName PackageName);
diff --git a/Engine/Source/Runtime/Engine/Private/LevelActor.cpp b/Engine/Source/Runtime/Engine/Private/LevelActor.cpp
--- a/Engine/Source/Runtime/Engine/Private/LevelActor.cpp
+++ b/Engine/Source/Runtime/Engine/Private/LevelActor.cpp
@@ -1508,7 +1508,7 @@ bool UWorld::EncroachingBlockingGeometry(const AActor* TestActor, FVector TestLo
 }
 
 
-void UWorld::LoadSecondaryLevels(bool bForce, TSet<FName>* FilenamesToSkip)
+void UWorld::LoadSecondaryLevels(bool bForce, TSet<FName>* PackageNamesToSkip)
 {
 	check( GIsEditor );
 
@@ -1527,15 +1527,11 @@ void UWorld::LoadSecondaryLevels(bool bForce, TSet<FName>* FilenamesToSkip)
 				// If we are cooking don't cook sub levels multiple times if they've already been cooked
 				FString PackageFilename;
 				const FString StreamingLevelWorldAssetPackageName = StreamingLevel->GetWorldAssetPackageName();
-				if (FilenamesToSkip)
+				if (PackageNamesToSkip)
 				{
-					if (FPackageName::DoesPackageExist(StreamingLevelWorldAssetPackageName, &PackageFilename))
-					{
-						bSkipFile |= FilenamesToSkip->Contains( FName(*PackageFilename) );
-					}
+					bSkipFile |= PackageNamesToSkip->Contains(StreamingLevel->GetWorldAssetPackageFName());
 				}
 
-
 				bool bAlreadyLoaded = false;
 				UPackage* LevelPackage = FindObject<UPackage>(NULL, *StreamingLevelWorldAssetPackageName,true);
 				// don't need to do any extra work if the level is already loaded
