Use GetObjectsWithPackage rather than ForEachObjectWithPackage when executing cleanup as Blueprints/Worlds may execute delegates/arbitrary code which could try to re-enter the UObject Hash.

diff --git a/Engine/Source/Editor/UnrealEd/Private/PackageTools.cpp b/Engine/Source/Editor/UnrealEd/Private/PackageTools.cpp
--- a/Engine/Source/Editor/UnrealEd/Private/PackageTools.cpp
+++ b/Engine/Source/Editor/UnrealEd/Private/PackageTools.cpp
@@ -1048,7 +1048,9 @@ UPackageTools::UPackageTools(const FObjectInitializer& ObjectInitializer)
 			GEngine->NotifyToolsOfObjectReplacement(InPackageReloadedEvent->GetRepointedObjects());
 
 			// Notify any Blueprints that are about to be unloaded, and destroy any leftover worlds.
-			ForEachObjectWithPackage(InPackageReloadedEvent->GetOldPackage(), [](UObject* InObject)
+			TArray<UObject*> Objects;
+			GetObjectsWithPackage(InPackageReloadedEvent->GetOldPackage(), Objects, true, RF_Transient, EInternalObjectFlags::Garbage);
+			for (UObject* InObject : Objects)
 			{
 				if (UBlueprint* BP = Cast<UBlueprint>(InObject))
 				{
@@ -1073,8 +1075,7 @@ UPackageTools::UPackageTools(const FObjectInitializer& ObjectInitializer)
 						World->CleanupWorld();
 					}
 				}
-				return true;
-			}, true, RF_Transient, EInternalObjectFlags::Garbage);
+			}
 		}
 
 		if (InPackageReloadPhase == EPackageReloadPhase::OnPackageFixup)
