Update the Transform submenu of the level editor's new viewport toolbar

* Add icon the the Transforms menu.
* Rename to "Transform" (dropping trailing s).
* Rename FName to "Transform" (new full menu name is "LevelEditor.ViewportToolbar.Transform".)
* Update Animation Mode to extend via the new submenu name.
* Set icon on the command "EditorViewport.RelativeCoordinateSystem_Local" from StarshipStyle.cpp to make it actually work.

diff --git a/Engine/Plugins/Animation/ControlRig/Source/ControlRigEditor/Private/EditMode/ControlRigEditMode.cpp b/Engine/Plugins/Animation/ControlRig/Source/ControlRigEditor/Private/EditMode/ControlRigEditMode.cpp
--- a/Engine/Plugins/Animation/ControlRig/Source/ControlRigEditor/Private/EditMode/ControlRigEditMode.cpp
+++ b/Engine/Plugins/Animation/ControlRig/Source/ControlRigEditor/Private/EditMode/ControlRigEditMode.cpp
@@ -477,7 +477,7 @@ void FControlRigEditMode::Enter()
 		}
 	}
 
-	UE::ControlRig::PopulateControlRigViewportToolbarTransformsSubmenu("LevelEditor.ViewportToolbar.Transforms");
+	UE::ControlRig::PopulateControlRigViewportToolbarTransformsSubmenu("LevelEditor.ViewportToolbar.Transform");
 	UE::ControlRig::PopulateControlRigViewportToolbarSelectionSubmenu("LevelEditor.ViewportToolbar.Select");
 	UE::ControlRig::PopulateControlRigViewportToolbarShowSubmenu(
 		"LevelEditor.ViewportToolbar.Show", GetToolkit()->GetToolkitCommands()
diff --git a/Engine/Source/Editor/EditorStyle/Private/StarshipStyle.cpp b/Engine/Source/Editor/EditorStyle/Private/StarshipStyle.cpp
--- a/Engine/Source/Editor/EditorStyle/Private/StarshipStyle.cpp
+++ b/Engine/Source/Editor/EditorStyle/Private/StarshipStyle.cpp
@@ -4252,6 +4252,8 @@ void FStarshipEditorStyle::FStyle::SetupLevelEditorStyle()
 		Set("EditorViewport.ToggleSurfaceSnapping", new IMAGE_BRUSH_SVG("Starship/EditorViewport/surface-snap", Icon16x16));
 
 		Set( "EditorViewport.RelativeCoordinateSystem_World", new IMAGE_BRUSH_SVG( "Starship/EditorViewport/globe", Icon16x16 ) );
+		Set("EditorViewport.RelativeCoordinateSystem_Local",
+			new IMAGE_BRUSH_SVG("Starship/Common/transform-local", Icon16x16));
 
 		Set( "EditorViewport.CamSpeedSetting", new IMAGE_BRUSH_SVG( "Starship/EditorViewport/camera", Icon16x16) );
 		
diff --git a/Engine/Source/Editor/UnrealEd/Private/ViewportToolbar/UnrealEdViewportToolbar.cpp b/Engine/Source/Editor/UnrealEd/Private/ViewportToolbar/UnrealEdViewportToolbar.cpp
--- a/Engine/Source/Editor/UnrealEd/Private/ViewportToolbar/UnrealEdViewportToolbar.cpp
+++ b/Engine/Source/Editor/UnrealEd/Private/ViewportToolbar/UnrealEdViewportToolbar.cpp
@@ -554,9 +554,9 @@ bool ShowNewViewportToolbars()
 
 FToolMenuEntry CreateViewportToolbarTransformsSection()
 {
-	return FToolMenuEntry::InitSubMenu(
-		"Transforms",
-		LOCTEXT("TransformsSubmenuLabel", "Transforms"),
+	FToolMenuEntry Entry = FToolMenuEntry::InitSubMenu(
+		"Transform",
+		LOCTEXT("TransformsSubmenuLabel", "Transform"),
 		LOCTEXT("TransformsSubmenuTooltip", "Viewport-related transforms tools"),
 		FNewToolMenuDelegate::CreateLambda(
 			[](UToolMenu* Submenu) -> void
@@ -654,6 +654,9 @@ FToolMenuEntry CreateViewportToolbarTransformsSection()
 			}
 		)
 	);
+
+	Entry.Icon = FSlateIcon(FAppStyle::GetAppStyleSetName(), "EditorViewport.RelativeCoordinateSystem_Local");
+	return Entry;
 }
 
 FToolMenuEntry CreateViewportToolbarSelectionSection()
