Improve icons and labels for the Snapping submenu of the new level editor viewport toolbar

* Add icon and remove toolbar label from the Snapping submenu.
* Add icon and dynamic label on the Snapping submenu's Location submenu.

diff --git a/Engine/Source/Editor/UnrealEd/Private/ViewportToolbar/UnrealEdViewportToolbar.cpp b/Engine/Source/Editor/UnrealEd/Private/ViewportToolbar/UnrealEdViewportToolbar.cpp
--- a/Engine/Source/Editor/UnrealEd/Private/ViewportToolbar/UnrealEdViewportToolbar.cpp
+++ b/Engine/Source/Editor/UnrealEd/Private/ViewportToolbar/UnrealEdViewportToolbar.cpp
@@ -271,73 +271,98 @@ FToolMenuEntry CreateLocationSnapCheckboxMenu()
 	const FName LocationSnapName = "LocationSnap";
 	const FText LocationSnapLabel = LOCTEXT("LocationSnapLabel", "Location");
 
+	FToolMenuEntry Entry;
+
 	if (!GEditor)
 	{
-		return FToolMenuEntry::InitMenuEntry(
+		Entry = FToolMenuEntry::InitMenuEntry(
 			LocationSnapName, LocationSnapLabel, FText(), FSlateIcon(), UnrealEd::Private::DisabledAction()
 		);
 	}
+	else
+	{
 
-	FNewToolMenuDelegate MakeMenuDelegate = FNewToolMenuDelegate::CreateLambda(
-		[LocationSnapName](UToolMenu* InToolMenu)
-		{
-			UnrealEd::FLocationGridCheckboxListExecuteActionDelegate ExecuteDelegate =
-				UnrealEd::FLocationGridCheckboxListExecuteActionDelegate::CreateUObject(GEditor, &UEditorEngine::SetGridSize);
+		FNewToolMenuDelegate MakeMenuDelegate = FNewToolMenuDelegate::CreateLambda(
+			[LocationSnapName](UToolMenu* InToolMenu)
+			{
+				UnrealEd::FLocationGridCheckboxListExecuteActionDelegate ExecuteDelegate =
+					UnrealEd::FLocationGridCheckboxListExecuteActionDelegate::CreateUObject(
+						GEditor, &UEditorEngine::SetGridSize
+					);
 
-			UnrealEd::FLocationGridCheckboxListIsCheckedDelegate IsCheckedDelegate =
-				UnrealEd::FLocationGridCheckboxListIsCheckedDelegate::CreateLambda(
-					[](int CurrGridSizeIndex)
-					{
-						const ULevelEditorViewportSettings* ViewportSettings = GetDefault<ULevelEditorViewportSettings>();
-						return ViewportSettings->CurrentPosGridSize == CurrGridSizeIndex;
-					}
+				UnrealEd::FLocationGridCheckboxListIsCheckedDelegate IsCheckedDelegate =
+					UnrealEd::FLocationGridCheckboxListIsCheckedDelegate::CreateLambda(
+						[](int CurrGridSizeIndex)
+						{
+							const ULevelEditorViewportSettings* ViewportSettings =
+								GetDefault<ULevelEditorViewportSettings>();
+							return ViewportSettings->CurrentPosGridSize == CurrGridSizeIndex;
+						}
+					);
+
+				const ULevelEditorViewportSettings* ViewportSettings = GetDefault<ULevelEditorViewportSettings>();
+				TArray<float> GridSizes = ViewportSettings->bUsePowerOf2SnapSize ? ViewportSettings->Pow2GridSizes
+																				 : ViewportSettings->DecimalGridSizes;
+
+				InToolMenu->AddMenuEntry(
+					LocationSnapName,
+					FToolMenuEntry::InitWidget(
+						LocationSnapName,
+						UnrealEd::CreateLocationGridSnapMenu(
+							ExecuteDelegate,
+							IsCheckedDelegate,
+							GridSizes,
+							TAttribute<bool>::CreateLambda(
+								[]()
+								{
+									return FLevelEditorActionCallbacks::LocationGridSnap_IsChecked();
+								}
+							)
+						),
+						FText()
+					)
 				);
+			}
+		);
 
+		Entry = UnrealEd::CreateCheckboxSubmenu(
+			"GridSnapping",
+			LocationSnapLabel,
+			FEditorViewportCommands::Get().SurfaceSnapping->MakeTooltip()->GetTextTooltip(),
+			FToolMenuExecuteAction::CreateLambda(
+				[](const FToolMenuContext& InContext)
+				{
+					FLevelEditorActionCallbacks::LocationGridSnap_Clicked();
+				}
+			),
+			FToolMenuCanExecuteAction(),
+			FToolMenuGetActionCheckState::CreateLambda(
+				[](const FToolMenuContext& InContext)
+				{
+					return FLevelEditorActionCallbacks::LocationGridSnap_IsChecked() ? ECheckBoxState::Checked
+																					 : ECheckBoxState::Unchecked;
+				}
+			),
+			MakeMenuDelegate
+		);
+	}
+
+	Entry.ToolbarLabelOverride = TAttribute<FText>::CreateLambda(
+		[]()
+		{
 			const ULevelEditorViewportSettings* ViewportSettings = GetDefault<ULevelEditorViewportSettings>();
-			TArray<float> GridSizes = ViewportSettings->bUsePowerOf2SnapSize ? ViewportSettings->Pow2GridSizes
-																			 : ViewportSettings->DecimalGridSizes;
 
-			InToolMenu->AddMenuEntry(
-				LocationSnapName,
-				FToolMenuEntry::InitWidget(
-					LocationSnapName,
-					UnrealEd::CreateLocationGridSnapMenu(
-						ExecuteDelegate,
-						IsCheckedDelegate,
-						GridSizes,
-						TAttribute<bool>::CreateLambda(
-							[]()
-							{
-								return FLevelEditorActionCallbacks::LocationGridSnap_IsChecked();
-							}
-						)
-					),
-					FText()
-				)
-			);
+			const TArray<float> GridSizes = ViewportSettings->bUsePowerOf2SnapSize ? ViewportSettings->Pow2GridSizes
+																				   : ViewportSettings->DecimalGridSizes;
+
+			const float CurrentGridSize = GridSizes[ViewportSettings->CurrentPosGridSize];
+
+			return FText::AsNumber(CurrentGridSize);
 		}
 	);
+	Entry.Icon = FSlateIcon(FAppStyle::GetAppStyleSetName(), "EditorViewport.LocationGridSnap");
 
-	return UnrealEd::CreateCheckboxSubmenu(
-		"GridSnapping",
-		LocationSnapLabel,
-		FEditorViewportCommands::Get().SurfaceSnapping->MakeTooltip()->GetTextTooltip(),
-		FToolMenuExecuteAction::CreateLambda(
-			[](const FToolMenuContext& InContext)
-			{
-				FLevelEditorActionCallbacks::LocationGridSnap_Clicked();
-			}
-		),
-		FToolMenuCanExecuteAction(),
-		FToolMenuGetActionCheckState::CreateLambda(
-			[](const FToolMenuContext& InContext)
-			{
-				return FLevelEditorActionCallbacks::LocationGridSnap_IsChecked() ? ECheckBoxState::Checked
-																				 : ECheckBoxState::Unchecked;
-			}
-		),
-		MakeMenuDelegate
-	);
+	return Entry;
 }
 
 FToolMenuEntry CreateRotationSnapCheckboxMenu()
@@ -970,7 +995,7 @@ FToolMenuEntry CreateViewportToolbarSelectionSection()
 
 FToolMenuEntry CreateViewportToolbarSnappingSubmenu()
 {
-	return FToolMenuEntry::InitSubMenu(
+	FToolMenuEntry Entry = FToolMenuEntry::InitSubMenu(
 		"Snapping",
 		LOCTEXT("SnappingSubmenuLabel", "Snapping"),
 		LOCTEXT("SnappingSubmenuTooltip", "Viewport-related snapping settings"),
@@ -999,6 +1024,10 @@ FToolMenuEntry CreateViewportToolbarSnappingSubmenu()
 			}
 		)
 	);
+
+	Entry.ToolbarLabelOverride = FText();
+	Entry.Icon = FSlateIcon(FAppStyle::GetAppStyleSetName(), "Icons.Snap");
+	return Entry;
 }
 
 FText GetViewModesSubmenuLabel(TWeakPtr<SEditorViewport> InViewport)
