Add icon to the Select submenu of the new level editor viewport toolbar

Also rename UE::UnrealEd::CreateViewportToolbarSelectionSection to UE::UnrealEd::CreateViewportToolbarSelectSection to map better to the submenu's FName.

diff --git a/Engine/Source/Editor/Kismet/Private/SSCSEditorViewport.cpp b/Engine/Source/Editor/Kismet/Private/SSCSEditorViewport.cpp
--- a/Engine/Source/Editor/Kismet/Private/SSCSEditorViewport.cpp
+++ b/Engine/Source/Editor/Kismet/Private/SSCSEditorViewport.cpp
@@ -472,7 +472,7 @@ TSharedPtr<SWidget> SSCSEditorViewport::BuildViewportToolbar()
 
 			// Add the "Selection" sub menu.
 			{
-				FToolMenuEntry SelectionSubmenu = UE::UnrealEd::CreateViewportToolbarSelectionSection();
+				FToolMenuEntry SelectionSubmenu = UE::UnrealEd::CreateViewportToolbarSelectSection();
 				SelectionSubmenu.InsertPosition.Position = EToolMenuInsertType::First;
 				LeftSection.AddEntry(SelectionSubmenu);
 			}
diff --git a/Engine/Source/Editor/LevelEditor/Private/SLevelViewport.cpp b/Engine/Source/Editor/LevelEditor/Private/SLevelViewport.cpp
--- a/Engine/Source/Editor/LevelEditor/Private/SLevelViewport.cpp
+++ b/Engine/Source/Editor/LevelEditor/Private/SLevelViewport.cpp
@@ -1957,7 +1957,7 @@ TSharedPtr<SWidget> SLevelViewport::BuildViewportToolbar()
 
 				// Add the "Selection" sub menu.
 				{
-					FToolMenuEntry SelectionSubmenu = UE::UnrealEd::CreateViewportToolbarSelectionSection();
+					FToolMenuEntry SelectionSubmenu = UE::UnrealEd::CreateViewportToolbarSelectSection();
 					SelectionSubmenu.InsertPosition.Position = EToolMenuInsertType::First;
 					LeftSection.AddEntry(SelectionSubmenu);
 				}
diff --git a/Engine/Source/Editor/StaticMeshEditor/Private/SStaticMeshEditorViewport.cpp b/Engine/Source/Editor/StaticMeshEditor/Private/SStaticMeshEditorViewport.cpp
--- a/Engine/Source/Editor/StaticMeshEditor/Private/SStaticMeshEditorViewport.cpp
+++ b/Engine/Source/Editor/StaticMeshEditor/Private/SStaticMeshEditorViewport.cpp
@@ -800,7 +800,7 @@ TSharedPtr<SWidget> SStaticMeshEditorViewport::BuildViewportToolbar()
 
 			// Add the "Selection" sub menu.
 			{
-				FToolMenuEntry SelectionSubmenu = UE::UnrealEd::CreateViewportToolbarSelectionSection();
+				FToolMenuEntry SelectionSubmenu = UE::UnrealEd::CreateViewportToolbarSelectSection();
 				SelectionSubmenu.InsertPosition.Position = EToolMenuInsertType::First;
 				LeftSection.AddEntry(SelectionSubmenu);
 			}
diff --git a/Engine/Source/Editor/UnrealEd/Private/ViewportToolbar/UnrealEdViewportToolbar.cpp b/Engine/Source/Editor/UnrealEd/Private/ViewportToolbar/UnrealEdViewportToolbar.cpp
--- a/Engine/Source/Editor/UnrealEd/Private/ViewportToolbar/UnrealEdViewportToolbar.cpp
+++ b/Engine/Source/Editor/UnrealEd/Private/ViewportToolbar/UnrealEdViewportToolbar.cpp
@@ -684,9 +684,9 @@ FToolMenuEntry CreateViewportToolbarTransformsSection()
 	return Entry;
 }
 
-FToolMenuEntry CreateViewportToolbarSelectionSection()
+FToolMenuEntry CreateViewportToolbarSelectSection()
 {
-	return FToolMenuEntry::InitSubMenu(
+	FToolMenuEntry Entry = FToolMenuEntry::InitSubMenu(
 		"Select",
 		LOCTEXT("SelectonSubmenuLabel", "Select"),
 		LOCTEXT("SelectionSubmenuTooltip", "Viewport-related selection tools"),
@@ -991,6 +991,9 @@ FToolMenuEntry CreateViewportToolbarSelectionSection()
 			}
 		)
 	);
+
+	Entry.Icon = FSlateIcon(FAppStyle::GetAppStyleSetName(), "LevelEditor.SelectMode");
+	return Entry;
 }
 
 FToolMenuEntry CreateViewportToolbarSnappingSubmenu()
diff --git a/Engine/Source/Editor/UnrealEd/Public/ViewportToolbar/UnrealEdViewportToolbar.h b/Engine/Source/Editor/UnrealEd/Public/ViewportToolbar/UnrealEdViewportToolbar.h
--- a/Engine/Source/Editor/UnrealEd/Public/ViewportToolbar/UnrealEdViewportToolbar.h
+++ b/Engine/Source/Editor/UnrealEd/Public/ViewportToolbar/UnrealEdViewportToolbar.h
@@ -37,7 +37,7 @@ UNREALED_API bool ShowNewViewportToolbars();
 
 UNREALED_API FToolMenuEntry CreateViewportToolbarTransformsSection();
 
-UNREALED_API FToolMenuEntry CreateViewportToolbarSelectionSection();
+UNREALED_API FToolMenuEntry CreateViewportToolbarSelectSection();
 
 UNREALED_API FToolMenuEntry CreateViewportToolbarSnappingSubmenu();
 
