Make the Camera submenu's label dynamic in the new viewport toolbar

* Add UE::UnrealEd::GetCameraSubenuLabelFromViewportType and use it for the Camera submenu label.
* Deprecate SViewportToolbar::GetCameraMenuLabelFromViewportType from UE 5.6 (it calls UE::UnrealEd::GetCameraSubenuLabelFromViewportType in the meantime.)
* Update callers to use UE::UnrealEd::GetCameraSubenuLabelFromViewportType directly.

diff --git a/Engine/Plugins/Runtime/nDisplay/Source/DisplayClusterConfigurator/Private/Views/Viewport/DisplayClusterConfiguratorSCSEditorViewport.cpp b/Engine/Plugins/Runtime/nDisplay/Source/DisplayClusterConfigurator/Private/Views/Viewport/DisplayClusterConfiguratorSCSEditorViewport.cpp
--- a/Engine/Plugins/Runtime/nDisplay/Source/DisplayClusterConfigurator/Private/Views/Viewport/DisplayClusterConfiguratorSCSEditorViewport.cpp
+++ b/Engine/Plugins/Runtime/nDisplay/Source/DisplayClusterConfigurator/Private/Views/Viewport/DisplayClusterConfiguratorSCSEditorViewport.cpp
@@ -27,6 +27,7 @@
 #include "ViewportTabContent.h"
 #include "Framework/MultiBox/MultiBoxBuilder.h"
 #include "Slate/SceneViewport.h"
+#include "ViewportToolbar/UnrealEdViewportToolbar.h"
 #include "Widgets/Input/SNumericEntryBox.h"
 
 #define LOCTEXT_NAMESPACE "DisplayClusterSCSEditorViewport"
@@ -144,7 +145,9 @@ public:
 	{
 		if(EditorViewport.IsValid())
 		{
-			return GetCameraMenuLabelFromViewportType(EditorViewport.Pin()->GetViewportClient()->GetViewportType());
+			return UE::UnrealEd::GetCameraSubenuLabelFromViewportType(
+				EditorViewport.Pin()->GetViewportClient()->GetViewportType()
+			);
 		}
 
 		return NSLOCTEXT("BlueprintEditor", "CameraMenuTitle_Default", "Camera");
diff --git a/Engine/Source/Editor/Kismet/Private/SSCSEditorViewport.cpp b/Engine/Source/Editor/Kismet/Private/SSCSEditorViewport.cpp
--- a/Engine/Source/Editor/Kismet/Private/SSCSEditorViewport.cpp
+++ b/Engine/Source/Editor/Kismet/Private/SSCSEditorViewport.cpp
@@ -270,7 +270,9 @@ public:
 	{
 		if(EditorViewport.IsValid())
 		{
-			return GetCameraMenuLabelFromViewportType(EditorViewport.Pin()->GetViewportClient()->GetViewportType());
+			return UE::UnrealEd::GetCameraSubenuLabelFromViewportType(
+				EditorViewport.Pin()->GetViewportClient()->GetViewportType()
+			);
 		}
 
 		return NSLOCTEXT("BlueprintEditor", "CameraMenuTitle_Default", "Camera");
diff --git a/Engine/Source/Editor/LevelEditor/Private/SLevelViewportToolBar.cpp b/Engine/Source/Editor/LevelEditor/Private/SLevelViewportToolBar.cpp
--- a/Engine/Source/Editor/LevelEditor/Private/SLevelViewportToolBar.cpp
+++ b/Engine/Source/Editor/LevelEditor/Private/SLevelViewportToolBar.cpp
@@ -59,6 +59,7 @@
 #include "Engine/SceneCapture.h"
 #include "ViewportToolbar/LevelEditorViewportToolbarSections.h"
 #include "ViewportToolbar/LevelViewportContext.h"
+#include "ViewportToolbar/UnrealEdViewportToolbar.h"
 
 #define LOCTEXT_NAMESPACE "LevelViewportToolBar"
 
@@ -346,7 +347,7 @@ FText SLevelViewportToolBar::GetCameraMenuLabel() const
 	TSharedPtr< SLevelViewport > PinnedViewport( Viewport.Pin() );
 	if( PinnedViewport.IsValid() )
 	{
-		return GetCameraMenuLabelFromViewportType( PinnedViewport->GetLevelViewportClient().ViewportType );
+		return UE::UnrealEd::GetCameraSubenuLabelFromViewportType(PinnedViewport->GetLevelViewportClient().ViewportType);
 	}
 
 	return LOCTEXT("CameraMenuTitle_Default", "Camera");
diff --git a/Engine/Source/Editor/Persona/Private/SAnimViewportToolBar.cpp b/Engine/Source/Editor/Persona/Private/SAnimViewportToolBar.cpp
--- a/Engine/Source/Editor/Persona/Private/SAnimViewportToolBar.cpp
+++ b/Engine/Source/Editor/Persona/Private/SAnimViewportToolBar.cpp
@@ -55,6 +55,7 @@
 #include "Animation/MirrorDataTable.h"
 #include "ScopedTransaction.h"
 #include "SNameComboBox.h"
+#include "ViewportToolbar/UnrealEdViewportToolbar.h"
 
 #define LOCTEXT_NAMESPACE "AnimViewportToolBar"
 
@@ -1321,7 +1322,7 @@ FText SAnimViewportToolBar::GetCameraMenuLabel() const
 	TSharedPtr< SAnimationEditorViewportTabBody > PinnedViewport(Viewport.Pin());
 	if( PinnedViewport.IsValid() )
 	{
-		return GetCameraMenuLabelFromViewportType( PinnedViewport->GetLevelViewportClient().ViewportType );
+		return UE::UnrealEd::GetCameraSubenuLabelFromViewportType(PinnedViewport->GetLevelViewportClient().ViewportType);
 	}
 
 	return LOCTEXT("Viewport_Default", "Camera");
diff --git a/Engine/Source/Editor/UnrealEd/Private/SCommonEditorViewportToolbarBase.cpp b/Engine/Source/Editor/UnrealEd/Private/SCommonEditorViewportToolbarBase.cpp
--- a/Engine/Source/Editor/UnrealEd/Private/SCommonEditorViewportToolbarBase.cpp
+++ b/Engine/Source/Editor/UnrealEd/Private/SCommonEditorViewportToolbarBase.cpp
@@ -334,7 +334,7 @@ void SCommonEditorViewportToolbarBase::ConstructScreenPercentageMenu(FMenuBuilde
 
 FText SCommonEditorViewportToolbarBase::GetCameraMenuLabel() const
 {
-	return GetCameraMenuLabelFromViewportType( GetViewportClient().GetViewportType() );
+	return UE::UnrealEd::GetCameraSubenuLabelFromViewportType(GetViewportClient().GetViewportType());
 }
 
 
diff --git a/Engine/Source/Editor/UnrealEd/Private/SViewportToolBar.cpp b/Engine/Source/Editor/UnrealEd/Private/SViewportToolBar.cpp
--- a/Engine/Source/Editor/UnrealEd/Private/SViewportToolBar.cpp
+++ b/Engine/Source/Editor/UnrealEd/Private/SViewportToolBar.cpp
@@ -7,6 +7,7 @@
 #include "Styling/AppStyle.h"
 #include "UObject/NameTypes.h"
 #include "UObject/UnrealNames.h"
+#include "ViewportToolbar/UnrealEdViewportToolbar.h"
 #include "Widgets/Input/SMenuAnchor.h"
 
 struct FSlateBrush;
@@ -32,46 +33,9 @@ void SViewportToolBar::SetOpenMenu( TSharedPtr< SMenuAnchor >& NewMenu )
 	OpenedMenu = NewMenu;
 }
 
-
-
-
 FText SViewportToolBar::GetCameraMenuLabelFromViewportType(const ELevelViewportType ViewportType) const
 {
-	FText Label = LOCTEXT("CameraMenuTitle_Default", "Camera");
-	switch (ViewportType)
-	{
-	case LVT_Perspective:
-		Label = LOCTEXT("CameraMenuTitle_Perspective", "Perspective");
-		break;
-
-	case LVT_OrthoXY:
-		Label = LOCTEXT("CameraMenuTitle_Top", "Top");
-		break;
-
-	case LVT_OrthoNegativeXZ:
-		Label = LOCTEXT("CameraMenuTitle_Left", "Left");
-		break;
-
-	case LVT_OrthoNegativeYZ:
-		Label = LOCTEXT("CameraMenuTitle_Front", "Front");
-		break;
-
-	case LVT_OrthoNegativeXY:
-		Label = LOCTEXT("CameraMenuTitle_Bottom", "Bottom");
-		break;
-
-	case LVT_OrthoXZ:
-		Label = LOCTEXT("CameraMenuTitle_Right", "Right");
-		break;
-
-	case LVT_OrthoYZ:
-		Label = LOCTEXT("CameraMenuTitle_Back", "Back");
-		break;
-	case LVT_OrthoFreelook:
-		break;
-	}
-
-	return Label;
+	return UE::UnrealEd::GetCameraSubenuLabelFromViewportType(ViewportType);
 }
 
 const FSlateBrush* SViewportToolBar::GetCameraMenuLabelIconFromViewportType(const ELevelViewportType ViewportType) const
diff --git a/Engine/Source/Editor/UnrealEd/Private/ViewportToolbar/UnrealEdViewportToolbar.cpp b/Engine/Source/Editor/UnrealEd/Private/ViewportToolbar/UnrealEdViewportToolbar.cpp
--- a/Engine/Source/Editor/UnrealEd/Private/ViewportToolbar/UnrealEdViewportToolbar.cpp
+++ b/Engine/Source/Editor/UnrealEd/Private/ViewportToolbar/UnrealEdViewportToolbar.cpp
@@ -1951,6 +1951,45 @@ TSharedRef<SWidget> CreateFarViewPlaneMenuWidget(const TSharedRef<SEditorViewpor
 	// clang-format on
 }
 
+FText GetCameraSubenuLabelFromViewportType(const ELevelViewportType ViewportType)
+{
+	FText Label = LOCTEXT("CameraMenuTitle_Default", "Camera");
+	switch (ViewportType)
+	{
+	case LVT_Perspective:
+		Label = LOCTEXT("CameraMenuTitle_Perspective", "Perspective");
+		break;
+
+	case LVT_OrthoXY:
+		Label = LOCTEXT("CameraMenuTitle_Top", "Top");
+		break;
+
+	case LVT_OrthoNegativeXZ:
+		Label = LOCTEXT("CameraMenuTitle_Left", "Left");
+		break;
+
+	case LVT_OrthoNegativeYZ:
+		Label = LOCTEXT("CameraMenuTitle_Front", "Front");
+		break;
+
+	case LVT_OrthoNegativeXY:
+		Label = LOCTEXT("CameraMenuTitle_Bottom", "Bottom");
+		break;
+
+	case LVT_OrthoXZ:
+		Label = LOCTEXT("CameraMenuTitle_Right", "Right");
+		break;
+
+	case LVT_OrthoYZ:
+		Label = LOCTEXT("CameraMenuTitle_Back", "Back");
+		break;
+	case LVT_OrthoFreelook:
+		break;
+	}
+
+	return Label;
+}
+
 FToolMenuEntry CreateViewportToolbarCameraSubmenu()
 {
 	return FToolMenuEntry::InitDynamicEntry(
@@ -1958,9 +1997,28 @@ FToolMenuEntry CreateViewportToolbarCameraSubmenu()
 		FNewToolMenuSectionDelegate::CreateLambda(
 			[](FToolMenuSection& InDynamicSection) -> void
 			{
+				TWeakPtr<SEditorViewport> WeakViewport;
+				if (UUnrealEdViewportToolbarContext* const EditorViewportContext =
+						InDynamicSection.FindContext<UUnrealEdViewportToolbarContext>())
+				{
+					WeakViewport = EditorViewportContext->Viewport;
+				}
+
+				const TAttribute<FText> Label = TAttribute<FText>::CreateLambda(
+					[WeakViewport]()
+					{
+						if (TSharedPtr<SEditorViewport> Viewport = WeakViewport.Pin())
+						{
+							return UE::UnrealEd::GetCameraSubenuLabelFromViewportType(Viewport->GetViewportClient()->ViewportType
+							);
+						}
+						return LOCTEXT("CameraSubmenuLabel", "Camera");
+					}
+				);
+
 				InDynamicSection.AddSubMenu(
 					"CameraOptions",
-					LOCTEXT("CameraSubmenuLabel", "Camera"),
+					Label,
 					LOCTEXT("CameraSubmenuTooltip", "Camera options"),
 					FNewToolMenuDelegate::CreateLambda(
 						[](UToolMenu* Submenu) -> void
diff --git a/Engine/Source/Editor/UnrealEd/Public/SViewportToolBar.h b/Engine/Source/Editor/UnrealEd/Public/SViewportToolBar.h
--- a/Engine/Source/Editor/UnrealEd/Public/SViewportToolBar.h
+++ b/Engine/Source/Editor/UnrealEd/Public/SViewportToolBar.h
@@ -50,6 +50,7 @@ protected:
 	*
 	* @return	Label to use for this menu label
 	*/
+	UE_DEPRECATED(5.6, "Please use UE::UnrealEd::GetCameraSubenuLabelFromViewportType instead.")
 	UNREALED_API virtual FText GetCameraMenuLabelFromViewportType(const ELevelViewportType ViewportType) const;
 
 	/**
diff --git a/Engine/Source/Editor/UnrealEd/Public/ViewportToolbar/UnrealEdViewportToolbar.h b/Engine/Source/Editor/UnrealEd/Public/ViewportToolbar/UnrealEdViewportToolbar.h
--- a/Engine/Source/Editor/UnrealEd/Public/ViewportToolbar/UnrealEdViewportToolbar.h
+++ b/Engine/Source/Editor/UnrealEd/Public/ViewportToolbar/UnrealEdViewportToolbar.h
@@ -8,15 +8,16 @@
 #include "ToolMenuDelegates.h"
 #include "UnrealEdViewportToolbarContext.h"
 
-class IPreviewProfileController;
 class FEditorViewportClient;
 class FText;
 class IPreviewProfileController;
+class IPreviewProfileController;
 class SEditorViewport;
 class UToolMenu;
+enum ELevelViewportType : int;
+enum ERotationGridMode : int;
 struct FNewToolMenuChoice;
 struct FToolMenuEntry;
-enum ERotationGridMode : int;
 
 namespace UE::UnrealEd
 {
@@ -124,6 +125,7 @@ UNREALED_API FToolMenuEntry CreateNumericEntry(
 	int32 InMaxFractionalDigits = 2
 );
 
+UNREALED_API FText GetCameraSubenuLabelFromViewportType(const ELevelViewportType ViewportType);
 UNREALED_API FToolMenuEntry CreateViewportToolbarCameraSubmenu();
 
 UNREALED_API FToolMenuEntry CreateViewportToolbarAssetViewerProfileSubmenu(const TSharedPtr<IPreviewProfileController>& InPreviewProfileController);
