Make the Camera submenu's icon dynamic in the new viewport toolbar

* Change the icon based on the current camera mode just like in the old viewport toolbar.
* Rename UE::UnrealEd::GetCameraSubenuLabelFromViewportType to UE::UnrealEd::GetCameraSubmenuLabelFromViewportType to fix a typo.

diff --git a/Engine/Plugins/Runtime/nDisplay/Source/DisplayClusterConfigurator/Private/Views/Viewport/DisplayClusterConfiguratorSCSEditorViewport.cpp b/Engine/Plugins/Runtime/nDisplay/Source/DisplayClusterConfigurator/Private/Views/Viewport/DisplayClusterConfiguratorSCSEditorViewport.cpp
--- a/Engine/Plugins/Runtime/nDisplay/Source/DisplayClusterConfigurator/Private/Views/Viewport/DisplayClusterConfiguratorSCSEditorViewport.cpp
+++ b/Engine/Plugins/Runtime/nDisplay/Source/DisplayClusterConfigurator/Private/Views/Viewport/DisplayClusterConfiguratorSCSEditorViewport.cpp
@@ -145,7 +145,7 @@ public:
 	{
 		if(EditorViewport.IsValid())
 		{
-			return UE::UnrealEd::GetCameraSubenuLabelFromViewportType(
+			return UE::UnrealEd::GetCameraSubmenuLabelFromViewportType(
 				EditorViewport.Pin()->GetViewportClient()->GetViewportType()
 			);
 		}
diff --git a/Engine/Source/Editor/Kismet/Private/SSCSEditorViewport.cpp b/Engine/Source/Editor/Kismet/Private/SSCSEditorViewport.cpp
--- a/Engine/Source/Editor/Kismet/Private/SSCSEditorViewport.cpp
+++ b/Engine/Source/Editor/Kismet/Private/SSCSEditorViewport.cpp
@@ -270,7 +270,7 @@ public:
 	{
 		if(EditorViewport.IsValid())
 		{
-			return UE::UnrealEd::GetCameraSubenuLabelFromViewportType(
+			return UE::UnrealEd::GetCameraSubmenuLabelFromViewportType(
 				EditorViewport.Pin()->GetViewportClient()->GetViewportType()
 			);
 		}
diff --git a/Engine/Source/Editor/LevelEditor/Private/SLevelViewportToolBar.cpp b/Engine/Source/Editor/LevelEditor/Private/SLevelViewportToolBar.cpp
--- a/Engine/Source/Editor/LevelEditor/Private/SLevelViewportToolBar.cpp
+++ b/Engine/Source/Editor/LevelEditor/Private/SLevelViewportToolBar.cpp
@@ -347,7 +347,7 @@ FText SLevelViewportToolBar::GetCameraMenuLabel() const
 	TSharedPtr< SLevelViewport > PinnedViewport( Viewport.Pin() );
 	if( PinnedViewport.IsValid() )
 	{
-		return UE::UnrealEd::GetCameraSubenuLabelFromViewportType(PinnedViewport->GetLevelViewportClient().ViewportType);
+		return UE::UnrealEd::GetCameraSubmenuLabelFromViewportType(PinnedViewport->GetLevelViewportClient().ViewportType);
 	}
 
 	return LOCTEXT("CameraMenuTitle_Default", "Camera");
diff --git a/Engine/Source/Editor/Persona/Private/SAnimViewportToolBar.cpp b/Engine/Source/Editor/Persona/Private/SAnimViewportToolBar.cpp
--- a/Engine/Source/Editor/Persona/Private/SAnimViewportToolBar.cpp
+++ b/Engine/Source/Editor/Persona/Private/SAnimViewportToolBar.cpp
@@ -1322,7 +1322,7 @@ FText SAnimViewportToolBar::GetCameraMenuLabel() const
 	TSharedPtr< SAnimationEditorViewportTabBody > PinnedViewport(Viewport.Pin());
 	if( PinnedViewport.IsValid() )
 	{
-		return UE::UnrealEd::GetCameraSubenuLabelFromViewportType(PinnedViewport->GetLevelViewportClient().ViewportType);
+		return UE::UnrealEd::GetCameraSubmenuLabelFromViewportType(PinnedViewport->GetLevelViewportClient().ViewportType);
 	}
 
 	return LOCTEXT("Viewport_Default", "Camera");
diff --git a/Engine/Source/Editor/UnrealEd/Private/SCommonEditorViewportToolbarBase.cpp b/Engine/Source/Editor/UnrealEd/Private/SCommonEditorViewportToolbarBase.cpp
--- a/Engine/Source/Editor/UnrealEd/Private/SCommonEditorViewportToolbarBase.cpp
+++ b/Engine/Source/Editor/UnrealEd/Private/SCommonEditorViewportToolbarBase.cpp
@@ -334,7 +334,7 @@ void SCommonEditorViewportToolbarBase::ConstructScreenPercentageMenu(FMenuBuilde
 
 FText SCommonEditorViewportToolbarBase::GetCameraMenuLabel() const
 {
-	return UE::UnrealEd::GetCameraSubenuLabelFromViewportType(GetViewportClient().GetViewportType());
+	return UE::UnrealEd::GetCameraSubmenuLabelFromViewportType(GetViewportClient().GetViewportType());
 }
 
 
diff --git a/Engine/Source/Editor/UnrealEd/Private/SViewportToolBar.cpp b/Engine/Source/Editor/UnrealEd/Private/SViewportToolBar.cpp
--- a/Engine/Source/Editor/UnrealEd/Private/SViewportToolBar.cpp
+++ b/Engine/Source/Editor/UnrealEd/Private/SViewportToolBar.cpp
@@ -35,54 +35,12 @@ void SViewportToolBar::SetOpenMenu( TSharedPtr< SMenuAnchor >& NewMenu )
 
 FText SViewportToolBar::GetCameraMenuLabelFromViewportType(const ELevelViewportType ViewportType) const
 {
-	return UE::UnrealEd::GetCameraSubenuLabelFromViewportType(ViewportType);
+	return UE::UnrealEd::GetCameraSubmenuLabelFromViewportType(ViewportType);
 }
 
 const FSlateBrush* SViewportToolBar::GetCameraMenuLabelIconFromViewportType(const ELevelViewportType ViewportType) const
 {
-	static FName PerspectiveIcon("EditorViewport.Perspective");
-	static FName TopIcon("EditorViewport.Top");
-	static FName LeftIcon("EditorViewport.Left");
-	static FName FrontIcon("EditorViewport.Front");
-	static FName BottomIcon("EditorViewport.Bottom");
-	static FName RightIcon("EditorViewport.Right");
-	static FName BackIcon("EditorViewport.Back");
-
-	FName Icon = NAME_None;
-
-	switch (ViewportType)
-	{
-	case LVT_Perspective:
-		Icon = PerspectiveIcon;
-		break;
-
-	case LVT_OrthoXY:
-		Icon = TopIcon;
-		break;
-
-	case LVT_OrthoNegativeXZ:
-		Icon = LeftIcon;
-		break;
-
-	case LVT_OrthoNegativeYZ:
-		Icon = FrontIcon;
-		break;
-
-	case LVT_OrthoNegativeXY:
-		Icon = BottomIcon;
-		break;
-
-	case LVT_OrthoXZ:
-		Icon = RightIcon;
-		break;
-
-	case LVT_OrthoYZ:
-		Icon = BackIcon;
-		break;
-	case LVT_OrthoFreelook:
-		break;
-	}
-
+	const FName Icon = UE::UnrealEd::GetCameraSubmenuIconFNameFromViewportType(ViewportType);
 	return FAppStyle::GetBrush(Icon);
 }
 
diff --git a/Engine/Source/Editor/UnrealEd/Private/ViewportToolbar/UnrealEdViewportToolbar.cpp b/Engine/Source/Editor/UnrealEd/Private/ViewportToolbar/UnrealEdViewportToolbar.cpp
--- a/Engine/Source/Editor/UnrealEd/Private/ViewportToolbar/UnrealEdViewportToolbar.cpp
+++ b/Engine/Source/Editor/UnrealEd/Private/ViewportToolbar/UnrealEdViewportToolbar.cpp
@@ -1951,7 +1951,7 @@ TSharedRef<SWidget> CreateFarViewPlaneMenuWidget(const TSharedRef<SEditorViewpor
 	// clang-format on
 }
 
-FText GetCameraSubenuLabelFromViewportType(const ELevelViewportType ViewportType)
+FText GetCameraSubmenuLabelFromViewportType(const ELevelViewportType ViewportType)
 {
 	FText Label = LOCTEXT("CameraMenuTitle_Default", "Camera");
 	switch (ViewportType)
@@ -1990,6 +1990,54 @@ FText GetCameraSubenuLabelFromViewportType(const ELevelViewportType ViewportType
 	return Label;
 }
 
+FName GetCameraSubmenuIconFNameFromViewportType(const ELevelViewportType ViewportType)
+{
+	static FName PerspectiveIcon("EditorViewport.Perspective");
+	static FName TopIcon("EditorViewport.Top");
+	static FName LeftIcon("EditorViewport.Left");
+	static FName FrontIcon("EditorViewport.Front");
+	static FName BottomIcon("EditorViewport.Bottom");
+	static FName RightIcon("EditorViewport.Right");
+	static FName BackIcon("EditorViewport.Back");
+
+	FName Icon = NAME_None;
+
+	switch (ViewportType)
+	{
+	case LVT_Perspective:
+		Icon = PerspectiveIcon;
+		break;
+
+	case LVT_OrthoXY:
+		Icon = TopIcon;
+		break;
+
+	case LVT_OrthoNegativeXZ:
+		Icon = LeftIcon;
+		break;
+
+	case LVT_OrthoNegativeYZ:
+		Icon = FrontIcon;
+		break;
+
+	case LVT_OrthoNegativeXY:
+		Icon = BottomIcon;
+		break;
+
+	case LVT_OrthoXZ:
+		Icon = RightIcon;
+		break;
+
+	case LVT_OrthoYZ:
+		Icon = BackIcon;
+		break;
+	case LVT_OrthoFreelook:
+		break;
+	}
+
+	return Icon;
+}
+
 FToolMenuEntry CreateViewportToolbarCameraSubmenu()
 {
 	return FToolMenuEntry::InitDynamicEntry(
@@ -2009,13 +2057,27 @@ FToolMenuEntry CreateViewportToolbarCameraSubmenu()
 					{
 						if (TSharedPtr<SEditorViewport> Viewport = WeakViewport.Pin())
 						{
-							return UE::UnrealEd::GetCameraSubenuLabelFromViewportType(Viewport->GetViewportClient()->ViewportType
+							return UE::UnrealEd::GetCameraSubmenuLabelFromViewportType(Viewport->GetViewportClient()->ViewportType
 							);
 						}
 						return LOCTEXT("CameraSubmenuLabel", "Camera");
 					}
 				);
 
+				const TAttribute<FSlateIcon> Icon = TAttribute<FSlateIcon>::CreateLambda(
+					[WeakViewport]()
+					{
+						if (TSharedPtr<SEditorViewport> Viewport = WeakViewport.Pin())
+						{
+							const FName IconFName = UE::UnrealEd::GetCameraSubmenuIconFNameFromViewportType(
+								Viewport->GetViewportClient()->ViewportType
+							);
+							return FSlateIcon(FAppStyle::GetAppStyleSetName(), IconFName);
+						}
+						return FSlateIcon();
+					}
+				);
+
 				InDynamicSection.AddSubMenu(
 					"CameraOptions",
 					Label,
@@ -2025,7 +2087,9 @@ FToolMenuEntry CreateViewportToolbarCameraSubmenu()
 						{
 							PopulateCameraMenu(Submenu);
 						}
-					)
+					),
+					false,
+					Icon
 				);
 			}
 		)
diff --git a/Engine/Source/Editor/UnrealEd/Public/SViewportToolBar.h b/Engine/Source/Editor/UnrealEd/Public/SViewportToolBar.h
--- a/Engine/Source/Editor/UnrealEd/Public/SViewportToolBar.h
+++ b/Engine/Source/Editor/UnrealEd/Public/SViewportToolBar.h
@@ -50,7 +50,7 @@ protected:
 	*
 	* @return	Label to use for this menu label
 	*/
-	UE_DEPRECATED(5.6, "Please use UE::UnrealEd::GetCameraSubenuLabelFromViewportType instead.")
+	UE_DEPRECATED(5.6, "Please use UE::UnrealEd::GetCameraSubmenuLabelFromViewportType instead.")
 	UNREALED_API virtual FText GetCameraMenuLabelFromViewportType(const ELevelViewportType ViewportType) const;
 
 	/**
diff --git a/Engine/Source/Editor/UnrealEd/Public/ViewportToolbar/UnrealEdViewportToolbar.h b/Engine/Source/Editor/UnrealEd/Public/ViewportToolbar/UnrealEdViewportToolbar.h
--- a/Engine/Source/Editor/UnrealEd/Public/ViewportToolbar/UnrealEdViewportToolbar.h
+++ b/Engine/Source/Editor/UnrealEd/Public/ViewportToolbar/UnrealEdViewportToolbar.h
@@ -125,7 +125,8 @@ UNREALED_API FToolMenuEntry CreateNumericEntry(
 	int32 InMaxFractionalDigits = 2
 );
 
-UNREALED_API FText GetCameraSubenuLabelFromViewportType(const ELevelViewportType ViewportType);
+UNREALED_API FText GetCameraSubmenuLabelFromViewportType(const ELevelViewportType ViewportType);
+UNREALED_API FName GetCameraSubmenuIconFNameFromViewportType(const ELevelViewportType ViewportType);
 UNREALED_API FToolMenuEntry CreateViewportToolbarCameraSubmenu();
 
 UNREALED_API FToolMenuEntry CreateViewportToolbarAssetViewerProfileSubmenu(const TSharedPtr<IPreviewProfileController>& InPreviewProfileController);
