Added an advanced overload of DuplicateActor/s to support overriding the level and transactional settings of the function call

#rnx

diff --git a/Engine/Source/Editor/UnrealEd/Private/Subsystems/EditorActorSubsystem.cpp b/Engine/Source/Editor/UnrealEd/Private/Subsystems/EditorActorSubsystem.cpp
--- a/Engine/Source/Editor/UnrealEd/Private/Subsystems/EditorActorSubsystem.cpp
+++ b/Engine/Source/Editor/UnrealEd/Private/Subsystems/EditorActorSubsystem.cpp
@@ -682,15 +682,25 @@ TArray<class AActor*> UEditorActorSubsystem::ConvertActors(const TArray<class AA
 
 AActor* UEditorActorSubsystem::DuplicateActor(AActor* ActorToDuplicate, UWorld* ToWorld/*= nullptr*/, FVector Offset/* = FVector::ZeroVector*/)
 {
-	TArray<AActor*> Duplicate = DuplicateActors({ ActorToDuplicate }, ToWorld, Offset);
+	return DuplicateActor(ActorToDuplicate, ToWorld, Offset, FActorDuplicateParameters());
+}
+
+AActor* UEditorActorSubsystem::DuplicateActor(AActor* ActorToDuplicate, UWorld* ToWorld, FVector Offset, const FActorDuplicateParameters& DuplicateParams)
+{
+	TArray<AActor*> Duplicate = DuplicateActors({ ActorToDuplicate }, ToWorld, Offset, DuplicateParams);
 	return (Duplicate.Num() > 0) ? Duplicate[0] : nullptr;
 }
 
-TArray<AActor*> UEditorActorSubsystem::DuplicateActors(const TArray<AActor*>& ActorsToDuplicate, UWorld* InToWorld/*= nullptr*/, FVector Offset/* = FVector::ZeroVector*/)
+TArray<AActor*> UEditorActorSubsystem::DuplicateActors(const TArray<AActor*>& ActorsToDuplicate, UWorld* ToWorld/*= nullptr*/, FVector Offset/* = FVector::ZeroVector*/)
+{
+	return DuplicateActors(ActorsToDuplicate, ToWorld, Offset, FActorDuplicateParameters());
+}
+
+TArray<AActor*> UEditorActorSubsystem::DuplicateActors(const TArray<AActor*>& ActorsToDuplicate, UWorld* InToWorld, FVector Offset, const FActorDuplicateParameters& DuplicateParams)
 {
 	TGuardValue<bool> UnattendedScriptGuard(GIsRunningUnattendedScript, true);
 
-	FScopedTransaction Transaction(LOCTEXT("DuplicateActors", "Duplicate Actors"));
+	FScopedTransaction Transaction(LOCTEXT("DuplicateActors", "Duplicate Actors"), DuplicateParams.bTransact);
 
 	UUnrealEditorSubsystem* UnrealEditorSubsystem = GEditor->GetEditorSubsystem<UUnrealEditorSubsystem>();
 
@@ -700,17 +710,22 @@ TArray<AActor*> UEditorActorSubsystem::DuplicateActors(const TArray<AActor*>& Ac
 	}
 
 	UWorld* ToWorld = InToWorld ? InToWorld : UnrealEditorSubsystem->GetEditorWorld();
-
 	if (!ToWorld)
 	{
 		return TArray<AActor*>();
 	}
 
+	ULevel* ToLevel = DuplicateParams.LevelOverride;
+	if (!ToLevel || ToLevel->GetWorld() != ToWorld)
+	{
+		ToLevel = ToWorld->GetCurrentLevel();
+	}
+
 	FEditorDelegates::OnDuplicateActorsBegin.Broadcast();
 
 	TArray<AActor*> NewActors;
 	ABrush::SetSuppressBSPRegeneration(true);
-	GUnrealEd->DuplicateActors(ActorsToDuplicate, NewActors, ToWorld->GetCurrentLevel(), Offset);
+	GUnrealEd->DuplicateActors(ActorsToDuplicate, NewActors, ToLevel, Offset);
 	ABrush::SetSuppressBSPRegeneration(false);
 
 	// Find out if any of the actors will change the BSP.
diff --git a/Engine/Source/Editor/UnrealEd/Public/Subsystems/EditorActorSubsystem.h b/Engine/Source/Editor/UnrealEd/Public/Subsystems/EditorActorSubsystem.h
--- a/Engine/Source/Editor/UnrealEd/Public/Subsystems/EditorActorSubsystem.h
+++ b/Engine/Source/Editor/UnrealEd/Public/Subsystems/EditorActorSubsystem.h
@@ -86,6 +86,12 @@ public:
 	UPROPERTY(BlueprintAssignable, Category = "Editor Scripting | Level Utility")
 	FOnDeleteActorsEnd OnDeleteActorsEnd;
 
+	struct FActorDuplicateParameters
+	{
+		ULevel* LevelOverride = nullptr;
+		bool bTransact = true;
+	};
+
 	/**
 	 * Duplicate an actor from the world editor.
 	 * @param	ActorToDuplicate	Actor to duplicate.
@@ -95,6 +101,7 @@ public:
 	 */
 	UFUNCTION(BlueprintCallable, Category = "Editor Scripting | Level Utility", meta = (AdvancedDisplay = 1))
 	UNREALED_API AActor* DuplicateActor(AActor* ActorToDuplicate, UWorld* ToWorld = nullptr, FVector Offset = FVector::ZeroVector);
+	UNREALED_API AActor* DuplicateActor(AActor* ActorToDuplicate, UWorld* ToWorld, FVector Offset, const FActorDuplicateParameters& DuplicateParams);
 
 	/**
 	 * Duplicate actors from the world editor.
@@ -105,6 +112,7 @@ public:
 	 */
 	UFUNCTION(BlueprintCallable, Category = "Editor Scripting | Level Utility", meta = (AdvancedDisplay = 1))
 	UNREALED_API TArray<AActor*> DuplicateActors(const TArray<AActor*>& ActorsToDuplicate, UWorld* ToWorld = nullptr, FVector Offset = FVector::ZeroVector);
+	UNREALED_API TArray<AActor*> DuplicateActors(const TArray<AActor*>& ActorsToDuplicate, UWorld* ToWorld, FVector Offset, const FActorDuplicateParameters& DuplicateParams);
 
 	/**
 	 * Duplicate all the selected actors in the given world
