Fix issue where some chunk overrides were being overrwritten from the base set.
Add utility function GetChunkDependencies for help in testing.

diff --git a/Engine/Plugins/Runtime/GameFeatures/Source/GameFeatures/Private/GameFeatureAction_AddChunkOverride.cpp b/Engine/Plugins/Runtime/GameFeatures/Source/GameFeatures/Private/GameFeatureAction_AddChunkOverride.cpp
--- a/Engine/Plugins/Runtime/GameFeatures/Source/GameFeatures/Private/GameFeatureAction_AddChunkOverride.cpp
+++ b/Engine/Plugins/Runtime/GameFeatures/Source/GameFeatures/Private/GameFeatureAction_AddChunkOverride.cpp
@@ -156,7 +156,10 @@ void UGameFeatureAction_AddChunkOverride::AddChunkIdOverride()
 			}))
 		{
 			// If we found this chunk it might have been auto generated. Update this instead of adding ours.
-			ExistingDep->ParentChunkID = ParentChunk;
+			if (ExistingDep->ParentChunkID == 0)
+			{
+				ExistingDep->ParentChunkID = ParentChunk;
+			}
 		}
 		else
 		{
diff --git a/Engine/Source/Editor/UnrealEd/Private/Commandlets/ChunkDependencyInfo.cpp b/Engine/Source/Editor/UnrealEd/Private/Commandlets/ChunkDependencyInfo.cpp
--- a/Engine/Source/Editor/UnrealEd/Private/Commandlets/ChunkDependencyInfo.cpp
+++ b/Engine/Source/Editor/UnrealEd/Private/Commandlets/ChunkDependencyInfo.cpp
@@ -208,4 +208,12 @@ int32 UChunkDependencyInfo::FindHighestSharedChunk(const TArray<int32>& ChunkIDs
 		return 0;
 	}
 	return TopologicallySortedChunks[HighestIndex];
-}
\ No newline at end of file
+}
+
+void UChunkDependencyInfo::GetChunkDependencies(const int32 InChunk, TSet<int32>& OutChunkDependencies) const
+{
+	if (const TSet<int32>* Parents = ChildToParentMap.Find(InChunk))
+	{
+		OutChunkDependencies.Append(*Parents);
+	}
+}
diff --git a/Engine/Source/Editor/UnrealEd/Public/Commandlets/ChunkDependencyInfo.h b/Engine/Source/Editor/UnrealEd/Public/Commandlets/ChunkDependencyInfo.h
--- a/Engine/Source/Editor/UnrealEd/Public/Commandlets/ChunkDependencyInfo.h
+++ b/Engine/Source/Editor/UnrealEd/Public/Commandlets/ChunkDependencyInfo.h
@@ -74,6 +74,11 @@ class UChunkDependencyInfo : public UObject
 	Potentially will return a chunk that is not in the given input array */
 	UNREALED_API int32 FindHighestSharedChunk(const TArray<int32>& ChunkIDs) const;
 
+	/**
+	 * Returns the chunks that the in chunk depends upon through the parent rules.
+	 */
+	UNREALED_API void GetChunkDependencies(const int32 InChunk, TSet<int32>& OutChunkDependencies) const;
+
 	/** List of dependencies used to remove redundant chunks */
 	UPROPERTY(config)
 	TArray<FChunkDependency> DependencyArray;
