UE-221492 - Changing categories of macros in libraries does not update in search menu of other BPs

When changing a macro category, we need to actively update the action database with the change. While SetBlueprintFunctionOrMacroCategory _does_ call MarkBlueprintAsModified, this won't trigger the Blueprint's OnChanged callback, which the database normally relies on for changes.

diff --git a/Engine/Source/Editor/UnrealEd/Private/Kismet2/BlueprintEditorUtils.cpp b/Engine/Source/Editor/UnrealEd/Private/Kismet2/BlueprintEditorUtils.cpp
--- a/Engine/Source/Editor/UnrealEd/Private/Kismet2/BlueprintEditorUtils.cpp
+++ b/Engine/Source/Editor/UnrealEd/Private/Kismet2/BlueprintEditorUtils.cpp
@@ -108,6 +108,7 @@
 #include "ObjectEditorUtils.h"
 #include "Toolkits/ToolkitManager.h"
 #include "UnrealExporter.h"
+#include "BlueprintActionDatabase.h"
 #include "BlueprintEditorSettings.h"
 
 #include "Framework/Notifications/NotificationManager.h"
@@ -4181,6 +4182,8 @@ void FBlueprintEditorUtils::SetBlueprintFunctionOrMacroCategory(UEdGraph* Graph,
 			{
 				FBlueprintEditorUtils::MarkBlueprintAsModified(Blueprint);
 			}
+
+			FBlueprintActionDatabase::Get().RefreshAssetActions(Blueprint);
 		}
 	}
 }
