Introduce an interface (ICookArtifactReader) for reading data back from a package store used at cook time, with implementations to either read artifacts from zen (if not present on the filesystem), or read artifacts from the filesystem alone.  Also ensured that the reader is used in the places that currently expect to read data for incremental cook.  Specifically:
- CookSettings.txt
- DevelopmentAssetRegistry.bin
- Shader code libraries

diff --git a/Engine/Source/Developer/IoStoreUtilities/IoStoreUtilities.Build.cs b/Engine/Source/Developer/IoStoreUtilities/IoStoreUtilities.Build.cs
--- a/Engine/Source/Developer/IoStoreUtilities/IoStoreUtilities.Build.cs
+++ b/Engine/Source/Developer/IoStoreUtilities/IoStoreUtilities.Build.cs
@@ -6,30 +6,30 @@ public class IoStoreUtilities : ModuleRules
 {
 	public IoStoreUtilities (ReadOnlyTargetRules Target) : base(Target)
 	{
-        PrivateIncludePathModuleNames.AddRange(new string[] {
+		PrivateIncludePathModuleNames.AddRange(new string[] {
+			"StorageServerClient",
 			"TargetPlatform",
 		});
 
 		PrivateDependencyModuleNames.AddRange(new string[] {
-            "Core",
-            "CoreUObject",
-            "DerivedDataCache",
 			"AssetRegistry",
 			"CookMetadata",
+			"Core",
+			"CoreUObject",
+			"DerivedDataCache",
+			"DeveloperToolSettings",
+			"Json",
+			"PakFile",
 			"Projects",
-			"Zen",
 			"RenderCore",
+			"RSA",
+			"SandboxFile",
 			"Sockets",
+			"Zen",
 		});
 
 		PublicIncludePathModuleNames.AddRange(new string[] {
 			"Zen",
 		});
-
-		PrivateDependencyModuleNames.Add("PakFile");
-        PrivateDependencyModuleNames.Add("Json");
-        PrivateDependencyModuleNames.Add("RSA");
-        PrivateDependencyModuleNames.Add("DeveloperToolSettings");
-        PrivateDependencyModuleNames.Add("SandboxFile");
 	}
 }
diff --git a/Engine/Source/Developer/IoStoreUtilities/Private/CookArtifactReaderCommon.cpp b/Engine/Source/Developer/IoStoreUtilities/Private/CookArtifactReaderCommon.cpp
new file mode 100644
--- /dev/null
+++ b/Engine/Source/Developer/IoStoreUtilities/Private/CookArtifactReaderCommon.cpp
@@ -0,0 +1,174 @@
+// Copyright Epic Games, Inc. All Rights Reserved.
+
+#include "CookArtifactReaderCommon.h"
+
+#include "Algo/Accumulate.h"
+#include "Async/ParallelFor.h"
+#include "GenericPlatform/GenericPlatformFile.h"
+#include "HAL/CriticalSection.h"
+#include "HAL/FileManager.h"
+#include "HAL/FileManagerGeneric.h"
+#include "Misc/Paths.h"
+#include "Misc/ScopeRWLock.h"
+#include "Serialization/Archive.h"
+#include "Templates/UniquePtr.h"
+
+FArchive* FCookArtifactReaderCommon::CreateFileReader(const TCHAR* Filename)
+{
+	if (IFileHandle* File = OpenRead(Filename))
+	{
+		if (TUniquePtr<FArchive> Reader = MakeUnique<FArchiveFileReaderGeneric>(File, Filename, File->Size()))
+		{
+			return Reader.Release();
+		}
+	}
+	return nullptr;
+}
+
+bool FCookArtifactReaderCommon::IterateDirectoryRecursively(const TCHAR* Directory, IPlatformFile::FDirectoryVisitor& Visitor)
+{
+	class FRecurse : public IPlatformFile::FDirectoryVisitor
+	{
+	public:
+		FDirectoryVisitor&  Visitor;
+		TArray<FString>&    Directories;
+		FRecurse(FDirectoryVisitor& InVisitor, TArray<FString>& InDirectories)
+			: FDirectoryVisitor(InVisitor.DirectoryVisitorFlags)
+			, Visitor(InVisitor)
+			, Directories(InDirectories)
+		{
+		}
+		virtual bool Visit(const TCHAR* FilenameOrDirectory, bool bIsDirectory) override
+		{
+			bool bResult = Visitor.CallShouldVisitAndVisit(FilenameOrDirectory, bIsDirectory);
+			if (bResult && bIsDirectory)
+			{
+				Directories.Emplace(FilenameOrDirectory);
+			}
+			return bResult;
+		}
+	};
+
+	TArray<FString> DirectoriesToVisit;
+	DirectoriesToVisit.Add(Directory);
+
+	constexpr int32 MinBatchSize = 1;
+	const EParallelForFlags ParallelForFlags = FTaskGraphInterface::IsRunning() && Visitor.IsThreadSafe()
+		? EParallelForFlags::Unbalanced : EParallelForFlags::ForceSingleThread;
+	std::atomic<bool> bResult{true};
+	TArray<TArray<FString>> DirectoriesToVisitNext;
+	while (bResult && DirectoriesToVisit.Num() > 0)
+	{
+		ParallelForWithTaskContext(TEXT("IterateDirectoryRecursively.PF"),
+			DirectoriesToVisitNext,
+			DirectoriesToVisit.Num(),
+			MinBatchSize,
+			[this, &Visitor, &DirectoriesToVisit, &bResult](TArray<FString>& Directories, int32 Index)
+			{
+				FRecurse Recurse(Visitor, Directories);
+				if (bResult.load(std::memory_order_relaxed) && !IterateDirectory(*DirectoriesToVisit[Index], Recurse))
+				{
+					bResult.store(false, std::memory_order_relaxed);
+				}
+			},
+			ParallelForFlags);
+		DirectoriesToVisit.Reset(Algo::TransformAccumulate(DirectoriesToVisitNext, &TArray<FString>::Num, 0));
+		for (TArray<FString>& Directories : DirectoriesToVisitNext)
+		{
+			DirectoriesToVisit.Append(MoveTemp(Directories));
+		}
+	}
+
+	return bResult;
+}
+
+namespace CookArtifactReaderImpl
+{
+	class FFileMatch : public IPlatformFile::FDirectoryVisitor
+	{
+	public:
+		TArray<FString>& Result;
+		FRWLock ResultLock;
+		FString WildCard;
+		bool bFiles;
+		bool bDirectories;
+		bool bStoreFullPath;
+		FFileMatch(TArray<FString>& InResult, const FString& InWildCard, bool bInFiles, bool bInDirectories, bool bInStoreFullPath = false)
+			: IPlatformFile::FDirectoryVisitor(EDirectoryVisitorFlags::ThreadSafe)
+			, Result(InResult)
+			, WildCard(InWildCard)
+			, bFiles(bInFiles)
+			, bDirectories(bInDirectories)
+			, bStoreFullPath(bInStoreFullPath)
+		{
+		}
+
+		virtual bool ShouldVisitLeafPathname(FStringView LeafFilename) override
+		{
+			return FString(LeafFilename).MatchesWildcard(WildCard);
+		}
+
+		virtual bool Visit(const TCHAR* FilenameOrDirectory, bool bIsDirectory)
+		{
+			if ((bIsDirectory && bDirectories) || (!bIsDirectory && bFiles))
+			{
+				FString Filename = FPaths::GetCleanFilename(FilenameOrDirectory);
+				if (ensureMsgf(ShouldVisitLeafPathname(Filename),
+					TEXT("PlatformFile.IterateDirectory needs to call ShouldVisitLeafFilename before calling Visit.")))
+				{
+					FString FullPath = bStoreFullPath ? FString(FilenameOrDirectory) : MoveTemp(Filename);
+					FWriteScopeLock ScopeLock(ResultLock);
+					Result.Add(MoveTemp(FullPath));
+				}
+			}
+			return true;
+		}
+	};
+}
+
+void FCookArtifactReaderCommon::FindFiles( TArray<FString>& Result, const TCHAR* InFilename, bool Files, bool Directories )
+{
+	FString Filename( InFilename );
+	FPaths::NormalizeFilename( Filename );
+	const FString CleanFilename = FPaths::GetCleanFilename(Filename);
+	const bool bFindAllFiles = CleanFilename == TEXT("*") || CleanFilename == TEXT("*.*");
+	CookArtifactReaderImpl::FFileMatch FileMatch( Result, bFindAllFiles ? TEXT("*") : CleanFilename, Files, Directories );
+	IterateDirectory( *FPaths::GetPath(Filename), FileMatch );
+}
+
+void FCookArtifactReaderCommon::FindFiles(TArray<FString>& FoundFiles, const TCHAR* Directory, const TCHAR* FileExtension)
+{
+	if (!Directory)
+	{
+		return;
+	}
+
+	FString RootDir(Directory);
+	FString ExtStr = (FileExtension != nullptr) ? FString(FileExtension) : "";
+
+	// No Directory?
+	if (RootDir.Len() < 1)
+	{
+		return;
+	}
+
+	FPaths::NormalizeDirectoryName(RootDir);
+
+	// Don't modify the ExtStr if the user supplied the form "*.EXT" or "*" or "*.*" or "Name.*"
+	if (!ExtStr.Contains(TEXT("*")))
+	{
+		if (ExtStr == "")
+		{
+			ExtStr = "*.*";
+		}
+		else
+		{
+			//Complete the supplied extension with * or *. to yield "*.EXT"
+			ExtStr = (ExtStr.Left(1) == ".") ? "*" + ExtStr : "*." + ExtStr;
+		}
+	}
+
+	// Create the full filter, which is "Directory/*.EXT".
+	FString FinalPath = RootDir + "/" + ExtStr;
+	FindFiles(FoundFiles, *FinalPath, true, false);
+}
diff --git a/Engine/Source/Developer/IoStoreUtilities/Private/IoStoreUtilities.cpp b/Engine/Source/Developer/IoStoreUtilities/Private/IoStoreUtilities.cpp
--- a/Engine/Source/Developer/IoStoreUtilities/Private/IoStoreUtilities.cpp
+++ b/Engine/Source/Developer/IoStoreUtilities/Private/IoStoreUtilities.cpp
@@ -75,6 +75,7 @@
 #include "ShaderCodeArchive.h"
 #include "ZenStoreHttpClient.h"
 #include "IPlatformFilePak.h"
+#include "ZenCookArtifactReader.h"
 #include "ZenStoreWriter.h"
 #include "IO/IoContainerHeader.h"
 #include "ProfilingDebugging/CountersTrace.h"
@@ -8539,7 +8540,8 @@ int32 Staged2Zen(const FString& BuildPath, const FKeyChain& KeyChain, const FStr
 	}
 
 	FString MetaDataOutputPath = FPaths::Combine(CookedOutputPath, ProjectName, TEXT("Metadata"));
-	TUniquePtr<FZenStoreWriter> ZenStoreWriter = MakeUnique<FZenStoreWriter>(CookedOutputPath, MetaDataOutputPath, TargetPlatform);
+	TSharedRef<FZenCookArtifactReader> ZenCookArtifactReader = MakeShared<FZenCookArtifactReader>(CookedOutputPath, MetaDataOutputPath, TargetPlatform);
+	TUniquePtr<FZenStoreWriter> ZenStoreWriter = MakeUnique<FZenStoreWriter>(CookedOutputPath, MetaDataOutputPath, TargetPlatform, ZenCookArtifactReader);
 	
 	ICookedPackageWriter::FCookInfo CookInfo;
 	CookInfo.bFullBuild = true;
diff --git a/Engine/Source/Developer/IoStoreUtilities/Private/LayeredCookArtifactReader.cpp b/Engine/Source/Developer/IoStoreUtilities/Private/LayeredCookArtifactReader.cpp
new file mode 100644
--- /dev/null
+++ b/Engine/Source/Developer/IoStoreUtilities/Private/LayeredCookArtifactReader.cpp
@@ -0,0 +1,106 @@
+// Copyright Epic Games, Inc. All Rights Reserved.
+
+#include "LayeredCookArtifactReader.h"
+
+#include "HAL/CriticalSection.h"
+#include "Misc/ScopeRWLock.h"
+
+void FLayeredCookArtifactReader::AddLayer(TSharedRef<ICookArtifactReader> InLayer)
+{
+	Layers.AddUnique(InLayer);
+}
+
+bool FLayeredCookArtifactReader::RemoveLayer(TSharedRef<ICookArtifactReader> InLayer)
+{
+	return Layers.Remove(InLayer) > 0;
+}
+
+void FLayeredCookArtifactReader::EmptyLayers()
+{
+	Layers.Empty();
+}
+
+bool FLayeredCookArtifactReader::FileExists(const TCHAR* Filename)
+{
+	for (const TSharedRef<ICookArtifactReader>& Layer : Layers)
+	{
+		if (Layer->FileExists(Filename))
+		{
+			return true;
+		}
+	}
+	return false;
+}
+
+int64 FLayeredCookArtifactReader::FileSize(const TCHAR* Filename)
+{
+	for (const TSharedRef<ICookArtifactReader>& Layer : Layers)
+	{
+		int64 Size = Layer->FileSize(Filename);
+		if (Size >= 0)
+		{
+			return Size;
+		}
+	}
+
+	return -1;
+}
+
+IFileHandle* FLayeredCookArtifactReader::OpenRead(const TCHAR* Filename)
+{
+	for (const TSharedRef<ICookArtifactReader>& Layer : Layers)
+	{
+		if (IFileHandle* File = Layer->OpenRead(Filename))
+		{
+			return File;
+		}
+	}
+	return nullptr;
+}
+
+struct FDirectoryEntry
+{
+	FString Name;
+	bool bIsDirectory;
+};
+
+class FMergingVisitor : public IPlatformFile::FDirectoryVisitor
+{
+public:
+	FRWLock FoundEntriesLock;
+	TMap<FString, bool>& FoundEntries;
+	FMergingVisitor(TMap<FString, bool>& InFoundEntries)
+		: IPlatformFile::FDirectoryVisitor(EDirectoryVisitorFlags::ThreadSafe)
+		, FoundEntries(InFoundEntries)
+	{
+	}
+
+	virtual bool Visit(const TCHAR* FilenameOrDirectory, bool bIsDirectory) override
+	{
+		FString FileName(FilenameOrDirectory);
+		FRWScopeLock ScopeLock(FoundEntriesLock, SLT_Write);
+		FoundEntries.FindOrAdd(MoveTemp(FileName), bIsDirectory);
+		return true;
+	}
+};
+
+bool FLayeredCookArtifactReader::IterateDirectory(const TCHAR* Directory, IPlatformFile::FDirectoryVisitor& Visitor)
+{
+	TMap<FString, bool> Entries;
+	FMergingVisitor MergingVisitor(Entries);
+	bool bRetVal = false;
+
+	for (const TSharedRef<ICookArtifactReader>& Layer : Layers)
+	{
+		bRetVal |= Layer->IterateDirectory(Directory, MergingVisitor);
+	}
+
+	for (const TPair<FString, bool>& Entry : Entries)
+	{
+		if (!Visitor.CallShouldVisitAndVisit(*Entry.Key, Entry.Value))
+		{
+			return  bRetVal;
+		}
+	}
+	return bRetVal;
+}
diff --git a/Engine/Source/Developer/IoStoreUtilities/Private/LooseFilesCookArtifactReader.cpp b/Engine/Source/Developer/IoStoreUtilities/Private/LooseFilesCookArtifactReader.cpp
new file mode 100644
--- /dev/null
+++ b/Engine/Source/Developer/IoStoreUtilities/Private/LooseFilesCookArtifactReader.cpp
@@ -0,0 +1,24 @@
+// Copyright Epic Games, Inc. All Rights Reserved.
+
+#include "LooseFilesCookArtifactReader.h"
+#include "HAL/PlatformFileManager.h"
+
+bool FLooseFilesCookArtifactReader::FileExists(const TCHAR* Filename)
+{
+	return FPlatformFileManager::Get().GetPlatformFile().FileExists(Filename);
+}
+
+int64 FLooseFilesCookArtifactReader::FileSize(const TCHAR* Filename)
+{
+	return FPlatformFileManager::Get().GetPlatformFile().FileSize(Filename);
+}
+
+IFileHandle* FLooseFilesCookArtifactReader::OpenRead(const TCHAR* Filename)
+{
+	return FPlatformFileManager::Get().GetPlatformFile().OpenRead(Filename);
+}
+
+bool FLooseFilesCookArtifactReader::IterateDirectory(const TCHAR* Directory, IPlatformFile::FDirectoryVisitor& Visitor)
+{
+	return FPlatformFileManager::Get().GetPlatformFile().IterateDirectory(Directory, Visitor);
+}
diff --git a/Engine/Source/Developer/IoStoreUtilities/Private/ZenCookArtifactReader.cpp b/Engine/Source/Developer/IoStoreUtilities/Private/ZenCookArtifactReader.cpp
new file mode 100644
--- /dev/null
+++ b/Engine/Source/Developer/IoStoreUtilities/Private/ZenCookArtifactReader.cpp
@@ -0,0 +1,101 @@
+// Copyright Epic Games, Inc. All Rights Reserved.
+
+#include "ZenCookArtifactReader.h"
+
+#include "HAL/PlatformFileManager.h"
+#include "Misc/App.h"
+#include "Misc/Paths.h"
+#include "StorageServerClientModule.h"
+#include "ZenStoreHttpClient.h"
+
+FZenCookArtifactReader::FZenCookArtifactReader(
+	const FString& InputPath, 
+	const FString& InMetadataDirectoryPath, 
+	const ITargetPlatform* InTargetPlatform
+)
+	: ZenRootPath(InputPath)
+	, StorageServerPlatformFile(IStorageServerClientModule::Get().TryCreateCustomPlatformFile(*InputPath, &FPlatformFileManager::Get().GetPlatformFile()))
+{
+	if (StorageServerPlatformFile)
+	{
+		StorageServerPlatformFile->SetLowerLevel(nullptr);
+	}
+}
+
+FZenCookArtifactReader::~FZenCookArtifactReader()
+{
+}
+
+bool FZenCookArtifactReader::FileExists(const TCHAR* Filename)
+{
+	if (StorageServerPlatformFile)
+	{
+		FString StandardFilename;
+		if (MakeStorageServerPath(Filename, StandardFilename))
+		{
+			return StorageServerPlatformFile->FileExists(*StandardFilename);
+		}
+	}
+	return false;
+}
+
+int64 FZenCookArtifactReader::FileSize(const TCHAR* Filename)
+{
+	if (StorageServerPlatformFile)
+	{
+		FString StandardFilename;
+		if (MakeStorageServerPath(Filename, StandardFilename))
+		{
+			return StorageServerPlatformFile->FileSize(*StandardFilename);
+		}
+	}
+
+	return -1;
+}
+
+IFileHandle* FZenCookArtifactReader::OpenRead(const TCHAR* Filename)
+{
+	if (StorageServerPlatformFile)
+	{
+		FString StandardFilename;
+		if (MakeStorageServerPath(Filename, StandardFilename))
+		{
+			return StorageServerPlatformFile->OpenRead(*StandardFilename);
+		}
+	}
+	return nullptr;
+}
+
+bool FZenCookArtifactReader::IterateDirectory(const TCHAR* Directory, IPlatformFile::FDirectoryVisitor& Visitor)
+{
+	if (StorageServerPlatformFile)
+	{
+		FString StandardDirectory;
+		if (MakeStorageServerPath(Directory, StandardDirectory))
+		{
+			return StorageServerPlatformFile->IterateDirectory(Directory, Visitor);
+		}
+	}
+
+	return false;
+}
+
+bool FZenCookArtifactReader::MakeStorageServerPath(const TCHAR* Filename, FString& OutFilename) const
+{
+	OutFilename = Filename;
+	if (FPaths::IsUnderDirectory(OutFilename, ZenRootPath) && FPaths::MakePathRelativeTo(OutFilename, *ZenRootPath))
+	{
+		static FString ProjectPrefix = FString::Printf(TEXT("%s/"), FApp::GetProjectName());
+		if (OutFilename.StartsWith(*ProjectPrefix))
+		{
+			OutFilename.RightChopInline(ProjectPrefix.Len());
+			OutFilename.InsertAt(0, *FPaths::ProjectDir());
+		}
+		else
+		{
+			OutFilename.InsertAt(0, "../../../");
+		}
+		return true;
+	}
+	return false;
+}
\ No newline at end of file
diff --git a/Engine/Source/Developer/IoStoreUtilities/Private/ZenStoreWriter.cpp b/Engine/Source/Developer/IoStoreUtilities/Private/ZenStoreWriter.cpp
--- a/Engine/Source/Developer/IoStoreUtilities/Private/ZenStoreWriter.cpp
+++ b/Engine/Source/Developer/IoStoreUtilities/Private/ZenStoreWriter.cpp
@@ -31,6 +31,7 @@
 #include "Serialization/LargeMemoryWriter.h" 
 #include "SocketSubsystem.h"
 #include "UObject/SavePackage.h"
+#include "ZenCookArtifactReader.h"
 #include "ZenFileSystemManifest.h"
 #include "ZenStoreHttpClient.h"
 
@@ -162,9 +163,11 @@ void FZenStoreWriter::StaticInit()
 FZenStoreWriter::FZenStoreWriter(
 	const FString& InOutputPath, 
 	const FString& InMetadataDirectoryPath, 
-	const ITargetPlatform* InTargetPlatform
+	const ITargetPlatform* InTargetPlatform,
+	TSharedRef<FZenCookArtifactReader> InCookArtifactReader
 )
-	: TargetPlatform(*InTargetPlatform)
+	: CookArtifactReader(InCookArtifactReader)
+	, TargetPlatform(*InTargetPlatform)
 	, TargetPlatformFName(*InTargetPlatform->PlatformName())
 	, OutputPath(InOutputPath)
 	, MetadataDirectoryPath(InMetadataDirectoryPath)
@@ -553,10 +556,13 @@ void FZenStoreWriter::Initialize(const FCookInfo& Info)
 					FString ServerPath = FString(FileObj["serverpath"].AsString());
 					FString ClientPath = FString(FileObj["clientpath"].AsString());
 
-					FIoChunkId FileChunkId;
-					FileChunkId.Set(FileId.GetView());
+					if (!ServerPath.IsEmpty())
+					{
+						FIoChunkId FileChunkId;
+						FileChunkId.Set(FileId.GetView());
 
-					ZenFileSystemManifest->AddManifestEntry(FileChunkId, MoveTemp(ServerPath), MoveTemp(ClientPath));
+						ZenFileSystemManifest->AddManifestEntry(FileChunkId, MoveTemp(ServerPath), MoveTemp(ClientPath));
+					}
 				}
 
 				UE_LOG(LogZenStoreWriter, Display, TEXT("Fetched '%d' file(s) from oplog '%s/%s'"), ZenFileSystemManifest->NumEntries(), *ProjectId, *OplogId);
@@ -1041,16 +1047,16 @@ TUniquePtr<FAssetRegistryState> FZenStoreWriter::LoadPreviousAssetRegistry()
 	// the returned asset registry to the intersection of the oplog and the previous asset registry;
 	// to report a package as already cooked we have to have the information from both sources.
 	FString PreviousAssetRegistryFile = FPaths::Combine(MetadataDirectoryPath, GetDevelopmentAssetRegistryFilename());
-	FArrayReader SerializedAssetData;
-	if (!IFileManager::Get().FileExists(*PreviousAssetRegistryFile) ||
-		!FFileHelper::LoadFileToArray(SerializedAssetData, *PreviousAssetRegistryFile))
+	TUniquePtr<FArchive> Reader(CookArtifactReader->CreateFileReader(*PreviousAssetRegistryFile));
+
+	if (!Reader)
 	{
 		RemoveCookedPackages();
 		return TUniquePtr<FAssetRegistryState>();
 	}
 
 	TUniquePtr<FAssetRegistryState> PreviousState = MakeUnique<FAssetRegistryState>();
-	PreviousState->Load(SerializedAssetData);
+	PreviousState->Load(*Reader);
 
 	TSet<FName> RemoveSet;
 	const TMap<FName, const FAssetPackageData*>& PreviousStatePackages = PreviousState->GetAssetPackageDataMap(); 
diff --git a/Engine/Source/Developer/IoStoreUtilities/Public/CookArtifactReader.h b/Engine/Source/Developer/IoStoreUtilities/Public/CookArtifactReader.h
new file mode 100644
--- /dev/null
+++ b/Engine/Source/Developer/IoStoreUtilities/Public/CookArtifactReader.h
@@ -0,0 +1,27 @@
+// Copyright Epic Games, Inc. All Rights Reserved.
+
+#pragma once
+
+#include "Containers/Array.h"
+#include "Containers/UnrealString.h"
+#include "GenericPlatform/GenericPlatformFile.h"
+
+class FArchive;
+
+class ICookArtifactReader
+{
+public:
+	virtual ~ICookArtifactReader() = default;
+	
+	virtual bool FileExists(const TCHAR* Filename) = 0;
+	virtual int64 FileSize(const TCHAR* Filename) = 0;
+	virtual IFileHandle* OpenRead(const TCHAR* Filename) = 0;
+
+	virtual bool IterateDirectory(const TCHAR* Directory, IPlatformFile::FDirectoryVisitor& Visitor) = 0;
+
+	// Utilities
+	virtual FArchive* CreateFileReader(const TCHAR* Filename) = 0;
+	virtual bool IterateDirectoryRecursively(const TCHAR* Directory, IPlatformFile::FDirectoryVisitor& Visitor) = 0;
+	virtual void FindFiles(TArray<FString>& Result, const TCHAR* Filename, bool Files, bool Directories) = 0;
+	virtual void FindFiles(TArray<FString>& FoundFiles, const TCHAR* Directory, const TCHAR* FileExtension = nullptr) = 0;
+};
diff --git a/Engine/Source/Developer/IoStoreUtilities/Public/CookArtifactReaderCommon.h b/Engine/Source/Developer/IoStoreUtilities/Public/CookArtifactReaderCommon.h
new file mode 100644
--- /dev/null
+++ b/Engine/Source/Developer/IoStoreUtilities/Public/CookArtifactReaderCommon.h
@@ -0,0 +1,23 @@
+// Copyright Epic Games, Inc. All Rights Reserved.
+
+#pragma once
+
+#include "Containers/Array.h"
+#include "Containers/UnrealString.h"
+#include "CookArtifactReader.h"
+#include "GenericPlatform/GenericPlatformFile.h"
+
+class FArchive;
+
+class FCookArtifactReaderCommon
+	: public ICookArtifactReader
+{
+public:
+	IOSTOREUTILITIES_API virtual ~FCookArtifactReaderCommon() = default;
+
+	// Utilities
+	IOSTOREUTILITIES_API FArchive* CreateFileReader(const TCHAR* Filename) override;
+	IOSTOREUTILITIES_API bool IterateDirectoryRecursively(const TCHAR* Directory, IPlatformFile::FDirectoryVisitor& Visitor) override;
+	IOSTOREUTILITIES_API void FindFiles(TArray<FString>& Result, const TCHAR* Filename, bool Files, bool Directories) override;
+	IOSTOREUTILITIES_API void FindFiles(TArray<FString>& FoundFiles, const TCHAR* Directory, const TCHAR* FileExtension = nullptr) override;
+};
diff --git a/Engine/Source/Developer/IoStoreUtilities/Public/LayeredCookArtifactReader.h b/Engine/Source/Developer/IoStoreUtilities/Public/LayeredCookArtifactReader.h
new file mode 100644
--- /dev/null
+++ b/Engine/Source/Developer/IoStoreUtilities/Public/LayeredCookArtifactReader.h
@@ -0,0 +1,26 @@
+// Copyright Epic Games, Inc. All Rights Reserved.
+
+#pragma once
+
+#include "Containers/Array.h"
+#include "CookArtifactReaderCommon.h"
+#include "Templates/SharedPointer.h"
+
+class FLayeredCookArtifactReader
+	: public FCookArtifactReaderCommon
+{
+public:
+	IOSTOREUTILITIES_API FLayeredCookArtifactReader() = default;
+
+	IOSTOREUTILITIES_API void AddLayer(TSharedRef<ICookArtifactReader> InLayer);
+	IOSTOREUTILITIES_API bool RemoveLayer(TSharedRef<ICookArtifactReader> InLayer);
+	IOSTOREUTILITIES_API void EmptyLayers();
+
+	IOSTOREUTILITIES_API bool FileExists(const TCHAR* Filename) override;
+	IOSTOREUTILITIES_API int64 FileSize(const TCHAR* Filename) override;
+	IOSTOREUTILITIES_API IFileHandle* OpenRead(const TCHAR* Filename) override;
+
+	IOSTOREUTILITIES_API bool IterateDirectory(const TCHAR* Directory, IPlatformFile::FDirectoryVisitor& Visitor) override;
+private:
+	TArray<TSharedRef<ICookArtifactReader>> Layers;
+};
diff --git a/Engine/Source/Developer/IoStoreUtilities/Public/LooseFilesCookArtifactReader.h b/Engine/Source/Developer/IoStoreUtilities/Public/LooseFilesCookArtifactReader.h
new file mode 100644
--- /dev/null
+++ b/Engine/Source/Developer/IoStoreUtilities/Public/LooseFilesCookArtifactReader.h
@@ -0,0 +1,18 @@
+// Copyright Epic Games, Inc. All Rights Reserved.
+
+#pragma once
+
+#include "CookArtifactReaderCommon.h"
+
+class FLooseFilesCookArtifactReader
+	: public FCookArtifactReaderCommon
+{
+public:
+	IOSTOREUTILITIES_API FLooseFilesCookArtifactReader() = default;
+
+	IOSTOREUTILITIES_API bool FileExists(const TCHAR* Filename) override;
+	IOSTOREUTILITIES_API int64 FileSize(const TCHAR* Filename) override;
+	IOSTOREUTILITIES_API IFileHandle* OpenRead(const TCHAR* Filename) override;
+
+	IOSTOREUTILITIES_API bool IterateDirectory(const TCHAR* Directory, IPlatformFile::FDirectoryVisitor& Visitor) override;
+};
diff --git a/Engine/Source/Developer/IoStoreUtilities/Public/ZenCookArtifactReader.h b/Engine/Source/Developer/IoStoreUtilities/Public/ZenCookArtifactReader.h
new file mode 100644
--- /dev/null
+++ b/Engine/Source/Developer/IoStoreUtilities/Public/ZenCookArtifactReader.h
@@ -0,0 +1,30 @@
+// Copyright Epic Games, Inc. All Rights Reserved.
+
+#pragma once
+
+#include "Containers/UnrealString.h"
+#include "CookArtifactReaderCommon.h"
+#include "Templates/UniquePtr.h"
+
+class IStorageServerPlatformFile;
+
+class FZenCookArtifactReader
+	: public FCookArtifactReaderCommon
+{
+public:
+	IOSTOREUTILITIES_API FZenCookArtifactReader(const FString& InputPath, 
+												const FString& MetadataDirectoryPath, 
+												const ITargetPlatform* TargetPlatform);
+	IOSTOREUTILITIES_API virtual ~FZenCookArtifactReader();
+
+	IOSTOREUTILITIES_API bool FileExists(const TCHAR* Filename) override;
+	IOSTOREUTILITIES_API int64 FileSize(const TCHAR* Filename) override;
+	IOSTOREUTILITIES_API IFileHandle* OpenRead(const TCHAR* Filename) override;
+
+	IOSTOREUTILITIES_API bool IterateDirectory(const TCHAR* Directory, IPlatformFile::FDirectoryVisitor& Visitor) override;
+private:
+	bool MakeStorageServerPath(const TCHAR* Filename, FString& OutFilename) const;
+
+	FString ZenRootPath;
+	TUniquePtr<IStorageServerPlatformFile> StorageServerPlatformFile;
+};
diff --git a/Engine/Source/Developer/IoStoreUtilities/Public/ZenStoreWriter.h b/Engine/Source/Developer/IoStoreUtilities/Public/ZenStoreWriter.h
--- a/Engine/Source/Developer/IoStoreUtilities/Public/ZenStoreWriter.h
+++ b/Engine/Source/Developer/IoStoreUtilities/Public/ZenStoreWriter.h
@@ -25,6 +25,7 @@
 #include "Serialization/CompactBinary.h"
 #include "Serialization/PackageWriter.h"
 #include "Templates/RefCounting.h"
+#include "Templates/SharedPointer.h"
 #include "Templates/UniquePtr.h"
 #include "UObject/NameTypes.h"
 
@@ -46,6 +47,7 @@ class FCbPackage;
 class FCbWriter;
 class FPackageStoreOptimizer;
 class FPackageStorePackage;
+class FZenCookArtifactReader;
 class FZenFileSystemManifest;
 
 /** 
@@ -57,7 +59,8 @@ class FZenStoreWriter
 public:
 	IOSTOREUTILITIES_API FZenStoreWriter(	const FString& OutputPath, 
 											const FString& MetadataDirectoryPath, 
-											const ITargetPlatform* TargetPlatform);
+											const ITargetPlatform* TargetPlatform,
+											TSharedRef<FZenCookArtifactReader> CookArtifactReader);
 
 	IOSTOREUTILITIES_API ~FZenStoreWriter();
 
@@ -233,6 +236,7 @@ private:
 	TUniquePtr<UE::FZenStoreHttpClient>	HttpClient;
 	bool IsLocalConnection = true;
 
+	TSharedRef<FZenCookArtifactReader>	CookArtifactReader;
 	const ITargetPlatform&				TargetPlatform;
 	const FName							TargetPlatformFName;
 	FString								ProjectId;
diff --git a/Engine/Source/Editor/UnrealEd/Classes/CookOnTheSide/CookOnTheFlyServer.h b/Engine/Source/Editor/UnrealEd/Classes/CookOnTheSide/CookOnTheFlyServer.h
--- a/Engine/Source/Editor/UnrealEd/Classes/CookOnTheSide/CookOnTheFlyServer.h
+++ b/Engine/Source/Editor/UnrealEd/Classes/CookOnTheSide/CookOnTheFlyServer.h
@@ -14,6 +14,7 @@
 #include "Misc/PackageAccessTracking.h"
 #include "Templates/Function.h"
 #include "Templates/RefCounting.h"
+#include "Templates/SharedPointer.h"
 #include "Templates/UniquePtr.h"
 #include "TickableEditorObject.h"
 #include "UObject/ICookInfo.h"
@@ -30,9 +31,12 @@ class FAssetRegistryGenerator;
 class FAsyncIODelete;
 class FDiffModeCookServerUtils;
 class FIterativeValidatePackageWriter;
+class FLayeredCookArtifactReader;
+class FLooseFilesCookArtifactReader;
 class FReferenceCollector;
 class FSavePackageContext;
 class IAssetRegistry;
+class ICookArtifactReader;
 class ICookedPackageWriter;
 class IPlugin;
 class ITargetPlatform;
@@ -1096,7 +1100,7 @@ private:
 	 */
 	void LoadBeginCookIterativeFlags(FBeginCookContext& BeginContext);
 	/** const because it is not always calledand should avoid sideeffects */
-	void LoadBeginCookIterativeFlagsLocal(FBeginCookContext& BeginContext) const;
+	void LoadBeginCookIterativeFlagsLocal(FBeginCookContext& BeginContext);
 	/** Initialize the sandbox for a new cook session */
 	void BeginCookSandbox(FBeginCookContext& BeginContext);
 
@@ -1226,7 +1230,7 @@ private:
 	 * previous cook even when running iteratively.
 	 */
 	bool ArePreviousCookSettingsCompatible(const TMap<FName, FString>& CurrentCookSettings,
-		const ITargetPlatform* TargetPlatform) const;
+		const ITargetPlatform* TargetPlatform);
 	/** Save the CurrentCookSettings into the output directory. */
 	void SaveCookSettings(const TMap<FName, FString>& CurrentCookSettings, const ITargetPlatform* TargetPlatform);
 	/** Load the CookSettings written at beginning of cook and rewrite them with CookInProgress flag removed. */
@@ -1394,6 +1398,8 @@ private:
 	UE::Cook::EPollStatus TryPopulateGeneratedPackage(UE::Cook::FGenerationHelper& GenerationHelper,
 		UE::Cook::FCookGenerationInfo& GeneratedInfo);
 
+	ICookArtifactReader& FindOrCreateCookArtifactReader(const ITargetPlatform* TargetPlatform);
+	const ICookArtifactReader* FindCookArtifactReader(const ITargetPlatform* TargetPlatform) const;
 	ICookedPackageWriter& FindOrCreatePackageWriter(const ITargetPlatform* TargetPlatform);
 	const ICookedPackageWriter* FindPackageWriter(const ITargetPlatform* TargetPlatform) const;
 	void FindOrCreateSaveContexts(TConstArrayView<const ITargetPlatform*> TargetPlatforms);
@@ -1567,6 +1573,8 @@ private:
 	TUniquePtr<UE::Cook::FBuildDefinitions> BuildDefinitions;
 	TUniquePtr<UE::Cook::FCookDirector> CookDirector;
 	TUniquePtr<UE::Cook::FCookWorkerClient> CookWorkerClient;
+	TUniquePtr<FLayeredCookArtifactReader> AllContextArtifactReader;
+	TSharedPtr<FLooseFilesCookArtifactReader> SharedLooseFilesCookArtifactReader;
 
 	TArray<UE::Cook::FCookSavePackageContext*> SavePackageContexts;
 	/**
diff --git a/Engine/Source/Editor/UnrealEd/Private/CookOnTheFlyServer.cpp b/Engine/Source/Editor/UnrealEd/Private/CookOnTheFlyServer.cpp
--- a/Engine/Source/Editor/UnrealEd/Private/CookOnTheFlyServer.cpp
+++ b/Engine/Source/Editor/UnrealEd/Private/CookOnTheFlyServer.cpp
@@ -89,10 +89,12 @@
 #include "Internationalization/Culture.h"
 #include "Internationalization/PackageLocalizationManager.h"
 #include "IPAddress.h"
+#include "LayeredCookArtifactReader.h"
 #include "LocalizationChunkDataGenerator.h"
 #include "LockFile.h"
 #include "Logging/MessageLog.h"
 #include "Logging/TokenizedMessage.h"
+#include "LooseFilesCookArtifactReader.h"
 #include "Materials/Material.h"
 #include "Materials/MaterialInterface.h"
 #include "MeshCardRepresentation.h"
@@ -115,6 +117,7 @@
 #include "Modules/ModuleManager.h"
 #include "ObjectTools.h"
 #include "PackageHelperFunctions.h"
+#include "PipelineCacheChunkDataGenerator.h"
 #include "PlatformInfo.h"
 #include "ProfilingDebugging/CookStats.h"
 #include "ProfilingDebugging/PlatformFileTrace.h"
@@ -131,10 +134,9 @@
 #include "Settings/ProjectPackagingSettings.h"
 #include "ShaderCodeLibrary.h"
 #include "ShaderCompiler.h"
+#include "ShaderLibraryChunkDataGenerator.h"
 #include "ShaderStats.h"
 #include "ShaderStatsCollector.h"
-#include "ShaderLibraryChunkDataGenerator.h"
-#include "PipelineCacheChunkDataGenerator.h"
 #include "String/Find.h"
 #include "String/ParseLines.h"
 #include "String/ParseTokens.h"
@@ -156,6 +158,7 @@
 #include "UObject/UObjectArray.h"
 #include "UObject/UObjectIterator.h"
 #include "UserGeneratedContentLocalization.h"
+#include "ZenCookArtifactReader.h"
 #include "ZenStoreWriter.h"
 
 #define LOCTEXT_NAMESPACE "Cooker"
@@ -6897,6 +6900,9 @@ void UCookOnTheFlyServer::Initialize( ECookMode::Type DesiredCookMode, ECookInit
 	PackageTracker = MakeUnique<UE::Cook::FPackageTracker>(*this);
 	DiffModeHelper = MakeUnique<FDiffModeCookServerUtils>();
 	BuildDefinitions = MakeUnique<UE::Cook::FBuildDefinitions>();
+	SharedLooseFilesCookArtifactReader = MakeShared<FLooseFilesCookArtifactReader>();
+	AllContextArtifactReader = MakeUnique<FLayeredCookArtifactReader>();
+	AllContextArtifactReader->AddLayer(SharedLooseFilesCookArtifactReader.ToSharedRef());
 	CookByTheBookOptions = MakeUnique<UE::Cook::FCookByTheBookOptions>();
 	CookOnTheFlyOptions = MakeUnique<UE::Cook::FCookOnTheFlyOptions>();
 	AssetRegistry = IAssetRegistry::Get();
@@ -8250,11 +8256,18 @@ FString UCookOnTheFlyServer::GetCookSettingsFileName(const ITargetPlatform* Targ
 	return ConvertToFullSandboxPath(*CookedSettingsIni, true, TargetPlatform->PlatformName());
 }
 
-bool UCookOnTheFlyServer::ArePreviousCookSettingsCompatible(const TMap<FName, FString>& CurrentCookSettings, const ITargetPlatform* TargetPlatform) const
+bool UCookOnTheFlyServer::ArePreviousCookSettingsCompatible(const TMap<FName, FString>& CurrentCookSettings, const ITargetPlatform* TargetPlatform)
 {
 	FConfigFile ConfigFile;
 	FString CookSettingsFileName = GetCookSettingsFileName(TargetPlatform);
-	ConfigFile.Read(CookSettingsFileName);
+	if (TUniquePtr<FArchive> Reader(FindOrCreateCookArtifactReader(TargetPlatform).CreateFileReader(*CookSettingsFileName)); Reader)
+	{
+		FString CookSettingsFileContents;
+		if (FFileHelper::LoadFileToString(CookSettingsFileContents, *Reader.Get()))
+		{
+			ConfigFile.ProcessInputFileContents(CookSettingsFileContents, CookSettingsFileName);
+		}
+	}
 
 	const FConfigSection* CookSettings = ConfigFile.FindSection(TEXT_CookSettings);
 	if (CookSettings == nullptr)
@@ -9665,7 +9678,7 @@ void UCookOnTheFlyServer::BeginCookStartShaderCodeLibrary(FBeginCookContext& Beg
 	bool const bCacheShaderLibraries = IsUsingShaderCodeLibrary();
 	if (bCacheShaderLibraries)
 	{
-		FShaderLibraryCooker::InitForCooking(PackagingSettings->bSharedMaterialNativeLibraries);
+		FShaderLibraryCooker::InitForCooking(PackagingSettings->bSharedMaterialNativeLibraries, AllContextArtifactReader.Get());
 
 		bool bAllPlatformsNeedStableKeys = false;
 		// support setting without Hungarian prefix for the compatibility, but allow newer one to override
@@ -11195,7 +11208,7 @@ void UCookOnTheFlyServer::LoadBeginCookIterativeFlags(FBeginCookContext& BeginCo
 	WorkerRequests->GetBeginCookIterativeFlags(*this, BeginContext);
 }
 
-void UCookOnTheFlyServer::LoadBeginCookIterativeFlagsLocal(FBeginCookContext& BeginContext) const
+void UCookOnTheFlyServer::LoadBeginCookIterativeFlagsLocal(FBeginCookContext& BeginContext)
 {
 	const bool bIsDiffOnly = FParse::Param(FCommandLine::Get(), TEXT("DIFFONLY"));
 	bool bForceRecook = FParse::Param(FCommandLine::Get(), TEXT("fullcook")) || FParse::Param(FCommandLine::Get(), TEXT("forcerecook"));
@@ -11415,6 +11428,7 @@ UE::Cook::FCookSavePackageContext* UCookOnTheFlyServer::CreateSaveContext(const
 	const FString ResolvedMetadataPath = MetadataPathSandbox.Replace(TEXT("[Platform]"), *PlatformString);
 
 	TUniquePtr<FDeterminismManager> DeterminismManager;
+	TSharedPtr<ICookArtifactReader> CookArtifactReader = nullptr;
 	ICookedPackageWriter* PackageWriter = nullptr;
 	FString WriterDebugName;
 	ICookedPackageWriter::FBeginCacheCallback BeginCacheCallback(
@@ -11437,7 +11451,13 @@ UE::Cook::FCookSavePackageContext* UCookOnTheFlyServer::CreateSaveContext(const
 	}
 	if (IsUsingZenStore())
 	{
-		FZenStoreWriter* ZenWriter = new FZenStoreWriter(ResolvedRootPath, ResolvedMetadataPath, TargetPlatform);
+		TSharedRef<FLayeredCookArtifactReader> LayeredReader = MakeShared<FLayeredCookArtifactReader>();
+		TSharedRef<FZenCookArtifactReader> ZenReader = MakeShared<FZenCookArtifactReader>(ResolvedRootPath, ResolvedMetadataPath, TargetPlatform);
+		LayeredReader->AddLayer(SharedLooseFilesCookArtifactReader.ToSharedRef());
+		LayeredReader->AddLayer(ZenReader);
+		AllContextArtifactReader->AddLayer(ZenReader);
+		CookArtifactReader = LayeredReader;
+		FZenStoreWriter* ZenWriter = new FZenStoreWriter(ResolvedRootPath, ResolvedMetadataPath, TargetPlatform, ZenReader);
 		ZenWriter->SetBeginCacheCallback(MoveTemp(BeginCacheCallback));
 		ZenWriter->SetRegisterDeterminismHelperCallback(MoveTemp(RegisterDeterminismHelperCallback));
 		PackageWriter = ZenWriter;
@@ -11445,9 +11465,11 @@ UE::Cook::FCookSavePackageContext* UCookOnTheFlyServer::CreateSaveContext(const
 	}
 	else
 	{
+		CookArtifactReader = SharedLooseFilesCookArtifactReader;
 		PackageWriter = new FLooseCookedPackageWriter(ResolvedRootPath, ResolvedMetadataPath, TargetPlatform,
 			GetAsyncIODelete(), *SandboxFile, MoveTemp(BeginCacheCallback),
-			MoveTemp(RegisterDeterminismHelperCallback));
+			MoveTemp(RegisterDeterminismHelperCallback),
+			SharedLooseFilesCookArtifactReader.ToSharedRef());
 		WriterDebugName = TEXT("LooseCookedPackageWriter");
 	}
 
@@ -11505,7 +11527,7 @@ UE::Cook::FCookSavePackageContext* UCookOnTheFlyServer::CreateSaveContext(const
 			});
 	}
 
-	FCookSavePackageContext* Context = new FCookSavePackageContext(TargetPlatform, PackageWriter, WriterDebugName, MoveTemp(SavePackageSettings),
+	FCookSavePackageContext* Context = new FCookSavePackageContext(TargetPlatform, CookArtifactReader, PackageWriter, WriterDebugName, MoveTemp(SavePackageSettings),
 		MoveTemp(DeterminismManager));
 	return Context;
 
@@ -11536,6 +11558,11 @@ void UCookOnTheFlyServer::FinalizePackageStore()
 
 void UCookOnTheFlyServer::ClearPackageStoreContexts()
 {
+	if (AllContextArtifactReader)
+	{
+		AllContextArtifactReader->EmptyLayers();
+		AllContextArtifactReader->AddLayer(SharedLooseFilesCookArtifactReader.ToSharedRef());
+	}
 	for (UE::Cook::FCookSavePackageContext* Context : SavePackageContexts)
 	{
 		delete Context;
@@ -12943,6 +12970,17 @@ bool UCookOnTheFlyServer::GetAllPackageFilenamesFromAssetRegistry(const FString&
 	return false;
 }
 
+ICookArtifactReader& UCookOnTheFlyServer::FindOrCreateCookArtifactReader(const ITargetPlatform* TargetPlatform)
+{
+	return *FindOrCreateSaveContext(TargetPlatform).ArtifactReader;
+}
+
+const ICookArtifactReader* UCookOnTheFlyServer::FindCookArtifactReader(const ITargetPlatform* TargetPlatform) const
+{
+	const UE::Cook::FCookSavePackageContext* Context = FindSaveContext(TargetPlatform);
+	return Context ? Context->ArtifactReader.Get() : nullptr;
+}
+
 ICookedPackageWriter& UCookOnTheFlyServer::FindOrCreatePackageWriter(const ITargetPlatform* TargetPlatform)
 {
 	return *FindOrCreateSaveContext(TargetPlatform).PackageWriter;
diff --git a/Engine/Source/Editor/UnrealEd/Private/Cooker/CookTypes.cpp b/Engine/Source/Editor/UnrealEd/Private/Cooker/CookTypes.cpp
--- a/Engine/Source/Editor/UnrealEd/Private/Cooker/CookTypes.cpp
+++ b/Engine/Source/Editor/UnrealEd/Private/Cooker/CookTypes.cpp
@@ -3,14 +3,15 @@
 #include "CookTypes.h"
 
 #include "Containers/StringView.h"
+#include "CookArtifactReader.h"
 #include "Cooker/CompactBinaryTCP.h"
 #include "Cooker/CookDeterminismManager.h"
 #include "Cooker/CookPackageData.h"
 #include "Cooker/PackageTracker.h"
 #include "DerivedDataRequest.h"
 #include "Editor.h"
-#include "HAL/PlatformTLS.h"
 #include "HAL/PlatformTime.h"
+#include "HAL/PlatformTLS.h"
 #include "Interfaces/ITargetPlatform.h"
 #include "Interfaces/ITargetPlatformManagerModule.h"
 #include "Math/NumericLimits.h"
@@ -221,11 +222,12 @@ void SetIsSchedulerThread(bool bValue)
 	FPlatformTLS::SetTlsValue(SchedulerThreadTlsSlot, bValue ? (void*)0x1 : (void*)0x0);
 }
 
-FCookSavePackageContext::FCookSavePackageContext(const ITargetPlatform* InTargetPlatform,
+FCookSavePackageContext::FCookSavePackageContext(const ITargetPlatform* InTargetPlatform, TSharedPtr<ICookArtifactReader> InCookArtifactReader,
 	ICookedPackageWriter* InPackageWriter, FStringView InWriterDebugName, FSavePackageSettings InSettings,
 	TUniquePtr<FDeterminismManager>&& InDeterminismManager)
 	: SaveContext(InTargetPlatform, InPackageWriter, MoveTemp(InSettings))
 	, WriterDebugName(InWriterDebugName)
+	, ArtifactReader(InCookArtifactReader)
 	, PackageWriter(InPackageWriter)
 	, DeterminismManager(MoveTemp(InDeterminismManager))
 {
diff --git a/Engine/Source/Editor/UnrealEd/Private/Cooker/CookTypes.h b/Engine/Source/Editor/UnrealEd/Private/Cooker/CookTypes.h
--- a/Engine/Source/Editor/UnrealEd/Private/Cooker/CookTypes.h
+++ b/Engine/Source/Editor/UnrealEd/Private/Cooker/CookTypes.h
@@ -17,12 +17,15 @@
 #include "Misc/AssertionMacros.h"
 #include "Serialization/PackageWriter.h"
 #include "Templates/Function.h"
+#include "Templates/SharedPointer.h"
+#include "Templates/UniquePtr.h"
 #include "UObject/CookEnums.h"
 #include "UObject/NameTypes.h"
 #include "UObject/SavePackage.h"
 
 class FCbFieldView;
 class FCbWriter;
+class ICookArtifactReader;
 class ITargetPlatform;
 struct FWeakObjectPtr;
 namespace UE::DerivedData { class FBuildDefinition; }
@@ -312,13 +315,14 @@ namespace UE::Cook
 	*/
 	struct FCookSavePackageContext
 	{
-		FCookSavePackageContext(const ITargetPlatform* InTargetPlatform,
+		FCookSavePackageContext(const ITargetPlatform* InTargetPlatform, TSharedPtr<ICookArtifactReader> InCookArtifactReader,
 			ICookedPackageWriter* InPackageWriter, FStringView InWriterDebugName, FSavePackageSettings InSettings,
 			TUniquePtr<FDeterminismManager>&& InDeterminismManager);
 		~FCookSavePackageContext();
 
 		FSavePackageContext SaveContext;
 		FString WriterDebugName;
+		TSharedPtr<ICookArtifactReader> ArtifactReader;
 		ICookedPackageWriter* PackageWriter;
 		ICookedPackageWriter::FCookCapabilities PackageWriterCapabilities;
 		TUniquePtr<FDeterminismManager> DeterminismManager;
@@ -616,7 +620,7 @@ struct FBeginCookContext
 	TArray<FBeginCookContextPlatform> PlatformContexts;
 	/** The list of platforms by themselves, for passing to functions that need just a list of platforms */
 	TArray<ITargetPlatform*> TargetPlatforms;
-	const UCookOnTheFlyServer& COTFS;
+	UCookOnTheFlyServer& COTFS;
 };
 
 /** Helper struct for FBeginCookContextForWorker; holds the context data for each platform being cooked */
diff --git a/Engine/Source/Editor/UnrealEd/Private/Cooker/LooseCookedPackageWriter.cpp b/Engine/Source/Editor/UnrealEd/Private/Cooker/LooseCookedPackageWriter.cpp
--- a/Engine/Source/Editor/UnrealEd/Private/Cooker/LooseCookedPackageWriter.cpp
+++ b/Engine/Source/Editor/UnrealEd/Private/Cooker/LooseCookedPackageWriter.cpp
@@ -21,10 +21,11 @@
 #include "HAL/PlatformMisc.h"
 #include "HAL/PlatformString.h"
 #include "Hash/Blake3.h"
-#include "IO/IoHash.h"
 #include "Interfaces/ITargetPlatform.h"
+#include "IO/IoHash.h"
 #include "Logging/LogCategory.h"
 #include "Logging/LogMacros.h"
+#include "LooseFilesCookArtifactReader.h"
 #include "Misc/App.h"
 #include "Misc/AssertionMacros.h"
 #include "Misc/CString.h"
@@ -33,8 +34,8 @@
 #include "Misc/Optional.h"
 #include "Misc/PackageName.h"
 #include "Misc/PackagePath.h"
-#include "Misc/PathViews.h"
 #include "Misc/Paths.h"
+#include "Misc/PathViews.h"
 #include "Misc/ScopeLock.h"
 #include "Misc/SecureHash.h"
 #include "Misc/StringBuilder.h"
@@ -60,8 +61,10 @@ LLM_DEFINE_TAG(Cooker_PackageStoreManifest);
 FLooseCookedPackageWriter::FLooseCookedPackageWriter(const FString& InOutputPath,
 	const FString& InMetadataDirectoryPath, const ITargetPlatform* InTargetPlatform, FAsyncIODelete& InAsyncIODelete,
 	UE::Cook::FCookSandbox& InSandboxFile, FBeginCacheCallback&& InBeginCacheCallback,
-	FRegisterDeterminismHelperCallback&& InRegisterDeterminismHelperCallback)
-	: OutputPath(InOutputPath)
+	FRegisterDeterminismHelperCallback&& InRegisterDeterminismHelperCallback,
+	TSharedRef<FLooseFilesCookArtifactReader> InCookArtifactReader)
+	: CookArtifactReader(InCookArtifactReader)
+	, OutputPath(InOutputPath)
 	, MetadataDirectoryPath(InMetadataDirectoryPath)
 	, TargetPlatform(*InTargetPlatform)
 	, SandboxFile(InSandboxFile)
@@ -729,15 +732,15 @@ TUniquePtr<FAssetRegistryState> FLooseCookedPackageWriter::LoadPreviousAssetRegi
 
 	PackageNameToCookedFiles.Reset();
 
-	FArrayReader SerializedAssetData;
-	if (!IFileManager::Get().FileExists(*PreviousAssetRegistryFile) || !FFileHelper::LoadFileToArray(SerializedAssetData, *PreviousAssetRegistryFile))
+	TUniquePtr<FArchive> Reader(CookArtifactReader->CreateFileReader(*PreviousAssetRegistryFile));
+	if (!Reader)
 	{ 
 		RemoveCookedPackages();
 		return TUniquePtr<FAssetRegistryState>();
 	}
 
 	TUniquePtr<FAssetRegistryState> PreviousState = MakeUnique<FAssetRegistryState>();
-	PreviousState->Load(SerializedAssetData);
+	PreviousState->Load(*Reader);
 
 	// If we are iterating from a shared build the cooked files do not exist in the local cooked directory;
 	// we assume they are packaged in the pak file (which we don't want to extract to confirm) and keep them all.
diff --git a/Engine/Source/Editor/UnrealEd/Private/Cooker/LooseCookedPackageWriter.h b/Engine/Source/Editor/UnrealEd/Private/Cooker/LooseCookedPackageWriter.h
--- a/Engine/Source/Editor/UnrealEd/Private/Cooker/LooseCookedPackageWriter.h
+++ b/Engine/Source/Editor/UnrealEd/Private/Cooker/LooseCookedPackageWriter.h
@@ -22,9 +22,10 @@
 #include "Templates/UniquePtr.h"
 #include "UObject/NameTypes.h"
 
-class FAsyncIODelete;
 class FAssetRegistryState;
+class FAsyncIODelete;
 class FLargeMemoryWriter;
+class FLooseFilesCookArtifactReader;
 class FMD5;
 class ITargetPlatform;
 template <typename ReferencedType> class TRefCountPtr;
@@ -40,7 +41,7 @@ public:
 	FLooseCookedPackageWriter(const FString& OutputPath, const FString& MetadataDirectoryPath,
 		const ITargetPlatform* TargetPlatform, FAsyncIODelete& InAsyncIODelete,
 		UE::Cook::FCookSandbox& InSandboxFile, FBeginCacheCallback&& InBeginCacheCallback,
-		FRegisterDeterminismHelperCallback&& InRegisterDeterminismHelperCallback);
+		FRegisterDeterminismHelperCallback&& InRegisterDeterminismHelperCallback, TSharedRef<FLooseFilesCookArtifactReader> CookArtifactReader);
 	~FLooseCookedPackageWriter();
 
 	virtual FCapabilities GetCapabilities() const override
@@ -162,6 +163,8 @@ private:
 		return AllPackageHashes;
 	}
 
+	TSharedRef<FLooseFilesCookArtifactReader> CookArtifactReader;
+
 	// If EWriteOptions::ComputeHash is not set, the package will not get added to this.
 	TMap<FName, TRefCountPtr<FPackageHashes>> AllPackageHashes;
 
diff --git a/Engine/Source/Runtime/RenderCore/Private/ShaderCodeArchive.cpp b/Engine/Source/Runtime/RenderCore/Private/ShaderCodeArchive.cpp
--- a/Engine/Source/Runtime/RenderCore/Private/ShaderCodeArchive.cpp
+++ b/Engine/Source/Runtime/RenderCore/Private/ShaderCodeArchive.cpp
@@ -623,14 +623,19 @@ void FSerializedShaderArchive::SaveAssetInfo(FArchive& Ar)
 
 bool FSerializedShaderArchive::LoadAssetInfo(const FString& Filename)
 {
-	TArray<uint8> FileData;
-	if (!FFileHelper::LoadFileToArray(FileData, *Filename))
+	TUniquePtr<FArchive> Reader(IFileManager::Get().CreateFileReader(*Filename));
+	return LoadAssetInfo(Reader.Get());
+}
+
+bool FSerializedShaderArchive::LoadAssetInfo(FArchive* Ar)
+{
+	if (!Ar)
 	{
 		return false;
 	}
 
 	FString JsonText;
-	FFileHelper::BufferToString(JsonText, FileData.GetData(), FileData.Num());
+	FFileHelper::LoadFileToString(JsonText, *Ar);
 
 	TSharedPtr<FJsonObject> JsonObject;
 	TSharedRef<TJsonReader<TCHAR>> Reader = TJsonReaderFactory<TCHAR>::Create(JsonText);
@@ -645,7 +650,7 @@ bool FSerializedShaderArchive::LoadAssetInfo(const FString& Filename)
 	if (!AssetInfoVersion.IsValid())
 	{
 		UE_LOG(LogShaderLibrary, Warning, TEXT("Rejecting asset info file %s: missing AssetInfoVersion (damaged file?)"), 
-			*Filename);
+			*Ar->GetArchiveName());
 		return false;
 	}
 	
@@ -653,7 +658,7 @@ bool FSerializedShaderArchive::LoadAssetInfo(const FString& Filename)
 	if (FileVersion != EAssetInfoVersion::CurrentVersion)
 	{
 		UE_LOG(LogShaderLibrary, Warning, TEXT("Rejecting asset info file %s: expected version %d, got unsupported version %d."),
-			*Filename, static_cast<int32>(EAssetInfoVersion::CurrentVersion), static_cast<int32>(FileVersion));
+			*Ar->GetArchiveName(), static_cast<int32>(EAssetInfoVersion::CurrentVersion), static_cast<int32>(FileVersion));
 		return false;
 	}
 
@@ -661,13 +666,13 @@ bool FSerializedShaderArchive::LoadAssetInfo(const FString& Filename)
 	if (!AssetInfoArrayValue.IsValid())
 	{
 		UE_LOG(LogShaderLibrary, Warning, TEXT("Rejecting asset info file %s: missing ShaderCodeToAssets array (damaged file?)"),
-			*Filename);
+			*Ar->GetArchiveName());
 		return false;
 	}
 	
 	TArray<TSharedPtr<FJsonValue>> AssetInfoArray = AssetInfoArrayValue->AsArray();
 	UE_LOG(LogShaderLibrary, Display, TEXT("Reading asset info file %s: found %d existing mappings"),
-		*Filename, AssetInfoArray.Num());
+		*Ar->GetArchiveName(), AssetInfoArray.Num());
 
 	for (int32 IdxPair = 0, NumPairs = AssetInfoArray.Num(); IdxPair < NumPairs; ++IdxPair)
 	{
@@ -675,7 +680,7 @@ bool FSerializedShaderArchive::LoadAssetInfo(const FString& Filename)
 		if (UNLIKELY(!Pair.IsValid()))
 		{
 			UE_LOG(LogShaderLibrary, Warning, TEXT("Rejecting asset info file %s: ShaderCodeToAssets array contains unreadable mapping #%d (damaged file?)"),
-				*Filename,
+				*Ar->GetArchiveName(),
 				IdxPair
 				);
 			return false;
@@ -685,7 +690,7 @@ bool FSerializedShaderArchive::LoadAssetInfo(const FString& Filename)
 		if (UNLIKELY(!ShaderMapHashJson.IsValid()))
 		{
 			UE_LOG(LogShaderLibrary, Warning, TEXT("Rejecting asset info file %s: ShaderCodeToAssets array contains unreadable ShaderMapHash for mapping %d (damaged file?)"),
-				*Filename,
+				*Ar->GetArchiveName(),
 				IdxPair
 				);
 			return false;
@@ -698,7 +703,7 @@ bool FSerializedShaderArchive::LoadAssetInfo(const FString& Filename)
 		if (UNLIKELY(!AssetPathsArrayValue.IsValid()))
 		{
 			UE_LOG(LogShaderLibrary, Warning, TEXT("Rejecting asset info file %s: ShaderCodeToAssets array contains unreadable Assets array for mapping %d (damaged file?)"),
-				*Filename,
+				*Ar->GetArchiveName(),
 				IdxPair
 			);
 			return false;
diff --git a/Engine/Source/Runtime/RenderCore/Private/ShaderCodeLibrary.cpp b/Engine/Source/Runtime/RenderCore/Private/ShaderCodeLibrary.cpp
--- a/Engine/Source/Runtime/RenderCore/Private/ShaderCodeLibrary.cpp
+++ b/Engine/Source/Runtime/RenderCore/Private/ShaderCodeLibrary.cpp
@@ -55,6 +55,7 @@ ShaderCodeLibrary.cpp: Bound shader state cache implementation.
 #include "Serialization/CompactBinarySerialization.h"
 #include "Serialization/CompactBinaryWriter.h"
 #include "RHIStrings.h"
+#include "CookArtifactReader.h"
 #endif
 
 // allow introspection (e.g. dumping the contents) for easier debugging
@@ -258,6 +259,10 @@ namespace UE
 	}
 }
 
+#if WITH_EDITOR
+static ICookArtifactReader* CookArtifactReader = nullptr;
+#endif // WITH_EDITOR
+
 TSet<UE::ShaderLibrary::Private::FMountedPakFileInfo> UE::ShaderLibrary::Private::FMountedPakFileInfo::KnownPakFiles;
 FCriticalSection UE::ShaderLibrary::Private::FMountedPakFileInfo::KnownPakFilesAccessLock;
 
@@ -290,6 +295,39 @@ public:
 	bool bShaderMapPreloaded;
 };
 
+static FArchive* CreateShaderFileReader(const TCHAR* Filename)
+{
+#if WITH_EDITOR
+	if (CookArtifactReader)
+	{
+		return CookArtifactReader->CreateFileReader(Filename);
+	}
+#endif
+	return IFileManager::Get().CreateFileReader(Filename);
+}
+
+static void ShaderFindFiles(TArray<FString>& FoundFiles, const TCHAR* Directory, const TCHAR* FileExtension)
+{
+#if WITH_EDITOR
+	if (CookArtifactReader)
+	{
+		return CookArtifactReader->FindFiles(FoundFiles, Directory, FileExtension);
+	}
+#endif
+	return IFileManager::Get().FindFiles(FoundFiles, Directory, FileExtension);
+}
+
+static void ShaderFindFiles(TArray<FString>& Result, const TCHAR* Filename, bool Files, bool Directories)
+{
+#if WITH_EDITOR
+	if (CookArtifactReader)
+	{
+		return CookArtifactReader->FindFiles(Result, Filename, Files, Directories);
+	}
+#endif
+	return IFileManager::Get().FindFiles(Result, Filename, Files, Directories);
+}
+
 static FString GetCodeArchiveFilename(const FString& BaseDir, const FString& LibraryName, FName Platform)
 {
 	return BaseDir / FString::Printf(TEXT("ShaderArchive-%s-"), *LibraryName) + Platform.ToString() + ShaderExtension;
@@ -902,7 +940,7 @@ public:
 			FString ShaderFormatAndPlatform = ShaderFormatName.ToString() + TEXT("-") + PlatformName.ToString();
 
 			const FString DestFilePath = GetCodeArchiveFilename(ShaderCodeDir, InLibraryName, FName(ShaderFormatAndPlatform));
-			TUniquePtr<FArchive> Ar(IFileManager::Get().CreateFileReader(*DestFilePath));
+			TUniquePtr<FArchive> Ar(CreateShaderFileReader(*DestFilePath));
 			if (Ar)
 			{
 				uint32 Version = 0;
@@ -1763,7 +1801,7 @@ struct FEditorShaderCodeArchive
 	bool LoadExistingShaderCodeLibrary(FString const& MetaDataDir)
 	{
 		FString IntermediateFormatPath = GetCodeArchiveFilename(MetaDataDir / TEXT("ShaderLibrarySource"), LibraryName, FormatName);
-		FArchive* PrevCookedAr = IFileManager::Get().CreateFileReader(*IntermediateFormatPath);
+		FArchive* PrevCookedAr = CreateShaderFileReader(*IntermediateFormatPath);
 		bool bOK = true;
 		if (PrevCookedAr)
 		{
@@ -1811,7 +1849,7 @@ struct FEditorShaderCodeArchive
 
 	void AddShaderCodeLibraryByName(const FString& BaseDir, const FString& InLibraryName)
 	{
-		if (FArchive* PrevCookedAr = IFileManager::Get().CreateFileReader(*GetCodeArchiveFilename(BaseDir, InLibraryName, FormatName)))
+		if (FArchive* PrevCookedAr = CreateShaderFileReader(*GetCodeArchiveFilename(BaseDir, InLibraryName, FormatName)))
 		{
 			uint32 Version = 0;
 			*PrevCookedAr << Version;
@@ -1823,7 +1861,8 @@ struct FEditorShaderCodeArchive
 				*PrevCookedAr << PrevCookedShaders;
 
 				// check if it also contains the asset info file
-				if (PrevCookedShaders.LoadAssetInfo(GetShaderAssetInfoFilename(BaseDir, InLibraryName, FormatName)))
+				TUniquePtr<FArchive> ShaderAssetInfoReader(CreateShaderFileReader(*GetShaderAssetInfoFilename(BaseDir, InLibraryName, FormatName)));
+				if (PrevCookedShaders.LoadAssetInfo(ShaderAssetInfoReader.Get()))
 				{
 					UE_LOG(LogShaderLibrary, Display, TEXT("Loaded asset info %s for the shader library %s: %d entries"),
 						*GetShaderAssetInfoFilename(BaseDir, InLibraryName, FormatName),
@@ -1865,7 +1904,7 @@ struct FEditorShaderCodeArchive
 		const FString ShaderIntermediateLocation = FPaths::ProjectSavedDir() / TEXT("Shaders") / FormatNameStr;
 
 		TArray<FString> ShaderFiles;
-		IFileManager::Get().FindFiles(ShaderFiles, *ShaderIntermediateLocation, *ShaderExtension);
+		ShaderFindFiles(ShaderFiles, *ShaderIntermediateLocation, *ShaderExtension);
 
 		for (const FString& ShaderFileName : ShaderFiles)
 		{
@@ -2304,7 +2343,7 @@ struct FEditorShaderStableInfo
 		{
 			const FString ShaderIntermediateLocation = FPaths::ProjectSavedDir() / TEXT("Shaders") / FormatName.ToString();
 			TArray<FString> ShaderFiles;
-			IFileManager::Get().FindFiles(ShaderFiles, *ShaderIntermediateLocation, *StableExtension);
+			ShaderFindFiles(ShaderFiles, *ShaderIntermediateLocation, *StableExtension);
 			FString ExpectedFileNameText = LibraryName + TEXT("-") + FormatName.ToString() + TEXT(".");
 
 			for (const FString& ShaderFileName : ShaderFiles)
@@ -2687,7 +2726,7 @@ public:
 #else
 						FString SearchMask = Directory / FString::Printf(TEXT("ShaderArchive-*%s*.ushaderbytecode"), *Name);
 #endif
-						IFileManager::Get().FindFiles(UshaderbytecodeFiles, *SearchMask, true, false);
+						ShaderFindFiles(UshaderbytecodeFiles, *SearchMask, true, false);
 
 						if (UshaderbytecodeFiles.Num() > 0)
 						{
@@ -3810,8 +3849,9 @@ void FShaderCodeLibrary::CloseLibrary(FString const& Name)
 
 #if WITH_EDITOR
 // for now a lot of FShaderLibraryCooker code is aliased with the runtime code, but this will be refactored (UE-103486)
-void FShaderLibraryCooker::InitForCooking(bool bNativeFormat)
+void FShaderLibraryCooker::InitForCooking(bool bNativeFormat, ICookArtifactReader* InCookArtifactReader)
 {
+	CookArtifactReader = InCookArtifactReader;
 	FShaderLibrariesCollection::Impl = new FShaderLibrariesCollection(SP_NumPlatforms, bNativeFormat);
 }
 
@@ -3927,7 +3967,7 @@ bool FShaderLibraryCooker::CreatePatchLibrary(TArray<FString> const& OldMetaData
 {
 	TMap<FName, TSet<FString>> FormatLibraryMap;
 	TArray<FString> LibraryFiles;
-	IFileManager::Get().FindFiles(LibraryFiles, *(NewMetaDataDir / TEXT("ShaderLibrarySource")), *ShaderExtension);
+	ShaderFindFiles(LibraryFiles, *(NewMetaDataDir / TEXT("ShaderLibrarySource")), *ShaderExtension);
 	
 	for (FString const& Path : LibraryFiles)
 	{
@@ -3996,7 +4036,7 @@ bool FShaderLibraryCooker::MergeShaderCodeArchive(const TArray<FString>& CookedM
 		const FString ShaderStableInfoDir = MetadataDir / TEXT("PipelineCaches");
 
 		TArray<FString> ShaderBytecodeFiles;
-		IFileManager::Get().FindFiles(ShaderBytecodeFiles, *ShaderCodeDir, *ShaderExtension);
+		ShaderFindFiles(ShaderBytecodeFiles, *ShaderCodeDir, *ShaderExtension);
 		for (const FString& ByteCodeFile : ShaderBytecodeFiles)
 		{
 			if (ShaderCodeArchives.Contains(ByteCodeFile))
@@ -4019,7 +4059,7 @@ bool FShaderLibraryCooker::MergeShaderCodeArchive(const TArray<FString>& CookedM
 		}
 
 		TArray<FString> StableInfoFiles;
-		IFileManager::Get().FindFiles(StableInfoFiles, *ShaderStableInfoDir, *StableExtension);
+		ShaderFindFiles(StableInfoFiles, *ShaderStableInfoDir, *StableExtension);
 		for (const FString& StableInfoFile : StableInfoFiles)
 		{
 			if (ShaderStableInfos.Contains(StableInfoFile))
diff --git a/Engine/Source/Runtime/RenderCore/Public/ShaderCodeArchive.h b/Engine/Source/Runtime/RenderCore/Public/ShaderCodeArchive.h
--- a/Engine/Source/Runtime/RenderCore/Public/ShaderCodeArchive.h
+++ b/Engine/Source/Runtime/RenderCore/Public/ShaderCodeArchive.h
@@ -262,6 +262,7 @@ public:
 #if WITH_EDITOR
 	RENDERCORE_API void SaveAssetInfo(FArchive& Ar);
 	RENDERCORE_API bool LoadAssetInfo(const FString& Filename);
+	RENDERCORE_API bool LoadAssetInfo(FArchive* Ar);
 	RENDERCORE_API void CreateAsChunkFrom(const FSerializedShaderArchive& Parent, const TSet<FName>& PackagesInChunk, TArray<int32>& OutShaderCodeEntriesNeeded);
 	RENDERCORE_API void CollectStatsAndDebugInfo(FDebugStats& OutDebugStats, FExtendedDebugStats* OutExtendedDebugStats);
 	RENDERCORE_API void DumpContentsInPlaintext(FString& OutText) const;
diff --git a/Engine/Source/Runtime/RenderCore/Public/ShaderCodeLibrary.h b/Engine/Source/Runtime/RenderCore/Public/ShaderCodeLibrary.h
--- a/Engine/Source/Runtime/RenderCore/Public/ShaderCodeLibrary.h
+++ b/Engine/Source/Runtime/RenderCore/Public/ShaderCodeLibrary.h
@@ -41,6 +41,7 @@ class UObject;
 #if WITH_EDITOR
 class FCbFieldView;
 class FCbWriter;
+class ICookArtifactReader;
 #endif
 
 struct FShaderCodeLibraryPipeline
@@ -266,7 +267,7 @@ class ITargetPlatform;
 struct FShaderLibraryCooker
 {
 	// Initialize the library cooker
-	static RENDERCORE_API void InitForCooking(bool bNativeFormat);
+	static RENDERCORE_API void InitForCooking(bool bNativeFormat, ICookArtifactReader* CookArtifactReader);
 	// Shutdown the library cooker
 	static RENDERCORE_API void Shutdown();
 
diff --git a/Engine/Source/Runtime/RenderCore/RenderCore.Build.cs b/Engine/Source/Runtime/RenderCore/RenderCore.Build.cs
--- a/Engine/Source/Runtime/RenderCore/RenderCore.Build.cs
+++ b/Engine/Source/Runtime/RenderCore/RenderCore.Build.cs
@@ -23,6 +23,7 @@ public class RenderCore : ModuleRules
 			DynamicallyLoadedModuleNames.Add("TargetPlatform");
 			// UObjects are used to produce the full path of the asset by which the shaders are identified
 			PrivateDependencyModuleNames.Add("CoreUObject");
+			PrivateIncludePathModuleNames.Add("IoStoreUtilities");
 		}
 		// shader runtime usage visualization requires ability to create images - it is only used in non-editor desktop development targets
 		// UE_BUILD_DEVELOPMENT is also defined for DebugGame
