Build Engine Localization - LOCTEXT macro has an empty identifier and cannot be gathered for UnrealEdViewportToolbar.cpp

diff --git a/Engine/Source/Editor/UnrealEd/Private/ViewportToolbar/UnrealEdViewportToolbar.cpp b/Engine/Source/Editor/UnrealEd/Private/ViewportToolbar/UnrealEdViewportToolbar.cpp
--- a/Engine/Source/Editor/UnrealEd/Private/ViewportToolbar/UnrealEdViewportToolbar.cpp
+++ b/Engine/Source/Editor/UnrealEd/Private/ViewportToolbar/UnrealEdViewportToolbar.cpp
@@ -2155,7 +2155,7 @@ FToolMenuEntry CreateViewportToolbarAssetViewerProfileSubmenu(const TSharedPtr<I
 					const FString& ProfileName = PreviewProfiles[ProfileIndex];
 					PreviewProfilesSelectionMenuBuilder.AddMenuEntry(
 						FText::FromString(ProfileName),
-						LOCTEXT("", ""),
+						FText(),
 						FSlateIcon(),
 						FUIAction(
 							FExecuteAction::CreateLambda(
