Clarified Refraction pin label.

#rnx

diff --git a/Engine/Source/Editor/UnrealEd/Private/MaterialGraphNode_Root.cpp b/Engine/Source/Editor/UnrealEd/Private/MaterialGraphNode_Root.cpp
--- a/Engine/Source/Editor/UnrealEd/Private/MaterialGraphNode_Root.cpp
+++ b/Engine/Source/Editor/UnrealEd/Private/MaterialGraphNode_Root.cpp
@@ -243,7 +243,7 @@ void UMaterialGraphNode_Root::CreateInputPins()
 					TempStr = TEXT("Disabled");
 					break;
 				case ERefractionMode::RM_IndexOfRefraction:
-					TempStr = TEXT("Index Of Refraction");
+					TempStr = TEXT("IOR, from Material F0 if Unplugged");
 					break;
 				case ERefractionMode::RM_PixelNormalOffset:
 					TempStr = TEXT("Pixel Normal Offset");
