CookGarbageCollect: Fix compile error in shipping editor due to ENABLE_GC_HISTORY 0.
#rnx

diff --git a/Engine/Source/Editor/UnrealEd/Private/Cooker/CookGarbageCollect.cpp b/Engine/Source/Editor/UnrealEd/Private/Cooker/CookGarbageCollect.cpp
--- a/Engine/Source/Editor/UnrealEd/Private/Cooker/CookGarbageCollect.cpp
+++ b/Engine/Source/Editor/UnrealEd/Private/Cooker/CookGarbageCollect.cpp
@@ -24,12 +24,16 @@ bool FCookGCDiagnosticContext::CurrentGCHasHistory() const
 
 bool FCookGCDiagnosticContext::TryRequestGCWithHistory()
 {
+#if ENABLE_GC_HISTORY
 	if (!bRequestsAvailable || !bGCInProgress || bCurrentGCHasHistory)
 	{
 		return false;
 	}
 	SetGCWithHistoryRequested(true);
 	return true;
+#else
+	return false;
+#endif
 }
 
 void FCookGCDiagnosticContext::OnCookerStartCollectGarbage()
@@ -37,7 +41,11 @@ void FCookGCDiagnosticContext::OnCookerStartCollectGarbage()
 	bRequestsAvailable = true;
 
 	bGCInProgress = true;
+#if ENABLE_GC_HISTORY
 	bCurrentGCHasHistory = FGCHistory::Get().GetHistorySize() > 0;
+#else
+	bCurrentGCHasHistory = false;
+#endif
 }
 
 void FCookGCDiagnosticContext::OnCookerEndCollectGarbage()
@@ -53,6 +61,7 @@ void FCookGCDiagnosticContext::OnEvaluateResultsComplete()
 
 void FCookGCDiagnosticContext::SetGCWithHistoryRequested(bool bValue)
 {
+#if ENABLE_GC_HISTORY
 	if (bValue == bRequestGCWithHistory)
 	{
 		return;
@@ -75,6 +84,7 @@ void FCookGCDiagnosticContext::SetGCWithHistoryRequested(bool bValue)
 		SavedGCHistorySize = 0;
 	}
 	bRequestGCWithHistory = bValue;
+#endif
 }
 
 } // namespace UE::Cook
\ No newline at end of file
