Localization:
Optimized include and exclude path filtering for the GatherTextFromAssetsCommandlet to reduce gather times of projects with heavy localization plugin use.
- Filtering of include exclude paths was a slow process that used FString::MatchesWildcard for all of the comparisons. We do 2 optimizations.
- Running all the assets through the include exclude paths is now parallelized with a ParallelFor() loop controlled by the cvar bool Localization.GatherTextFromAssetsCommandlet.ParallelizeIncludeExcludePathFiltering which is set to true by default
- We do some pre-processing on fuzzy paths and package file filters to reduce the number of FString::MatchesWildcard function calls and instead rely on FString::EndsWith() and FString::StartsWith() calls for filters that only contain a single * wildcard at the start or end of the fuzzy path.

diff --git a/Engine/Source/Editor/UnrealEd/Classes/Commandlets/GatherTextCommandletBase.h b/Engine/Source/Editor/UnrealEd/Classes/Commandlets/GatherTextCommandletBase.h
--- a/Engine/Source/Editor/UnrealEd/Classes/Commandlets/GatherTextCommandletBase.h
+++ b/Engine/Source/Editor/UnrealEd/Classes/Commandlets/GatherTextCommandletBase.h
@@ -40,16 +40,24 @@ private:
 		Exclude,
 	};
 
+	/**
+	* * The algorithm used to test path matches for a fuzzy path. Defaults to FString::MatchesWildcard
+	*/
+	enum class EPathTestPolicy : uint8
+	{
+		/** Performs the path test with FString::MatchesWildcard. This is the default algorithm to use for fuzzy paths that can't be optimized with FString::StartsWith.*/
+		MatchesWildcard,
+		/** Uses FSTring::StartsWith to perform the path test against this fuzzy path. This is an optimization for fuzzy paths that only contain a single wildcard and the * wildcard only exists at the end of the fuzzy path. */
+		StartsWith
+	};
+
 	struct FFuzzyPath
 	{
-		FFuzzyPath(FString InPathFilter, const EPathType InPathType)
-			: PathFilter(MoveTemp(InPathFilter))
-			, PathType(InPathType)
-		{
-		}
+		FFuzzyPath(FString InPathFilter, const EPathType InPathType);
 
 		FString PathFilter;
 		EPathType PathType;
+		EPathTestPolicy PathTestPolicy;
 	};
 
 	TArray<FFuzzyPath> FuzzyPaths;
diff --git a/Engine/Source/Editor/UnrealEd/Private/Commandlets/GatherTextCommandletBase.cpp b/Engine/Source/Editor/UnrealEd/Private/Commandlets/GatherTextCommandletBase.cpp
--- a/Engine/Source/Editor/UnrealEd/Private/Commandlets/GatherTextCommandletBase.cpp
+++ b/Engine/Source/Editor/UnrealEd/Private/Commandlets/GatherTextCommandletBase.cpp
@@ -171,6 +171,22 @@ const FString& UGatherTextCommandletBase::GetProjectBasePath()
 	return ProjectBasePath;
 }
 
+FFuzzyPathMatcher::FFuzzyPath::FFuzzyPath(FString InPathFilter, const EPathType InPathType)
+	: PathFilter(MoveTemp(InPathFilter))
+	, PathType(InPathType)
+	, PathTestPolicy(EPathTestPolicy::MatchesWildcard)
+{
+	const int32 PathLastIndex = PathFilter.Len() - 1;
+	bool bEndsWithWildcard = PathFilter[PathLastIndex] == TEXT('*');
+	// If we know we end with a * but searching from the front, the index of the first * is not the end of the string, there's more than 1 *
+	int32 FirstAsteriskIndex = INDEX_NONE;
+	PathFilter.FindChar(TEXT('*'), FirstAsteriskIndex);
+	bool bOnlyContainsOneWildcard = FirstAsteriskIndex == PathLastIndex;
+	if (bEndsWithWildcard && bOnlyContainsOneWildcard)
+	{
+		PathTestPolicy = EPathTestPolicy::StartsWith;
+	}
+}
 
 FFuzzyPathMatcher::FFuzzyPathMatcher(const TArray<FString>& InIncludePathFilters, const TArray<FString>& InExcludePathFilters)
 {
@@ -217,15 +233,37 @@ FFuzzyPathMatcher::FFuzzyPathMatcher(const TArray<FString>& InIncludePathFilters
 		}
 		return PathOneFuzzRating < PathTwoFuzzRating;
 	});
+
+	// Now we pre-process and alter the path filter for paths that will be compared with EPathTestPolicy::StartsWith
+	// We only do that here because we need the paths that end with the * wildcard to be intact for the above sorting 
+	for (FFuzzyPath& FuzzyPath : FuzzyPaths)
+	{
+		if (FuzzyPath.PathTestPolicy == EPathTestPolicy::StartsWith)
+		{
+			FuzzyPath.PathFilter.LeftChopInline(1);
+		}
+	}
 }
 
 FFuzzyPathMatcher::EPathMatch FFuzzyPathMatcher::TestPath(const FString& InPathToTest) const
 {
 	for (const FFuzzyPath& FuzzyPath : FuzzyPaths)
 	{
-		if (InPathToTest.MatchesWildcard(FuzzyPath.PathFilter))
+		if (FuzzyPath.PathTestPolicy == EPathTestPolicy::StartsWith)
 		{
-			return (FuzzyPath.PathType == EPathType::Include) ? EPathMatch::Included : EPathMatch::Excluded;
+			// The wildcard at the end should already be be removed as part of a preprocessing step.
+			check(FuzzyPath.PathFilter[FuzzyPath.PathFilter.Len() - 1] != TEXT('*'));
+			if (InPathToTest.StartsWith(FuzzyPath.PathFilter))
+			{
+				return (FuzzyPath.PathType == EPathType::Include) ? EPathMatch::Included : EPathMatch::Excluded;
+			}
+		}
+		else if (FuzzyPath.PathTestPolicy == EPathTestPolicy::MatchesWildcard)
+		{
+			if (InPathToTest.MatchesWildcard(FuzzyPath.PathFilter))
+			{
+				return (FuzzyPath.PathType == EPathType::Include) ? EPathMatch::Included : EPathMatch::Excluded;
+			}
 		}
 	}
 
diff --git a/Engine/Source/Editor/UnrealEd/Private/Commandlets/GatherTextFromAssetsCommandlet.cpp b/Engine/Source/Editor/UnrealEd/Private/Commandlets/GatherTextFromAssetsCommandlet.cpp
--- a/Engine/Source/Editor/UnrealEd/Private/Commandlets/GatherTextFromAssetsCommandlet.cpp
+++ b/Engine/Source/Editor/UnrealEd/Private/Commandlets/GatherTextFromAssetsCommandlet.cpp
@@ -40,6 +40,8 @@
 #include "CollectionManagerModule.h"
 #include "ICollectionManager.h"
 #include "WorldPartition/WorldPartitionActorDescInstance.h"
+#include "Async/ParallelFor.h"
+#include "ProfilingDebugging/ScopedTimers.h"
 
 DEFINE_LOG_CATEGORY_STATIC(LogGatherTextFromAssetsCommandlet, Log, All);
 
@@ -248,6 +250,8 @@ namespace UE::Private::GatherTextFromAssetsCommandlet
 		}
 	}
 
+	static bool bParallelizeIncludeExcludePathFiltering = true;
+	static FAutoConsoleVariableRef CVarParallelizeIncludeExcludePathFiltering(TEXT("Localization.GatherTextFromAssetsCommandlet.ParallelizeIncludeExcludePathFiltering"), bParallelizeIncludeExcludePathFiltering, TEXT("True to parallelize the include exclude path filtering. False to force the include exclu de process to be single threaded for easier debugging."));
 }
 
 #define LOC_DEFINE_REGION
@@ -276,6 +280,7 @@ UGatherTextFromAssetsCommandlet::UGatherTextFromAssetsCommandlet(const FObjectIn
 
 void UGatherTextFromAssetsCommandlet::ProcessGatherableTextDataArray(const TArray<FGatherableTextData>& GatherableTextDataArray)
 {
+	
 	for (const FGatherableTextData& GatherableTextData : GatherableTextDataArray)
 	{
 		for (const FTextSourceSiteContext& TextSourceSiteContext : GatherableTextData.SourceSiteContexts)
@@ -338,6 +343,7 @@ void CalculateDependenciesImpl(IAssetRegistry& InAssetRegistry, const FName& InP
 
 void UGatherTextFromAssetsCommandlet::CalculateDependenciesForPackagesPendingGather()
 {
+	UE_SCOPED_TIMER(TEXT("UGatherTextFromAssetsCommandlet::CalculateDependenciesForPackagesPendingGather"), LogGatherTextFromAssetsCommandlet, Display);
 	FAssetRegistryModule& AssetRegistryModule = FModuleManager::LoadModuleChecked<FAssetRegistryModule>(TEXT("AssetRegistry"));
 	IAssetRegistry& AssetRegistry = AssetRegistryModule.Get();
 	TMap<FName, TSet<FName>> PackageNameToDependencies;
@@ -506,6 +512,7 @@ bool UGatherTextFromAssetsCommandlet::ShouldRunInPreview(const TArray<FString>&
 	// If the param is not specified, it is assumed that both source and assets are to be gathered 
 	return !GatherType || *GatherType == TEXT("Asset") || *GatherType == TEXT("All");
 }
+
 /**
  * Builds the first pass filter which currently consists of the collection filter and the optional derived class filter.
  * See BuildCollectionFilter and BuildExcludeDerivedClassesFilter
@@ -592,28 +599,57 @@ bool UGatherTextFromAssetsCommandlet::BuildExcludeExactClassesFilter(FARFilter&
 /** Filters out assets that fail the IncludePath and ExcludePath wildcard filters. */
 void UGatherTextFromAssetsCommandlet::FilterAssetsBasedOnIncludeExcludePaths(TArray<FAssetData>& InOutAssetDataArray) const
 {
+	UE_SCOPED_TIMER(TEXT("UGatherTextFromAssetsCommandlet::FilterAssetsBasedOnIncludeExcludePaths"), LogGatherTextFromAssetsCommandlet, Display);
+	// We pre-process the package filters into 2 sets because comparing wild cards is expensive 
+	// This is the array for cases like *.uasset, *.umap 
+	// We only store the extension without the wildcard for an optimization later 
+	TArray<FString> PackageFileFiltersStartingWithWildcard;
+	// For everything else. We will assume that we will need a wildcard match in this case 
+	TArray<FString> OtherPackageFileFilters;
+
+	for (const FString& PackageFileNameFilter : PackageFileNameFilters)
+	{
+		FString Extension;
+		FString CleanPackageFileName;
+		PackageFileNameFilter.Split(TEXT("."), &CleanPackageFileName, &Extension);
+		if ((CleanPackageFileName.Len() == 1) && (CleanPackageFileName[0] == TEXT('*')) && !Extension.Contains(TEXT("*")))
+		{
+			// We drop the * from say *.uasset and just keep the extension 
+			PackageFileFiltersStartingWithWildcard.Add(PackageFileNameFilter.RightChop(1));
+		}
+		else
+		{
+			OtherPackageFileFilters.Add(PackageFileNameFilter);
+		}
+	}
+
 	const FFuzzyPathMatcher FuzzyPathMatcher = FFuzzyPathMatcher(IncludePathFilters, ExcludePathFilters);
-	const double FilteringAssetsByIncludeExcludePathsStartTime = FPlatformTime::Seconds();
-	InOutAssetDataArray.RemoveAll([&](const FAssetData& PartiallyFilteredAssetData) -> bool
+	TArray<bool> PackagesToFilter;
+	PackagesToFilter.Init( false, InOutAssetDataArray.Num());
+	ParallelFor(InOutAssetDataArray.Num(), [&](int32 Index)
 		{
+			const FAssetData& PartiallyFilteredAssetData = InOutAssetDataArray[Index];
 			if (PartiallyFilteredAssetData.IsRedirector())
 			{
 				// Redirectors never have localization
-				return true;
+				PackagesToFilter[Index] = true;
+				return;
 			}
 
 			FString PackageFilePathWithoutExtension;
 			if (!FPackageName::TryConvertLongPackageNameToFilename(PartiallyFilteredAssetData.PackageName.ToString(), PackageFilePathWithoutExtension))
 			{
 				// This means the asset data is for content that isn't mounted - this can happen when using a cooked asset registry
-				return true;
+				PackagesToFilter[Index] = true;
+				return;
 			}
 
 			FString PackageFilePathWithExtension;
 			if (!FPackageName::FindPackageFileWithoutExtension(PackageFilePathWithoutExtension, PackageFilePathWithExtension))
 			{
 				// This means the package file doesn't exist on disk, which means we cannot gather it
-				return true;
+				PackagesToFilter[Index] = true;
+				return;
 			}
 
 			PackageFilePathWithExtension = FPaths::ConvertRelativePathToFull(PackageFilePathWithExtension);
@@ -621,30 +657,49 @@ void UGatherTextFromAssetsCommandlet::FilterAssetsBasedOnIncludeExcludePaths(TAr
 
 			// Filter out assets whose package file names DO NOT match any of the package file name filters.
 			{
-				bool HasPassedAnyFileNameFilter = false;
-				for (const FString& PackageFileNameFilter : PackageFileNameFilters)
+				bool bHasPassedAnyFileNameFilter = false;
+				// This is an optimization to process package file filters in the form *.uasset or *.umap differently
+				// FString::MatchesWildcard is an expensive call so we try and minimize the call to that and we go with FString::EndsWith instead for better performance
+				for (const FString& PackageFileNameFilter : PackageFileFiltersStartingWithWildcard)
+				{
+					if (PackageFileName.EndsWith(PackageFileNameFilter))
+					{
+						bHasPassedAnyFileNameFilter = true;
+						break;
+					}
+				}
+
+				for (const FString& PackageFileNameFilter : OtherPackageFileFilters)
 				{
 					if (PackageFileName.MatchesWildcard(PackageFileNameFilter))
 					{
-						HasPassedAnyFileNameFilter = true;
+						bHasPassedAnyFileNameFilter = true;
 						break;
 					}
 				}
-				if (!HasPassedAnyFileNameFilter)
+				if (!bHasPassedAnyFileNameFilter)
 				{
-					return true;
+					PackagesToFilter[Index] = true;
+					return;
 				}
 			}
 
 			// Filter out assets whose package file paths do not pass the "fuzzy path" filters.
 			if (FuzzyPathMatcher.TestPath(PackageFilePathWithExtension) != FFuzzyPathMatcher::EPathMatch::Included)
 			{
-				return true;
+				PackagesToFilter[Index] = true;
+				return;
 			}
+		}, UE::Private::GatherTextFromAssetsCommandlet::bParallelizeIncludeExcludePathFiltering ? EParallelForFlags::None: EParallelForFlags::ForceSingleThread);
 
-			return false;
-		});
-	UE_LOG(LogGatherTextFromAssetsCommandlet, Display, TEXT("Filtering assets by include exclude paths took %.2f seconds."), FPlatformTime::Seconds() - FilteringAssetsByIncludeExcludePathsStartTime);
+	check(PackagesToFilter.Num() == InOutAssetDataArray.Num());
+	for (int32 Index = InOutAssetDataArray.Num() - 1; Index >= 0; --Index)
+	{
+		if (PackagesToFilter[Index])
+		{
+			InOutAssetDataArray.RemoveAtSwap(Index, EAllowShrinking::No);
+		}
+	}
 }
 
 /** Remove any external actors that currently exist in InOutAssetDataArray. OutPartitionedWorldPackageNames is populated with the package paths of partitioned worlds.*/
@@ -952,7 +1007,7 @@ void UGatherTextFromAssetsCommandlet::PopulatePackagesPendingGather(TSet<FName>
 /** Process packages with loc data cached in its header and removes them from the pending packages.*/
 void UGatherTextFromAssetsCommandlet::ProcessAndRemoveCachedPackages(TMap<FName, TSet<FGuid>>& OutExternalActorsWithStaleOrMissingCaches)
 {
-	const double ProcessingStartTime = FPlatformTime::Seconds();
+	UE_SCOPED_TIMER(TEXT("UGatherTextFromAssetsCommandlet::ProcessAndRemoveCachedPackages"), LogGatherTextFromAssetsCommandlet, Display);
 	int32 NumPackagesProcessed = 0;
 	int32 PackageCount = PackagesPendingGather.Num();
 	TMap<FString, FName> AssignedPackageLocalizationIds;
@@ -1040,7 +1095,6 @@ void UGatherTextFromAssetsCommandlet::ProcessAndRemoveCachedPackages(TMap<FName,
 
 			return true;
 		});
-	UE_LOG(LogGatherTextFromAssetsCommandlet, Display, TEXT("Processing and removing packages with cached localization data took %.2f"), FPlatformTime::Seconds() - ProcessingStartTime);
 }
 
 void UGatherTextFromAssetsCommandlet::MergeInExternalActorsWithStaleOrMissingCaches(TMap<FName, TSet<FGuid>>& StaleExternalActors)
@@ -1071,12 +1125,13 @@ void UGatherTextFromAssetsCommandlet::MergeInExternalActorsWithStaleOrMissingCac
 /** Load the remaining pending packages for gather.*/
 void UGatherTextFromAssetsCommandlet::LoadAndProcessUncachedPackages(TArray<FName>& OutPackagesWithStaleGatherCache)
 {
+	UE_SCOPED_TIMER(TEXT("UGatherTextFromAssetsCommandlet::LoadAndProcessUncachedPackages"), LogGatherTextFromAssetsCommandlet, Display);
 	FLoadPackageLogOutputRedirector LogOutputRedirector;
 	TArray<FGatherableTextData> GatherableTextDataArray;
 	int32 NumPackagesProcessed = 0;
 	int32 PackageCount = PackagesPendingGather.Num();
 	int32 NumPackagesFailedLoading = 0;
-	const double LoadingStartTime = FPlatformTime::Seconds();
+	
 	while (PackagesPendingGather.Num() > 0)
 	{
 		const FPackagePendingGather PackagePendingGather = PackagesPendingGather.Pop(EAllowShrinking::No);
@@ -1232,7 +1287,6 @@ void UGatherTextFromAssetsCommandlet::LoadAndProcessUncachedPackages(TArray<FNam
 		}
 	}
 	UE_LOG(LogGatherTextFromAssetsCommandlet, Display, TEXT("Loaded %d packages. %d failed."), NumPackagesProcessed, NumPackagesFailedLoading);
-	UE_LOG(LogGatherTextFromAssetsCommandlet, Display, TEXT("Loading all uncached packages took %.2f seconds."), FPlatformTime::Seconds() - LoadingStartTime);
 }
 
 void UGatherTextFromAssetsCommandlet::ReportStaleGatherCache(TArray<FName>& InPackagesWithStaleGatherCache) const
@@ -1318,6 +1372,7 @@ UGatherTextFromAssetsCommandlet::EPackageLocCacheState UGatherTextFromAssetsComm
 
 int32 UGatherTextFromAssetsCommandlet::Main(const FString& Params)
 {
+	UE_SCOPED_TIMER(TEXT("UGatherTextFromAssetsCommandlet::Main"), LogGatherTextFromAssetsCommandlet, Display);
 	// Parse command line.
 	if (!ParseCommandLineHelper(Params))
 	{
@@ -1334,14 +1389,14 @@ int32 UGatherTextFromAssetsCommandlet::Main(const FString& Params)
 			GEditor->CreateNewMapForEditing(/*bPromptForSave*/false);
 		}
 	}
-
 	UE_LOG(LogGatherTextFromAssetsCommandlet, Display, TEXT("Discovering assets to gather..."));
 	const double DiscoveringAssetsStartTime = FPlatformTime::Seconds();
-	const double SearchAssetRegistryStartTime = FPlatformTime::Seconds();
-	FAssetRegistryModule& AssetRegistryModule = FModuleManager::LoadModuleChecked<FAssetRegistryModule>(TEXT("AssetRegistry"));
-	IAssetRegistry& AssetRegistry = AssetRegistryModule.Get();
-	AssetRegistry.SearchAllAssets(true);
-	UE_LOG(LogGatherTextFromAssetsCommandlet, Display, TEXT("Searching all assets took %.2f seconds."), FPlatformTime::Seconds() - SearchAssetRegistryStartTime);
+	{
+		UE_SCOPED_TIMER(TEXT("UGatherTextFromAssetsCommandlet::SearchAssetRegistryForAllAssets"), LogGatherTextFromAssetsCommandlet, Display);
+		FAssetRegistryModule& AssetRegistryModule = FModuleManager::LoadModuleChecked<FAssetRegistryModule>(TEXT("AssetRegistry"));
+		IAssetRegistry& AssetRegistry = AssetRegistryModule.Get();
+		AssetRegistry.SearchAllAssets(true);
+	}
 
 	TArray<FAssetData> AssetDataArray;
 	if (!PerformFirstPassFilter(AssetDataArray))
