Cooker: Speculative fix for error PackageSplitter found the Generator package still in memory after it should have been deleted by GC.
If we encounter a CookPackageSplitter that needs to be garbage collected but was not GC'd by a soft GC, request a full GC.
#rnx

diff --git a/Engine/Source/Editor/UnrealEd/Classes/CookOnTheSide/CookOnTheFlyServer.h b/Engine/Source/Editor/UnrealEd/Classes/CookOnTheSide/CookOnTheFlyServer.h
--- a/Engine/Source/Editor/UnrealEd/Classes/CookOnTheSide/CookOnTheFlyServer.h
+++ b/Engine/Source/Editor/UnrealEd/Classes/CookOnTheSide/CookOnTheFlyServer.h
@@ -828,8 +828,8 @@ public:
 	UNREALED_API void SetGarbageCollectType(uint32 ResultFlagsFromTick);
 	UNREALED_API void ClearGarbageCollectType();
 
-	UNREALED_API void OnCookerStartCollectGarbage();
-	UNREALED_API void OnCookerEndCollectGarbage();
+	UNREALED_API void OnCookerStartCollectGarbage(uint32& ResultFlagsFromTick);
+	UNREALED_API void OnCookerEndCollectGarbage(uint32& ResultFlagsFromTick);
 	UNREALED_API void EvaluateGarbageCollectionResults(bool bWasDueToOOM, bool bWasPartialGC, uint32 ResultFlags,
 		int32 NumObjectsBeforeGC, const FPlatformMemoryStats& MemStatsBeforeGC,
 		const FGenericMemoryStats& AllocatorStatsBeforeGC,
@@ -1628,6 +1628,7 @@ private:
 	friend UE::Cook::FAssetRegistryMPCollector;
 	friend UE::Cook::FBeginCookConfigSettings;
 	friend UE::Cook::FCookDirector;
+	friend UE::Cook::FCookGCDiagnosticContext;
 	friend UE::Cook::FCookGenerationInfo;
 	friend UE::Cook::FCookWorkerClient;
 	friend UE::Cook::FCookWorkerServer;
diff --git a/Engine/Source/Editor/UnrealEd/Private/Commandlets/CookCommandlet.cpp b/Engine/Source/Editor/UnrealEd/Private/Commandlets/CookCommandlet.cpp
--- a/Engine/Source/Editor/UnrealEd/Private/Commandlets/CookCommandlet.cpp
+++ b/Engine/Source/Editor/UnrealEd/Private/Commandlets/CookCommandlet.cpp
@@ -734,14 +734,15 @@ void UCookCommandlet::ConditionalCollectGarbage(uint32 TickResults, UCookOnTheFl
 	UE_LOG(LogCookCommandlet, Display, TEXT("GarbageCollection...%s (%s)"), *GCType, *GCReason);
 	{
 		TGuardValue<bool> SoftGCGuard(UPackage::bSupportCookerSoftGC, true);
-		COTFS.OnCookerStartCollectGarbage();
+		COTFS.OnCookerStartCollectGarbage(TickResults);
 		CollectGarbage(RF_NoFlags);
-		COTFS.OnCookerEndCollectGarbage();
+		COTFS.OnCookerEndCollectGarbage(TickResults);
 		if (COTFS.NeedsDiagnosticSecondGC())
 		{
-			COTFS.OnCookerStartCollectGarbage();
+			UE_LOG(LogCookCommandlet, Display, TEXT("Second GarbageCollect requested by cooker..."));
+			COTFS.OnCookerStartCollectGarbage(TickResults);
 			CollectGarbage(RF_NoFlags);
-			COTFS.OnCookerEndCollectGarbage();
+			COTFS.OnCookerEndCollectGarbage(TickResults);
 		}
 	}
 
diff --git a/Engine/Source/Editor/UnrealEd/Private/CookOnTheFlyServer.cpp b/Engine/Source/Editor/UnrealEd/Private/CookOnTheFlyServer.cpp
--- a/Engine/Source/Editor/UnrealEd/Private/CookOnTheFlyServer.cpp
+++ b/Engine/Source/Editor/UnrealEd/Private/CookOnTheFlyServer.cpp
@@ -5589,14 +5589,14 @@ bool UCookOnTheFlyServer::NeedsDiagnosticSecondGC() const
 	return GCDiagnosticContext->NeedsDiagnosticSecondGC();
 }
 
-void UCookOnTheFlyServer::OnCookerStartCollectGarbage()
+void UCookOnTheFlyServer::OnCookerStartCollectGarbage(uint32& ResultFlagsFromTick)
 {
-	GCDiagnosticContext->OnCookerStartCollectGarbage();
+	GCDiagnosticContext->OnCookerStartCollectGarbage(*this, ResultFlagsFromTick);
 }
 
-void UCookOnTheFlyServer::OnCookerEndCollectGarbage()
+void UCookOnTheFlyServer::OnCookerEndCollectGarbage(uint32& ResultFlagsFromTick)
 {
-	GCDiagnosticContext->OnCookerEndCollectGarbage();
+	GCDiagnosticContext->OnCookerEndCollectGarbage(*this, ResultFlagsFromTick);
 }
 
 void UCookOnTheFlyServer::EvaluateGarbageCollectionResults(bool bWasDueToOOM, bool bWasPartialGC, uint32 ResultFlags,
diff --git a/Engine/Source/Editor/UnrealEd/Private/Cooker/CookGarbageCollect.cpp b/Engine/Source/Editor/UnrealEd/Private/Cooker/CookGarbageCollect.cpp
--- a/Engine/Source/Editor/UnrealEd/Private/Cooker/CookGarbageCollect.cpp
+++ b/Engine/Source/Editor/UnrealEd/Private/Cooker/CookGarbageCollect.cpp
@@ -2,6 +2,8 @@
 
 #include "Cooker/CookGarbageCollect.h"
 
+#include "CookOnTheSide/CookOnTheFlyServer.h"
+#include "Misc/EnumClassFlags.h"
 #include "UObject/GarbageCollectionHistory.h"
 
 namespace UE::Cook
@@ -14,7 +16,7 @@ FCookGCDiagnosticContext::~FCookGCDiagnosticContext()
 
 bool FCookGCDiagnosticContext::NeedsDiagnosticSecondGC() const
 {
-	return bRequestGCWithHistory;
+	return bRequestGCWithHistory || bRequestFullGC;
 }
 
 bool FCookGCDiagnosticContext::CurrentGCHasHistory() const
@@ -36,7 +38,17 @@ bool FCookGCDiagnosticContext::TryRequestGCWithHistory()
 #endif
 }
 
-void FCookGCDiagnosticContext::OnCookerStartCollectGarbage()
+bool FCookGCDiagnosticContext::TryRequestFullGC()
+{
+	if (!bRequestsAvailable || !bGCInProgress || bCurrentGCIsFull)
+	{
+		return false;
+	}
+	bRequestFullGC = true;
+	return true;
+}
+
+void FCookGCDiagnosticContext::OnCookerStartCollectGarbage(UCookOnTheFlyServer& COTFS, uint32& ResultFlagsFromTick)
 {
 	bRequestsAvailable = true;
 
@@ -46,17 +58,25 @@ void FCookGCDiagnosticContext::OnCookerStartCollectGarbage()
 #else
 	bCurrentGCHasHistory = false;
 #endif
+	if (bRequestFullGC)
+	{
+		COTFS.bGarbageCollectTypeSoft = false;
+		ResultFlagsFromTick = ResultFlagsFromTick & ~UCookOnTheFlyServer::COSR_RequiresGC_Soft_OOM;
+	}
+	bCurrentGCIsFull = !COTFS.bGarbageCollectTypeSoft;
 }
 
-void FCookGCDiagnosticContext::OnCookerEndCollectGarbage()
+void FCookGCDiagnosticContext::OnCookerEndCollectGarbage(UCookOnTheFlyServer& COTFS, uint32& ResultFlagsFromTick)
 {
 	bGCInProgress = false;
 	bCurrentGCHasHistory = false;
+	bCurrentGCIsFull = false;
 }
 
 void FCookGCDiagnosticContext::OnEvaluateResultsComplete()
 {
 	SetGCWithHistoryRequested(false);
+	bRequestFullGC = false;
 }
 
 void FCookGCDiagnosticContext::SetGCWithHistoryRequested(bool bValue)
diff --git a/Engine/Source/Editor/UnrealEd/Private/Cooker/CookGarbageCollect.h b/Engine/Source/Editor/UnrealEd/Private/Cooker/CookGarbageCollect.h
--- a/Engine/Source/Editor/UnrealEd/Private/Cooker/CookGarbageCollect.h
+++ b/Engine/Source/Editor/UnrealEd/Private/Cooker/CookGarbageCollect.h
@@ -4,6 +4,8 @@
 
 #include "CoreTypes.h"
 
+class UCookOnTheFlyServer;
+
 namespace UE::Cook
 {
 
@@ -22,13 +24,19 @@ public:
 
 	/**
 	 * Add a request to reexecute the current GC after all of the PostGarbageCollect calls run and
-	 * control returns back to the caller of CollectGarbage. Returns false if not currently in post-GC,
-	 * or the garbage collect that just ran already had history.
+	 * control returns back to the caller of CollectGarbage, and with history turned on.
+	 * Returns false if not currently in post-GC, or the garbage collect that just ran already had history.
 	 */
 	bool TryRequestGCWithHistory();
+	/**
+	 * Add a request to reexecute the current GC after all of the PostGarbageCollect calls run and
+	 * control returns back to the caller of CollectGarbage, and with soft GC turned off.
+	 * Returns false if not currently in post-GC, or the garbage collect that just ran already was a full GC.
+	 */
+	bool TryRequestFullGC();
 
-	void OnCookerStartCollectGarbage();
-	void OnCookerEndCollectGarbage();
+	void OnCookerStartCollectGarbage(UCookOnTheFlyServer& COTFS, uint32& ResultFlagsFromTick);
+	void OnCookerEndCollectGarbage(UCookOnTheFlyServer& COTFS, uint32& ResultFlagsFromTick);
 	void OnEvaluateResultsComplete();
 
 private:
@@ -38,7 +46,9 @@ private:
 	bool bRequestsAvailable = false;
 	bool bGCInProgress = false;
 	bool bRequestGCWithHistory = false;
+	bool bRequestFullGC = false;
 	bool bCurrentGCHasHistory = false;
+	bool bCurrentGCIsFull = false;
 };
 
 } // namespace UE::Cook
diff --git a/Engine/Source/Editor/UnrealEd/Private/Cooker/CookGenerationHelper.cpp b/Engine/Source/Editor/UnrealEd/Private/Cooker/CookGenerationHelper.cpp
--- a/Engine/Source/Editor/UnrealEd/Private/Cooker/CookGenerationHelper.cpp
+++ b/Engine/Source/Editor/UnrealEd/Private/Cooker/CookGenerationHelper.cpp
@@ -1677,7 +1677,10 @@ void FGenerationHelper::VerifyGeneratorPackageGarbageCollected(FCookGCDiagnostic
 	UPackage* LocalOwnerPackage = FindObject<UPackage>(nullptr, *GeneratorPackageName);
 	if (LocalOwnerPackage)
 	{
-		if (!Context.TryRequestGCWithHistory())
+		bool bWillRetry = false;
+		bWillRetry = Context.TryRequestGCWithHistory() || bWillRetry;
+		bWillRetry = Context.TryRequestFullGC() || bWillRetry;
+		if (!bWillRetry)
 		{
 			// Might be called when uninitialized, so do not call GetSplitDataObjectNameIfAvailable
 			FString Identifier;
