Localization:
Add timers to all the main localization commandlets to get aggregate timings for their runs across multiple plugin localization targets.

#rnx

diff --git a/Engine/Source/Editor/UnrealEd/Private/Commandlets/GatherTextFromMetadataCommandlet.cpp b/Engine/Source/Editor/UnrealEd/Private/Commandlets/GatherTextFromMetadataCommandlet.cpp
--- a/Engine/Source/Editor/UnrealEd/Private/Commandlets/GatherTextFromMetadataCommandlet.cpp
+++ b/Engine/Source/Editor/UnrealEd/Private/Commandlets/GatherTextFromMetadataCommandlet.cpp
@@ -23,6 +23,7 @@
 #include "UObject/UObjectHash.h"
 #include "UObject/UObjectIterator.h"
 #include "UObject/UnrealType.h"
+#include "ProfilingDebugging/ScopedTimers.h"
 
 class UObject;
 
@@ -45,6 +46,7 @@ bool UGatherTextFromMetaDataCommandlet::ShouldRunInPreview(const TArray<FString>
 
 int32 UGatherTextFromMetaDataCommandlet::Main( const FString& Params )
 {
+	UE_SCOPED_TIMER(TEXT("UGatherTextFromMetaDataCommandlet::Main"), LogGatherTextFromMetaDataCommandlet, Display);
 	// Parse command line - we're interested in the param vals
 	TArray<FString> Tokens;
 	TArray<FString> Switches;
diff --git a/Engine/Source/Editor/UnrealEd/Private/Commandlets/GatherTextFromSourceCommandlet.cpp b/Engine/Source/Editor/UnrealEd/Private/Commandlets/GatherTextFromSourceCommandlet.cpp
--- a/Engine/Source/Editor/UnrealEd/Private/Commandlets/GatherTextFromSourceCommandlet.cpp
+++ b/Engine/Source/Editor/UnrealEd/Private/Commandlets/GatherTextFromSourceCommandlet.cpp
@@ -9,6 +9,7 @@
 #include "Algo/Transform.h"
 #include "Internationalization/InternationalizationMetadata.h"
 #include "Internationalization/TextNamespaceUtil.h"
+#include "ProfilingDebugging/ScopedTimers.h"
 
 DEFINE_LOG_CATEGORY_STATIC(LogGatherTextFromSourceCommandlet, Log, All);
 
@@ -43,6 +44,7 @@ bool UGatherTextFromSourceCommandlet::ShouldRunInPreview(const TArray<FString>&
 
 int32 UGatherTextFromSourceCommandlet::Main( const FString& Params )
 {
+	UE_SCOPED_TIMER(TEXT("UGatherTextFromSourceCommandlet::Main"), LogGatherTextFromSourceCommandlet, Display);
 	// Parse command line - we're interested in the param vals
 	TArray<FString> Tokens;
 	TArray<FString> Switches;
diff --git a/Engine/Source/Editor/UnrealEd/Private/Commandlets/GenerateGatherArchiveCommandlet.cpp b/Engine/Source/Editor/UnrealEd/Private/Commandlets/GenerateGatherArchiveCommandlet.cpp
--- a/Engine/Source/Editor/UnrealEd/Private/Commandlets/GenerateGatherArchiveCommandlet.cpp
+++ b/Engine/Source/Editor/UnrealEd/Private/Commandlets/GenerateGatherArchiveCommandlet.cpp
@@ -16,6 +16,7 @@
 #include "Logging/LogMacros.h"
 #include "Templates/SharedPointer.h"
 #include "Trace/Detail/Channel.h"
+#include "ProfilingDebugging/ScopedTimers.h"
 
 DEFINE_LOG_CATEGORY_STATIC(LogGenerateArchiveCommandlet, Log, All);
 
@@ -140,6 +141,7 @@ UGenerateGatherArchiveCommandlet::UGenerateGatherArchiveCommandlet(const FObject
 
 int32 UGenerateGatherArchiveCommandlet::Main( const FString& Params )
 {
+	UE_SCOPED_TIMER(TEXT("UGenerateGatherArchiveCommandlet::Main"), LogGenerateArchiveCommandlet, Display);
 	// Parse command line - we're interested in the param vals
 	TArray<FString> Tokens;
 	TArray<FString> Switches;
diff --git a/Engine/Source/Editor/UnrealEd/Private/Commandlets/GenerateGatherManifestCommandlet.cpp b/Engine/Source/Editor/UnrealEd/Private/Commandlets/GenerateGatherManifestCommandlet.cpp
--- a/Engine/Source/Editor/UnrealEd/Private/Commandlets/GenerateGatherManifestCommandlet.cpp
+++ b/Engine/Source/Editor/UnrealEd/Private/Commandlets/GenerateGatherManifestCommandlet.cpp
@@ -14,6 +14,7 @@
 #include "Misc/Paths.h"
 #include "Templates/SharedPointer.h"
 #include "Trace/Detail/Channel.h"
+#include "ProfilingDebugging/ScopedTimers.h"
 
 DEFINE_LOG_CATEGORY_STATIC(LogGenerateManifestCommandlet, Log, All);
 
@@ -68,6 +69,7 @@ UGenerateGatherManifestCommandlet::UGenerateGatherManifestCommandlet(const FObje
 
 int32 UGenerateGatherManifestCommandlet::Main( const FString& Params )
 {
+	UE_SCOPED_TIMER(TEXT("UGenerateGatherManifestCommandlet::Main"), LogGenerateManifestCommandlet, Display);
 	// Parse command line - we're interested in the param vals
 	TArray<FString> Tokens;
 	TArray<FString> Switches;
diff --git a/Engine/Source/Editor/UnrealEd/Private/Commandlets/GenerateTextLocalizationReportCommandlet.cpp b/Engine/Source/Editor/UnrealEd/Private/Commandlets/GenerateTextLocalizationReportCommandlet.cpp
--- a/Engine/Source/Editor/UnrealEd/Private/Commandlets/GenerateTextLocalizationReportCommandlet.cpp
+++ b/Engine/Source/Editor/UnrealEd/Private/Commandlets/GenerateTextLocalizationReportCommandlet.cpp
@@ -16,6 +16,7 @@
 #include "Misc/DateTime.h"
 #include "Templates/SharedPointer.h"
 #include "Trace/Detail/Channel.h"
+#include "ProfilingDebugging/ScopedTimers.h"
 
 DEFINE_LOG_CATEGORY_STATIC(LogGenerateTextLocalizationReportCommandlet, Log, All);
 
@@ -26,6 +27,7 @@ UGenerateTextLocalizationReportCommandlet::UGenerateTextLocalizationReportComman
 
 int32 UGenerateTextLocalizationReportCommandlet::Main(const FString& Params)
 {
+	UE_SCOPED_TIMER(TEXT("UGenerateTextLocalizationReportCommandlet::Main"), LogGenerateTextLocalizationReportCommandlet, Display);
 	// Parse command line - we're interested in the param vals
 	TArray<FString> Tokens;
 	TArray<FString> Switches;
diff --git a/Engine/Source/Editor/UnrealEd/Private/Commandlets/GenerateTextLocalizationResourceCommandlet.cpp b/Engine/Source/Editor/UnrealEd/Private/Commandlets/GenerateTextLocalizationResourceCommandlet.cpp
--- a/Engine/Source/Editor/UnrealEd/Private/Commandlets/GenerateTextLocalizationResourceCommandlet.cpp
+++ b/Engine/Source/Editor/UnrealEd/Private/Commandlets/GenerateTextLocalizationResourceCommandlet.cpp
@@ -20,6 +20,7 @@
 #include "TextLocalizationResourceGenerator.h"
 #include "Trace/Detail/Channel.h"
 #include "UObject/NameTypes.h"
+#include "ProfilingDebugging/ScopedTimers.h"
 
 DEFINE_LOG_CATEGORY_STATIC(LogGenerateTextLocalizationResourceCommandlet, Log, All);
 
@@ -30,6 +31,7 @@ UGenerateTextLocalizationResourceCommandlet::UGenerateTextLocalizationResourceCo
 
 int32 UGenerateTextLocalizationResourceCommandlet::Main(const FString& Params)
 {
+	UE_SCOPED_TIMER(TEXT("UGenerateTextLocalizationResourceCommandlet::Main"), LogGenerateTextLocalizationResourceCommandlet, Display);
 	// Parse command line - we're interested in the param vals
 	TArray<FString> Tokens;
 	TArray<FString> Switches;
diff --git a/Engine/Source/Editor/UnrealEd/Private/Commandlets/InternationalizationExportCommandlet.cpp b/Engine/Source/Editor/UnrealEd/Private/Commandlets/InternationalizationExportCommandlet.cpp
--- a/Engine/Source/Editor/UnrealEd/Private/Commandlets/InternationalizationExportCommandlet.cpp
+++ b/Engine/Source/Editor/UnrealEd/Private/Commandlets/InternationalizationExportCommandlet.cpp
@@ -15,6 +15,7 @@
 #include "Trace/Detail/Channel.h"
 #include "UObject/Class.h"
 #include "UObject/UObjectGlobals.h"
+#include "ProfilingDebugging/ScopedTimers.h"
 
 DEFINE_LOG_CATEGORY_STATIC(LogInternationalizationExportCommandlet, Log, All);
 
@@ -159,6 +160,7 @@ int32 UInternationalizationExportCommandlet::Main(const FString& Params)
 
 	if (bDoImport)
 	{
+		UE_SCOPED_TIMER(TEXT("UInternationalizationExportCommandlet::Main Import"), LogInternationalizationExportCommandlet, Display);
 		// Load the manifest and all archives
 		FLocTextHelper LocTextHelper(DestinationPath, ManifestName, ArchiveName, NativeCultureName, CulturesToGenerate, GatherManifestHelper->GetLocFileNotifies(), GatherManifestHelper->GetPlatformSplitMode());
 		LocTextHelper.SetCopyrightNotice(GatherManifestHelper->GetCopyrightNotice());
@@ -181,6 +183,7 @@ int32 UInternationalizationExportCommandlet::Main(const FString& Params)
 
 	if (bDoExport)
 	{
+		UE_SCOPED_TIMER(TEXT("UInternationalizationExportCommandlet::Main (Export)"), LogInternationalizationExportCommandlet, Display);
 		bool bShouldPersistComments = false;
 		GetBoolFromConfig(*SectionName, TEXT("ShouldPersistCommentsOnExport"), bShouldPersistComments, ConfigPath);
 
