Actor Coloration Framework: fix crash when property coloration mode is active and changing maps.

diff --git a/Engine/Source/Editor/UnrealEd/Private/EditorServer.cpp b/Engine/Source/Editor/UnrealEd/Private/EditorServer.cpp
--- a/Engine/Source/Editor/UnrealEd/Private/EditorServer.cpp
+++ b/Engine/Source/Editor/UnrealEd/Private/EditorServer.cpp
@@ -3726,7 +3726,7 @@ void UEditorEngine::SetPropertyColorationTarget(UWorld* InWorld, const FString&
 		if (PropertyChain)
 		{
 			GPropertyColorationChain = new TSharedRef<FEditPropertyChain>(*PropertyChain);
-			GbColorationClassIsActor = GPropertyColorationClass->IsChildOf( AActor::StaticClass() );
+			GbColorationClassIsActor = GPropertyColorationClass && GPropertyColorationClass->IsChildOf(AActor::StaticClass());
 			GbColorationPropertyIsObjectProperty = CastField<FObjectPropertyBase>(GPropertyColorationProperty) != NULL;
 			
 			FActorPrimitiveColorHandler::Get().RefreshPrimitiveColorHandler(TEXT("PropertyColor"), InWorld);
