Actor Coloration Framework QoL: make the property color for mathing objects configurable in editor settings.

diff --git a/Engine/Source/Editor/UnrealEd/Classes/Settings/LevelEditorViewportSettings.h b/Engine/Source/Editor/UnrealEd/Classes/Settings/LevelEditorViewportSettings.h
--- a/Engine/Source/Editor/UnrealEd/Classes/Settings/LevelEditorViewportSettings.h
+++ b/Engine/Source/Editor/UnrealEd/Classes/Settings/LevelEditorViewportSettings.h
@@ -597,6 +597,10 @@ public:
 	UPROPERTY(config)
 	TMap<TSoftObjectPtr<UWorld>, FLevelEditorViewporEditorViews> EditorViews;
 
+	/** Property coloration color for matching objects */
+	UPROPERTY(EditAnywhere, config, Category = LookAndFeel)
+	FColor PropertyColorationColorForMatchingObjects = FColor::Red;
+
 private:
 
 	// Per-instance viewport settings.
diff --git a/Engine/Source/Editor/UnrealEd/Private/EditorEngine.cpp b/Engine/Source/Editor/UnrealEd/Private/EditorEngine.cpp
--- a/Engine/Source/Editor/UnrealEd/Private/EditorEngine.cpp
+++ b/Engine/Source/Editor/UnrealEd/Private/EditorEngine.cpp
@@ -618,7 +618,7 @@ UEditorEngine::UEditorEngine(const FObjectInitializer& ObjectInitializer)
 				{
 					if (GetPropertyColorationMatch(Actor))
 					{
-						PropertyColor = FColor::Red;
+						PropertyColor = GetDefault<ULevelEditorViewportSettings>()->PropertyColorationColorForMatchingObjects;
 					}
 				}
 				return PropertyColor;
