Cooker: Move BroadcastCookByTheBookStarted to before shutting down the session, so that session data is still available for subscribers.
#rnx

diff --git a/Engine/Source/Editor/UnrealEd/Private/CookOnTheFlyServer.cpp b/Engine/Source/Editor/UnrealEd/Private/CookOnTheFlyServer.cpp
--- a/Engine/Source/Editor/UnrealEd/Private/CookOnTheFlyServer.cpp
+++ b/Engine/Source/Editor/UnrealEd/Private/CookOnTheFlyServer.cpp
@@ -10377,7 +10377,6 @@ void UCookOnTheFlyServer::CookByTheBookFinished()
 		ClearCookInProgressFlagFromCookSettings(TargetPlatform);
 	}
 	ShutdownCookSession();
-	BroadcastCookByTheBookFinished();
 	UE_LOG(LogCook, Display, TEXT("Done!"));
 }
 
@@ -10691,19 +10690,28 @@ void UCookOnTheFlyServer::ShutdownCookSession()
 	{
 		CookDirector->ShutdownCookSession();
 	}
-	if (CookWorkerClient)
-	{
-		CookAsCookWorkerFinished();
-	}
 
 	if (IsCookByTheBookMode())
 	{
+		// CookWorkers report false for IsCookByTheBookMode; they are CookWorker mode. They need to shutdown in a
+		// custom manner, which we do in the else if below.
+		check(!CookWorkerClient);
 		UnregisterCookByTheBookDelegates();
 
 		PrintFinishStats();
 		OutputHierarchyTimers();
 		PrintDetailedCookStats();
+
+		// BroadcastCookByTheBookFinished needs to be called before clearing the session data, so that subscribers
+		// can access information about the session such as DLCName.
+		BroadcastCookByTheBookFinished();
 	}
+	else if (CookWorkerClient)
+	{
+		CookAsCookWorkerFinished();
+		// CookAsCookWorkerFinished is responsible for calling BroadcastCookByTheBookFinished.
+	}
+
 	CookByTheBookOptions->ClearSessionData();
 	PlatformManager->ClearSessionPlatforms(*this);
 	ClearHierarchyTimers();
