IncrementalCook: IterativeValidate: Add display of number of falsepositiveiterativeskips per asset class to the end-of-cook error message.
#rnx

diff --git a/Engine/Source/Editor/UnrealEd/Private/Cooker/IterativeValidatePackageWriter.cpp b/Engine/Source/Editor/UnrealEd/Private/Cooker/IterativeValidatePackageWriter.cpp
--- a/Engine/Source/Editor/UnrealEd/Private/Cooker/IterativeValidatePackageWriter.cpp
+++ b/Engine/Source/Editor/UnrealEd/Private/Cooker/IterativeValidatePackageWriter.cpp
@@ -109,6 +109,47 @@ bool LoadFromCompactBinary(FCbFieldView Field, FIterativeValidatePackageWriter::
 	return false;
 }
 
+FArchive& operator<<(FArchive& Ar, FIterativeValidatePackageWriter::FPackageStatusInfo& Info)
+{
+	Ar << Info.Status;
+	bool bHasAssetClass = Info.AssetClass.IsValid();
+	Ar << bHasAssetClass;
+	if (bHasAssetClass)
+	{
+		Ar << Info.AssetClass;
+	}
+	return Ar;
+}
+
+FCbWriter& operator<<(FCbWriter& Writer, const FIterativeValidatePackageWriter::FPackageStatusInfo& Info)
+{
+	Writer.BeginArray();
+	Writer << Info.Status;
+	if (Info.AssetClass.IsValid())
+	{
+		Writer << Info.AssetClass;
+	}
+	Writer.EndArray();
+	return Writer;
+}
+
+bool LoadFromCompactBinary(FCbFieldView Field, FIterativeValidatePackageWriter::FPackageStatusInfo& Info)
+{
+	FCbArrayView ArrayField = Field.AsArrayView();
+	if (ArrayField.Num() < 1)
+	{
+		return false;
+	}
+	FCbFieldViewIterator Iter = ArrayField.CreateViewIterator();
+	bool bOk = true;
+	bOk = LoadFromCompactBinary(*(Iter++), Info.Status) & bOk;
+	if (ArrayField.Num() >= 2)
+	{
+		bOk = LoadFromCompactBinary(*(Iter++), Info.AssetClass) & bOk;
+	}
+	return bOk;
+}
+
 void FIterativeValidateMPCollector::ServerTick(UE::Cook::FMPCollectorServerTickContext& Context)
 {
 	static_assert(sizeof(EMessageSubtype) == sizeof(uint8));
@@ -236,11 +277,12 @@ FGuid FIterativeValidateMPCollector::MessageType(TEXT("5E56C5D96F3B455E9452C15AD
 
 bool FIterativeValidateMPCollector::TryWritePackageStatus(FCbWriter& Writer, FName PackageName)
 {
-	FIterativeValidatePackageWriter::EPackageStatus PackageStatus = Owner->GetPackageStatus(PackageName);
+	const FIterativeValidatePackageWriter::FPackageStatusInfo* PackageStatus
+		= Owner->PackageStatusMap.Find(PackageName);
 
-	if (PackageStatus != FIterativeValidatePackageWriter::EPackageStatus::NotYetProcessed)
+	if (PackageStatus && PackageStatus->Status != FIterativeValidatePackageWriter::EPackageStatus::NotYetProcessed)
 	{
-		Writer << "Status" << (uint8)PackageStatus;
+		Writer << "Status" << *PackageStatus;
 
 		const TArray<FIterativeValidatePackageWriter::FMessage>* Messages = Owner->PackageMessageMap.Find(PackageName);
 		if (Messages != nullptr)
@@ -254,14 +296,12 @@ bool FIterativeValidateMPCollector::TryWritePackageStatus(FCbWriter& Writer, FNa
 
 void FIterativeValidateMPCollector::ReadAndSyncPackageStatus(FCbObjectView Message, FName PackageName)
 {
-	uint8 PackageStatusInteger = (uint8)FIterativeValidatePackageWriter::EPackageStatus::Count;
-	bool bOk = LoadFromCompactBinary(Message["Status"], PackageStatusInteger);
-	bOk = bOk && PackageStatusInteger < (uint8)FIterativeValidatePackageWriter::EPackageStatus::Count;
+	FIterativeValidatePackageWriter::FPackageStatusInfo Info;
+	bool bOk = LoadFromCompactBinary(Message["Status"], Info);
 
 	if (bOk)
 	{
-		FIterativeValidatePackageWriter::EPackageStatus PackageStatus = (FIterativeValidatePackageWriter::EPackageStatus)PackageStatusInteger;
-		Owner->SetPackageStatus(PackageName, PackageStatus);
+		Owner->PackageStatusMap.FindOrAdd(PackageName) = MoveTemp(Info);
 		if (Message.FindView("MessageArray").HasValue())
 		{
 			TArray<FIterativeValidatePackageWriter::FMessage>& MessageArray = Owner->PackageMessageMap.FindOrAdd(PackageName);
@@ -271,8 +311,8 @@ void FIterativeValidateMPCollector::ReadAndSyncPackageStatus(FCbObjectView Messa
 	
 	if (!bOk)
 	{
-		UE_LOG(LogCook, Error, TEXT("Invalid message received in ReadAndSyncPackageStatus. Status received (%d) for package \"%s\""),
-			PackageStatusInteger, *PackageName.ToString());
+		UE_LOG(LogCook, Error, TEXT("Invalid message received in ReadAndSyncPackageStatus. Failed to load Info from Message[\"Status\"] for package \"%s\""),
+			*PackageName.ToString());
 	}
 }
 
@@ -486,6 +526,11 @@ int64 FIterativeValidatePackageWriter::GetExportsFooterSize()
 
 TUniquePtr<FLargeMemoryWriter> FIterativeValidatePackageWriter::CreateLinkerArchive(FName PackageName, UObject* Asset, uint16 MultiOutputIndex)
 {
+	if (Asset)
+	{
+		PackageStatusMap.FindOrAdd(PackageName).AssetClass = Asset->GetClass()->GetClassPathName();
+	}
+
 	switch (SaveAction)
 	{
 	case ESaveAction::CheckForDiffs:
@@ -731,8 +776,25 @@ void FIterativeValidatePackageWriter::EndCook(const FCookInfo& Info)
 			StatusCounts[EPackageStatus::DeclaredUnmodified_FoundModified_FalsePositive]);
 		if (StatusCounts[EPackageStatus::DeclaredUnmodified_FoundModified_FalsePositive] > 0)
 		{
+			TStringBuilder<1024> MessageWithDiagnostics;
+			MessageWithDiagnostics << Message;
+			TMap<FTopLevelAssetPath, int32> ClassFalsePositiveCounts = GetSummaryFalsePositiveCounts();
+			int32 NumPrinted = 0;
+			constexpr int32 MaxNumPrinted = 5;
+			for (const TPair<FTopLevelAssetPath, int32>& Pair : ClassFalsePositiveCounts)
+			{
+				MessageWithDiagnostics << TEXT("\n\t") << WriteToString<256>(Pair.Key) << TEXT(": ") << Pair.Value;
+				if (++NumPrinted >= MaxNumPrinted)
+				{
+					break;
+				}
+			}
+			if (ClassFalsePositiveCounts.Num() > NumPrinted)
+			{
+				MessageWithDiagnostics << TEXT("\n\t...");
+			}
 			UE_CLOG(COTFS.GetCookMode() != ECookMode::CookWorker,
-				LogIterativeValidate, Error, TEXT("%s"), *Message);
+				LogIterativeValidate, Error, TEXT("%s"), *MessageWithDiagnostics);
 		}
 		else if (StatusCounts[EPackageStatus::DeclaredUnmodified_FoundModified_OnIgnoreList] > 0)
 		{
@@ -1083,26 +1145,56 @@ FString FIterativeValidatePackageWriter::GetIterativeValidatePath() const
 
 FIterativeValidatePackageWriter::EPackageStatus FIterativeValidatePackageWriter::GetPackageStatus(FName PackageName) const
 {
-	if (const EPackageStatus* Status = PackageStatusMap.Find(PackageName))
+	if (const FPackageStatusInfo* Info = PackageStatusMap.Find(PackageName))
 	{
-		return *Status;
+		return Info->Status;
 	}
 	return EPackageStatus::NotYetProcessed;
 }
 
 void FIterativeValidatePackageWriter::SetPackageStatus(FName PackageName, EPackageStatus NewStatus)
 {
-	PackageStatusMap.FindOrAdd(PackageName) = NewStatus;
+	FPackageStatusInfo& Info = PackageStatusMap.FindOrAdd(PackageName);
+	Info.Status = NewStatus;
+	switch (NewStatus)
+	{
+	case EPackageStatus::DeclaredUnmodified_ConfirmedUnmodified: [[fallthrough]];
+	case EPackageStatus::DeclaredModified_WillNotVerify:
+		// For the non-error cases, clear the AssetClass so that we don't waste bandwidth or diskspace to store it.
+		Info.AssetClass = FTopLevelAssetPath();
+		break;
+	default:
+		break;
+	}
 }
 
 FIterativeValidatePackageWriter::FStatusCounts FIterativeValidatePackageWriter::CountPackagesByStatus()
 {
 	FIterativeValidatePackageWriter::FStatusCounts StatusCounts;
 
-	for (TPair<FName, EPackageStatus>& PackageStatusPair : PackageStatusMap)
+	for (const TPair<FName,FPackageStatusInfo>& Pair : PackageStatusMap)
 	{
-		StatusCounts[PackageStatusPair.Value]++;
+		StatusCounts[Pair.Value.Status]++;
 	}
 
 	return StatusCounts;
 }
+
+TMap<FTopLevelAssetPath, int32> FIterativeValidatePackageWriter::GetSummaryFalsePositiveCounts()
+{
+	TMap<FTopLevelAssetPath, int32> Result;
+	for (const TPair<FName, FPackageStatusInfo>& Pair : PackageStatusMap)
+	{
+		if (Pair.Value.Status == EPackageStatus::DeclaredUnmodified_FoundModified_FalsePositive
+			&& Pair.Value.AssetClass.IsValid())
+		{
+			++(Result.FindOrAdd(Pair.Value.AssetClass, 0));
+		}
+	}
+
+	Result.ValueSort([](const int32& A, const int32& B)
+		{
+			return A > B;
+		});
+	return Result;
+}
diff --git a/Engine/Source/Editor/UnrealEd/Private/Cooker/IterativeValidatePackageWriter.h b/Engine/Source/Editor/UnrealEd/Private/Cooker/IterativeValidatePackageWriter.h
--- a/Engine/Source/Editor/UnrealEd/Private/Cooker/IterativeValidatePackageWriter.h
+++ b/Engine/Source/Editor/UnrealEd/Private/Cooker/IterativeValidatePackageWriter.h
@@ -97,6 +97,12 @@ protected:
 		TStaticArray<uint32, (uint32)EPackageStatus::Count> Data;
 	};
 
+	struct FPackageStatusInfo
+	{
+		FTopLevelAssetPath AssetClass;
+		EPackageStatus Status;
+	};
+
 protected:
 	virtual void OnDiffWriterMessage(ELogVerbosity::Type Verbosity, FStringView Message) override;
 	void LogIterativeDifferences();
@@ -107,9 +113,10 @@ protected:
 	EPackageStatus GetPackageStatus(FName PackageName) const;
 	void SetPackageStatus(FName PackageName, EPackageStatus NewStatus);
 	FStatusCounts CountPackagesByStatus();
+	TMap<FTopLevelAssetPath, int32> GetSummaryFalsePositiveCounts();
 
 protected:
-	TMap<FName, EPackageStatus> PackageStatusMap;
+	TMap<FName, FPackageStatusInfo> PackageStatusMap;
 	TMap<FName, TArray<FMessage>> PackageMessageMap;
 	TSet<FName> PackageIgnoreList;
 	FString MetadataPath;
@@ -124,6 +131,9 @@ protected:
 	friend bool LoadFromCompactBinary(FCbFieldView Field, FMessage& Path);
 	friend FCbWriter& operator<<(FCbWriter& Writer, EPackageStatus Status);
 	friend bool LoadFromCompactBinary(FCbFieldView Field, EPackageStatus& Status);
+	friend FArchive& operator<<(FArchive& Writer, FPackageStatusInfo& Info);
+	friend FCbWriter& operator<<(FCbWriter& Writer, const FPackageStatusInfo& Info);
+	friend bool LoadFromCompactBinary(FCbFieldView Field, FPackageStatusInfo& Info);
 };
 
 
