[NavigationSystem] no longer add AI/Navigation system when creating inactive world with physics

diff --git a/Engine/Source/Editor/UnrealEd/Private/EditorEngine.cpp b/Engine/Source/Editor/UnrealEd/Private/EditorEngine.cpp
--- a/Engine/Source/Editor/UnrealEd/Private/EditorEngine.cpp
+++ b/Engine/Source/Editor/UnrealEd/Private/EditorEngine.cpp
@@ -4715,7 +4715,12 @@ bool UEditorEngine::InitializePhysicsSceneForSaveIfNecessary(UWorld* World, bool
 		else
 		{
 			// If we aren't already initialized, initialize now and create a physics scene. Don't create an FX system because it uses too much video memory for bulk operations
-			World->InitWorld(GetEditorWorldInitializationValues().CreateFXSystem(false).CreatePhysicsScene(true));
+			World->InitWorld(
+				GetEditorWorldInitializationValues()
+				.CreateFXSystem(false)
+				.CreateAISystem(false)
+				.CreateNavigation(false)
+				.CreatePhysicsScene(true));
 			bOutForceInitialized = true;
 		}
 
