Cooker: Change CookByTheBookStarted and CookByTheBookFinished to CookStarted and CookFinished, and call them for all cookmodes. This provides necessary hooks for licensees who need to modify CookOnTheFly.
[RN] Minor, Cooking

diff --git a/Engine/Plugins/Experimental/EditorTelemetry/Source/Private/EditorTelemetry.cpp b/Engine/Plugins/Experimental/EditorTelemetry/Source/Private/EditorTelemetry.cpp
--- a/Engine/Plugins/Experimental/EditorTelemetry/Source/Private/EditorTelemetry.cpp
+++ b/Engine/Plugins/Experimental/EditorTelemetry/Source/Private/EditorTelemetry.cpp
@@ -663,8 +663,12 @@ void FEditorTelemetry::StartSession()
 			
 			if (GUnrealEd!=nullptr && GUnrealEd->CookServer != nullptr )
 			{
-				UE::Cook::FDelegates::CookByTheBookFinished.AddLambda([this](UE::Cook::ICookInfo& CookInfo)
+				UE::Cook::FDelegates::CookFinished.AddLambda([this](UE::Cook::ICookInfo& CookInfo)
 					{
+						if (CookInfo.GetCookType() != UE::Cook::ECookType::ByTheBook)
+						{
+							return;
+						}
 						TArray<FAnalyticsEventAttribute> Attributes;
 						Attributes.Emplace(TEXT("MapName"), EditorMapName);
 						
@@ -801,8 +805,12 @@ void FEditorTelemetry::StartSession()
 
 	
 	// Install Cooking Callbacks
-	UE::Cook::FDelegates::CookByTheBookStarted.AddLambda([this](UE::Cook::ICookInfo& CookInfo)
+	UE::Cook::FDelegates::CookStarted.AddLambda([this](UE::Cook::ICookInfo& CookInfo)
 	{
+		if (CookInfo.GetCookType() != UE::Cook::ECookType::ByTheBook)
+		{
+			return;
+		}
 		// Begin the cooking span	
 		CookingSpan = FStudioTelemetry::Get().StartSpan(TEXT("Cooking"));
 
@@ -812,8 +820,12 @@ void FEditorTelemetry::StartSession()
 		CookingSpan->AddAttributes(Attributes);
 	});
 
-	UE::Cook::FDelegates::CookByTheBookFinished.AddLambda([this](UE::Cook::ICookInfo& CookInfo)
+	UE::Cook::FDelegates::CookFinished.AddLambda([this](UE::Cook::ICookInfo& CookInfo)
 	{
+		if (CookInfo.GetCookType() != UE::Cook::ECookType::ByTheBook)
+		{
+			return;
+		}
 		// End the cooking span
 	
 		// Suppress sending telemetry from CookWorkers for now.
diff --git a/Engine/Source/Editor/UnrealEd/Classes/CookOnTheSide/CookOnTheFlyServer.h b/Engine/Source/Editor/UnrealEd/Classes/CookOnTheSide/CookOnTheFlyServer.h
--- a/Engine/Source/Editor/UnrealEd/Classes/CookOnTheSide/CookOnTheFlyServer.h
+++ b/Engine/Source/Editor/UnrealEd/Classes/CookOnTheSide/CookOnTheFlyServer.h
@@ -867,11 +867,11 @@ public:
 	UNREALED_API void OnObjectSaved( UObject *ObjectSaved, FObjectPreSaveContext SaveContext );
 
 	DECLARE_MULTICAST_DELEGATE(FOnCookByTheBookStarted);
-	UE_DEPRECATED(5.4, "Use UE::Cook::FDelegates::CookByTheBookStarted (CoreUObject/Public/UObject/ICookInfo.h.")
+	UE_DEPRECATED(5.4, "Use UE::Cook::FDelegates::CookStarted, possibly restricting to the case CookInfo.GetCookType() == ECookType::ByTheBook (see CoreUObject/Public/UObject/ICookInfo.h).")
 	static FOnCookByTheBookStarted& OnCookByTheBookStarted() { return CookByTheBookStartedEvent; };
 
 	DECLARE_MULTICAST_DELEGATE(FOnCookByTheBookFinished);
-	UE_DEPRECATED(5.4, "Use UE::Cook::FDelegates::CookByTheBookFinished (CoreUObject/Public/UObject/ICookInfo.h.")
+	UE_DEPRECATED(5.4, "Use UE::Cook::FDelegates::CookFinished, possibly restricting to the case CookInfo.GetCookType() == ECookType::ByTheBook (see CoreUObject/Public/UObject/ICookInfo.h).")
 	static FOnCookByTheBookFinished& OnCookByTheBookFinished() { return CookByTheBookFinishedEvent; };
 	/**
 	* Marks a package as dirty for cook
@@ -1456,8 +1456,8 @@ private:
 	void OnObjectHandleReadDebug(const TArrayView<const UObject*const>& ReadObjects);
 	/** Send warnings/telemetry when a discovered or read package is found to be a hidden dependency. */
 	void ReportHiddenDependency(FName Referencer, FName Dependency);
-	void BroadcastCookByTheBookStarted();
-	void BroadcastCookByTheBookFinished();
+	void BroadcastCookStarted();
+	void BroadcastCookFinished();
 
 	static UCookOnTheFlyServer* ActiveCOTFS;
 	uint32		StatLoadedPackageCount = 0;
diff --git a/Engine/Source/Editor/UnrealEd/Private/CookOnTheFlyServer.cpp b/Engine/Source/Editor/UnrealEd/Private/CookOnTheFlyServer.cpp
--- a/Engine/Source/Editor/UnrealEd/Private/CookOnTheFlyServer.cpp
+++ b/Engine/Source/Editor/UnrealEd/Private/CookOnTheFlyServer.cpp
@@ -10702,14 +10702,18 @@ void UCookOnTheFlyServer::ShutdownCookSession()
 		OutputHierarchyTimers();
 		PrintDetailedCookStats();
 
-		// BroadcastCookByTheBookFinished needs to be called before clearing the session data, so that subscribers
+		// BroadcastCookFinished needs to be called before clearing the session data, so that subscribers
 		// can access information about the session such as DLCName.
-		BroadcastCookByTheBookFinished();
+		BroadcastCookFinished();
+	}
+	else if (IsCookOnTheFlyMode())
+	{
+		BroadcastCookFinished();
 	}
 	else if (CookWorkerClient)
 	{
 		CookAsCookWorkerFinished();
-		// CookAsCookWorkerFinished is responsible for calling BroadcastCookByTheBookFinished.
+		// CookAsCookWorkerFinished is responsible for calling BroadcastCookFinished.
 	}
 
 	CookByTheBookOptions->ClearSessionData();
@@ -11714,7 +11718,7 @@ void UCookOnTheFlyServer::StartCookByTheBook( const FCookByTheBookStartupOptions
 	RecordDLCPackagesFromBaseGame(BeginContext);
 	RegisterCookByTheBookDelegates();
 
-	BroadcastCookByTheBookStarted();
+	BroadcastCookStarted();
 }
 
 const UCookOnTheFlyServer::FCookByTheBookStartupOptions& UCookOnTheFlyServer::BlockOnPrebootCookGate(bool& bOutAbortCook,
@@ -11891,8 +11895,8 @@ void UCookOnTheFlyServer::StartCookAsCookWorker()
 	{
 		RegisterCookByTheBookDelegates();
 		BeginCookFinishShaderCodeLibrary(BeginContext);
-		BroadcastCookByTheBookStarted();
 	}
+	BroadcastCookStarted();
 }
 
 void UCookOnTheFlyServer::LogCookWorkerStats()
@@ -11933,10 +11937,7 @@ void UCookOnTheFlyServer::CookAsCookWorkerFinished()
 		GShaderCompilingManager->SetAllowForIncompleteShaderMaps(false);
 	}
 	LogCookWorkerStats();
-	if (IsDirectorCookByTheBook())
-	{
-		BroadcastCookByTheBookFinished();
-	}
+	BroadcastCookFinished();
 	CookWorkerClient->FlushLogs();
 }
 
@@ -13718,18 +13719,22 @@ void ConditionalWaitOnCommandFile(FStringView GateName, TFunctionRef<void (FStri
 	CommandHandler(CommandContents);
 }
 
-void UCookOnTheFlyServer::BroadcastCookByTheBookStarted()
+void UCookOnTheFlyServer::BroadcastCookStarted()
 {
-	PRAGMA_DISABLE_DEPRECATION_WARNINGS;
-	CookByTheBookStartedEvent.Broadcast();
-	PRAGMA_ENABLE_DEPRECATION_WARNINGS;
-	UE::Cook::FDelegates::CookByTheBookStarted.Broadcast(*this);
+	if (IsDirectorCookByTheBook())
+	{
+		PRAGMA_DISABLE_DEPRECATION_WARNINGS;
+		CookByTheBookStartedEvent.Broadcast();
+		UE::Cook::FDelegates::CookByTheBookStarted.Broadcast(*this);
+		PRAGMA_ENABLE_DEPRECATION_WARNINGS;
+	}
+	UE::Cook::FDelegates::CookStarted.Broadcast(*this);
 #if ENABLE_LOW_LEVEL_MEM_TRACKER
 	FLowLevelMemTracker::Get().UpdateStatsPerFrame();
 #endif
 
 	// Register collectors used internally by CookOnTheFlyServer.
-	// External systems would do this during CookByTheBookStarted.Broadcast
+	// External systems would do this during CookStarted.Broadcast
 	if (GetProcessType() != UE::Cook::EProcessType::SingleProcess)
 	{
 #if UE_WITH_CONFIG_TRACKING
@@ -13739,7 +13744,7 @@ void UCookOnTheFlyServer::BroadcastCookByTheBookStarted()
 	}
 }
 
-void UCookOnTheFlyServer::BroadcastCookByTheBookFinished()
+void UCookOnTheFlyServer::BroadcastCookFinished()
 {
 	// Unregister collectors used internally by CookOnTheFlyServer.
 	if (GetProcessType() != UE::Cook::EProcessType::SingleProcess)
@@ -13750,10 +13755,14 @@ void UCookOnTheFlyServer::BroadcastCookByTheBookFinished()
 #endif
 	}
 
-	PRAGMA_DISABLE_DEPRECATION_WARNINGS;
-	CookByTheBookFinishedEvent.Broadcast();
-	PRAGMA_ENABLE_DEPRECATION_WARNINGS;
-	UE::Cook::FDelegates::CookByTheBookFinished.Broadcast(*this);
+	if (IsDirectorCookByTheBook())
+	{
+		PRAGMA_DISABLE_DEPRECATION_WARNINGS;
+		CookByTheBookFinishedEvent.Broadcast();
+		UE::Cook::FDelegates::CookByTheBookFinished.Broadcast(*this);
+		PRAGMA_ENABLE_DEPRECATION_WARNINGS;
+	}
+	UE::Cook::FDelegates::CookFinished.Broadcast(*this);
 }
 
 #undef LOCTEXT_NAMESPACE
diff --git a/Engine/Source/Runtime/CoreUObject/Private/UObject/ICookInfo.cpp b/Engine/Source/Runtime/CoreUObject/Private/UObject/ICookInfo.cpp
--- a/Engine/Source/Runtime/CoreUObject/Private/UObject/ICookInfo.cpp
+++ b/Engine/Source/Runtime/CoreUObject/Private/UObject/ICookInfo.cpp
@@ -45,8 +45,12 @@ FString FInstigator::ToString() const
 	return FString(Result);
 }
 
+PRAGMA_DISABLE_DEPRECATION_WARNINGS;
 FCookInfoEvent FDelegates::CookByTheBookStarted;
 FCookInfoEvent FDelegates::CookByTheBookFinished;
+PRAGMA_ENABLE_DEPRECATION_WARNINGS;
+FCookInfoEvent FDelegates::CookStarted;
+FCookInfoEvent FDelegates::CookFinished;
 FValidateSourcePackage FDelegates::ValidateSourcePackage;
 
 }
diff --git a/Engine/Source/Runtime/CoreUObject/Public/UObject/ICookInfo.h b/Engine/Source/Runtime/CoreUObject/Public/UObject/ICookInfo.h
--- a/Engine/Source/Runtime/CoreUObject/Public/UObject/ICookInfo.h
+++ b/Engine/Source/Runtime/CoreUObject/Public/UObject/ICookInfo.h
@@ -148,8 +148,12 @@ DECLARE_DELEGATE_RetVal_TwoParams(EDataValidationResult, FValidateSourcePackage,
 struct FDelegates
 {
 public:
+	UE_DEPRECATED(5.5, "Use CookStarted, possibly restricting to the case CookInfo.GetCookType() == ECookType::ByTheBook.")
 	static COREUOBJECT_API FCookInfoEvent CookByTheBookStarted;
+	UE_DEPRECATED(5.5, "Use CookFinished, possibly restricting to the case CookInfo.GetCookType() == ECookType::ByTheBook.")
 	static COREUOBJECT_API FCookInfoEvent CookByTheBookFinished;
+	static COREUOBJECT_API FCookInfoEvent CookStarted;
+	static COREUOBJECT_API FCookInfoEvent CookFinished;
 	static COREUOBJECT_API FValidateSourcePackage ValidateSourcePackage;
 };
 
