Prevent null dereference of standalone host in FAssetEditorToolkit.
#rnx

diff --git a/Engine/Source/Editor/UnrealEd/Private/Toolkits/AssetEditorToolkit.cpp b/Engine/Source/Editor/UnrealEd/Private/Toolkits/AssetEditorToolkit.cpp
--- a/Engine/Source/Editor/UnrealEd/Private/Toolkits/AssetEditorToolkit.cpp
+++ b/Engine/Source/Editor/UnrealEd/Private/Toolkits/AssetEditorToolkit.cpp
@@ -314,7 +314,10 @@ void FAssetEditorToolkit::InitAssetEditor( const EToolkitMode::Type Mode, const
 		TSharedRef<FTabManager::FLayout> LayoutToUse = FLayoutSaveRestore::LoadFromConfig(GEditorLayoutIni, StandaloneDefaultLayout);
 
 		// Actually create the widget content
-		NewStandaloneHost->SetupInitialContent(LayoutToUse, NewMajorTab, bCreateDefaultStandaloneMenu);
+		if (NewStandaloneHost)
+		{
+			NewStandaloneHost->SetupInitialContent(LayoutToUse, NewMajorTab, bCreateDefaultStandaloneMenu);
+		}	
 	}
 	
 	// Create toolbars
@@ -1546,12 +1549,18 @@ bool FAssetEditorToolkit::IsActuallyAnAsset() const
 
 void FAssetEditorToolkit::AddMenuExtender(TSharedPtr<FExtender> Extender)
 {
-	StandaloneHost.Pin()->GetMenuExtenders().AddUnique(Extender);
+	if (TSharedPtr<SStandaloneAssetEditorToolkitHost> StandaloneHostPtr = StandaloneHost.Pin())
+	{
+		StandaloneHostPtr->GetMenuExtenders().AddUnique(Extender);
+	}
 }
 
 void FAssetEditorToolkit::RemoveMenuExtender(TSharedPtr<FExtender> Extender)
 {
-	StandaloneHost.Pin()->GetMenuExtenders().Remove(Extender);
+	if (TSharedPtr<SStandaloneAssetEditorToolkitHost> StandaloneHostPtr = StandaloneHost.Pin())
+	{
+		StandaloneHostPtr->GetMenuExtenders().Remove(Extender);
+	}
 }
 
 void FAssetEditorToolkit::AddToolbarExtender(TSharedPtr<FExtender> Extender)
@@ -1584,7 +1593,10 @@ TSharedPtr<FExtensibilityManager> FAssetEditorToolkit::GetSharedToolBarExtensibi
 
 void FAssetEditorToolkit::SetMenuOverlay( TSharedRef<SWidget> Widget )
 {
-	StandaloneHost.Pin()->SetMenuOverlay( Widget );
+	if (TSharedPtr<SStandaloneAssetEditorToolkitHost> StandaloneHostPtr = StandaloneHost.Pin())
+	{
+		StandaloneHostPtr->SetMenuOverlay(Widget);
+	}
 }
 
 void FAssetEditorToolkit::AddToolbarWidget(TSharedRef<SWidget> Widget)
