Fix Insights status bar and editor viewport when screenshot trace is disabled.

diff --git a/Engine/Plugins/TraceUtilities/Source/EditorTraceUtilities/Private/InsightsStatusBar/SInsightsStatusBar.cpp b/Engine/Plugins/TraceUtilities/Source/EditorTraceUtilities/Private/InsightsStatusBar/SInsightsStatusBar.cpp
--- a/Engine/Plugins/TraceUtilities/Source/EditorTraceUtilities/Private/InsightsStatusBar/SInsightsStatusBar.cpp
+++ b/Engine/Plugins/TraceUtilities/Source/EditorTraceUtilities/Private/InsightsStatusBar/SInsightsStatusBar.cpp
@@ -1006,7 +1006,9 @@ bool SInsightsStatusBarWidget::TraceScreenshot_CanExecute()
 
 void SInsightsStatusBarWidget::TraceScreenshot_Execute()
 {
+#if UE_SCREENSHOT_TRACE_ENABLED
 	FTraceScreenshot::RequestScreenshot(TEXT(""), false);
+#endif
 }
 
 bool SInsightsStatusBarWidget::TraceBookmark_CanExecute()
diff --git a/Engine/Source/Editor/UnrealEd/Private/EditorViewportClient.cpp b/Engine/Source/Editor/UnrealEd/Private/EditorViewportClient.cpp
--- a/Engine/Source/Editor/UnrealEd/Private/EditorViewportClient.cpp
+++ b/Engine/Source/Editor/UnrealEd/Private/EditorViewportClient.cpp
@@ -6134,11 +6134,13 @@ bool RequestSaveScreenshot(bool bWriteAlpha, TArray<FColorType>& Bitmap, FIntPoi
 	FHighResScreenshotConfig& HighResScreenshotConfig = GetHighResScreenshotConfig();
 	bool bIsScreenshotSaved = false;
 	bool bSuppressWritingToFile = false;
+#if UE_SCREENSHOT_TRACE_ENABLED
 	if (SHOULD_TRACE_SCREENSHOT())
 	{
 		bSuppressWritingToFile = FTraceScreenshot::ShouldSuppressWritingToFile();
 		FTraceScreenshot::TraceScreenshot(BitmapSize.X, BitmapSize.Y, Bitmap, FScreenshotRequest::GetFilename());
 	}
+#endif
 
 	if (!bSuppressWritingToFile)
 	{
@@ -6259,7 +6261,9 @@ bool FEditorViewportClient::ProcessScreenShots(FViewport* InViewport)
 
 		// Done with the request
 		FScreenshotRequest::Reset();
+#if UE_SCREENSHOT_TRACE_ENABLED
 		FTraceScreenshot::Reset();
+#endif
 		FScreenshotRequest::OnScreenshotRequestProcessed().Broadcast();
 
 		// Re-enable screen messages - if we are NOT capturing a movie
