Highlight the realtime toggle button of the new viewport toolbar in yellow when the viewport is not realtime

diff --git a/Engine/Source/Editor/EditorStyle/Private/StarshipStyle.cpp b/Engine/Source/Editor/EditorStyle/Private/StarshipStyle.cpp
--- a/Engine/Source/Editor/EditorStyle/Private/StarshipStyle.cpp
+++ b/Engine/Source/Editor/EditorStyle/Private/StarshipStyle.cpp
@@ -4327,6 +4327,8 @@ void FStarshipEditorStyle::FStyle::SetupLevelEditorStyle()
 		Set( "EditorViewport.TranslateRotate2DMode", new IMAGE_BRUSH_SVG("Starship/EditorViewport/TranslateRotate2D", Icon16x16 ) );
 
 		Set( "EditorViewport.ToggleRealTime", new IMAGE_BRUSH_SVG("Starship/Common/Realtime", Icon16x16));
+		Set("EditorViewport.ToggleRealTimeWarning",
+			new IMAGE_BRUSH_SVG("Starship/Common/Realtime", Icon16x16, FStyleColors::Warning));
 
 		Set( "EditorViewport.LocationGridSnap", new IMAGE_BRUSH_SVG("Starship/EditorViewport/grid", Icon16x16));
 		Set( "EditorViewport.RotationGridSnap", new IMAGE_BRUSH_SVG("Starship/EditorViewport/angle", Icon16x16));
diff --git a/Engine/Source/Editor/UnrealEd/Private/ViewportToolbar/UnrealEdViewportToolbar.cpp b/Engine/Source/Editor/UnrealEd/Private/ViewportToolbar/UnrealEdViewportToolbar.cpp
--- a/Engine/Source/Editor/UnrealEd/Private/ViewportToolbar/UnrealEdViewportToolbar.cpp
+++ b/Engine/Source/Editor/UnrealEd/Private/ViewportToolbar/UnrealEdViewportToolbar.cpp
@@ -2706,12 +2706,18 @@ FToolMenuEntry CreateToggleRealtimeEntry()
 					EditorViewportWeak = EditorViewportContext->Viewport;
 
 					RealtimeToggleAction.ExecuteAction = FToolMenuExecuteAction::CreateLambda(
-					[EditorViewportWeak](const FToolMenuContext& Context) -> void
+						[EditorViewportWeak](const FToolMenuContext& Context) -> void
 						{
-						if (TSharedPtr<SEditorViewport> EditorViewport = EditorViewportWeak.Pin())
-						{
-							EditorViewport->OnToggleRealtime();
-						}
+							if (TSharedPtr<SEditorViewport> EditorViewport = EditorViewportWeak.Pin())
+							{
+								EditorViewport->OnToggleRealtime();
+								// Calling UToolMenu::RefreshAllWidgets here is cheating. We do it because the menu
+								// entry's TAttribute<FSlateIcon> is only called once when the menu is opened (because
+								// FBaseMenuBuilder::AddMenuEntry takes an FSlateIcon and not a TAttribute<FSlateIcon>).
+								// So when we refresh all widgets here, we force the open menu to close and hide the
+								// fact that the icon wouldn't have updated if the menu stayed open.
+								UToolMenus::Get()->RefreshAllWidgets();
+							}
 						}
 					);
 
@@ -2760,11 +2766,26 @@ FToolMenuEntry CreateToggleRealtimeEntry()
 					}
 				}
 
+				const TAttribute<FSlateIcon> Icon = TAttribute<FSlateIcon>::CreateLambda(
+					[EditorViewportWeak]() -> FSlateIcon
+					{
+						bool bIsViewportRealtime = true;
+						if (const TSharedPtr<SEditorViewport> EditorViewport = EditorViewportWeak.Pin())
+						{
+							bIsViewportRealtime = EditorViewport->IsRealtime();
+						}
+
+						return bIsViewportRealtime
+								 ? FSlateIcon(FAppStyle::GetAppStyleSetName(), "EditorViewport.ToggleRealTime")
+								 : FSlateIcon(FAppStyle::GetAppStyleSetName(), "EditorViewport.ToggleRealTimeWarning");
+					}
+				);
+
 				FToolMenuEntry ToggleRealtime = FToolMenuEntry::InitMenuEntry(
 					"ToggleRealtime",
 					LOCTEXT("ToggleRealtimeLabel", "Realtime Viewport"),
 					Tooltip,
-					FSlateIcon(FAppStyle::GetAppStyleSetName(), "EditorViewport.ToggleRealTime"),
+					Icon,
 					RealtimeToggleAction,
 					EUserInterfaceActionType::ToggleButton
 				);
