Reinitialize timecode and custom timestep after blueprint compile. This should allow modifying timecode providers without having to manually set them again after compiling them.

LiveLinkTimecodeProvider - Call SetSubject when initializing in case the provider was reinitialized.

diff --git a/Engine/Plugins/Animation/LiveLink/Source/LiveLink/Private/LiveLinkTimecodeProvider.cpp b/Engine/Plugins/Animation/LiveLink/Source/LiveLink/Private/LiveLinkTimecodeProvider.cpp
--- a/Engine/Plugins/Animation/LiveLink/Source/LiveLink/Private/LiveLinkTimecodeProvider.cpp
+++ b/Engine/Plugins/Animation/LiveLink/Source/LiveLink/Private/LiveLinkTimecodeProvider.cpp
@@ -111,6 +111,9 @@ bool ULiveLinkTimecodeProvider::Initialize(class UEngine* InEngine)
 
 	InitClient();
 
+	UnregisterSubject();
+	RegisterSubject();
+	
 	return true;
 }
 
diff --git a/Engine/Source/Editor/UnrealEd/Classes/Editor/EditorEngine.h b/Engine/Source/Editor/UnrealEd/Classes/Editor/EditorEngine.h
--- a/Engine/Source/Editor/UnrealEd/Classes/Editor/EditorEngine.h
+++ b/Engine/Source/Editor/UnrealEd/Classes/Editor/EditorEngine.h
@@ -3073,6 +3073,18 @@ private:
 	/** Handler for when an asset finishes compiling (used to notify AssetRegistry to update tags after Load) */
 	UNREALED_API void OnAssetPostCompile(const TArray<FAssetCompileData>& CompiledAssets);
 
+	/** Handler for when the timecode provider changed, so that we can register the BlueprintCompiled delegate to it.*/
+	UNREALED_API void RegisterTimecodeProviderCompiledDelegate();
+
+	/** Handler for when the custom timestep changed, so that we can register the BlueprintCompiled delegate to it.*/
+	UNREALED_API void RegisterCustomTimeStepCompiledDelegate();
+
+	/** Handler used to reinitialize Timecode when their blueprint is compiled. */
+	UNREALED_API void HandleTimecodeProviderCompiled(UBlueprint* InBlueprint);
+
+	/** Handler used to reinitialize CustomTimeStep when their blueprint is compiled. */
+	UNREALED_API void HandleCustomTimeStepCompiled(UBlueprint* InBlueprint);
+
 	/** Handler for when a world is duplicated in the editor */
 	UNREALED_API void InitializeNewlyCreatedInactiveWorld(UWorld* World);
 
@@ -3149,6 +3161,12 @@ private:
 	/** Delegate handle for game viewport close requests in PIE sessions. */
 	FDelegateHandle ViewportCloseRequestedDelegateHandle;
 
+	/** Delegate handle for when a TimecodeProvider blueprint is compiled. */
+	FDelegateHandle TimecodeProviderCompiledDelegateHandle;
+
+	/** Delegate handle for when a CustomTimeStep blueprint is compiled. */
+	FDelegateHandle CustomTimeStepCompiledDelegateHandle;
+
 	/** Minimized Windows during PIE */
 	TArray<TWeakPtr<SWindow>> MinimizedWindowsDuringPIE;
 
diff --git a/Engine/Source/Editor/UnrealEd/Private/EditorEngine.cpp b/Engine/Source/Editor/UnrealEd/Private/EditorEngine.cpp
--- a/Engine/Source/Editor/UnrealEd/Private/EditorEngine.cpp
+++ b/Engine/Source/Editor/UnrealEd/Private/EditorEngine.cpp
@@ -271,6 +271,9 @@
 #include "IDocumentation.h"
 #include "StereoRenderTargetManager.h"
 
+#include "Engine/EngineCustomTimeStep.h"
+#include "Engine/TimecodeProvider.h"
+
 DEFINE_LOG_CATEGORY_STATIC(LogEditor, Log, All);
 
 #define LOCTEXT_NAMESPACE "UnrealEd.Editor"
@@ -1160,8 +1163,14 @@ void UEditorEngine::InitEditor(IEngineLoop* InEngineLoop)
 	}
 
 	FAssetCompilingManager::Get().OnAssetPostCompileEvent().AddUObject(this, &UEditorEngine::OnAssetPostCompile);
-
+	
 	WorldAddExtraDeletionObjectsHandle = FEditorDelegates::OnAddExtraObjectsToDelete.AddStatic(&UWorld::OnAddExtraObjectsToDelete);
+
+	OnTimecodeProviderChanged().AddUObject(this, &UEditorEngine::RegisterTimecodeProviderCompiledDelegate);
+	OnCustomTimeStepChanged().AddUObject(this, &UEditorEngine::RegisterCustomTimeStepCompiledDelegate);
+
+	RegisterTimecodeProviderCompiledDelegate();
+	RegisterCustomTimeStepCompiledDelegate();
 }
 
 bool UEditorEngine::HandleOpenAsset(UObject* Asset)
@@ -1614,6 +1623,10 @@ void UEditorEngine::FinishDestroy()
 				AssetRegistry->OnInMemoryAssetCreated().RemoveAll(this);
 			}
 		}
+
+		OnTimecodeProviderChanged().RemoveAll(this);
+		OnCustomTimeStepChanged().RemoveAll(this);
+
 		FAssetCompilingManager::Get().OnAssetPostCompileEvent().RemoveAll(this);
 
 		FEditorDelegates::OnAddExtraObjectsToDelete.Remove(WorldAddExtraDeletionObjectsHandle);
@@ -7603,6 +7616,40 @@ void UEditorEngine::OnAssetPostCompile(const TArray<FAssetCompileData>& Compiled
 	}
 }
 
+void UEditorEngine::HandleTimecodeProviderCompiled(UBlueprint* InBlueprint)
+{
+	ReinitializeTimecodeProvider();
+}
+
+void UEditorEngine::HandleCustomTimeStepCompiled(UBlueprint* InBlueprint)
+{
+	ReinitializeCustomTimeStep();
+}
+
+void UEditorEngine::RegisterTimecodeProviderCompiledDelegate()
+{
+	if (UTimecodeProvider* NewProvider = GetTimecodeProvider())
+	{
+		if (UBlueprint* Blueprint = UBlueprint::GetBlueprintFromClass(NewProvider->GetClass()))
+		{
+			Blueprint->OnCompiled().Remove(TimecodeProviderCompiledDelegateHandle);
+			TimecodeProviderCompiledDelegateHandle = Blueprint->OnCompiled().AddUObject(this, &UEditorEngine::HandleTimecodeProviderCompiled);
+		}
+	}
+}
+
+void UEditorEngine::RegisterCustomTimeStepCompiledDelegate()
+{
+	if (UEngineCustomTimeStep* NewCustomTimeStep = GetCustomTimeStep())
+	{
+		if (UBlueprint* Blueprint = UBlueprint::GetBlueprintFromClass(NewCustomTimeStep->GetClass()))
+		{
+			Blueprint->OnCompiled().Remove(CustomTimeStepCompiledDelegateHandle);
+			CustomTimeStepCompiledDelegateHandle = Blueprint->OnCompiled().AddUObject(this, &UEditorEngine::HandleCustomTimeStepCompiled);
+		}
+	}
+}
+
 void UEditorEngine::InitializeNewlyCreatedInactiveWorld(UWorld* World)
 {
 	check(World);
