Use the same logic to show/hide the RealtimeToggle in the new viewport toolbar as in the old

* Move and rename UE::LevelEditor::ShowViewportRealtimeWarning to UE::UnrealEd::ShouldShowViewportRealtimeWarning.
* Use UE::UnrealEd::ShouldShowViewportRealtimeWarning to drive the visibility of the new viewport toolbar's RealtimeToggle.

diff --git a/Engine/Source/Editor/LevelEditor/Private/SLevelViewportToolBar.cpp b/Engine/Source/Editor/LevelEditor/Private/SLevelViewportToolBar.cpp
--- a/Engine/Source/Editor/LevelEditor/Private/SLevelViewportToolBar.cpp
+++ b/Engine/Source/Editor/LevelEditor/Private/SLevelViewportToolBar.cpp
@@ -1089,8 +1089,8 @@ FReply SLevelViewportToolBar::OnRealtimeWarningClicked()
 
 EVisibility SLevelViewportToolBar::GetRealtimeWarningVisibility() const
 {
-	FLevelEditorViewportClient& ViewportClient = Viewport.Pin()->GetLevelViewportClient();
-	const bool bWarn = UE::LevelEditor::ShowViewportRealtimeWarning(ViewportClient);
+	const FEditorViewportClient& ViewportClient = Viewport.Pin()->GetLevelViewportClient();
+	const bool bWarn = UE::UnrealEd::ShouldShowViewportRealtimeWarning(ViewportClient);
 	return bWarn ? EVisibility::Visible : EVisibility::Collapsed;
 }
 
diff --git a/Engine/Source/Editor/LevelEditor/Private/ViewportToolbar/LevelEditorViewportToolbarSections.cpp b/Engine/Source/Editor/LevelEditor/Private/ViewportToolbar/LevelEditorViewportToolbarSections.cpp
--- a/Engine/Source/Editor/LevelEditor/Private/ViewportToolbar/LevelEditorViewportToolbarSections.cpp
+++ b/Engine/Source/Editor/LevelEditor/Private/ViewportToolbar/LevelEditorViewportToolbarSections.cpp
@@ -462,11 +462,6 @@ FSlateIcon GetCameraSubmenuIconFromLevelViewport(const TWeakPtr<::SLevelViewport
 namespace UE::LevelEditor
 {
 
-bool ShowViewportRealtimeWarning(FLevelEditorViewportClient& ViewportClient)
-{
-	return !ViewportClient.IsRealtime() && !ViewportClient.IsRealtimeOverrideSet() && ViewportClient.IsPerspective();
-}
-
 TSharedPtr<FExtender> GetViewModesLegacyExtenders()
 {
 	FLevelEditorModule& LevelEditorModule = FModuleManager::GetModuleChecked<FLevelEditorModule>(TEXT("LevelEditor"));
diff --git a/Engine/Source/Editor/LevelEditor/Private/ViewportToolbar/LevelEditorViewportToolbarSections.h b/Engine/Source/Editor/LevelEditor/Private/ViewportToolbar/LevelEditorViewportToolbarSections.h
--- a/Engine/Source/Editor/LevelEditor/Private/ViewportToolbar/LevelEditorViewportToolbarSections.h
+++ b/Engine/Source/Editor/LevelEditor/Private/ViewportToolbar/LevelEditorViewportToolbarSections.h
@@ -25,8 +25,6 @@ FToolMenuEntry CreateCameraSpeedSlider(TWeakPtr<SLevelViewport> InLevelViewportW
 FToolMenuEntry CreateCameraSpeedScalarSlider(TWeakPtr<SLevelViewport> InLevelViewportWeak);
 FToolMenuEntry CreateLevelViewportToolbarCameraSubmenu();
 
-bool ShowViewportRealtimeWarning(FLevelEditorViewportClient& ViewportClient);
-
 TSharedPtr<FExtender> GetViewModesLegacyExtenders();
 void PopulateViewModesMenu(UToolMenu* InMenu);
 void ExtendViewModesSubmenu(FName InViewModesSubmenuName);
diff --git a/Engine/Source/Editor/UnrealEd/Private/ViewportToolbar/UnrealEdViewportToolbar.cpp b/Engine/Source/Editor/UnrealEd/Private/ViewportToolbar/UnrealEdViewportToolbar.cpp
--- a/Engine/Source/Editor/UnrealEd/Private/ViewportToolbar/UnrealEdViewportToolbar.cpp
+++ b/Engine/Source/Editor/UnrealEd/Private/ViewportToolbar/UnrealEdViewportToolbar.cpp
@@ -2629,7 +2629,12 @@ static void ConstructScreenPercentageMenu(UToolMenu* InMenu)
 	}
 }
 
-UNREALED_API FToolMenuEntry CreatePerformanceAndScalabilitySubmenu()
+bool ShouldShowViewportRealtimeWarning(const FEditorViewportClient& ViewportClient)
+{
+	return !ViewportClient.IsRealtime() && !ViewportClient.IsRealtimeOverrideSet() && ViewportClient.IsPerspective();
+}
+
+FToolMenuEntry CreatePerformanceAndScalabilitySubmenu()
 {
 	return FToolMenuEntry::InitSubMenu(
 		"PerformanceAndScalability",
@@ -2800,7 +2805,7 @@ FToolMenuEntry CreateToggleRealtimeEntry()
 						{
 							if (const TSharedPtr<SEditorViewport> EditorViewport = WeakViewport.Pin())
 							{
-								return !EditorViewport->IsRealtime();
+								return ShouldShowViewportRealtimeWarning(*EditorViewport->GetViewportClient());
 							}
 							return false;
 						}
diff --git a/Engine/Source/Editor/UnrealEd/Public/ViewportToolbar/UnrealEdViewportToolbar.h b/Engine/Source/Editor/UnrealEd/Public/ViewportToolbar/UnrealEdViewportToolbar.h
--- a/Engine/Source/Editor/UnrealEd/Public/ViewportToolbar/UnrealEdViewportToolbar.h
+++ b/Engine/Source/Editor/UnrealEd/Public/ViewportToolbar/UnrealEdViewportToolbar.h
@@ -139,6 +139,8 @@ UNREALED_API void PopulateCameraMenu(UToolMenu* InMenu);
  */
 UNREALED_API void ExtendCameraSubmenu(FName InCameraOptionsSubmenuName);
 
+UNREALED_API bool ShouldShowViewportRealtimeWarning(const FEditorViewportClient& ViewportClient);
+
 UNREALED_API FToolMenuEntry CreatePerformanceAndScalabilitySubmenu();
 
 /**
