Cooker: DiffOnly feature improvements:
Display differences found in the ExportBundles and DependencyBundles fields of zen package headers.
Add a handler for operator<< on FName when displaying the text of FNames that differ, rather than failing by attempting to deserialize the FName as a string.
#rnx

diff --git a/Engine/Source/Editor/UnrealEd/Private/Cooker/DiffWriterArchive.cpp b/Engine/Source/Editor/UnrealEd/Private/Cooker/DiffWriterArchive.cpp
--- a/Engine/Source/Editor/UnrealEd/Private/Cooker/DiffWriterArchive.cpp
+++ b/Engine/Source/Editor/UnrealEd/Private/Cooker/DiffWriterArchive.cpp
@@ -19,6 +19,7 @@
 #include "Serialization/AsyncLoading2.h"
 #include "Serialization/MemoryReader.h"
 #include "Serialization/StaticMemoryReader.h"
+#include "Templates/Function.h"
 #include "UObject/LinkerLoad.h"
 #include "UObject/ObjectMacros.h"
 #include "UObject/Package.h"
@@ -35,11 +36,66 @@ static const ANSICHAR* DebugDataStackMarker = "\r\nDebugDataStack:\r\n";
 const TCHAR* const IndentToken = TEXT("%DWA%    ");
 const TCHAR* const NewLineToken = TEXT("%DWA%\n");
 
+/**
+ * Data parsed from the header in the package. Might be stored either as a ZenPackageHeader or
+ * LinkerLoad PackageHeader, depending on the EPackageHeaderFormat of the package.
+ * Provides an interface for the data needed to interpret the bytes of the exports of the package, such as the NameMap.
+ */
+class FPackageHeaderData
+{
+public:
+	FPackageHeaderData(const TCHAR* InName, bool bInReadFromPackageStore,
+		const FString& InAssetFilename, const FPackageData& InPackageData, EPackageHeaderFormat InFormat,
+		FAccumulatorGlobals& InGlobals, FMessageCallback& InMessageCallback);
+	~FPackageHeaderData();
+	void Initialize();
+
+	const TCHAR* GetName() const;
+	const FString& GetAssetFilename() const;
+	const FPackageData& GetPackageData() const;
+	EPackageHeaderFormat GetFormat() const;
+	FAccumulatorGlobals& GetGlobals() const;
+	FMessageCallback& GetMessageCallback() const;
+
+	FDiffWriterZenHeader& GetZenHeader() const;
+	FLinkerLoad* GetLinker() const;
+
+	bool IsValid() const;
+	bool TryGetMappedName(int32 Index, int32 Number, FName& OutName) const;
+
+private:
+	const TCHAR* Name = nullptr;
+	const FString& AssetFilename;
+	const FPackageData& PackageData;
+	FAccumulatorGlobals& Globals;
+	FMessageCallback& MessageCallback;
+	EPackageHeaderFormat Format = EPackageHeaderFormat::PackageFileSummary;
+	bool bReadFromPackageStore = false;
+	bool bInitialized = false;
+
+	TUniquePtr<FDiffWriterZenHeader> ZenHeader;
+	FLinkerLoad* Linker = nullptr;
+};
+
+/**
+ * Interprets the read of an FName in the manner that it is serialized by FLinkerSave,
+ * and implements that read on inner archive which e.g. is a FMemoryArchive to the
+ * bytes of the package with the given header.
+ */
+class FPackageHeaderDataProxyArchive : public FArchiveProxy
+{
+public:
+	FPackageHeaderDataProxyArchive(FPackageHeaderData& InHeader, FArchive& InInner);
+	virtual FArchive& operator<<(FName& Value) override;
+
+private:
+	FPackageHeaderData& Header;
+};
+
 /** Logs any mismatching header data. */
-void DumpPackageHeaderDiffs(
-	FAccumulatorGlobals& Globals, const FPackageData& SourcePackage, const FPackageData& DestPackage,
-	const FString& AssetFilename, const int32 MaxDiffsToLog, const EPackageHeaderFormat PackageHeaderFormat,
-	const FMessageCallback& MessageCallback);
+void DumpPackageHeaderDiffs(FPackageHeaderData& SourceHeaderData, FPackageHeaderData& DestHeaderData,
+	const int32 MaxDiffsToLog);
+
 /** Returns a new linker for loading the specified package. */
 FLinkerLoad* CreateLinkerForPackage(
 	FUObjectSerializeContext* LoadContext,
@@ -680,16 +736,20 @@ void FAccumulator::OnSecondSaveComplete(int64 InHeaderSize)
 		check(FirstSaveLinkerData.Num() >= HeaderSize);
 		check(LinkerArchive && LinkerArchive->TotalSize() >= InHeaderSize);
 
-		FPackageData FirstSaveHeader{ FirstSaveLinkerData.GetData(), HeaderSize};
-		FPackageData SecondSaveHeader{ LinkerArchive->GetData(), InHeaderSize };
+		FPackageData FirstSaveHeaderSegment{ FirstSaveLinkerData.GetData(), HeaderSize};
+		FPackageData SecondSaveHeaderSegment{ LinkerArchive->GetData(), InHeaderSize };
 		int32 NumHeaderDiffMessages = 0;
-		DumpPackageHeaderDiffs(Globals, FirstSaveHeader, SecondSaveHeader, Filename, MaxDiffsToLog,
-			PackageHeaderFormat,
-			[&NumHeaderDiffMessages, this](ELogVerbosity::Type Verbosity, FStringView Message)
+		FMessageCallback HeaderMessageCallback = [&NumHeaderDiffMessages, this](ELogVerbosity::Type Verbosity, FStringView Message)
 			{
 				MessageCallback(Verbosity, Message);
 				++NumHeaderDiffMessages;
-			});
+			};
+		FPackageHeaderData FirstSaveHeader(TEXT("source"), false /* bReadFromPackageStore */, Filename,
+			FirstSaveHeaderSegment, PackageHeaderFormat, Globals, HeaderMessageCallback);
+		FPackageHeaderData SecondSaveHeader(TEXT("dest"), false /* bReadFromPackageStore */, Filename,
+			SecondSaveHeaderSegment, PackageHeaderFormat, Globals, HeaderMessageCallback);
+
+		DumpPackageHeaderDiffs(FirstSaveHeader, SecondSaveHeader, MaxDiffsToLog);
 		if (NumHeaderDiffMessages == 0)
 		{
 			MessageCallback(ELogVerbosity::Warning, FString::Printf(
@@ -829,6 +889,7 @@ bool ShouldDumpPropertyValueState(FProperty* Prop)
 }
 
 void FAccumulator::CompareWithPreviousForSection(const FPackageData& SourcePackage, const FPackageData& DestPackage,
+	FPackageHeaderData& SourceHeader, FPackageHeaderData& DestHeader,
 	const TCHAR* CallstackCutoffText, int32& InOutDiffsLogged, TMap<FName, FArchiveDiffStats>& OutStats,
 	const FString& SectionFilename)
 {
@@ -909,7 +970,8 @@ void FAccumulator::CompareWithPreviousForSection(const FPackageData& SourcePacka
 			return;
 		}
 		// Skip logging of header differences if requested
-		if (bIgnoreHeaderDiffs && DestAbsoluteOffset < HeaderSize)
+		bool bIsHeaderDiff = DestAbsoluteOffset < HeaderSize;
+		if (bIgnoreHeaderDiffs && bIsHeaderDiff)
 		{
 			continue;
 		}
@@ -963,7 +1025,8 @@ void FAccumulator::CompareWithPreviousForSection(const FPackageData& SourcePacka
 
 			FString BeforePropertyVal;
 			FString AfterPropertyVal;
-			if (FProperty* SerProp = DifferenceCallstackData.SerializedProp)
+			FProperty* SerProp = DifferenceCallstackData.SerializedProp;
+			if (SerProp && !bIsHeaderDiff)
 			{
 				// We don't attempt to serialize properties when we have already encountered at least one diff in the asset.
 				// That is because we don't handle length differences in the source and destination data caused by the previous
@@ -1003,9 +1066,13 @@ void FAccumulator::CompareWithPreviousForSection(const FPackageData& SourcePacka
 
 					FStaticMemoryReader SourceReader(&SourcePackage.Data[SourceAbsoluteOffset - (DestAbsoluteOffset - OffsetX)], SourcePackage.Size - SourceAbsoluteOffset);
 					FStaticMemoryReader DestReader(&DestPackage.Data[OffsetX], DestPackage.Size - DestAbsoluteOffset);
+					SourceHeader.Initialize();
+					DestHeader.Initialize();
+					FPackageHeaderDataProxyArchive SourceAr(SourceHeader, SourceReader);
+					FPackageHeaderDataProxyArchive DestAr(DestHeader, DestReader);
 
-					SourceVal.Serialize(SourceReader);
-					DestVal.Serialize(DestReader);
+					SourceVal.Serialize(SourceAr);
+					DestVal.Serialize(DestAr);
 
 					if (!SourceReader.IsError() && !DestReader.IsError())
 					{
@@ -1113,29 +1180,32 @@ void FAccumulator::CompareWithPrevious(const TCHAR* CallstackCutoffText, TMap<FN
 
 	int32 NumLoggedDiffs = 0;
 	
-	FPackageData SourcePackageHeader = SourcePackage;
-	SourcePackageHeader.Size = SourcePackage.HeaderSize;
-	SourcePackageHeader.HeaderSize = 0;
-	SourcePackageHeader.StartOffset = 0;
-
-	FPackageData DestPackageHeader = DestPackage;
-	DestPackageHeader.Size = HeaderSize;
-	DestPackageHeader.HeaderSize = 0;
-	DestPackageHeader.StartOffset = 0;
-
-	CompareWithPreviousForSection(SourcePackageHeader, DestPackageHeader, CallstackCutoffText,
-		NumLoggedDiffs, OutStats, Filename);
-
+	FPackageData SourceHeaderSegment = SourcePackage;
+	SourceHeaderSegment.Size = SourcePackage.HeaderSize;
+	SourceHeaderSegment.HeaderSize = 0;
+	SourceHeaderSegment.StartOffset = 0;
+
+	FPackageData DestHeaderSegment = DestPackage;
+	DestHeaderSegment.Size = HeaderSize;
+	DestHeaderSegment.HeaderSize = 0;
+	DestHeaderSegment.StartOffset = 0;
+
+	int32 NumHeaderDiffMessages = 0;
+	FMessageCallback HeaderMessageCallback = [&NumHeaderDiffMessages, this](ELogVerbosity::Type Verbosity, FStringView Message)
+		{
+			MessageCallback(Verbosity, Message);
+			++NumHeaderDiffMessages;
+		};
+	FPackageHeaderData SourceHeader(TEXT("source"), true /* bReadFromPackageStore */, Filename, SourcePackage,
+		PackageHeaderFormat, Globals, HeaderMessageCallback);
+	FPackageHeaderData DestHeader(TEXT("dest"), false /* bReadFromPackageStore */, Filename, DestPackage,
+		PackageHeaderFormat, Globals, HeaderMessageCallback);
+
+	CompareWithPreviousForSection(SourceHeaderSegment, DestHeaderSegment, SourceHeader, DestHeader,
+		CallstackCutoffText, NumLoggedDiffs, OutStats, Filename);
 	if (HeaderSize > 0 && OutStats.FindOrAdd(AssetClass).NumDiffs > 0)
 	{
-		int32 NumHeaderDiffMessages = 0;
-		DumpPackageHeaderDiffs(Globals, SourcePackage, DestPackage, Filename, MaxDiffsToLog,
-			PackageHeaderFormat,
-			[&NumHeaderDiffMessages, this](ELogVerbosity::Type Verbosity, FStringView Message)
-			{
-				MessageCallback(Verbosity, Message);
-				++NumHeaderDiffMessages;
-			});
+		DumpPackageHeaderDiffs(SourceHeader, DestHeader, MaxDiffsToLog);
 		if (NumHeaderDiffMessages == 0)
 		{
 			MessageCallback(ELogVerbosity::Warning, FString::Printf(
@@ -1162,8 +1232,8 @@ void FAccumulator::CompareWithPrevious(const TCHAR* CallstackCutoffText, TMap<FN
 		ExportsFilename = Filename;
 	}
 
-	CompareWithPreviousForSection(SourcePackageExports, DestPackageExports, CallstackCutoffText,
-		NumLoggedDiffs, OutStats, ExportsFilename);
+	CompareWithPreviousForSection(SourcePackageExports, DestPackageExports, SourceHeader, DestHeader,
+		CallstackCutoffText, NumLoggedDiffs, OutStats, ExportsFilename);
 
 	// Log comparison of the DeterminismHelper diagnostics, if we have any
 	if (DeterminismManager)
@@ -1210,22 +1280,30 @@ void FAccumulator::CompareWithPrevious(const TCHAR* CallstackCutoffText, TMap<FN
 
 				// Copy the original asset as '.before.uasset'.
 				{
-					TUniquePtr<FArchive> DiffUAssetArchive(FileManager.CreateFileWriter(*FPaths::SetExtension(OutputFilename, TEXT(".before.") + FPaths::GetExtension(Filename))));
-					DiffUAssetArchive->Serialize(SourcePackageHeader.Data + SourcePackageHeader.StartOffset, SourcePackageHeader.Size - SourcePackageHeader.StartOffset);
+					TUniquePtr<FArchive> DiffUAssetArchive(FileManager.CreateFileWriter(
+						*FPaths::SetExtension(OutputFilename, TEXT(".before.") + FPaths::GetExtension(Filename))));
+					DiffUAssetArchive->Serialize(SourceHeaderSegment.Data + SourceHeaderSegment.StartOffset,
+						SourceHeaderSegment.Size - SourceHeaderSegment.StartOffset);
 				}
 				{
-					TUniquePtr<FArchive> DiffUExpArchive(FileManager.CreateFileWriter(*FPaths::SetExtension(OutputFilename, TEXT(".before.uexp"))));
-					DiffUExpArchive->Serialize(SourcePackageExports.Data + SourcePackageExports.StartOffset, SourcePackageExports.Size - SourcePackageExports.StartOffset);
+					TUniquePtr<FArchive> DiffUExpArchive(FileManager.CreateFileWriter(
+						*FPaths::SetExtension(OutputFilename, TEXT(".before.uexp"))));
+					DiffUExpArchive->Serialize(SourcePackageExports.Data + SourcePackageExports.StartOffset,
+						SourcePackageExports.Size - SourcePackageExports.StartOffset);
 				}
 
 				// Save out the in-memory data as '.after.uasset'.
 				{
-					TUniquePtr<FArchive> DiffUAssetArchive(FileManager.CreateFileWriter(*FPaths::SetExtension(OutputFilename, TEXT(".after.") + FPaths::GetExtension(Filename))));
-					DiffUAssetArchive->Serialize(DestPackageHeader.Data + DestPackageHeader.StartOffset, DestPackageHeader.Size - DestPackageHeader.StartOffset);
+					TUniquePtr<FArchive> DiffUAssetArchive(FileManager.CreateFileWriter(
+						*FPaths::SetExtension(OutputFilename, TEXT(".after.") + FPaths::GetExtension(Filename))));
+					DiffUAssetArchive->Serialize(DestHeaderSegment.Data + DestHeaderSegment.StartOffset,
+						DestHeaderSegment.Size - DestHeaderSegment.StartOffset);
 				}
 				{
-					TUniquePtr<FArchive> DiffUExpArchive(FileManager.CreateFileWriter(*FPaths::SetExtension(OutputFilename, TEXT(".after.uexp"))));
-					DiffUExpArchive->Serialize(DestPackageExports.Data + DestPackageExports.StartOffset, DestPackageExports.Size - DestPackageExports.StartOffset);
+					TUniquePtr<FArchive> DiffUExpArchive(FileManager.CreateFileWriter(
+						*FPaths::SetExtension(OutputFilename, TEXT(".after.uexp"))));
+					DiffUExpArchive->Serialize(DestPackageExports.Data + DestPackageExports.StartOffset,
+						DestPackageExports.Size - DestPackageExports.StartOffset);
 				}
 			}
 			else
@@ -1507,93 +1585,123 @@ static void DumpTableDifferences(
 		*HumanReadableString));
 }
 
-void DumpPackageHeaderDiffs_LinkerLoad(
-	FAccumulatorGlobals& Globals,
-	const FPackageData& SourcePackage,
-	const FPackageData& DestPackage,
-	const FString& AssetFilename,
-	const int32 MaxDiffsToLog,
-	const FMessageCallback& MessageCallback)
+static void DumpOrderedArrayDifferences(
+	int32 SourceNum,
+	int32 DestNum,
+	TFunctionRef<bool(int32 Index)> IsElementsAtIndexEqual,
+	TFunctionRef<FString(int32 Index)> ConvertSourceIndexToText,
+	TFunctionRef<FString(int32 Index)> ConvertDestIndexToText,
+	TFunctionRef<void(ELogVerbosity::Type, FString)> LogMessage,
+	const TCHAR* AssetFilename,
+	const TCHAR* ItemName,
+	const int32 MaxDiffsToLog
+)
 {
-	FString AssetPathName = FPaths::Combine(*FPaths::GetPath(AssetFilename.Mid(AssetFilename.Find(TEXT(":"), ESearchCase::CaseSensitive) + 1)), *FPaths::GetBaseFilename(AssetFilename));
-	// The root directory could have a period in it (d:/Release5.0/EngineTest/Saved/Cooked),
-	// which is not a valid character for a LongPackageName. Remove it.
-	for (TCHAR c : FStringView(INVALID_LONGPACKAGE_CHARACTERS))
-	{
-		AssetPathName.ReplaceCharInline(c, TEXT('_'), ESearchCase::CaseSensitive);
-	}
-	FString SourceAssetPackageName = FPaths::Combine(TEXT("/Memory"), TEXT("/SourceForDiff"), *AssetPathName);
-	FString DestAssetPackageName = FPaths::Combine(TEXT("/Memory"), TEXT("/DestForDiff"), *AssetPathName);
-	check(FPackageName::IsValidLongPackageName(SourceAssetPackageName, true /* bIncludeReadOnlyRoots */));
-
-	TGuardValue<bool> GuardIsSavingPackage(GIsSavingPackage, false);
-	TGuardValue<int32> GuardAllowUnversionedContentInEditor(GAllowUnversionedContentInEditor, 1);
-	TGuardValue<int32> GuardAllowCookedDataInEditorBuilds(GAllowCookedDataInEditorBuilds, 1);
+	FString HumanReadableString;
+	int32 LoggedDiffs = 0;
+	int32 NumDiffs = 0;
 
-	FLinkerLoad* SourceLinker = nullptr;
-	FLinkerLoad* DestLinker = nullptr;
-	// Create linkers. Note there's no need to clean them up here since they will be removed by the package associated with them
+	int32 MaxIndex = FMath::Max(SourceNum, DestNum);
+	for (int32 Index = 0; Index < MaxIndex; ++Index)
 	{
-		TRefCountPtr<FUObjectSerializeContext> LinkerLoadContext(FUObjectThreadContext::Get().GetSerializeContext());
-		BeginLoad(LinkerLoadContext);
-		SourceLinker = CreateLinkerForPackage(LinkerLoadContext, SourceAssetPackageName, AssetFilename, SourcePackage);
-		EndLoad(LinkerLoadContext);
+		if (Index >= DestNum)
+		{
+			if (MaxDiffsToLog < 0 || NumDiffs < MaxDiffsToLog)
+			{
+				HumanReadableString += IndentToken;
+				HumanReadableString += FString::Printf(TEXT("-[%d] %s"), Index, *ConvertSourceIndexToText(Index));
+				HumanReadableString += NewLineToken;
+				++LoggedDiffs;
+			}
+			++NumDiffs;
+		}
+		else if (Index >= SourceNum)
+		{
+			if (MaxDiffsToLog < 0 || NumDiffs < MaxDiffsToLog)
+			{
+				HumanReadableString += IndentToken;
+				HumanReadableString += FString::Printf(TEXT("+[%d] %s"), Index, *ConvertDestIndexToText(Index));
+				HumanReadableString += NewLineToken;
+				++LoggedDiffs;
+			}
+			++NumDiffs;
+		}
+		else if (!IsElementsAtIndexEqual(Index))
+		{
+			if (MaxDiffsToLog < 0 || NumDiffs < MaxDiffsToLog)
+			{
+				HumanReadableString += IndentToken;
+				HumanReadableString += FString::Printf(TEXT("-[%d] %s"), Index, *ConvertSourceIndexToText(Index));
+				HumanReadableString += NewLineToken;
+				HumanReadableString += IndentToken;
+				HumanReadableString += FString::Printf(TEXT("+[%d] %s"), Index, *ConvertDestIndexToText(Index));
+				HumanReadableString += NewLineToken;
+				++LoggedDiffs;
+			}
+			++NumDiffs;
+		}
 	}
-	
+
+	if (NumDiffs > LoggedDiffs)
 	{
-		TRefCountPtr<FUObjectSerializeContext> LinkerLoadContext(FUObjectThreadContext::Get().GetSerializeContext());
-		BeginLoad(LinkerLoadContext);
-		DestLinker = CreateLinkerForPackage(LinkerLoadContext, DestAssetPackageName, AssetFilename, DestPackage);
-		EndLoad(LinkerLoadContext);
+		HumanReadableString += IndentToken;
+		HumanReadableString += FString::Printf(TEXT("+ %d differences not logged."), (NumDiffs - LoggedDiffs));
+		HumanReadableString += NewLineToken;
 	}
 
-	if (SourceLinker && DestLinker)
-	{
-		FDiffWriterLinkerLoadHeader SourceContext(SourceLinker, MessageCallback);
-		FDiffWriterLinkerLoadHeader DestContext(DestLinker, MessageCallback);
+	LogMessage(ELogVerbosity::Warning, FString::Printf(
+		TEXT("%s: %sMap is different:%s%s"),
+		AssetFilename,
+		ItemName,
+		NewLineToken,
+		*HumanReadableString));
+}
 
-		if (SourceLinker->NameMap != DestLinker->NameMap)
-		{
-			DumpTableDifferences<FNameEntryId>(SourceContext, DestContext, SourceLinker->NameMap, DestLinker->NameMap,
-				*AssetFilename, TEXT("Name"), MaxDiffsToLog);
-		}
+void DumpPackageHeaderDiffs_LinkerLoad(FPackageHeaderData& SourceHeaderData, FPackageHeaderData& DestHeaderData,
+	const int32 MaxDiffsToLog)
+{
+	SourceHeaderData.Initialize();
+	DestHeaderData.Initialize();
+	const FString& AssetFilename = SourceHeaderData.GetAssetFilename();
+	FLinkerLoad* SourceLinker = SourceHeaderData.GetLinker();
+	FLinkerLoad* DestLinker = DestHeaderData.GetLinker();
+	if (!SourceLinker || !DestLinker)
+	{
+		return;
+	}
 
-		if (!SourceContext.IsImportMapIdentical(DestContext))
-		{
-			DumpTableDifferences<FObjectImport>(SourceContext, DestContext, SourceLinker->ImportMap, DestLinker->ImportMap,
-				*AssetFilename, TEXT("Import"), MaxDiffsToLog);
-		}
+	FDiffWriterLinkerLoadHeader SourceContext(SourceLinker, SourceHeaderData.GetMessageCallback());
+	FDiffWriterLinkerLoadHeader DestContext(DestLinker, SourceHeaderData.GetMessageCallback());
 
-		if (!SourceContext.IsExportMapIdentical(DestContext))
-		{
-			DumpTableDifferences<FObjectExport>(SourceContext, DestContext, SourceLinker->ExportMap, DestLinker->ExportMap,
-				*AssetFilename, TEXT("Export"), MaxDiffsToLog);
-		}
+	if (SourceLinker->NameMap != DestLinker->NameMap)
+	{
+		DumpTableDifferences<FNameEntryId>(SourceContext, DestContext, SourceLinker->NameMap, DestLinker->NameMap,
+			*AssetFilename, TEXT("Name"), MaxDiffsToLog);
 	}
 
-	if (SourceLinker)
+	if (!SourceContext.IsImportMapIdentical(DestContext))
 	{
-		UE::ArchiveStackTrace::ForceKillPackageAndLinker(SourceLinker);
+		DumpTableDifferences<FObjectImport>(SourceContext, DestContext, SourceLinker->ImportMap, DestLinker->ImportMap,
+			*AssetFilename, TEXT("Import"), MaxDiffsToLog);
 	}
-	if (DestLinker)
+
+	if (!SourceContext.IsExportMapIdentical(DestContext))
 	{
-		UE::ArchiveStackTrace::ForceKillPackageAndLinker(DestLinker);
+		DumpTableDifferences<FObjectExport>(SourceContext, DestContext, SourceLinker->ExportMap, DestLinker->ExportMap,
+			*AssetFilename, TEXT("Export"), MaxDiffsToLog);
 	}
 }
 
-void DumpPackageHeaderDiffs_ZenPackage(
-	FAccumulatorGlobals& Globals,
-	const FPackageData& SourcePackage,
-	const FPackageData& DestPackage,
-	const FString& AssetFilename,
-	const int32 MaxDiffsToLog,
-	const FMessageCallback& MessageCallback)
+void DumpPackageHeaderDiffs_ZenPackage(FPackageHeaderData& SourceHeaderData,
+	FPackageHeaderData& DestHeaderData, const int32 MaxDiffsToLog)
 {
-	FDiffWriterZenHeader SourceHeader(Globals, MessageCallback, true /* bUseZenStore */, SourcePackage, AssetFilename,
-		TEXT("source"));
-	FDiffWriterZenHeader DestHeader(Globals, MessageCallback, false /* bUseZenStore */, DestPackage, AssetFilename,
-		TEXT("dest"));
-	if (!SourceHeader.IsValid() || !DestHeader.IsValid())
+	SourceHeaderData.Initialize();
+	DestHeaderData.Initialize();
+
+	const FString& AssetFilename = SourceHeaderData.GetAssetFilename();
+	FDiffWriterZenHeader& SourceHeader = SourceHeaderData.GetZenHeader();
+	FDiffWriterZenHeader& DestHeader = DestHeaderData.GetZenHeader();
+	if (!SourceHeaderData.IsValid() || !DestHeaderData.IsValid())
 	{
 		// Explanation was logged by TryConstructSummary
 		return;
@@ -1622,20 +1730,24 @@ void DumpPackageHeaderDiffs_ZenPackage(
 	Algo::Sort(SourceNames, StringSortByNoCaseThenCase);
 	Algo::Sort(DestNames, StringSortByNoCaseThenCase);
 
+	bool bFoundDifferenceInLinkerTableData = false;
 	if (!SourceHeader.IsNameMapIdentical(DestHeader, SourceNames, DestNames))
 	{
+		bFoundDifferenceInLinkerTableData = true;
 		DumpTableDifferences<FString>(SourceHeader, DestHeader, SourceNames, DestNames,
 			*AssetFilename, TEXT("Name"), MaxDiffsToLog);
 	}
 
 	if (!SourceHeader.IsImportMapIdentical(DestHeader))
 	{
+		bFoundDifferenceInLinkerTableData = true;
 		DumpTableDifferences<FPackageObjectIndex>(SourceHeader, DestHeader, SourceHeader.GetPackageHeader().ImportMap,
 			DestHeader.GetPackageHeader().ImportMap, *AssetFilename, TEXT("Import"), MaxDiffsToLog);
 	}
 
 	if (!SourceHeader.IsExportMapIdentical(DestHeader))
 	{
+		bFoundDifferenceInLinkerTableData = true;
 		int32 SourceNum = SourceHeader.GetPackageHeader().ExportMap.Num();
 		int32 DestNum = DestHeader.GetPackageHeader().ExportMap.Num();
 		TArray<FZenHeaderIndexIntoExportMap> SourceIndices;
@@ -1653,28 +1765,252 @@ void DumpPackageHeaderDiffs_ZenPackage(
 		DumpTableDifferences<FZenHeaderIndexIntoExportMap>(SourceHeader, DestHeader, SourceIndices, DestIndices,
 			*AssetFilename, TEXT("Export"), MaxDiffsToLog);
 	}
+
+	// Don't show differences in the dependency data if we found differences in the names/imports/exports data,
+	// because the dependency data depends on that basic data and will usually have differences if any of them
+	// changed, and it just creates noise for the developer.
+	if (!bFoundDifferenceInLinkerTableData)
+	{
+		if (!SourceHeader.IsExportBundlesIdentical(DestHeader))
+		{
+			DumpOrderedArrayDifferences(SourceHeader.GetPackageHeader().ExportBundleEntries.Num(),
+				DestHeader.GetPackageHeader().ExportBundleEntries.Num(),
+				[&SourceHeader, &DestHeader](int32 Index)
+				{
+					return SourceHeader.IsExportBundleIdentical(DestHeader, Index);
+				},
+				[&SourceHeader](int32 Index)
+				{
+					return SourceHeader.ConvertExportBundleToText(Index);
+				},
+				[&DestHeader](int32 Index)
+				{
+					return DestHeader.ConvertExportBundleToText(Index);
+				},
+				[&SourceHeader](ELogVerbosity::Type Verbosity, FString Message)
+				{
+					SourceHeader.LogMessage(Verbosity, Message);
+				},
+				*AssetFilename, TEXT("ExportBundles"), MaxDiffsToLog);
+		}
+		if (!SourceHeader.IsDependencyBundlesIdentical(DestHeader))
+		{
+			DumpOrderedArrayDifferences(SourceHeader.GetPackageHeader().DependencyBundleHeaders.Num(),
+				DestHeader.GetPackageHeader().DependencyBundleHeaders.Num(),
+				[&SourceHeader, &DestHeader](int32 Index)
+				{
+					return SourceHeader.IsDependencyBundleIdentical(DestHeader, Index);
+				},
+				[&SourceHeader](int32 Index)
+				{
+					return SourceHeader.ConvertDependencyBundleToText(Index);
+				},
+				[&DestHeader](int32 Index)
+				{
+					return DestHeader.ConvertDependencyBundleToText(Index);
+				},
+				[&SourceHeader](ELogVerbosity::Type Verbosity, FString Message)
+				{
+					SourceHeader.LogMessage(Verbosity, Message);
+				},
+				*AssetFilename, TEXT("DependencyBundles"), MaxDiffsToLog);
+		}
+	}
 }
 
-void DumpPackageHeaderDiffs(
-	FAccumulatorGlobals& Globals,
-	const FPackageData& SourcePackage,
-	const FPackageData& DestPackage,
-	const FString& AssetFilename,
-	const int32 MaxDiffsToLog,
-	const EPackageHeaderFormat PackageHeaderFormat,
-	const FMessageCallback& MessageCallback)
+void DumpPackageHeaderDiffs(FPackageHeaderData& SourceHeaderData, FPackageHeaderData& DestHeaderData,
+	const int32 MaxDiffsToLog)
 {
-	switch (PackageHeaderFormat)
+	switch (SourceHeaderData.GetFormat())
 	{
 	case EPackageHeaderFormat::PackageFileSummary:
-		DumpPackageHeaderDiffs_LinkerLoad(Globals, SourcePackage, DestPackage, AssetFilename, MaxDiffsToLog, MessageCallback);
+		DumpPackageHeaderDiffs_LinkerLoad(SourceHeaderData, DestHeaderData, MaxDiffsToLog);
 		break;
 	case EPackageHeaderFormat::ZenPackageSummary:
-		DumpPackageHeaderDiffs_ZenPackage(Globals, SourcePackage, DestPackage, AssetFilename, MaxDiffsToLog, MessageCallback);
+		DumpPackageHeaderDiffs_ZenPackage(SourceHeaderData, DestHeaderData, MaxDiffsToLog);
 		break;
 	default:
 		unimplemented();
 	}
 }
 
+FPackageHeaderData::FPackageHeaderData(const TCHAR* InName, bool bInReadFromPackageStore,
+	const FString& InAssetFilename, const FPackageData& InPackageData, EPackageHeaderFormat InFormat,
+	FAccumulatorGlobals& InGlobals, FMessageCallback& InMessageCallback)
+	: Name(InName)
+	, AssetFilename(InAssetFilename)
+	, PackageData(InPackageData)
+	, Globals(InGlobals)
+	, MessageCallback(InMessageCallback)
+	, Format(InFormat)
+	, bReadFromPackageStore(bInReadFromPackageStore)
+{
+}
+
+FPackageHeaderData::~FPackageHeaderData()
+{
+	if (Linker)
+	{
+		UE::ArchiveStackTrace::ForceKillPackageAndLinker(Linker);
+	}
+}
+
+void FPackageHeaderData::Initialize()
+{
+	if (bInitialized)
+	{
+		return;
+	}
+	bInitialized = true;
+
+	switch (Format)
+	{
+	case EPackageHeaderFormat::PackageFileSummary:
+	{
+		FString AssetPathName = FPaths::Combine(
+			*FPaths::GetPath(AssetFilename.Mid(AssetFilename.Find(TEXT(":"), ESearchCase::CaseSensitive) + 1)),
+			*FPaths::GetBaseFilename(AssetFilename));
+		// The root directory could have a period in it (d:/Release5.0/EngineTest/Saved/Cooked),
+		// which is not a valid character for a LongPackageName. Remove it.
+		for (TCHAR c : FStringView(INVALID_LONGPACKAGE_CHARACTERS))
+		{
+			AssetPathName.ReplaceCharInline(c, TEXT('_'), ESearchCase::CaseSensitive);
+		}
+		FString AssetPackageName = FPaths::Combine(TEXT("/Memory"),
+			WriteToString<32>(TEXT("/%sForDiff"), Name),
+			*AssetPathName);
+		check(FPackageName::IsValidLongPackageName(AssetPackageName, true /* bIncludeReadOnlyRoots */));
+
+		TGuardValue<bool> GuardIsSavingPackage(GIsSavingPackage, false);
+		TGuardValue<int32> GuardAllowUnversionedContentInEditor(GAllowUnversionedContentInEditor, 1);
+		TGuardValue<int32> GuardAllowCookedDataInEditorBuilds(GAllowCookedDataInEditorBuilds, 1);
+
+		// Create linkers. Note there's no need to clean them up here since they will be removed by the package associated with them
+		{
+			TRefCountPtr<FUObjectSerializeContext> LinkerLoadContext(FUObjectThreadContext::Get().GetSerializeContext());
+			BeginLoad(LinkerLoadContext);
+			Linker = CreateLinkerForPackage(LinkerLoadContext, AssetPackageName, AssetFilename, PackageData);
+			EndLoad(LinkerLoadContext);
+		}
+		break;
+	}
+	case EPackageHeaderFormat::ZenPackageSummary:
+		ZenHeader = MakeUnique<FDiffWriterZenHeader>(Globals, MessageCallback, bReadFromPackageStore, PackageData,
+			AssetFilename, Name);
+		break;
+	default:
+		unimplemented();
+		break;
+	}
+}
+
+const TCHAR* FPackageHeaderData::GetName() const
+{
+	return Name;
+}
+
+const FString& FPackageHeaderData::GetAssetFilename() const
+{
+	return AssetFilename;
+}
+
+const FPackageData& FPackageHeaderData::GetPackageData() const
+{
+	return PackageData;
+}
+
+EPackageHeaderFormat FPackageHeaderData::GetFormat() const
+{
+	return Format;
+}
+
+FAccumulatorGlobals& FPackageHeaderData::GetGlobals() const
+{
+	return Globals;
+}
+
+FMessageCallback& FPackageHeaderData::GetMessageCallback() const
+{
+	return MessageCallback;
+}
+
+FDiffWriterZenHeader& FPackageHeaderData::GetZenHeader() const
+{
+	check(Format == EPackageHeaderFormat::ZenPackageSummary && bInitialized);
+	check(ZenHeader.IsValid());
+	return *ZenHeader;
+}
+
+FLinkerLoad* FPackageHeaderData::GetLinker() const
+{
+	check(Format == EPackageHeaderFormat::PackageFileSummary && bInitialized);
+	return Linker;
+}
+
+bool FPackageHeaderData::IsValid() const
+{
+	if (!bInitialized)
+	{
+		return false;
+	}
+	switch (Format)
+	{
+	case EPackageHeaderFormat::PackageFileSummary:
+		return Linker != nullptr;
+	case EPackageHeaderFormat::ZenPackageSummary:
+		check(ZenHeader);
+		return ZenHeader->IsValid();
+	default:
+		unimplemented();
+		return false;
+	}
+}
+
+bool FPackageHeaderData::TryGetMappedName(int32 Index, int32 Number, FName& OutName) const
+{
+	if (!bInitialized || !IsValid())
+	{
+		return false;
+	}
+
+	switch (Format)
+	{
+	case EPackageHeaderFormat::PackageFileSummary:
+		check(Linker);
+		if (!Linker->NameMap.IsValidIndex(Index))
+		{
+			return false;
+		}
+		OutName = FName::CreateFromDisplayId(Linker->NameMap[Index], Number);
+		return true;
+	case EPackageHeaderFormat::ZenPackageSummary:
+		check(ZenHeader);
+		return ZenHeader->GetPackageHeader().NameMap.TryGetName(
+			FMappedName::Create((uint32)Index, (uint32)Number, FMappedName::EType::Package),
+			OutName);
+	default:
+		unimplemented();
+		return false;
+	}
+}
+
+FPackageHeaderDataProxyArchive::FPackageHeaderDataProxyArchive(FPackageHeaderData& InHeader, FArchive& InInner)
+	: FArchiveProxy(InInner)
+	, Header(InHeader)
+{
+}
+
+FArchive& FPackageHeaderDataProxyArchive::operator<<(FName& Value)
+{
+	int32 Index;
+	int32 Number;
+
+	InnerArchive << Index << Number;
+	if (!Header.TryGetMappedName(Index, Number, Value))
+	{
+		Value = NAME_None;
+	}
+	return *this;
+}
+
 } // namespace UE::DiffWriter
\ No newline at end of file
diff --git a/Engine/Source/Editor/UnrealEd/Private/Cooker/DiffWriterArchive.h b/Engine/Source/Editor/UnrealEd/Private/Cooker/DiffWriterArchive.h
--- a/Engine/Source/Editor/UnrealEd/Private/Cooker/DiffWriterArchive.h
+++ b/Engine/Source/Editor/UnrealEd/Private/Cooker/DiffWriterArchive.h
@@ -33,6 +33,7 @@ namespace UE::DiffWriter
 
 class FAccumulator;
 class FDiffArchive;
+class FPackageHeaderData;
 
 typedef TUniqueFunction<void(ELogVerbosity::Type, FStringView)> FMessageCallback;
 using EPackageHeaderFormat = ICookedPackageWriter::EPackageHeaderFormat;
@@ -272,6 +273,7 @@ private:
 	void GenerateDiffMap();
 	/** Compares two packages and logs the differences and calltacks. */
 	void CompareWithPreviousForSection(const FPackageData& SourcePackage, const FPackageData& DestPackage,
+		FPackageHeaderData& SourceHeader, FPackageHeaderData& DestHeader,
 		const TCHAR* CallstackCutoffText, int32& InOutLoggedDiffs,TMap<FName, FArchiveDiffStats>& OutStats,
 		const FString& SectionFilename);
 
diff --git a/Engine/Source/Editor/UnrealEd/Private/Cooker/DiffWriterZenHeader.cpp b/Engine/Source/Editor/UnrealEd/Private/Cooker/DiffWriterZenHeader.cpp
--- a/Engine/Source/Editor/UnrealEd/Private/Cooker/DiffWriterZenHeader.cpp
+++ b/Engine/Source/Editor/UnrealEd/Private/Cooker/DiffWriterZenHeader.cpp
@@ -278,6 +278,158 @@ FString FDiffWriterZenHeader::ConvertItemToText(const FZenHeaderIndexIntoExportM
 		(int32)Export.FilterFlags);
 }
 
+bool FDiffWriterZenHeader::IsExportBundlesIdentical(FDiffWriterZenHeader& DestContext)
+{
+	int32 NumBundles = PackageHeader.ExportBundleEntries.Num();
+	if (NumBundles != DestContext.PackageHeader.ExportBundleEntries.Num())
+	{
+		return false;
+	}
+
+	for (int32 BundleIndex = 0; BundleIndex < NumBundles; ++BundleIndex)
+	{
+		if (!IsExportBundleIdentical(DestContext, BundleIndex))
+		{
+			return false;
+		}
+	}
+	return true;
+}
+
+bool FDiffWriterZenHeader::IsExportBundleIdentical(FDiffWriterZenHeader& DestContext, int32 Index)
+{
+	if (!PackageHeader.ExportBundleEntries.IsValidIndex(Index) ||
+		!DestContext.PackageHeader.ExportBundleEntries.IsValidIndex(Index))
+	{
+		return false;
+	}
+	const FExportBundleEntry& Src = PackageHeader.ExportBundleEntries[Index];
+	const FExportBundleEntry& Dest = DestContext.PackageHeader.ExportBundleEntries[Index];
+	return Src.LocalExportIndex == Dest.LocalExportIndex
+		&& Src.CommandType == Dest.CommandType;
+}
+
+FString FDiffWriterZenHeader::ConvertExportBundleToText(int32 Index)
+{
+	if (!PackageHeader.ExportBundleEntries.IsValidIndex(Index))
+	{
+		return TEXT("<InvalidIndex>");
+	}
+	const FExportBundleEntry& Entry = PackageHeader.ExportBundleEntries[Index];
+	return FString::Printf(TEXT("[%d,%d]"), Entry.LocalExportIndex, Entry.CommandType);
+}
+
+bool FDiffWriterZenHeader::IsDependencyBundlesIdentical(FDiffWriterZenHeader& DestContext)
+{
+	int32 NumBundles = PackageHeader.DependencyBundleHeaders.Num();
+	if (NumBundles != DestContext.PackageHeader.DependencyBundleHeaders.Num())
+	{
+		return false;
+	}
+
+	for (int32 BundleIndex = 0; BundleIndex < NumBundles; ++BundleIndex)
+	{
+		if (!IsDependencyBundleIdentical(DestContext, BundleIndex))
+		{
+			return false;
+		}
+	}
+	return true;
+}
+
+bool FDiffWriterZenHeader::IsDependencyBundleIdentical(FDiffWriterZenHeader& DestContext, int32 Index)
+{
+	if (!PackageHeader.DependencyBundleHeaders.IsValidIndex(Index) ||
+		!DestContext.PackageHeader.DependencyBundleHeaders.IsValidIndex(Index))
+	{
+		return false;
+	}
+	const FDependencyBundleHeader& SrcHeader = PackageHeader.DependencyBundleHeaders[Index];
+	const FDependencyBundleHeader& DestHeader = DestContext.PackageHeader.DependencyBundleHeaders[Index];
+	if (SrcHeader.FirstEntryIndex == -1)
+	{
+		return DestHeader.FirstEntryIndex == -1;
+	}
+	if (DestHeader.FirstEntryIndex == -1)
+	{
+		return false;
+	}
+	TArrayView<const FDependencyBundleEntry> SrcEntries = PackageHeader.DependencyBundleEntries;
+	TArrayView<const FDependencyBundleEntry> DestEntries = DestContext.PackageHeader.DependencyBundleEntries;
+	int32 SrcIndex = SrcHeader.FirstEntryIndex;
+	int32 DestIndex = DestHeader.FirstEntryIndex;
+	for (uint32 FromType = 0; FromType < (uint32)FExportBundleEntry::ExportCommandType_Count; ++FromType)
+	{
+		for (uint32 ToType = 0; ToType < (uint32)FExportBundleEntry::ExportCommandType_Count; ++ToType)
+		{
+			int32 EntryCount = SrcHeader.EntryCount[FromType][ToType];
+			if (EntryCount != DestHeader.EntryCount[FromType][ToType])
+			{
+				return false;
+			}
+			for (int32 IndexWithinEntry = 0; IndexWithinEntry < EntryCount; ++IndexWithinEntry)
+			{
+				if (SrcIndex >= SrcEntries.Num() || DestIndex >= DestEntries.Num())
+				{
+					return false;
+				}
+				if (SrcEntries[SrcIndex].LocalImportOrExportIndex != DestEntries[SrcIndex].LocalImportOrExportIndex)
+				{
+					return false;
+				}
+				++SrcIndex;
+				++DestIndex;
+			}
+		}
+	}
+	return true;
+}
+
+FString FDiffWriterZenHeader::ConvertDependencyBundleToText(int32 Index)
+{
+	if (!PackageHeader.DependencyBundleEntries.IsValidIndex(Index))
+	{
+		return TEXT("<InvalidIndex>");
+	}
+	TStringBuilder<256> Result;
+
+	const FDependencyBundleHeader& BundleHeader = PackageHeader.DependencyBundleHeaders[Index];
+	TArrayView<const FDependencyBundleEntry> Entries = PackageHeader.DependencyBundleEntries;
+	int32 EntryIndex = BundleHeader.FirstEntryIndex;
+	for (uint32 FromType = 0; FromType < (uint32)FExportBundleEntry::ExportCommandType_Count; ++FromType)
+	{
+		for (uint32 ToType = 0; ToType < (uint32)FExportBundleEntry::ExportCommandType_Count; ++ToType)
+		{
+			Result << TEXT("[");
+			int32 CountWithinEntry = BundleHeader.EntryCount[FromType][ToType];
+			if (Entries.IsValidIndex(EntryIndex) && CountWithinEntry > 0)
+			{
+				for (int32 IndexWithinEntry = 0; IndexWithinEntry < CountWithinEntry; ++IndexWithinEntry)
+				{
+					if (Entries.IsValidIndex(EntryIndex))
+					{
+						FPackageIndex PackageIndex = Entries[EntryIndex].LocalImportOrExportIndex;
+						if (PackageIndex.IsExport())
+						{
+							Result << TEXT("Export[") << PackageIndex.ToExport() << TEXT("]");
+						}
+						else
+						{
+							Result << TEXT("Import[") << PackageIndex.ToImport() << TEXT("]");
+						}
+						Result << ",";
+						++EntryIndex;
+					}
+				}
+				Result.RemoveSuffix(1);
+			}
+			Result << TEXT("],");
+		}
+	}
+	Result.RemoveSuffix(1);
+	return FString(Result);
+}
+
 void FDiffWriterZenHeader::LogMessage(ELogVerbosity::Type Verbosity, FStringView Message)
 {
 	MessageCallback(Verbosity, Message);
diff --git a/Engine/Source/Editor/UnrealEd/Private/Cooker/DiffWriterZenHeader.h b/Engine/Source/Editor/UnrealEd/Private/Cooker/DiffWriterZenHeader.h
--- a/Engine/Source/Editor/UnrealEd/Private/Cooker/DiffWriterZenHeader.h
+++ b/Engine/Source/Editor/UnrealEd/Private/Cooker/DiffWriterZenHeader.h
@@ -71,6 +71,14 @@ public:
 		const FZenHeaderIndexIntoExportMap& DestExport);
 	FString ConvertItemToText(const FZenHeaderIndexIntoExportMap& Id);
 
+	bool IsExportBundlesIdentical(FDiffWriterZenHeader& DestContext);
+	bool IsExportBundleIdentical(FDiffWriterZenHeader& DestContext, int32 Index);
+	FString ConvertExportBundleToText(int32 Index);
+
+	bool IsDependencyBundlesIdentical(FDiffWriterZenHeader& DestContext);
+	bool IsDependencyBundleIdentical(FDiffWriterZenHeader& DestContext, int32 Index);
+	FString ConvertDependencyBundleToText(int32 Index);
+
 	void LogMessage(ELogVerbosity::Type Verbosity, FStringView Message);
 
 private:
