InteractiveToolsFramework: Add ensure to catch cases where a tool command object is such that the tool will crash on startup.

diff --git a/Engine/Source/Editor/UnrealEd/Public/Tools/InteractiveToolsCommands.h b/Engine/Source/Editor/UnrealEd/Public/Tools/InteractiveToolsCommands.h
--- a/Engine/Source/Editor/UnrealEd/Public/Tools/InteractiveToolsCommands.h
+++ b/Engine/Source/Editor/UnrealEd/Public/Tools/InteractiveToolsCommands.h
@@ -136,6 +136,7 @@ void TInteractiveToolCommands<CommandContextType>::RegisterCommands()
 	// register all the commands in all the actions
 	int NumActions = ToolActions.Num();
 	ActionCommands.SetNum(NumActions);
+	bool bEverRegistered = false;
 	for (int k = 0; k < NumActions; ++k)
 	{
 		const FInteractiveToolAction& ToolAction = ToolActions[k];
@@ -165,8 +166,21 @@ void TInteractiveToolCommands<CommandContextType>::RegisterCommands()
 		if (bRegistered == false)
 		{
 			RegisterUIToolCommand(ActionCommands[k].ToolAction, ActionCommands[k].UICommand);
+			bEverRegistered = true;
 		}
 	}
+
+	// If you hit this ensure, then your tool must not have registered any actions that were not of the "standard" variety
+	//  (e.g. with ID like EStandardToolModeCommands::IncreaseBrushSize, etc). This is currently not supported, and your
+	//  tool will likely crash when started. Either don't give the tool a command objects (i.e., not use hotkeys), or add
+	//  a hotkey of your own.
+	// TODO: Support all-standard-hotkey objects, jira UE-221911
+	ensureMsgf(bEverRegistered || ToolActions.Num() == 0, TEXT("Command object had actions, but never registered a FUICommandInfo "
+		"object of its own, so its singleton pointer may not be kept alive, causing a crash on a subsequent Get() call. Currently, "
+		"a command object must register at least one \"non-standard\" action through the tool to be safely used."));
+	// We could remove the ToolActions.Num() == 0 exception above, but it is convenient to set up empty hotkey objects
+	//  for tools to populate later in development. On the other hand, registering hotkeys and still hitting a crash is
+	//  more perplexing and harder to catch, hence this ensure.
 }
 
 
diff --git a/Engine/Source/Runtime/Slate/Public/Framework/Commands/Commands.h b/Engine/Source/Runtime/Slate/Public/Framework/Commands/Commands.h
--- a/Engine/Source/Runtime/Slate/Public/Framework/Commands/Commands.h
+++ b/Engine/Source/Runtime/Slate/Public/Framework/Commands/Commands.h
@@ -68,6 +68,9 @@ public:
 	/** Get the singleton instance of this set of commands. */
 	FORCENOINLINE static const CommandContextType& Get()
 	{
+		// If you get a crash here, then either your commands object wasn't registered, or it was not kept
+		//  alive. The latter could happen if your RegisterCommands() function does not actually register 
+		//  any FUICommandInfo objects that use your TCommands object, so nothing keeps it alive.
 		return *( Instance.Pin() );
 	}
 
