Cooker: Update to DiffOnly feature improvements
Supress false positive static analysis error.
#rnx

diff --git a/Engine/Source/Editor/UnrealEd/Private/Cooker/DiffWriterArchive.cpp b/Engine/Source/Editor/UnrealEd/Private/Cooker/DiffWriterArchive.cpp
--- a/Engine/Source/Editor/UnrealEd/Private/Cooker/DiffWriterArchive.cpp
+++ b/Engine/Source/Editor/UnrealEd/Private/Cooker/DiffWriterArchive.cpp
@@ -750,7 +750,8 @@ void FAccumulator::OnSecondSaveComplete(int64 InHeaderSize)
 			SecondSaveHeaderSegment, PackageHeaderFormat, Globals, HeaderMessageCallback);
 
 		DumpPackageHeaderDiffs(FirstSaveHeader, SecondSaveHeader, MaxDiffsToLog);
-		if (NumHeaderDiffMessages == 0)
+		// Suppress static analysis warning V547: Expression 'NumHeaderDiffMessages == 0' is always true
+		if (NumHeaderDiffMessages == 0) // -V547
 		{
 			MessageCallback(ELogVerbosity::Warning, FString::Printf(
 				TEXT("%s: headers are different, but DumpPackageHeaderDiffs does not yet implement describing the difference."),
@@ -1206,7 +1207,8 @@ void FAccumulator::CompareWithPrevious(const TCHAR* CallstackCutoffText, TMap<FN
 	if (HeaderSize > 0 && OutStats.FindOrAdd(AssetClass).NumDiffs > 0)
 	{
 		DumpPackageHeaderDiffs(SourceHeader, DestHeader, MaxDiffsToLog);
-		if (NumHeaderDiffMessages == 0)
+		// Suppress static analysis warning V547: Expression 'NumHeaderDiffMessages == 0' is always true
+		if (NumHeaderDiffMessages == 0) // -V547
 		{
 			MessageCallback(ELogVerbosity::Warning, FString::Printf(
 				TEXT("%s: headers are different, but DumpPackageHeaderDiffs does not yet implement describing the difference."),
