Correctly update ClassWithin and ClassConfigName for blueprints that do not go through full compilation

diff --git a/Engine/Source/Editor/KismetCompiler/Private/KismetCompiler.cpp b/Engine/Source/Editor/KismetCompiler/Private/KismetCompiler.cpp
--- a/Engine/Source/Editor/KismetCompiler/Private/KismetCompiler.cpp
+++ b/Engine/Source/Editor/KismetCompiler/Private/KismetCompiler.cpp
@@ -3130,6 +3130,7 @@ void FKismetCompilerContext::FinishCompilingClass(UClass* Class)
 		Class->ClassFlags |= (ParentClass->ClassFlags & CLASS_ScriptInherit);//@TODO: ChangeParentClass had this, but I don't think I want it: | UClass::StaticClassFlags;  // will end up with CLASS_Intrinsic
 		Class->ClassCastFlags |= ParentClass->ClassCastFlags;
 		Class->ClassConfigName = ParentClass->ClassConfigName;
+		Class->ClassWithin = ParentClass->ClassWithin;
 
 		// If the Blueprint was marked as deprecated, then flag the class as deprecated.
 		if(Blueprint->bDeprecate)
diff --git a/Engine/Source/Editor/UnrealEd/Private/Kismet2/Kismet2.cpp b/Engine/Source/Editor/UnrealEd/Private/Kismet2/Kismet2.cpp
--- a/Engine/Source/Editor/UnrealEd/Private/Kismet2/Kismet2.cpp
+++ b/Engine/Source/Editor/UnrealEd/Private/Kismet2/Kismet2.cpp
@@ -1003,6 +1003,8 @@ void FKismetEditorUtilities::ConformBlueprintFlagsAndComponents(UBlueprint* Blue
 	if( UClass* SkelClass = BlueprintObj->SkeletonGeneratedClass )
 	{
 		SkelClass->ClassFlags |= (ParentClass->ClassFlags & CLASS_ScriptInherit);
+		SkelClass->ClassConfigName = ParentClass->ClassConfigName;
+		SkelClass->ClassWithin = ParentClass->ClassWithin;
 		UObject* SkelCDO = SkelClass->GetDefaultObject();
 		// NOTE: we don't need to call ConformRemovedNativeComponents() for skel
 		//       classes, as they're generated on load (and not saved with stale 
@@ -1013,6 +1015,8 @@ void FKismetEditorUtilities::ConformBlueprintFlagsAndComponents(UBlueprint* Blue
 	if( UClass* GenClass = BlueprintObj->GeneratedClass )
 	{
 		GenClass->ClassFlags |= (ParentClass->ClassFlags & CLASS_ScriptInherit);
+		GenClass->ClassConfigName = ParentClass->ClassConfigName;
+		GenClass->ClassWithin = ParentClass->ClassWithin;
 		if (UObject* GenCDO = GenClass->ClassDefaultObject)
 		{
 			ConformComponentsUtils::ConformRemovedNativeComponents(GenCDO);
