Allow dumping specific keys from stable key file that match a given output hash value.

diff --git a/Engine/Source/Editor/UnrealEd/Private/Commandlets/ShaderPipelineCacheToolsCommandlet.cpp b/Engine/Source/Editor/UnrealEd/Private/Commandlets/ShaderPipelineCacheToolsCommandlet.cpp
--- a/Engine/Source/Editor/UnrealEd/Private/Commandlets/ShaderPipelineCacheToolsCommandlet.cpp
+++ b/Engine/Source/Editor/UnrealEd/Private/Commandlets/ShaderPipelineCacheToolsCommandlet.cpp
@@ -672,14 +672,25 @@ bool CouldBeUsedTogether(const FStableShaderKeyAndValue& A, const FStableShaderK
 	return true;
 }
 
-int32 DumpStableKeysFile(const FString& Token)
+int32 DumpStableKeysFile(const FString& Token, const TMap<FString, FString>& ParamVals)
 {
 	const FStringView File = Token;
+	FString OutputHashToSearch;
+	if (const FString* Param = ParamVals.Find(TEXT("outputhash")))
+	{
+		OutputHashToSearch = *Param;
+	}
+
 	TMultiMap<FStableShaderKeyAndValue, FSHAHash> StableMap;
 	LoadStableShaderKeysMultiple(StableMap, MakeArrayView(&File, 1));
 	UE_LOG(LogShaderPipelineCacheTools, Display, TEXT("    %s"), *FStableShaderKeyAndValue::HeaderLine());
 	for (const auto& Pair : StableMap)
 	{
+		if (!OutputHashToSearch.IsEmpty() && Pair.Key.OutputHash.ToString() != OutputHashToSearch)
+		{
+			continue;
+		}
+
 		FStableShaderKeyAndValue Temp(Pair.Key);
 		UE_LOG(LogShaderPipelineCacheTools, Display, TEXT("    %s"), *Temp.ToString());
 	}
@@ -2862,7 +2873,7 @@ int32 UShaderPipelineCacheToolsCommandlet::StaticMain(const FString& Params)
 				}
 				if (Tokens[Index].EndsWith(ShaderStableKeysFileExt))
 				{
-					return DumpStableKeysFile(Tokens[Index]);
+					return DumpStableKeysFile(Tokens[Index], ParamVals);
 				}
 			}
 		}
@@ -2878,9 +2889,9 @@ int32 UShaderPipelineCacheToolsCommandlet::StaticMain(const FString& Params)
 		}
 	}
 	
-	UE_LOG(LogShaderPipelineCacheTools, Warning, TEXT("Usage: Dump SCLInfo.%s [...]] - dumps stable keys file.\n"), ShaderStableKeysFileExt);
+	UE_LOG(LogShaderPipelineCacheTools, Warning, TEXT("Usage: Dump SCLInfo.%s [...] [-outputhash=X] - dumps stable keys file, optionally only dumping keys that match the specified output hash value.\n"), ShaderStableKeysFileExt);
 	UE_LOG(LogShaderPipelineCacheTools, Warning, TEXT("Usage: Dump PSOCache.upipelinecache [SCLInfo.%s] - dumps recorded PSOs. If optional stable keys file, converts shader hashes to readable stable shader descriptions.\n"), ShaderStableKeysFileExt);
-	UE_LOG(LogShaderPipelineCacheTools, Warning, TEXT("Usage: Dump StablePSOCache.csv[.compressed] - dumps content of a stable PSO.\n"), ShaderStableKeysFileExt);
+	UE_LOG(LogShaderPipelineCacheTools, Warning, TEXT("Usage: Dump StablePSOCache.csv[.compressed] - dumps content of a stable PSO.\n"));
 	UE_LOG(LogShaderPipelineCacheTools, Warning, TEXT("Usage: Diff ShaderCache1.stablepc.csv ShaderCache1.stablepc.csv [...]]\n"));
 	UE_LOG(LogShaderPipelineCacheTools, Warning, TEXT("Usage: Expand Input1.upipelinecache Dir2/*.upipelinecache InputSCLInfo1.%s Dir2/*.%s InputSCLInfo3.%s [...] Output.stablepc.csv\n"), ShaderStableKeysFileExt, ShaderStableKeysFileExt, ShaderStableKeysFileExt);
 	UE_LOG(LogShaderPipelineCacheTools, Warning, TEXT("Usage: Build Input.stablepc.csv InputDir2/*.stablepc.csv InputSCLInfo1.%s Dir2/*.%s InputSCLInfo3.%s [...] Output.upipelinecache\n"), ShaderStableKeysFileExt, ShaderStableKeysFileExt, ShaderStableKeysFileExt);
