[Viewport Toolbar] Fix empty labels, tooltips and section names

diff --git a/Engine/Plugins/FX/Niagara/Source/NiagaraEditor/Private/Widgets/NiagaraSystemEditorViewportToolbarSections.cpp b/Engine/Plugins/FX/Niagara/Source/NiagaraEditor/Private/Widgets/NiagaraSystemEditorViewportToolbarSections.cpp
--- a/Engine/Plugins/FX/Niagara/Source/NiagaraEditor/Private/Widgets/NiagaraSystemEditorViewportToolbarSections.cpp
+++ b/Engine/Plugins/FX/Niagara/Source/NiagaraEditor/Private/Widgets/NiagaraSystemEditorViewportToolbarSections.cpp
@@ -83,7 +83,7 @@ FToolMenuEntry UE::NiagaraSystemEditor::CreateShowSubmenu(const TWeakPtr<SNiagar
 					return;
 				}
 
-				FToolMenuSection& UnnamedSection = Submenu->FindOrAddSection("", LOCTEXT("UnnamedLabel", ""));
+				FToolMenuSection& UnnamedSection = Submenu->FindOrAddSection(NAME_None);
 
 				UnnamedSection.AddEntry(FToolMenuEntry::InitWidget(
 					"ShowMenuItems", UE::NiagaraSystemEditor::CreateShowMenuWidget(Viewport.ToSharedRef(), true), FText(), true
diff --git a/Engine/Source/Editor/Kismet/Private/SSCSEditorViewport.cpp b/Engine/Source/Editor/Kismet/Private/SSCSEditorViewport.cpp
--- a/Engine/Source/Editor/Kismet/Private/SSCSEditorViewport.cpp
+++ b/Engine/Source/Editor/Kismet/Private/SSCSEditorViewport.cpp
@@ -554,7 +554,7 @@ TSharedPtr<SWidget> SSCSEditorViewport::BuildViewportToolbar()
 					FNewToolMenuDelegate::CreateLambda(
 						[](UToolMenu* Submenu) -> void
 						{
-							FToolMenuSection& UnnamedSection = Submenu->FindOrAddSection("", LOCTEXT("UnnamedLabel", ""));
+							FToolMenuSection& UnnamedSection = Submenu->FindOrAddSection(NAME_None);
 							UnnamedSection.AddMenuEntry(FBlueprintEditorCommands::Get().ShowFloor);
 							UnnamedSection.AddMenuEntry(FBlueprintEditorCommands::Get().ShowGrid);
 						}
@@ -574,7 +574,7 @@ TSharedPtr<SWidget> SSCSEditorViewport::BuildViewportToolbar()
 					FNewToolMenuDelegate::CreateLambda(
 						[](UToolMenu* Submenu) -> void
 						{
-							FToolMenuSection& UnnamedSection = Submenu->FindOrAddSection("", LOCTEXT("UnnamedLabel", ""));
+							FToolMenuSection& UnnamedSection = Submenu->FindOrAddSection(NAME_None);
 							UnnamedSection.AddEntry(UE::UnrealEd::CreateToggleRealtimeEntry());
 						}
 					)
diff --git a/Engine/Source/Editor/LevelEditor/Private/ViewportToolbar/LevelEditorViewportToolbarSections.cpp b/Engine/Source/Editor/LevelEditor/Private/ViewportToolbar/LevelEditorViewportToolbarSections.cpp
--- a/Engine/Source/Editor/LevelEditor/Private/ViewportToolbar/LevelEditorViewportToolbarSections.cpp
+++ b/Engine/Source/Editor/LevelEditor/Private/ViewportToolbar/LevelEditorViewportToolbarSections.cpp
@@ -1510,7 +1510,7 @@ FToolMenuEntry CreateFeatureLevelPreviewSubmenu()
 						FNewToolMenuDelegate::CreateLambda(
 							[CommandList, CommandListJson](UToolMenu* InSubMenu)
 							{
-								FToolMenuSection& Section = InSubMenu->AddSection("", LOCTEXT("EditorPreviewModeDevices", "Preview Devices"));
+								FToolMenuSection& Section = InSubMenu->AddSection(NAME_None);
 								for (const TSharedPtr<FUICommandInfo>& Command : CommandList)
 								{
 									Section.AddMenuEntry(Command);
@@ -1525,7 +1525,7 @@ FToolMenuEntry CreateFeatureLevelPreviewSubmenu()
 										FNewToolMenuDelegate::CreateLambda(
 											[CommandListJson](UToolMenu* InSubMenu)
 											{
-												FToolMenuSection& Section = InSubMenu->AddSection("", LOCTEXT("EditorPreviewModeDevicesJson", "Preview Devices Json"));
+												FToolMenuSection& Section = InSubMenu->AddSection(NAME_None);
 												for (const TSharedPtr<FUICommandInfo>& Command : *CommandListJson)
 												{
 													Section.AddMenuEntry(Command);
@@ -2535,7 +2535,7 @@ void ExtendCameraSubmenu(FName InCameraOptionsSubmenuName)
 					// This additional options section is used to force certain elements to appear after extensions
 					{
 						FToolMenuSection& AdditionalOptions =
-							InDynamicMenu->FindOrAddSection("AdditionalOptions", LOCTEXT("AdditionalOptionsLabel", ""));
+							InDynamicMenu->FindOrAddSection("AdditionalOptions");
 						AdditionalOptions.AddSeparator("AdditionalOptionsSeparator");
 
 						FToolMenuEntry HighResolutionScreenshot =
diff --git a/Engine/Source/Editor/UnrealEd/Private/ViewportToolbar/UnrealEdViewportToolbar.cpp b/Engine/Source/Editor/UnrealEd/Private/ViewportToolbar/UnrealEdViewportToolbar.cpp
--- a/Engine/Source/Editor/UnrealEd/Private/ViewportToolbar/UnrealEdViewportToolbar.cpp
+++ b/Engine/Source/Editor/UnrealEd/Private/ViewportToolbar/UnrealEdViewportToolbar.cpp
@@ -2138,7 +2138,7 @@ FToolMenuEntry CreateViewportToolbarAssetViewerProfileSubmenu(const TSharedPtr<I
 					return;
 				}
 
-				FToolMenuSection& UnnamedSection = Submenu->FindOrAddSection("", LOCTEXT("UnnamedLabel", ""));
+				FToolMenuSection& UnnamedSection = Submenu->FindOrAddSection(NAME_None);
 
 				constexpr bool bInShouldCloseWindowAfterMenuSelection = true;
 				FMenuBuilder PreviewProfilesSelectionMenuBuilder(bInShouldCloseWindowAfterMenuSelection, nullptr);
@@ -2215,7 +2215,7 @@ void PopulateCameraMenu(UToolMenu* InMenu)
 	}
 
 	FToolMenuSection& PerspectiveCameraSection =
-		InMenu->FindOrAddSection("LevelViewportCameraType_Perspective", LOCTEXT("UnnamedLabel", ""));
+		InMenu->FindOrAddSection("LevelViewportCameraType_Perspective");
 	PerspectiveCameraSection.AddMenuEntry(FEditorViewportCommands::Get().Perspective);
 
 	FToolMenuSection& OrthographicCameraSection =
@@ -2253,7 +2253,7 @@ void ExtendCameraSubmenu(FName InCameraOptionsSubmenuName)
 				FToolMenuInsert InsertPosition("LevelViewportCameraType_Ortho", EToolMenuInsertType::After);
 
 				FToolMenuSection& UnnamedSection =
-					InDynamicMenu->FindOrAddSection("", LOCTEXT("UnnamedLabel", ""), InsertPosition);
+					InDynamicMenu->FindOrAddSection(NAME_None, FText(), InsertPosition);
 				UnnamedSection.AddSeparator("CameraSubmenuSeparator");
 
 				UnnamedSection.AddEntry(FToolMenuEntry::InitWidget(
@@ -2639,11 +2639,11 @@ FToolMenuEntry CreatePerformanceAndScalabilitySubmenu()
 	return FToolMenuEntry::InitSubMenu(
 		"PerformanceAndScalability",
 		LOCTEXT("PerformanceAndScalabilitySubmenuLabel", "Performance and Scalability"),
-		LOCTEXT("PerformanceAndScalabilitySubmenuTooltip", ""),
+		LOCTEXT("PerformanceAndScalabilitySubmenuTooltip", "Performance and scalability tools tied to this viewport."),
 		FNewToolMenuDelegate::CreateLambda(
 			[](UToolMenu* Submenu) -> void
 			{
-				FToolMenuSection& UnnamedSection = Submenu->FindOrAddSection("", LOCTEXT("UnnamedLabel", ""));
+				FToolMenuSection& UnnamedSection = Submenu->FindOrAddSection(NAME_None);
 
 				UnnamedSection.AddEntry(CreateToggleRealtimeEntry());
 
