[Reinstancing] unregister old actor components before spawning its replacement to avoid collision with component registration

diff --git a/Engine/Source/Editor/UnrealEd/Private/Kismet2/KismetReinstanceUtilities.cpp b/Engine/Source/Editor/UnrealEd/Private/Kismet2/KismetReinstanceUtilities.cpp
--- a/Engine/Source/Editor/UnrealEd/Private/Kismet2/KismetReinstanceUtilities.cpp
+++ b/Engine/Source/Editor/UnrealEd/Private/Kismet2/KismetReinstanceUtilities.cpp
@@ -2437,7 +2437,15 @@ static void ReplaceActorHelper(AActor* OldActor, UClass* OldClass, UObject*& New
 	OldActor->UObject::Rename(nullptr, OldActor->GetOuter(), REN_DoNotDirty | REN_DontCreateRedirectors);
 
 	const bool bPackageNewlyCreated = OldActor->GetExternalPackage() && OldActor->GetExternalPackage()->HasAnyPackageFlags(PKG_NewlyCreated);
-	
+
+	FActorAttachmentData& CurrentAttachmentData = ActorAttachmentData.FindChecked(OldActor);
+	ActorAttachmentData.Remove(OldActor);
+
+	// Unregister native components so we don't copy any sub-components they generate for themselves (like UCameraComponent does)
+	// Perform this before spawning the new one to avoid collisions for components registering to external systems (e.g., actor will reuse the same guid)
+	bool bHadRegisteredComponents = OldActor->HasActorRegisteredAllComponents();
+	OldActor->UnregisterAllComponents();
+
 	AActor* NewActor = nullptr;
 	{
 		FMakeClassSpawnableOnScope TemporarilySpawnable(SpawnClass);
@@ -2465,24 +2473,19 @@ static void ReplaceActorHelper(AActor* OldActor, UClass* OldClass, UObject*& New
 		}
 	}
 
-
-	NewUObject = NewActor;
 	// store the new actor for the second pass (NOTE: this detaches 
 	// OldActor from all child/parent attachments)
 	//
 	// running the NewActor's construction-script is saved for that 
 	// second pass (because the construction-script may reference 
 	// another instance that hasn't been replaced yet).
-	bool bHadRegisteredComponents = OldActor->HasActorRegisteredAllComponents();
-	FActorAttachmentData& CurrentAttachmentData = ActorAttachmentData.FindChecked(OldActor);
+	NewUObject = NewActor;
+
 	ReplacementActors.Add(FActorReplacementHelper(NewActor, OldActor, MoveTemp(CurrentAttachmentData)));
-	ActorAttachmentData.Remove(OldActor);
 
 	ReinstancedObjectsWeakReferenceMap.Add(OldActor, NewUObject);
 
 	OldActor->DestroyConstructedComponents(); // don't want to serialize components from the old actor
-												// Unregister native components so we don't copy any sub-components they generate for themselves (like UCameraComponent does)
-	OldActor->UnregisterAllComponents();
 
 	// Unregister any native components, might have cached state based on properties we are going to overwrite
 	NewActor->UnregisterAllComponents();
