World Partition - HLOD - Fix debugging code leftover...

diff --git a/Engine/Source/Editor/UnrealEd/Private/WorldPartition/WorldPartitionHLODsBuilder.cpp b/Engine/Source/Editor/UnrealEd/Private/WorldPartition/WorldPartitionHLODsBuilder.cpp
--- a/Engine/Source/Editor/UnrealEd/Private/WorldPartition/WorldPartitionHLODsBuilder.cpp
+++ b/Engine/Source/Editor/UnrealEd/Private/WorldPartition/WorldPartitionHLODsBuilder.cpp
@@ -393,45 +393,42 @@ bool UWorldPartitionHLODsBuilder::BuildHLODActors()
 			return false;
 		}
 
-		if (0)
+		UE_LOG(LogWorldPartitionHLODsBuilder, Display, TEXT("#### Building %d HLOD actors ####"), HLODActorsToBuild.Num());
+		if (bResumeBuild)
 		{
-			UE_LOG(LogWorldPartitionHLODsBuilder, Display, TEXT("#### Building %d HLOD actors ####"), HLODActorsToBuild.Num());
-			if (bResumeBuild)
-			{
-				UE_LOG(LogWorldPartitionHLODsBuilder, Display, TEXT("#### Resuming build at %d ####"), ResumeBuildIndex);
-			}
+			UE_LOG(LogWorldPartitionHLODsBuilder, Display, TEXT("#### Resuming build at %d ####"), ResumeBuildIndex);
+		}
 
-			for (int32 CurrentActor = ResumeBuildIndex; CurrentActor < HLODActorsToBuild.Num(); ++CurrentActor)
-			{
-				TRACE_BOOKMARK(TEXT("BuildHLOD Start - %d"), CurrentActor);
+		for (int32 CurrentActor = ResumeBuildIndex; CurrentActor < HLODActorsToBuild.Num(); ++CurrentActor)
+		{
+			TRACE_BOOKMARK(TEXT("BuildHLOD Start - %d"), CurrentActor);
 
-				{
-					const FGuid& HLODActorGuid = HLODActorsToBuild[CurrentActor];
+			{
+				const FGuid& HLODActorGuid = HLODActorsToBuild[CurrentActor];
 
-					FWorldPartitionReference ActorRef(WorldPartition, HLODActorGuid);
+				FWorldPartitionReference ActorRef(WorldPartition, HLODActorGuid);
 
-					AWorldPartitionHLOD* HLODActor = CastChecked<AWorldPartitionHLOD>(ActorRef.GetActor());
+				AWorldPartitionHLOD* HLODActor = CastChecked<AWorldPartitionHLOD>(ActorRef.GetActor());
 
-					UE_LOG(LogWorldPartitionHLODsBuilder, Display, TEXT("[%d / %d] Building HLOD actor %s..."), CurrentActor + 1, HLODActorsToBuild.Num(), *HLODActor->GetActorLabel());
+				UE_LOG(LogWorldPartitionHLODsBuilder, Display, TEXT("[%d / %d] Building HLOD actor %s..."), CurrentActor + 1, HLODActorsToBuild.Num(), *HLODActor->GetActorLabel());
 
-					// Simulate an engine tick to make sure engine & render resources that are queued for deletion are processed.
-					FWorldPartitionHelpers::FakeEngineTick(World);
+				// Simulate an engine tick to make sure engine & render resources that are queued for deletion are processed.
+				FWorldPartitionHelpers::FakeEngineTick(World);
 
-					HLODActor->BuildHLOD(bForceBuild);
+				HLODActor->BuildHLOD(bForceBuild);
 
-					bool bSaved = SaveHLODActor(HLODActor);
-					if (!bSaved)
-					{
-						return false;
-					}
+				bool bSaved = SaveHLODActor(HLODActor);
+				if (!bSaved)
+				{
+					return false;
 				}
+			}
 
-				TRACE_BOOKMARK(TEXT("BuildHLOD End - %d"), CurrentActor);
+			TRACE_BOOKMARK(TEXT("BuildHLOD End - %d"), CurrentActor);
 
-				if (FWorldPartitionHelpers::ShouldCollectGarbage())
-				{
-					FWorldPartitionHelpers::DoCollectGarbage();
-				}
+			if (FWorldPartitionHelpers::ShouldCollectGarbage())
+			{
+				FWorldPartitionHelpers::DoCollectGarbage();
 			}
 		}
 
