T3D Level Exporter calls into Actor T3D exporters
* Enables T3D customization based on Actor type
* Maintaining legacy code path behind CVar `UnrealEd.T3DUseActorExporters` as a precaution

diff --git a/Engine/Source/Editor/UnrealEd/Classes/Exporters/ActorExporterT3D.h b/Engine/Source/Editor/UnrealEd/Classes/Exporters/ActorExporterT3D.h
new file mode 100644
--- /dev/null
+++ b/Engine/Source/Editor/UnrealEd/Classes/Exporters/ActorExporterT3D.h
@@ -0,0 +1,50 @@
+﻿// Copyright Epic Games, Inc. All Rights Reserved.
+
+#pragma once
+
+#include "CoreMinimal.h"
+#include "UObject/ObjectMacros.h"
+#include "UObject/UObjectGlobals.h"
+#include "Exporters/Exporter.h"
+#include "ActorExporterT3D.generated.h"
+
+class FExportObjectInnerContext;
+class UActorComponent;
+
+UCLASS(MinimalAPI)
+class UActorExporterT3D : public UExporter
+{
+public:
+	GENERATED_BODY()
+
+public:
+	UNREALED_API UActorExporterT3D(const FObjectInitializer& ObjectInitializer = FObjectInitializer());
+
+	//~ Begin UExporter Interface
+	UNREALED_API virtual bool ExportText( const FExportObjectInnerContext* Context, UObject* Object, const TCHAR* Type, FOutputDevice& Ar, FFeedbackContext* Warn, uint32 PortFlags=0 ) override;
+	//~ End UExporter Interface
+};
+
+UCLASS()
+class UGroupActorExporterT3D : public UActorExporterT3D
+{
+	GENERATED_BODY()
+public:
+	UGroupActorExporterT3D(const FObjectInitializer& ObjectInitializer = FObjectInitializer());
+
+	//~ Begin UExporter Interface
+	virtual bool ExportText( const FExportObjectInnerContext* Context, UObject* Object, const TCHAR* Type, FOutputDevice& Ar, FFeedbackContext* Warn, uint32 PortFlags=0 ) override;
+	//~ End UExporter Interface
+};
+
+UCLASS()
+class UPhysicsVolumeExporterT3D : public UActorExporterT3D
+{
+	GENERATED_BODY()
+public:
+	UPhysicsVolumeExporterT3D(const FObjectInitializer& ObjectInitializer = FObjectInitializer());
+
+	//~ Begin UExporter Interface
+	virtual bool ExportText( const FExportObjectInnerContext* Context, UObject* Object, const TCHAR* Type, FOutputDevice& Ar, FFeedbackContext* Warn, uint32 PortFlags=0 ) override;
+	//~ End UExporter Interface
+};
\ No newline at end of file
diff --git a/Engine/Source/Editor/UnrealEd/Private/EditorExporters.cpp b/Engine/Source/Editor/UnrealEd/Private/EditorExporters.cpp
--- a/Engine/Source/Editor/UnrealEd/Private/EditorExporters.cpp
+++ b/Engine/Source/Editor/UnrealEd/Private/EditorExporters.cpp
@@ -86,6 +86,7 @@
 #include "Particles/Emitter.h"
 #include "Engine/Brush.h"
 #include "Components/BrushComponent.h"
+#include "Exporters/ActorExporterT3D.h"
 
 #include "Misc/App.h"
 #include "Widgets/Notifications/SNotificationList.h"
@@ -93,6 +94,17 @@
 
 DEFINE_LOG_CATEGORY_STATIC(LogEditorExporters, Log, All);
 
+namespace Private
+{
+	bool bLevelExportT3D_UseActorExporters = true;
+	static FAutoConsoleVariableRef CVarDumpHandlerTick(
+		TEXT("UnrealEd.T3DUseActorExporters"),
+		bLevelExportT3D_UseActorExporters,
+		TEXT(" \n"),
+		ECVF_Default
+	);
+}
+
 struct FScopedFbxExporterInstance
 {
 	FScopedFbxExporterInstance()
@@ -455,6 +467,109 @@ void ExporterHelper_DumpPackageInners(const FExportObjectInnerContext* Context,
 	}
 }
 
+UActorExporterT3D::UActorExporterT3D(const FObjectInitializer& ObjectInitializer)
+	: Super(ObjectInitializer)
+{
+	SupportedClass = AActor::StaticClass();
+	bText = true;
+	PreferredFormatIndex = 0;
+	FormatExtension.Add(TEXT("T3D"));
+	FormatExtension.Add(TEXT("COPY"));
+	FormatDescription.Add(TEXT("Unreal Actor text"));
+	FormatDescription.Add(TEXT("Unreal Actor text"));
+}
+
+bool UActorExporterT3D::ExportText(const FExportObjectInnerContext* Context, UObject* Object, const TCHAR* Type, FOutputDevice& Ar, FFeedbackContext* Warn, uint32 PortFlags)
+{
+	AActor* Actor = Cast<AActor>(Object);
+	if (Actor->ShouldExport())
+	{
+		// Temporarily unbind dynamic delegates so we don't export the bindings.
+		UBlueprintGeneratedClass::UnbindDynamicDelegates(Actor->GetClass(), Actor);
+
+		AActor* ParentActor = Actor->GetAttachParentActor();
+		FName SocketName = Actor->GetAttachParentSocketName();
+		FDetachmentTransformRules DetachmentTransformRules = FDetachmentTransformRules::KeepWorldTransform;
+		DetachmentTransformRules.bCallModify = false;
+		Actor->DetachFromActor(DetachmentTransformRules);
+
+		const bool IsUsingActorFolders = Actor->GetLevel() && Actor->GetLevel()->IsUsingActorFolders();
+
+		FString ParentActorString = ( ParentActor ? FString::Printf(TEXT(" ParentActor=%s"), *ParentActor->GetName() ) : TEXT(""));
+		FString SocketNameString = ( (ParentActor && SocketName != NAME_None) ? FString::Printf(TEXT(" SocketName=%s"), *SocketName.ToString() ) : TEXT(""));
+		FString GroupActor = (Actor->GroupActor ? FString::Printf(TEXT(" GroupActor=%s"), *Actor->GroupActor->GetName() ) : TEXT(""));
+		FString GroupFolder = (Actor->GroupActor ? FString::Printf(TEXT(" GroupFolder=%s"), *Actor->GroupActor->GetFolderPath().ToString() ) : TEXT(""));
+		FString ActorFolderPath = (IsUsingActorFolders ? FString::Printf(TEXT(" ActorFolderPath=\"%s\""), *Actor->GetFolderPath().ToString()) : TEXT(""));
+		FString CopyPasteId = (Actor->CopyPasteId != INDEX_NONE) ? FString::Printf(TEXT(" CopyPasteId=%d"), Actor->CopyPasteId ) : TEXT("");
+		FString ContentBundleGuid = (Actor->GetContentBundleGuid().IsValid() ? FString::Printf(TEXT(" ActorContentBundleGuid=%s"), *Actor->GetContentBundleGuid().ToString()) : TEXT(""));
+		FString ExternalDataLayer = (Actor->GetExternalDataLayerAsset() ? FString::Printf(TEXT(" ExternalDataLayerAsset=%s"), *FSoftObjectPath(Actor->GetExternalDataLayerAsset()).ToString()) : TEXT(""));
+
+		Ar.Logf( TEXT("%sBegin Actor Class=%s Name=%s Archetype=%s%s%s%s%s%s%s%s%s"), 
+			FCString::Spc(TextIndent), *Actor->GetClass()->GetPathName(), *Actor->GetName(),
+			*FObjectPropertyBase::GetExportPath(Actor->GetArchetype(), nullptr, nullptr, (PortFlags | PPF_Delimited) & ~PPF_ExportsNotFullyQualified), 
+			*ParentActorString, *SocketNameString, *GroupActor, *GroupFolder, *ActorFolderPath, *CopyPasteId, *ContentBundleGuid, *ExternalDataLayer);
+
+		// When exporting for diffs, export paths can cause false positives. since diff files don't get imported, we can
+		// skip adding this info the file.
+		if (!(PortFlags & PPF_ForDiff))
+		{
+			// Emit the actor path
+			Ar.Logf(TEXT(" ExportPath=%s"), *FObjectPropertyBase::GetExportPath(Actor, nullptr, nullptr, (PortFlags | PPF_Delimited) & ~PPF_ExportsNotFullyQualified));
+		}
+
+		Ar.Logf(LINE_TERMINATOR);
+
+		ExportRootScope = Actor;
+		ExportObjectInner( Context, Actor, Ar, PortFlags | PPF_ExportsNotFullyQualified );
+		ExportRootScope = nullptr;
+
+		Ar.Logf( TEXT("%sEnd Actor\r\n"), FCString::Spc(TextIndent) );
+		Actor->AttachToActor(ParentActor, FAttachmentTransformRules::KeepWorldTransform, SocketName);
+
+		// Restore dynamic delegate bindings.
+		UBlueprintGeneratedClass::BindDynamicDelegates(Actor->GetClass(), Actor);
+	}
+	else if (GEditor)
+	{
+		GEditor->GetSelectedActors()->Deselect(Actor);
+	}
+	return true;
+}
+
+UGroupActorExporterT3D::UGroupActorExporterT3D(const FObjectInitializer& ObjectInitializer)
+{
+	SupportedClass = AGroupActor::StaticClass();
+}
+
+bool UGroupActorExporterT3D::ExportText(const FExportObjectInnerContext* Context, UObject* Object, const TCHAR* Type, FOutputDevice& Ar,
+	FFeedbackContext* Warn, uint32 PortFlags)
+{
+	// Don't export group actors
+	return true;
+}
+
+UPhysicsVolumeExporterT3D::UPhysicsVolumeExporterT3D(const FObjectInitializer& ObjectInitializer)
+{
+	SupportedClass = APhysicsVolume::StaticClass();
+}
+
+bool UPhysicsVolumeExporterT3D::ExportText(const FExportObjectInnerContext* Context, UObject* Object, const TCHAR* Type, FOutputDevice& Ar,
+                                           FFeedbackContext* Warn, uint32 PortFlags)
+{
+	APhysicsVolume* PhysicsVolume = Cast<APhysicsVolume>(Object);
+	if (UWorld* World = PhysicsVolume->GetWorld())
+	{
+		if (World->GetDefaultPhysicsVolume() == PhysicsVolume)
+		{
+			// Don't export the default physics volume, as it doesn't have a UModel associated with it
+			// and thus will not import properly.
+			return true;
+		}
+		return Super::ExportText(Context, Object, Type, Ar, Warn, PortFlags);
+	}
+	return true;
+}
+
 ULevelExporterT3D::ULevelExporterT3D(const FObjectInitializer& ObjectInitializer)
 	: Super(ObjectInitializer)
 {
@@ -535,67 +650,89 @@ bool ULevelExporterT3D::ExportText( const FExportObjectInnerContext* Context, UO
 		{
 			break;
 		}
-		// Don't export the default physics volume, as it doesn't have a UModel associated with it
-		// and thus will not import properly.
-		if ( Actor == DefaultPhysicsVolume )
+		
+		if (Private::bLevelExportT3D_UseActorExporters)
 		{
-			continue;
+			// Ensure the actor is currently selected
+			if( Actor && ( bAllActors || IsObjectSelectedForExport(Context, Actor) ) )
+			{
+				UExporter* Exporter = UExporter::FindExporter(Actor, Type);
+				if (Exporter)
+				{
+					Exporter->TextIndent = TextIndent;
+					Exporter->bSelectedOnly = bSelectedOnly;
+					bool Result = Exporter->ExportText(Context, Actor, Type, Ar, Warn, PortFlags);
+					if (Result == false)
+					{
+						return Result;
+					}
+				}
+
+				SlowTask.EnterProgressFrame();
+			}
 		}
-		// Ensure actor is not a group if grouping is disabled and that the actor is currently selected
-		if( Actor && !Actor->IsA(AGroupActor::StaticClass()) &&
-			( bAllActors || IsObjectSelectedForExport(Context, Actor) ) )
+		else
 		{
-			if (Actor->ShouldExport())
+			if ( Actor == DefaultPhysicsVolume )
+			{
+				continue;
+			}
+			// Ensure actor is not a group if grouping is disabled and that the actor is currently selected
+			if( Actor && !Actor->IsA(AGroupActor::StaticClass()) &&
+				( bAllActors || IsObjectSelectedForExport(Context, Actor) ) )
 			{
-				// Temporarily unbind dynamic delegates so we don't export the bindings.
-				UBlueprintGeneratedClass::UnbindDynamicDelegates(Actor->GetClass(), Actor);
-
-				AActor* ParentActor = Actor->GetAttachParentActor();
-				FName SocketName = Actor->GetAttachParentSocketName();
-				FDetachmentTransformRules DetachmentTransformRules = FDetachmentTransformRules::KeepWorldTransform;
-				DetachmentTransformRules.bCallModify = false;
-				Actor->DetachFromActor(DetachmentTransformRules);
-
-				FString ParentActorString = ( ParentActor ? FString::Printf(TEXT(" ParentActor=%s"), *ParentActor->GetName() ) : TEXT(""));
-				FString SocketNameString = ( (ParentActor && SocketName != NAME_None) ? FString::Printf(TEXT(" SocketName=%s"), *SocketName.ToString() ) : TEXT(""));
-				FString GroupActor = (Actor->GroupActor ? FString::Printf(TEXT(" GroupActor=%s"), *Actor->GroupActor->GetName() ) : TEXT(""));
-				FString GroupFolder = (Actor->GroupActor ? FString::Printf(TEXT(" GroupFolder=%s"), *Actor->GroupActor->GetFolderPath().ToString() ) : TEXT(""));
-				FString ActorFolderPath = (Level->IsUsingActorFolders() ? FString::Printf(TEXT(" ActorFolderPath=\"%s\""), *Actor->GetFolderPath().ToString()) : TEXT(""));
-				FString CopyPasteId = (Actor->CopyPasteId != INDEX_NONE) ? FString::Printf(TEXT(" CopyPasteId=%d"), Actor->CopyPasteId ) : TEXT("");
-				FString ContentBundleGuid = (Actor->GetContentBundleGuid().IsValid() ? FString::Printf(TEXT(" ActorContentBundleGuid=%s"), *Actor->GetContentBundleGuid().ToString()) : TEXT(""));
-				FString ExternalDataLayer = (Actor->GetExternalDataLayerAsset() ? FString::Printf(TEXT(" ExternalDataLayerAsset=%s"), *FSoftObjectPath(Actor->GetExternalDataLayerAsset()).ToString()) : TEXT(""));
-
-				Ar.Logf( TEXT("%sBegin Actor Class=%s Name=%s Archetype=%s%s%s%s%s%s%s%s%s"), 
-					FCString::Spc(TextIndent), *Actor->GetClass()->GetPathName(), *Actor->GetName(),
-					*FObjectPropertyBase::GetExportPath(Actor->GetArchetype(), nullptr, nullptr, (PortFlags | PPF_Delimited) & ~PPF_ExportsNotFullyQualified), 
-					*ParentActorString, *SocketNameString, *GroupActor, *GroupFolder, *ActorFolderPath, *CopyPasteId, *ContentBundleGuid, *ExternalDataLayer);
-
-				// When exporting for diffs, export paths can cause false positives. since diff files don't get imported, we can
-				// skip adding this info the file.
-				if (!(PortFlags & PPF_ForDiff))
+				if (Actor->ShouldExport())
 				{
-					// Emit the actor path
-					Ar.Logf(TEXT(" ExportPath=%s"), *FObjectPropertyBase::GetExportPath(Actor, nullptr, nullptr, (PortFlags | PPF_Delimited) & ~PPF_ExportsNotFullyQualified));
-				}
+					// Temporarily unbind dynamic delegates so we don't export the bindings.
+					UBlueprintGeneratedClass::UnbindDynamicDelegates(Actor->GetClass(), Actor);
+
+					AActor* ParentActor = Actor->GetAttachParentActor();
+					FName SocketName = Actor->GetAttachParentSocketName();
+					FDetachmentTransformRules DetachmentTransformRules = FDetachmentTransformRules::KeepWorldTransform;
+					DetachmentTransformRules.bCallModify = false;
+					Actor->DetachFromActor(DetachmentTransformRules);
+
+					FString ParentActorString = ( ParentActor ? FString::Printf(TEXT(" ParentActor=%s"), *ParentActor->GetName() ) : TEXT(""));
+					FString SocketNameString = ( (ParentActor && SocketName != NAME_None) ? FString::Printf(TEXT(" SocketName=%s"), *SocketName.ToString() ) : TEXT(""));
+					FString GroupActor = (Actor->GroupActor ? FString::Printf(TEXT(" GroupActor=%s"), *Actor->GroupActor->GetName() ) : TEXT(""));
+					FString GroupFolder = (Actor->GroupActor ? FString::Printf(TEXT(" GroupFolder=%s"), *Actor->GroupActor->GetFolderPath().ToString() ) : TEXT(""));
+					FString ActorFolderPath = (Level->IsUsingActorFolders() ? FString::Printf(TEXT(" ActorFolderPath=\"%s\""), *Actor->GetFolderPath().ToString()) : TEXT(""));
+					FString CopyPasteId = (Actor->CopyPasteId != INDEX_NONE) ? FString::Printf(TEXT(" CopyPasteId=%d"), Actor->CopyPasteId ) : TEXT("");
+					FString ContentBundleGuid = (Actor->GetContentBundleGuid().IsValid() ? FString::Printf(TEXT(" ActorContentBundleGuid=%s"), *Actor->GetContentBundleGuid().ToString()) : TEXT(""));
+					FString ExternalDataLayer = (Actor->GetExternalDataLayerAsset() ? FString::Printf(TEXT(" ExternalDataLayerAsset=%s"), *FSoftObjectPath(Actor->GetExternalDataLayerAsset()).ToString()) : TEXT(""));
+
+					Ar.Logf( TEXT("%sBegin Actor Class=%s Name=%s Archetype=%s%s%s%s%s%s%s%s%s"), 
+						FCString::Spc(TextIndent), *Actor->GetClass()->GetPathName(), *Actor->GetName(),
+						*FObjectPropertyBase::GetExportPath(Actor->GetArchetype(), nullptr, nullptr, (PortFlags | PPF_Delimited) & ~PPF_ExportsNotFullyQualified), 
+						*ParentActorString, *SocketNameString, *GroupActor, *GroupFolder, *ActorFolderPath, *CopyPasteId, *ContentBundleGuid, *ExternalDataLayer);
+
+					// When exporting for diffs, export paths can cause false positives. since diff files don't get imported, we can
+					// skip adding this info the file.
+					if (!(PortFlags & PPF_ForDiff))
+					{
+						// Emit the actor path
+						Ar.Logf(TEXT(" ExportPath=%s"), *FObjectPropertyBase::GetExportPath(Actor, nullptr, nullptr, (PortFlags | PPF_Delimited) & ~PPF_ExportsNotFullyQualified));
+					}
 
-				Ar.Logf(LINE_TERMINATOR);
+					Ar.Logf(LINE_TERMINATOR);
 
-				ExportRootScope = Actor;
-				ExportObjectInner( Context, Actor, Ar, PortFlags | PPF_ExportsNotFullyQualified );
-				ExportRootScope = nullptr;
+					ExportRootScope = Actor;
+					ExportObjectInner( Context, Actor, Ar, PortFlags | PPF_ExportsNotFullyQualified );
+					ExportRootScope = nullptr;
 
-				Ar.Logf( TEXT("%sEnd Actor\r\n"), FCString::Spc(TextIndent) );
-				Actor->AttachToActor(ParentActor, FAttachmentTransformRules::KeepWorldTransform, SocketName);
+					Ar.Logf( TEXT("%sEnd Actor\r\n"), FCString::Spc(TextIndent) );
+					Actor->AttachToActor(ParentActor, FAttachmentTransformRules::KeepWorldTransform, SocketName);
 
-				// Restore dynamic delegate bindings.
-				UBlueprintGeneratedClass::BindDynamicDelegates(Actor->GetClass(), Actor);
-			}
-			else if (GEditor)
-			{
-				GEditor->GetSelectedActors()->Deselect(Actor);
-			}
+					// Restore dynamic delegate bindings.
+					UBlueprintGeneratedClass::BindDynamicDelegates(Actor->GetClass(), Actor);
+				}
+				else if (GEditor)
+				{
+					GEditor->GetSelectedActors()->Deselect(Actor);
+				}
 
-			SlowTask.EnterProgressFrame();
+				SlowTask.EnterProgressFrame();
+			}
 		}
 	}
 
