Cooker: CookPackageSplitter: Fix another case of "PackageSplitter found the Generator package still in memory after it should have been deleted by GC." The newly encountered edge case is that the generator (or a generated) package is in save state on a cook worker, working on on CachedCookedPlatformData, but does not yet have the GenerationHelper reference set. We need to expect the generationhelper to remain referenced in that case.
#rnx

diff --git a/Engine/Source/Editor/UnrealEd/Private/Cooker/CookGenerationHelper.cpp b/Engine/Source/Editor/UnrealEd/Private/Cooker/CookGenerationHelper.cpp
--- a/Engine/Source/Editor/UnrealEd/Private/Cooker/CookGenerationHelper.cpp
+++ b/Engine/Source/Editor/UnrealEd/Private/Cooker/CookGenerationHelper.cpp
@@ -1481,19 +1481,35 @@ void FGenerationHelper::PreGarbageCollectGCLifetimeData()
 	// Starts at one because the caller of PreGarbageCollect has a ref
 	uint32 HoldForGCRefCounts = 1;
 	HoldForGCRefCounts += ReferenceFromKeepForAllSavedOrGC.IsValid() ? 1 : 0;
-	// Every stalled package will be holding a refcount, but is allowed to be demoted and released
-	// if we are going to garbage collect.
+	// If the owner or any generated package is in progress and not stalled, do not uninitialize, because
+	// the cooker might keep the packgae referenced (if it is e.g. in save state) even if the cooker does not
+	// have a reference to the GenerationHelper from that package.
+	// For stalled packages, if a generated package is stalled, we want to keep it in memory until GC,
+	// but now that we have reached GC that stalled package is allowed to be demoted and released and does
+	// not prevent uninitialize.
 	// If the generator package is stalled, that's a complex case that we don't need to handle optimally;
-	// just keep the entire generation helper referenced.
-	if (OwnerInfo.PackageData->IsStalled())
+	// just keep the entire generation helper referenced while the generator package is stalled.
+	// Every stalled package will be holding a refcount; we need to subtract those refcounts when deciding
+	// whether we have a reference from any non-stalled package.
+	if (OwnerInfo.PackageData->IsStalled()
+		|| OwnerInfo.PackageData->IsInProgress())
 	{
+		// Owner packagedata is stalled or in progress; do not uninitialize
 		return;
 	}
 	for (FCookGenerationInfo& Info : PackagesToGenerate)
 	{
-		if (Info.PackageData->GetParentGenerationHelper() && Info.PackageData->IsStalled())
+		if (Info.PackageData->IsStalled())
 		{
-			HoldForGCRefCounts += 1;
+			if (Info.PackageData->GetParentGenerationHelper())
+			{
+				HoldForGCRefCounts += 1;
+			}
+		}
+		else if (Info.PackageData->IsInProgress())
+		{
+			// Generated package is in progress and not stalled; do not uninitialize
+			return;
 		}
 	}
 
