Improved "Browse to Asset" override handling

Added UBrowseToAssetOverrideSubsystem to provide per-class override handling for a given instance, in addition to the per-instance override handling that AActor already had embedded into it.
This allows the override logic to work more reliably in certain cases (eg, when mapping a given class to a given alias) regardless of how the actor was created, or whether it had the per-instance override data added to it.

diff --git a/Engine/Source/Editor/UnrealEd/Private/AssetSelection.cpp b/Engine/Source/Editor/UnrealEd/Private/AssetSelection.cpp
--- a/Engine/Source/Editor/UnrealEd/Private/AssetSelection.cpp
+++ b/Engine/Source/Editor/UnrealEd/Private/AssetSelection.cpp
@@ -54,6 +54,7 @@
 #include "ISourceControlModule.h"
 #include "ISourceControlProvider.h"
 #include "Misc/MessageDialog.h"
+#include "Subsystems/BrowseToAssetOverrideSubsystem.h"
 #include "Subsystems/PlacementSubsystem.h"
 #include "Elements/Framework/EngineElementsLibrary.h"
 #include "Elements/Framework/TypedElementRegistry.h"
@@ -98,6 +99,8 @@ namespace AssetSelectionUtils
 
 			if( FirstActor && !FirstActor->IsTemplate() )
 			{
+				UBrowseToAssetOverrideSubsystem* BrowseToAssetOverrideSubsystem = UBrowseToAssetOverrideSubsystem::Get();
+
 				UClass* FirstClass = FirstActor->GetClass();
 				UObject* FirstArchetype = FirstActor->GetArchetype();
 
@@ -317,7 +320,7 @@ namespace AssetSelectionUtils
 						++ActorInfo.NumSimulationChanges;
 					}
 
-					if (!CurrentActor->GetBrowseToAssetOverride().IsEmpty())
+					if (!BrowseToAssetOverrideSubsystem->GetBrowseToAssetOverride(CurrentActor).IsNone())
 					{
 						ActorInfo.bHaveBrowseOverride = true;
 					}
diff --git a/Engine/Source/Editor/UnrealEd/Private/EditorEngine.cpp b/Engine/Source/Editor/UnrealEd/Private/EditorEngine.cpp
--- a/Engine/Source/Editor/UnrealEd/Private/EditorEngine.cpp
+++ b/Engine/Source/Editor/UnrealEd/Private/EditorEngine.cpp
@@ -250,6 +250,7 @@
 #include "ToolMenus.h"
 #include "IToolMenusEditorModule.h"
 #include "Subsystems/AssetEditorSubsystem.h"
+#include "Subsystems/BrowseToAssetOverrideSubsystem.h"
 #include "LevelEditorSubsystem.h"
 #include "Engine/LevelScriptActor.h"
 #include "UObject/UnrealType.h"
@@ -3311,7 +3312,8 @@ void UEditorEngine::GetAssetsToSyncToContentBrowser(TArray<FAssetData>& Assets,
 	{
 		if (UTypedElementSelectionSet* SelectionSet = GEditor->GetSelectedActors()->GetElementSelectionSet())
 		{
-			SelectionSet->ForEachSelectedElementHandle([&Assets, bAllowBrowseToAssetOverride](const FTypedElementHandle& SelectedHandle)
+			UBrowseToAssetOverrideSubsystem* BrowseToAssetOverrideSubsystem = UBrowseToAssetOverrideSubsystem::Get();
+			SelectionSet->ForEachSelectedElementHandle([&Assets, bAllowBrowseToAssetOverride, BrowseToAssetOverrideSubsystem](const FTypedElementHandle& SelectedHandle)
 			{
 				if (bAllowBrowseToAssetOverride)
 				{
@@ -3320,13 +3322,13 @@ void UEditorEngine::GetAssetsToSyncToContentBrowser(TArray<FAssetData>& Assets,
 						if (AActor* Actor = ObjectInterface.GetObjectAs<AActor>())
 						{
 							// If BrowseToAssetOverride is set, then use the asset it points to instead of the selected asset
-							const FString& BrowseToAssetOverride = Actor->GetBrowseToAssetOverride();
-							if (!BrowseToAssetOverride.IsEmpty())
+							const FName BrowseToAssetOverride = BrowseToAssetOverrideSubsystem->GetBrowseToAssetOverride(Actor);
+							if (!BrowseToAssetOverride.IsNone())
 							{
 								if (IAssetRegistry* AssetRegistry = IAssetRegistry::Get())
 								{
 									TArray<FAssetData> FoundAssets;
-									if (AssetRegistry->GetAssetsByPackageName(*BrowseToAssetOverride, FoundAssets) && FoundAssets.Num() > 0)
+									if (AssetRegistry->GetAssetsByPackageName(BrowseToAssetOverride, FoundAssets) && FoundAssets.Num() > 0)
 									{
 										Assets.Add(FoundAssets[0]);
 										return true;
diff --git a/Engine/Source/Editor/UnrealEd/Private/Subsystems/BrowseToAssetOverrideSubsystem.cpp b/Engine/Source/Editor/UnrealEd/Private/Subsystems/BrowseToAssetOverrideSubsystem.cpp
new file mode 100644
--- /dev/null
+++ b/Engine/Source/Editor/UnrealEd/Private/Subsystems/BrowseToAssetOverrideSubsystem.cpp
@@ -0,0 +1,57 @@
+// Copyright Epic Games, Inc. All Rights Reserved.
+
+#include "Subsystems/BrowseToAssetOverrideSubsystem.h"
+
+#include "Editor.h"
+#include "GameFramework/Actor.h"
+
+UBrowseToAssetOverrideSubsystem* UBrowseToAssetOverrideSubsystem::Get()
+{
+	if (GEditor)
+	{
+		return GEditor->GetEditorSubsystem<UBrowseToAssetOverrideSubsystem>();
+	}
+	return nullptr;
+}
+
+FName UBrowseToAssetOverrideSubsystem::GetBrowseToAssetOverride(const UObject* Object)
+{
+	// Actors also allow this to be overridden per-instance via meta-data
+	// If set, that takes priority over any per-class overrides
+	if (const AActor* Actor = Cast<AActor>(Object))
+	{
+		const FString& ActorBrowseToAssetOverride = Actor->GetBrowseToAssetOverride();
+		if (!ActorBrowseToAssetOverride.IsEmpty())
+		{
+			return *ActorBrowseToAssetOverride;
+		}
+	}
+
+	// Walk the class hierarchy to see if there's a valid per-class override for this instance
+	if (PerClassOverrides.Num() > 0)
+	{
+		for (const UClass* Class = Object->GetClass(); Class; Class = Class->GetSuperClass())
+		{
+			if (const FBrowseToAssetOverrideDelegate* CallbackPtr = PerClassOverrides.Find(Class->GetClassPathName()))
+			{
+				FName CallbackBrowseToAssetOverride = CallbackPtr->IsBound() ? CallbackPtr->Execute(Object) : FName();
+				if (!CallbackBrowseToAssetOverride.IsNone())
+				{
+					return CallbackBrowseToAssetOverride;
+				}
+			}
+		}
+	}
+
+	return FName();
+}
+
+void UBrowseToAssetOverrideSubsystem::RegisterBrowseToAssetOverrideForClass(const FTopLevelAssetPath& Class, FBrowseToAssetOverrideDelegate&& Callback)
+{
+	PerClassOverrides.Add(Class, MoveTemp(Callback));
+}
+	
+void UBrowseToAssetOverrideSubsystem::UnregisterBrowseToAssetOverrideForClass(const FTopLevelAssetPath& Class)
+{
+	PerClassOverrides.Remove(Class);
+}
diff --git a/Engine/Source/Editor/UnrealEd/Public/Subsystems/BrowseToAssetOverrideSubsystem.h b/Engine/Source/Editor/UnrealEd/Public/Subsystems/BrowseToAssetOverrideSubsystem.h
new file mode 100644
--- /dev/null
+++ b/Engine/Source/Editor/UnrealEd/Public/Subsystems/BrowseToAssetOverrideSubsystem.h
@@ -0,0 +1,52 @@
+// Copyright Epic Games, Inc. All Rights Reserved.
+
+#pragma once
+
+#include "Containers/Map.h"
+#include "Delegates/Delegate.h"
+#include "EditorSubsystem.h"
+
+#include "BrowseToAssetOverrideSubsystem.generated.h"
+
+using FBrowseToAssetOverrideDelegate = TDelegate<FName(const UObject*)>;
+
+UCLASS(MinimalAPI)
+class UBrowseToAssetOverrideSubsystem : public UEditorSubsystem
+{
+	GENERATED_BODY()
+
+public:
+	static UNREALED_API UBrowseToAssetOverrideSubsystem* Get();
+
+	/**
+	 * Given an object, see if it has a "browse to asset" package name override.
+	 * @return The package name of override, or None if there is no override.
+	 */
+	UNREALED_API FName GetBrowseToAssetOverride(const UObject* Object);
+
+	/**
+	 * Register a per-class override for the "browse to asset" resolution.
+	 * @note Callback should return a package name, or None if there is no override.
+	 */
+	UNREALED_API void RegisterBrowseToAssetOverrideForClass(const FTopLevelAssetPath& Class, FBrowseToAssetOverrideDelegate&& Callback);
+
+	template <typename ClassType>
+	void RegisterBrowseToAssetOverrideForClass(FBrowseToAssetOverrideDelegate&& Callback)
+	{
+		RegisterBrowseToAssetOverrideForClass(ClassType::StaticClass()->GetClassPathName(), MoveTemp(Callback));
+	}
+	
+	/**
+	 * Unregister a per-class override for the "browse to asset" resolution.
+	 */
+	UNREALED_API void UnregisterBrowseToAssetOverrideForClass(const FTopLevelAssetPath& Class);
+
+	template <typename ClassType>
+	void UnregisterBrowseToAssetOverrideForClass()
+	{
+		UnregisterBrowseToAssetOverrideForClass(ClassType::StaticClass()->GetClassPathName());
+	}
+	
+private:
+	TMap<FTopLevelAssetPath, FBrowseToAssetOverrideDelegate> PerClassOverrides;
+};
