Show Rotation and Scale snapping in the top-level of the new Level and Static Mesh editor viewport toolbars

diff --git a/Engine/Source/Editor/UnrealEd/Private/STransformViewportToolbar.cpp b/Engine/Source/Editor/UnrealEd/Private/STransformViewportToolbar.cpp
--- a/Engine/Source/Editor/UnrealEd/Private/STransformViewportToolbar.cpp
+++ b/Engine/Source/Editor/UnrealEd/Private/STransformViewportToolbar.cpp
@@ -248,20 +248,20 @@ TSharedRef< SWidget > STransformViewportToolBar::MakeTransformToolBar( const TSh
 		static FName PositionSnapName = FName(TEXT("PositionSnap"));
 
 		// Setup a GridSnapSetting with the UICommand
-		ToolbarBuilder.AddWidget(	
+		ToolbarBuilder.AddWidget(
 			SNew(SViewportToolBarComboMenu)
 				.IsChecked(this, &STransformViewportToolBar::IsLocationGridSnapChecked)
 				.OnCheckStateChanged(this, &STransformViewportToolBar::HandleToggleLocationGridSnap)
-				.Label(this, &STransformViewportToolBar::GetLocationGridLabel)
+				.Label(TAttribute<FText>::Create(&UE::UnrealEd::GetLocationGridLabel))
 				.OnGetMenuContent(this, &STransformViewportToolBar::FillLocationGridSnapMenu)
 				.ToggleButtonToolTip(Command->GetDescription())
 				.MenuButtonToolTip(LOCTEXT("LocationGridSnap_ToolTip", "Set the Position Grid Snap value"))
 				.Icon(Command->GetIcon())
 				.MinDesiredButtonWidth(24.0f)
 				.ParentToolBar(SharedThis(this)),
-			PositionSnapName, 
-			false, 
-			HAlign_Fill, 
+			PositionSnapName,
+			false,
+			HAlign_Fill,
 
 			// explictly specify what this widget should look like as a menu item
 			FNewMenuDelegate::CreateLambda( [this, Command]( FMenuBuilder& InMenuBuilder )
@@ -288,19 +288,19 @@ TSharedRef< SWidget > STransformViewportToolBar::MakeTransformToolBar( const TSh
 		static FName RotationSnapName = FName(TEXT("RotationSnap"));
 
 		// Setup a GridSnapSetting with the UICommand
-		ToolbarBuilder.AddWidget(	
+		ToolbarBuilder.AddWidget(
 			SNew(SViewportToolBarComboMenu)
 				.IsChecked(this, &STransformViewportToolBar::IsRotationGridSnapChecked)
 				.OnCheckStateChanged(this, &STransformViewportToolBar::HandleToggleRotationGridSnap)
-				.Label(this, &STransformViewportToolBar::GetRotationGridLabel)
+				.Label(TAttribute<FText>::Create(&UE::UnrealEd::GetRotationGridLabel))
 				.OnGetMenuContent(this, &STransformViewportToolBar::FillRotationGridSnapMenu)
 				.ToggleButtonToolTip(Command->GetDescription())
 				.MenuButtonToolTip(LOCTEXT("RotationGridSnap_ToolTip", "Set the Rotation Grid Snap value"))
 				.Icon(Command->GetIcon())
 				.ParentToolBar(SharedThis(this)),
 			RotationSnapName,
-			false, 
-			HAlign_Fill, 
+			false,
+			HAlign_Fill,
 
 			// explictly specify what this widget should look like as a menu item
 			FNewMenuDelegate::CreateLambda( [this, Command]( FMenuBuilder& InMenuBuilder )
@@ -363,22 +363,20 @@ TSharedRef< SWidget > STransformViewportToolBar::MakeTransformToolBar( const TSh
 	}
 	ToolbarBuilder.EndSection();
 
-
-
 	ToolbarBuilder.BeginSection("ScaleGridSnap");
 	{
-		// Grab the existing UICommand 
+		// Grab the existing UICommand
 		TSharedPtr<FUICommandInfo> Command = FEditorViewportCommands::Get().ScaleGridSnap;
 
 		static FName ScaleSnapName = FName(TEXT("ScaleSnap"));
 
 		// Setup a GridSnapSetting with the UICommand
-		ToolbarBuilder.AddWidget(	
+		ToolbarBuilder.AddWidget(
 			SNew(SViewportToolBarComboMenu)
-				.Cursor( EMouseCursor::Default )
-				.IsChecked(this,&STransformViewportToolBar::IsScaleGridSnapChecked)
+				.Cursor(EMouseCursor::Default)
+				.IsChecked(this, &STransformViewportToolBar::IsScaleGridSnapChecked)
 				.OnCheckStateChanged(this, &STransformViewportToolBar::HandleToggleScaleGridSnap)
-				.Label(this ,&STransformViewportToolBar::GetScaleGridLabel)
+				.Label(TAttribute<FText>::Create(&UE::UnrealEd::GetScaleGridLabel))
 				.OnGetMenuContent(this, &STransformViewportToolBar::FillScaleGridSnapMenu)
 				.ToggleButtonToolTip(Command->GetDescription())
 				.MenuButtonToolTip(LOCTEXT("ScaleGridSnap_ToolTip", "Set scaling options"))
@@ -386,8 +384,8 @@ TSharedRef< SWidget > STransformViewportToolBar::MakeTransformToolBar( const TSh
 				.MinDesiredButtonWidth(24.0f)
 				.ParentToolBar(SharedThis(this)),
 			ScaleSnapName,
-			false, 
-			HAlign_Fill, 
+			false,
+			HAlign_Fill,
 
 			// explictly specify what this widget should look like as a menu item
 			FNewMenuDelegate::CreateLambda( [this, Command]( FMenuBuilder& InMenuBuilder )
@@ -550,16 +548,6 @@ FSlateIcon STransformViewportToolBar::GetLocalToWorldIcon() const
 	return FSlateIcon(FAppStyle::GetAppStyleSetName(), LocalIcon);
 }
 
-FText STransformViewportToolBar::GetLocationGridLabel() const
-{
-	return FText::AsNumber( GEditor->GetGridSize() );
-}
-
-FText STransformViewportToolBar::GetRotationGridLabel() const
-{
-	return FText::Format(LOCTEXT("GridRotation - Number - DegreeSymbol", "{0}\u00b0"), FText::AsNumber(GEditor->GetRotGridSize().Pitch));
-}
-
 FText STransformViewportToolBar::GetLayer2DLabel() const
 {
 	const ULevelEditorViewportSettings* ViewportSettings = GetDefault<ULevelEditorViewportSettings>();
@@ -568,19 +556,8 @@ FText STransformViewportToolBar::GetLayer2DLabel() const
 	{
 		return FText::FromString(Settings2D->SnapLayers[ViewportSettings->ActiveSnapLayerIndex].Name);
 	}
-	
-	return FText();
-}
-
-FText STransformViewportToolBar::GetScaleGridLabel() const
-{
-	FNumberFormattingOptions NumberFormattingOptions;
-	NumberFormattingOptions.MaximumFractionalDigits = 5;
 
-	const float CurGridAmount = GEditor->GetScaleGridSize();
-	return (GEditor->UsePercentageBasedScaling()) 
-		? FText::AsPercent(CurGridAmount / 100.0f, &NumberFormattingOptions) 
-		: FText::AsNumber(CurGridAmount, &NumberFormattingOptions);
+	return FText();
 }
 
 FText STransformViewportToolBar::GetCameraSpeedLabel() const
diff --git a/Engine/Source/Editor/UnrealEd/Private/ViewportToolbar/UnrealEdViewportToolbar.cpp b/Engine/Source/Editor/UnrealEd/Private/ViewportToolbar/UnrealEdViewportToolbar.cpp
--- a/Engine/Source/Editor/UnrealEd/Private/ViewportToolbar/UnrealEdViewportToolbar.cpp
+++ b/Engine/Source/Editor/UnrealEd/Private/ViewportToolbar/UnrealEdViewportToolbar.cpp
@@ -344,19 +344,7 @@ FToolMenuEntry CreateLocationSnapCheckboxMenu()
 		);
 	}
 
-	Entry.ToolbarLabelOverride = TAttribute<FText>::CreateLambda(
-		[]()
-		{
-			const ULevelEditorViewportSettings* ViewportSettings = GetDefault<ULevelEditorViewportSettings>();
-
-			const TArray<float> GridSizes = ViewportSettings->bUsePowerOf2SnapSize ? ViewportSettings->Pow2GridSizes
-																				   : ViewportSettings->DecimalGridSizes;
-
-			const float CurrentGridSize = GridSizes[ViewportSettings->CurrentPosGridSize];
-
-			return FText::AsNumber(CurrentGridSize);
-		}
-	);
+	Entry.ToolbarLabelOverride = TAttribute<FText>::Create(&UE::UnrealEd::GetLocationGridLabel);
 	Entry.Icon = FSlateIcon(FAppStyle::GetAppStyleSetName(), "EditorViewport.LocationGridSnap");
 
 	return Entry;
@@ -374,19 +362,15 @@ FToolMenuEntry CreateRotationSnapCheckboxMenu()
 		);
 	}
 
-	FNewToolMenuDelegate MakeMenuDelegate = FNewToolMenuDelegate::CreateLambda(
+	const FNewToolMenuDelegate MakeMenuDelegate = FNewToolMenuDelegate::CreateLambda(
 		[RotationSnapName](UToolMenu* InToolMenu)
 		{
-			const ULevelEditorViewportSettings* ViewportSettings = GetDefault<ULevelEditorViewportSettings>();
-			TArray<float> GridSizes = ViewportSettings->bUsePowerOf2SnapSize ? ViewportSettings->Pow2GridSizes
-																			 : ViewportSettings->DecimalGridSizes;
-
-			UnrealEd::FRotationGridCheckboxListExecuteActionDelegate ExecuteDelegate =
+			const UnrealEd::FRotationGridCheckboxListExecuteActionDelegate ExecuteDelegate =
 				UnrealEd::FRotationGridCheckboxListExecuteActionDelegate::CreateUObject(
 					GEditor, &UEditorEngine::SetRotGridSize
 				);
 
-			UnrealEd::FRotationGridCheckboxListIsCheckedDelegate IsCheckedDelegate =
+			const UnrealEd::FRotationGridCheckboxListIsCheckedDelegate IsCheckedDelegate =
 				UnrealEd::FRotationGridCheckboxListIsCheckedDelegate::CreateLambda(
 					[](int CurrGridAngleIndex, ERotationGridMode InGridMode)
 					{
@@ -396,27 +380,21 @@ FToolMenuEntry CreateRotationSnapCheckboxMenu()
 					}
 				);
 
+			const TAttribute<bool> IsEnabledDelegate =
+				TAttribute<bool>::Create(&FLevelEditorActionCallbacks::RotationGridSnap_IsChecked);
+
 			InToolMenu->AddMenuEntry(
 				RotationSnapName,
 				FToolMenuEntry::InitWidget(
 					RotationSnapName,
-					UnrealEd::CreateRotationGridSnapMenu(
-						ExecuteDelegate,
-						IsCheckedDelegate,
-						TAttribute<bool>::CreateLambda(
-							[]()
-							{
-								return FLevelEditorActionCallbacks::RotationGridSnap_IsChecked();
-							}
-						)
-					),
+					UnrealEd::CreateRotationGridSnapMenu(ExecuteDelegate, IsCheckedDelegate, IsEnabledDelegate),
 					FText()
 				)
 			);
 		}
 	);
 
-	return UnrealEd::CreateCheckboxSubmenu(
+	FToolMenuEntry Entry = UnrealEd::CreateCheckboxSubmenu(
 		"RotationSnapping",
 		RotationSnapLabel,
 		FEditorViewportCommands::Get().SurfaceSnapping->MakeTooltip()->GetTextTooltip(),
@@ -436,6 +414,11 @@ FToolMenuEntry CreateRotationSnapCheckboxMenu()
 		),
 		MakeMenuDelegate
 	);
+
+	Entry.ToolbarLabelOverride = TAttribute<FText>::Create(&UE::UnrealEd::GetRotationGridLabel);
+	Entry.Icon = FSlateIcon(FAppStyle::GetAppStyleSetName(), "EditorViewport.RotationGridSnap");
+
+	return Entry;
 }
 
 FToolMenuEntry CreateScaleSnapCheckboxMenu()
@@ -515,7 +498,7 @@ FToolMenuEntry CreateScaleSnapCheckboxMenu()
 		}
 	);
 
-	return UnrealEd::CreateCheckboxSubmenu(
+	FToolMenuEntry Entry = UnrealEd::CreateCheckboxSubmenu(
 		"ScaleSnapping",
 		ScaleSnapLabel,
 		FEditorViewportCommands::Get().SurfaceSnapping->MakeTooltip()->GetTextTooltip(),
@@ -535,6 +518,11 @@ FToolMenuEntry CreateScaleSnapCheckboxMenu()
 		),
 		MakeMenuDelegate
 	);
+
+	Entry.ToolbarLabelOverride = TAttribute<FText>::Create(&UE::UnrealEd::GetScaleGridLabel);
+	Entry.Icon = FSlateIcon(FAppStyle::GetAppStyleSetName(), "EditorViewport.ScaleGridSnap");
+
+	return Entry;
 }
 
 bool IsViewModeSupported(EViewModeIndex InViewModeIndex)
@@ -1007,8 +995,8 @@ FToolMenuEntry CreateViewportToolbarSnappingSubmenu()
 
 				SnappingSection.AddEntry(Private::CreateSurfaceSnapCheckboxMenu());
 				SnappingSection.AddEntry(Private::CreateLocationSnapCheckboxMenu()).SetShowInToolbarTopLevel(true);
-				SnappingSection.AddEntry(Private::CreateRotationSnapCheckboxMenu());
-				SnappingSection.AddEntry(Private::CreateScaleSnapCheckboxMenu());
+				SnappingSection.AddEntry(Private::CreateRotationSnapCheckboxMenu()).SetShowInToolbarTopLevel(true);
+				SnappingSection.AddEntry(Private::CreateScaleSnapCheckboxMenu()).SetShowInToolbarTopLevel(true);
 				SnappingSection.AddEntry(Private::CreateActorSnapCheckboxMenu());
 
 				FToolMenuEntry SocketSnapping =
@@ -1582,6 +1570,13 @@ TSharedRef<SWidget> BuildRotationGridCheckBoxList(
 	return RotationGridMenuBuilder.MakeWidget();
 }
 
+FText GetRotationGridLabel()
+{
+	return FText::Format(
+		LOCTEXT("GridRotation - Number - DegreeSymbol", "{0}\u00b0"), FText::AsNumber(GEditor->GetRotGridSize().Pitch)
+	);
+}
+
 TSharedRef<SWidget> CreateRotationGridSnapMenu(
 	const FRotationGridCheckboxListExecuteActionDelegate& InExecuteDelegate,
 	const FRotationGridCheckboxListIsCheckedDelegate& InIsCheckedDelegate,
@@ -1590,8 +1585,6 @@ TSharedRef<SWidget> CreateRotationGridSnapMenu(
 )
 {
 	const ULevelEditorViewportSettings* ViewportSettings = GetDefault<ULevelEditorViewportSettings>();
-	TArray<float> GridSizes = ViewportSettings->bUsePowerOf2SnapSize ? ViewportSettings->Pow2GridSizes
-																	 : ViewportSettings->DecimalGridSizes;
 
 	// clang-format off
 	return SNew(SUniformGridPanel)
@@ -1617,6 +1610,11 @@ TSharedRef<SWidget> CreateRotationGridSnapMenu(
 	// clang-format on
 }
 
+FText GetLocationGridLabel()
+{
+	return FText::AsNumber(GEditor->GetGridSize());
+}
+
 TSharedRef<SWidget> CreateLocationGridSnapMenu(
 	const FLocationGridCheckboxListExecuteActionDelegate& InExecuteDelegate,
 	const FLocationGridCheckboxListIsCheckedDelegate& InIsCheckedDelegate,
@@ -1666,6 +1664,16 @@ TSharedRef<SWidget> CreateLocationGridSnapMenu(
 	return LocationGridMenuBuilder.MakeWidget();
 }
 
+FText GetScaleGridLabel()
+{
+	FNumberFormattingOptions NumberFormattingOptions;
+	NumberFormattingOptions.MaximumFractionalDigits = 5;
+
+	const float CurGridAmount = GEditor->GetScaleGridSize();
+	return (GEditor->UsePercentageBasedScaling()) ? FText::AsPercent(CurGridAmount / 100.0f, &NumberFormattingOptions)
+												  : FText::AsNumber(CurGridAmount, &NumberFormattingOptions);
+}
+
 TSharedRef<SWidget> CreateScaleGridSnapMenu(
 	const FScaleGridCheckboxListExecuteActionDelegate& InExecuteDelegate,
 	const FScaleGridCheckboxListIsCheckedDelegate& InIsCheckedDelegate,
diff --git a/Engine/Source/Editor/UnrealEd/Public/STransformViewportToolbar.h b/Engine/Source/Editor/UnrealEd/Public/STransformViewportToolbar.h
--- a/Engine/Source/Editor/UnrealEd/Public/STransformViewportToolbar.h
+++ b/Engine/Source/Editor/UnrealEd/Public/STransformViewportToolbar.h
@@ -87,10 +87,7 @@ private:
 	FOnCamSpeedScalarChanged OnCamSpeedScalarChanged;
 
 	/** Grid snap label callbacks */
-	FText GetLocationGridLabel() const;
-	FText GetRotationGridLabel() const;
 	FText GetLayer2DLabel() const;
-	FText GetScaleGridLabel() const;
 
 	/** GridSnap menu construction callbacks */
 	TSharedRef<SWidget> FillLocationGridSnapMenu();
diff --git a/Engine/Source/Editor/UnrealEd/Public/ViewportToolbar/UnrealEdViewportToolbar.h b/Engine/Source/Editor/UnrealEd/Public/ViewportToolbar/UnrealEdViewportToolbar.h
--- a/Engine/Source/Editor/UnrealEd/Public/ViewportToolbar/UnrealEdViewportToolbar.h
+++ b/Engine/Source/Editor/UnrealEd/Public/ViewportToolbar/UnrealEdViewportToolbar.h
@@ -78,6 +78,7 @@ UNREALED_API TSharedRef<SWidget> BuildRotationGridCheckBoxList(
 	const TSharedPtr<FUICommandList>& InCommandList = {}
 );
 
+UNREALED_API FText GetRotationGridLabel();
 UNREALED_API TSharedRef<SWidget> CreateRotationGridSnapMenu(
 	const FRotationGridCheckboxListExecuteActionDelegate& InExecuteDelegate,
 	const FRotationGridCheckboxListIsCheckedDelegate& InIsCheckedDelegate,
@@ -85,6 +86,7 @@ UNREALED_API TSharedRef<SWidget> CreateRotationGridSnapMenu(
 	const TSharedPtr<FUICommandList>& InCommandList = {}
 );
 
+UNREALED_API FText GetLocationGridLabel();
 UNREALED_API TSharedRef<SWidget> CreateLocationGridSnapMenu(
 	const FLocationGridCheckboxListExecuteActionDelegate& InExecuteDelegate,
 	const FLocationGridCheckboxListIsCheckedDelegate& InIsCheckedDelegate,
@@ -93,6 +95,7 @@ UNREALED_API TSharedRef<SWidget> CreateLocationGridSnapMenu(
 	const TSharedPtr<FUICommandList>& InCommandList = {}
 );
 
+UNREALED_API FText GetScaleGridLabel();
 UNREALED_API TSharedRef<SWidget> CreateScaleGridSnapMenu(
 	const FScaleGridCheckboxListExecuteActionDelegate& InExecuteDelegate,
 	const FScaleGridCheckboxListIsCheckedDelegate& InIsCheckedDelegate,
