[World Partition Builder Commandlet] Don't terminate commandlet when a builder fails on a specific map, log a warning and continue to other maps instead.
Also disable auto submit if one of the maps failed.

diff --git a/Engine/Source/Editor/UnrealEd/Private/Commandlets/WorldPartitionBuilderCommandlet.cpp b/Engine/Source/Editor/UnrealEd/Private/Commandlets/WorldPartitionBuilderCommandlet.cpp
--- a/Engine/Source/Editor/UnrealEd/Private/Commandlets/WorldPartitionBuilderCommandlet.cpp
+++ b/Engine/Source/Editor/UnrealEd/Private/Commandlets/WorldPartitionBuilderCommandlet.cpp
@@ -117,6 +117,7 @@ int32 UWorldPartitionBuilderCommandlet::Main(const FString& Params)
 	}	
 
 	// Run the builder on the provided map(s)
+	int32 Result = 0;
 	uint32 PackageIndex = 0;
 	const uint32 PackageCount = MapPackagesNames.Num();
 	for (const FString& MapPackageName : MapPackagesNames)
@@ -129,17 +130,18 @@ int32 UWorldPartitionBuilderCommandlet::Main(const FString& Params)
 
 		if (!RunBuilder(BuilderClass, MapPackageName))
 		{
-			return 1;
+			UE_LOG(LogWorldPartitionBuilderCommandlet, Warning, TEXT("Failed to execute %s on map %s."), *BuilderClassName, *MapPackageName);
+			Result = 1;
 		}
 	}
 
 	// Autosubmit
-	if (!AutoSubmitModifiedFiles())
+	if (!Result && !AutoSubmitModifiedFiles())
 	{
 		return 1;
 	}
 
-	return 0;
+	return Result;
 }
 
 TSet<FString> UWorldPartitionBuilderCommandlet::GatherMapsFromCollection(const FString& CollectionName) const
