Allow to skip saving cook settings when cook finishes

This takes more than a second in some cases in the content worker.

#rnx

diff --git a/Engine/Source/Editor/UnrealEd/Private/CookOnTheFlyServer.cpp b/Engine/Source/Editor/UnrealEd/Private/CookOnTheFlyServer.cpp
--- a/Engine/Source/Editor/UnrealEd/Private/CookOnTheFlyServer.cpp
+++ b/Engine/Source/Editor/UnrealEd/Private/CookOnTheFlyServer.cpp
@@ -10443,6 +10443,8 @@ void UCookOnTheFlyServer::CookByTheBookFinishedInternal()
 
 		// previously shader library was saved at this spot, but it's too early to know the chunk assignments, we need to BuildChunkManifest in the asset registry first
 
+		// Saving the current ini settings. This is only required for iterative cooking and may take seconds.
+		if (!FParse::Param(FCommandLine::Get(), TEXT("SkipSaveCookSettings")))
 		{
 			UE_SCOPED_HIERARCHICAL_COOKTIMER(SavingCurrentIniSettings)
 			for (const ITargetPlatform* TargetPlatform : PlatformManager->GetSessionPlatforms() )
