IterativeCompare: Add -hide=modifiedpackages flag for DiffCookCommandlet.
#rnx

diff --git a/Engine/Build/Graph/Tests/IterativeCompare.xml b/Engine/Build/Graph/Tests/IterativeCompare.xml
--- a/Engine/Build/Graph/Tests/IterativeCompare.xml
+++ b/Engine/Build/Graph/Tests/IterativeCompare.xml
@@ -69,7 +69,7 @@
 
 			<!-- Compare the incremental cook with the non-incremental cook. -->
 			<Log Message="$(NewLine)========== IterativeCompare.xml: DiffCook =========="/>
-			<Commandlet Name="DiffCook" Project="$(ProjectFile)" Arguments="-base=$(PlatformCookedDir) -comp=$(CompareDir) -targetplatform=$(TargetPlatform) -addedverbosity=error -removedverbosity=error -modifiedverbosity=warning"/>
+			<Commandlet Name="DiffCook" Project="$(ProjectFile)" Arguments="-base=$(PlatformCookedDir) -comp=$(CompareDir) -targetplatform=$(TargetPlatform) -addedverbosity=error -removedverbosity=error -modifiedverbosity=warning -hide=modifiedpackages"/>
 		</Node>
 	</Agent>
 
diff --git a/Engine/Source/Editor/UnrealEd/Classes/Commandlets/DiffCookCommandlet.h b/Engine/Source/Editor/UnrealEd/Classes/Commandlets/DiffCookCommandlet.h
--- a/Engine/Source/Editor/UnrealEd/Classes/Commandlets/DiffCookCommandlet.h
+++ b/Engine/Source/Editor/UnrealEd/Classes/Commandlets/DiffCookCommandlet.h
@@ -199,6 +199,9 @@ private:
 	ELogVerbosity::Type ModifiedVerbosity = ELogVerbosity::Display;
 	bool bShowSummary = true;
 	bool bShowPackages = true;
+	bool bShowAddedPackages = true;
+	bool bShowRemovedPackages = true;
+	bool bShowModifiedPackages = true;
 	bool bShowHeaders = false; // Not yet implemented, change default to true once implemented
 	bool bShowSerialize = false;
 
diff --git a/Engine/Source/Editor/UnrealEd/Private/Commandlets/DiffCookCommandlet.cpp b/Engine/Source/Editor/UnrealEd/Private/Commandlets/DiffCookCommandlet.cpp
--- a/Engine/Source/Editor/UnrealEd/Private/Commandlets/DiffCookCommandlet.cpp
+++ b/Engine/Source/Editor/UnrealEd/Private/Commandlets/DiffCookCommandlet.cpp
@@ -152,6 +152,18 @@ bool UDiffCookCommandlet::TryParseCommandLine(const FString& CmdLineParams)
 					{
 						bShowPackages = bShow;
 					}
+					if (ShowStr.Equals(TEXT("addedpackages"), ESearchCase::IgnoreCase))
+					{
+						bShowAddedPackages = bShow;
+					}
+					if (ShowStr.Equals(TEXT("removedpackages"), ESearchCase::IgnoreCase))
+					{
+						bShowRemovedPackages = bShow;
+					}
+					if (ShowStr.Equals(TEXT("modifiedpackages"), ESearchCase::IgnoreCase))
+					{
+						bShowModifiedPackages = bShow;
+					}
 					else if (ShowStr.Equals(TEXT("header"), ESearchCase::IgnoreCase)
 						|| ShowStr.Equals(TEXT("headers"), ESearchCase::IgnoreCase))
 					{
@@ -164,7 +176,7 @@ bool UDiffCookCommandlet::TryParseCommandLine(const FString& CmdLineParams)
 					else
 					{
 						UE_LOG(LogCook, Warning, TEXT("Unrecognized showflag '-%s=%.*s'.")
-							TEXT(" Valid Options are {'packages', 'headers', 'serialize'}."),
+							TEXT(" Valid Options are {'packages', 'removedpackages', 'addedpackages', 'modifiedpackages', 'headers', 'serialize'}."),
 							bShow ? TEXT("show") : TEXT("hide"), ShowStr.Len(), ShowStr.GetData());
 					}
 				}, UE::String::EParseTokensOptions::SkipEmpty | UE::String::EParseTokensOptions::Trim);
@@ -983,13 +995,19 @@ void UDiffCookCommandlet::PrintPackageDiffs(FDiffResult& Diff)
 		{
 		case EPackageDiffResult::Removed:
 			++NumRemoved;
-			UE_LOG(LogCook, Display, TEXT("Removed:  %s\n\tBase: %s\n\tComp: %s (Missing)"),
-				*PackageNameStr, **BasePath, **CompPath);
+			if (bShowRemovedPackages)
+			{
+				UE_LOG(LogCook, Display, TEXT("Removed:  %s\n\tBase: %s\n\tComp: %s (Missing)"),
+					*PackageNameStr, **BasePath, **CompPath);
+			}
 			break;
 		case EPackageDiffResult::Added:
 			++NumAdded;
-			UE_LOG(LogCook, Display, TEXT("Added:    %s\n\tBase: %s (Missing)\n\tComp: %s"),
-				*PackageNameStr, **BasePath, **CompPath);
+			if (bShowAddedPackages)
+			{
+				UE_LOG(LogCook, Display, TEXT("Added:    %s\n\tBase: %s (Missing)\n\tComp: %s"),
+					*PackageNameStr, **BasePath, **CompPath);
+			}
 			break;
 		case EPackageDiffResult::Modified:
 			++NumModified;
@@ -1005,23 +1023,26 @@ void UDiffCookCommandlet::PrintPackageDiffs(FDiffResult& Diff)
 			continue;
 		}
 
-		FString WhichFile = FPathViews::SetExtension(FPathViews::GetCleanFilename(*BasePath),
-			LexToString(PackageDiff.Extension));
-		UE_LOG(LogCook, Display, TEXT("Modified: %s\n\tBase: %s\n\tComp: %s")
-			TEXT("\n\t      Different at Byte %d in %s%s."),
-			*PackageNameStr, **BasePath, **CompPath,
-			PackageDiff.Offset, *WhichFile,
-			(PackageDiff.CombinedOffset >= 0
-				? *FString::Printf(TEXT(", Combined/DiffBreak Offset %d"), PackageDiff.CombinedOffset)
-				: TEXT("")));
-
-		if (bShowHeaders)
-		{
-			UE_CALL_ONCE([]() { UE_LOG(LogCook, Error, TEXT("-show=Headers is not yet implemented")); });
-		}
-		if (bShowSerialize)
+		if (bShowModifiedPackages)
 		{
-			UE_CALL_ONCE([]() { UE_LOG(LogCook, Error, TEXT("-show=Serialize is not yet implemented")); });
+			FString WhichFile = FPathViews::SetExtension(FPathViews::GetCleanFilename(*BasePath),
+				LexToString(PackageDiff.Extension));
+			UE_LOG(LogCook, Display, TEXT("Modified: %s\n\tBase: %s\n\tComp: %s")
+				TEXT("\n\t      Different at Byte %d in %s%s."),
+				*PackageNameStr, **BasePath, **CompPath,
+				PackageDiff.Offset, *WhichFile,
+				(PackageDiff.CombinedOffset >= 0
+					? *FString::Printf(TEXT(", Combined/DiffBreak Offset %d"), PackageDiff.CombinedOffset)
+					: TEXT("")));
+
+			if (bShowHeaders)
+			{
+				UE_CALL_ONCE([]() { UE_LOG(LogCook, Error, TEXT("-show=Headers is not yet implemented")); });
+			}
+			if (bShowSerialize)
+			{
+				UE_CALL_ONCE([]() { UE_LOG(LogCook, Error, TEXT("-show=Serialize is not yet implemented")); });
+			}
 		}
 	}
 }
