[ContentBrowser] Thumbnail Status Indicators

diff --git a/Engine/Source/Editor/AssetDefinition/Private/AssetStatusAssetDataInfoProvider.cpp b/Engine/Source/Editor/AssetDefinition/Private/AssetStatusAssetDataInfoProvider.cpp
new file mode 100644
--- /dev/null
+++ b/Engine/Source/Editor/AssetDefinition/Private/AssetStatusAssetDataInfoProvider.cpp
@@ -0,0 +1,19 @@
+﻿// Copyright Epic Games, Inc. All Rights Reserved.
+
+#include "AssetStatusAssetDataInfoProvider.h"
+#include "AssetRegistry/AssetData.h"
+#include "Misc/PackageName.h"
+#include "UObject/UObjectGlobals.h"
+
+UPackage* FAssetStatusAssetDataInfoProvider::FindPackage() const
+{
+	return FindObjectSafe<UPackage>(NULL, *FNameBuilder(AssetData.PackageName), true);
+}
+
+FString FAssetStatusAssetDataInfoProvider::TryGetFilename() const
+{
+	FString OutFileName;
+	const FString* PackageExtension = AssetData.HasAnyPackageFlags(PKG_ContainsMap) ? &FPackageName::GetMapPackageExtension() : &FPackageName::GetAssetPackageExtension();
+	FPackageName::TryConvertLongPackageNameToFilename(AssetData.PackageName.ToString(), OutFileName, *PackageExtension);
+	return OutFileName;
+}
diff --git a/Engine/Source/Editor/AssetDefinition/Public/AssetDefinition.h b/Engine/Source/Editor/AssetDefinition/Public/AssetDefinition.h
--- a/Engine/Source/Editor/AssetDefinition/Public/AssetDefinition.h
+++ b/Engine/Source/Editor/AssetDefinition/Public/AssetDefinition.h
@@ -2,6 +2,7 @@
 
 #pragma once
 
+#include "AssetDefinitionStatusBase.h"
 #include "UObject/Object.h"
 #include "UObject/SoftObjectPtr.h"
 #include "AssetRegistry/AssetData.h"
@@ -588,6 +589,14 @@ public:
 		return false;
 	}
 
+#if UE_CONTENTBROWSER_NEW_STYLE
+	// Implementation will change in the AssetDefinitionDefault in a later update
+	virtual void GetAssetStatusInfo(const TSharedPtr<IAssetStatusInfoProvider>& InAssetStatusInfoProvider, FAssetStatusInfo& OutStatusInfo) const
+	{
+		return;
+	}
+#endif
+
 	/** Whether this asset has external packages associated with it (impacts saving and dirty-state behavior) */
 	virtual bool ShouldSaveExternalPackages() const
 	{
diff --git a/Engine/Source/Editor/AssetDefinition/Public/AssetDefinitionStatusBase.h b/Engine/Source/Editor/AssetDefinition/Public/AssetDefinitionStatusBase.h
new file mode 100644
--- /dev/null
+++ b/Engine/Source/Editor/AssetDefinition/Public/AssetDefinitionStatusBase.h
@@ -0,0 +1,101 @@
+﻿// Copyright Epic Games, Inc. All Rights Reserved.
+
+#pragma once
+
+#include "Containers/Array.h"
+#include "Containers/UnrealString.h"
+#include "Internationalization/Text.h"
+#include "Layout/Visibility.h"
+#include "Misc/Attribute.h"
+
+class SWidget;
+class UPackage;
+struct FSlateBrush;
+
+enum class EStatusSeverity
+{
+	Info,
+	Warning,
+	Error
+};
+ENUM_CLASS_FLAGS(EStatusSeverity)
+
+struct FAssetStatusPriority
+{
+	FAssetStatusPriority()
+		: Severity(EStatusSeverity::Info)
+		, SeverityPriority(0)
+	{}
+
+	explicit FAssetStatusPriority(const TAttribute<EStatusSeverity>& InSeverity)
+		: Severity(InSeverity)
+		, SeverityPriority(0)
+	{}
+
+	FAssetStatusPriority(const TAttribute<EStatusSeverity>& InSeverity, int32 InSeverityPriority)
+		: Severity(InSeverity)
+		, SeverityPriority(InSeverityPriority)
+	{}
+
+public:
+	bool operator==(const FAssetStatusPriority& InOtherStatusPriority) const
+	{
+		if (!Severity.IsSet() || !InOtherStatusPriority.Severity.IsSet())
+		{
+			return false;
+		}
+
+		return Severity.Get() == InOtherStatusPriority.Severity.Get() && SeverityPriority == InOtherStatusPriority.SeverityPriority;
+	}
+
+	bool operator<(const FAssetStatusPriority& InOtherStatusPriority) const
+	{
+		if (!Severity.IsSet())
+		{
+			return true;
+		}
+
+		if (!InOtherStatusPriority.Severity.IsSet())
+		{
+			return false;
+		}
+
+		if (Severity.Get() == InOtherStatusPriority.Severity.Get())
+		{
+			return SeverityPriority < InOtherStatusPriority.SeverityPriority;
+		}
+
+		return Severity.Get() < InOtherStatusPriority.Severity.Get();
+	}
+
+public:
+	TAttribute<EStatusSeverity> Severity;
+	int32 SeverityPriority;
+};
+
+struct FAssetStatus
+{
+public:
+	TAttribute<const FSlateBrush*> StatusIcon;
+	TAttribute<FText> StatusDescription;
+	TAttribute<EVisibility> IsVisible;
+	TAttribute<FAssetStatusPriority> Priority;
+};
+
+struct FAssetStatusInfo
+{
+public:
+	TArray<FAssetStatus> AssetStatus;
+};
+
+class ASSETDEFINITION_API IAssetStatusInfoProvider
+{
+public:
+	virtual ~IAssetStatusInfoProvider() = default;
+
+	/** Try to find the Package without loading it, if the package is not loaded it will return nullptr */
+	virtual UPackage* FindPackage() const = 0;
+
+	/** Try to get the filename, return empty otherwise */
+	virtual FString TryGetFilename() const = 0;
+};
diff --git a/Engine/Source/Editor/AssetDefinition/Public/AssetStatusAssetDataInfoProvider.h b/Engine/Source/Editor/AssetDefinition/Public/AssetStatusAssetDataInfoProvider.h
new file mode 100644
--- /dev/null
+++ b/Engine/Source/Editor/AssetDefinition/Public/AssetStatusAssetDataInfoProvider.h
@@ -0,0 +1,21 @@
+﻿// Copyright Epic Games, Inc. All Rights Reserved.
+
+#pragma once
+
+#include "AssetDefinitionStatusBase.h"
+#include "AssetRegistry/AssetData.h"
+
+class ASSETDEFINITION_API FAssetStatusAssetDataInfoProvider : public IAssetStatusInfoProvider
+{
+public:
+	FAssetStatusAssetDataInfoProvider(FAssetData InAssetData)
+		: AssetData(InAssetData)
+	{}
+
+	virtual UPackage* FindPackage() const override;
+
+	virtual FString TryGetFilename() const override;
+
+private:
+	FAssetData AssetData;
+};
diff --git a/Engine/Source/Editor/ContentBrowser/Private/AssetViewWidgets.cpp b/Engine/Source/Editor/ContentBrowser/Private/AssetViewWidgets.cpp
--- a/Engine/Source/Editor/ContentBrowser/Private/AssetViewWidgets.cpp
+++ b/Engine/Source/Editor/ContentBrowser/Private/AssetViewWidgets.cpp
@@ -247,6 +247,7 @@ TSharedRef<SWidget> FAssetViewItemHelper::CreateListTileItemContents(T* const In
 			]
 		];
 
+#if !UE_CONTENTBROWSER_NEW_STYLE
 		// Dirty state
 		ItemContentsOverlay->AddSlot()
 		.HAlign(HAlign_Left)
@@ -260,7 +261,7 @@ TSharedRef<SWidget> FAssetViewItemHelper::CreateListTileItemContents(T* const In
 				.Image(InTileOrListItem, &T::GetDirtyImage)
 			]
 		];
-
+#endif
 		// Tools for thumbnail edit mode
 		ItemContentsOverlay->AddSlot()
 		[
@@ -1952,6 +1953,7 @@ void SAssetTileItem::Construct( const FArguments& InArgs )
 		ThumbnailConfig.Padding= FMargin(2.0f);
 		ThumbnailConfig.GenericThumbnailSize = MakeAttributeSP(this, &SAssetTileItem::GetGenericThumbnailSize);
 #if UE_CONTENTBROWSER_NEW_STYLE
+		ThumbnailConfig.bAllowAssetStatusThumbnailOverlay = true;
 		ThumbnailConfig.bShowAssetChip = true;
 		ThumbnailConfig.AssetChipBorderImageOverride = TAttribute<const FSlateBrush*>::CreateSP(this, &SAssetTileItem::GetAssetAreaOverlayBackgroundImage);
 #endif
diff --git a/Engine/Source/Editor/EditorStyle/Private/StarshipStyle.cpp b/Engine/Source/Editor/EditorStyle/Private/StarshipStyle.cpp
--- a/Engine/Source/Editor/EditorStyle/Private/StarshipStyle.cpp
+++ b/Engine/Source/Editor/EditorStyle/Private/StarshipStyle.cpp
@@ -6496,6 +6496,10 @@ void FStarshipEditorStyle::FStyle::SetupContentBrowserStyle()
 		OverlayColorAssetThumbnailOverlay.A = 0.75f;
 		Set("ContentBrowser.AssetTileItem.AssetThumbnailBar", new FSlateRoundedBoxBrush(OverlayColorAssetThumbnailOverlay, 4.f));
 
+		FLinearColor OverlayColorAssetStatusOverlay = FStyleColors::Panel.GetSpecifiedColor();
+		OverlayColorAssetStatusOverlay.A = 0.75f;
+		Set("ContentBrowser.AssetTileItem.AssetThumbnailStatusBar", new FSlateRoundedBoxBrush(OverlayColorAssetStatusOverlay, 2.f));
+
 		Set( "ContentBrowser.ColumnViewAssetIcon", new IMAGE_BRUSH( "Icons/doc_16x", Icon16x16 ) );
 
 		Set( "ContentBrowser.ColumnViewFolderIcon", new CORE_IMAGE_BRUSH_SVG( "Starship/Common/folder-closed", Icon16x16 ) );
diff --git a/Engine/Source/Editor/UnrealEd/Private/AssetDefinitionDefault.cpp b/Engine/Source/Editor/UnrealEd/Private/AssetDefinitionDefault.cpp
--- a/Engine/Source/Editor/UnrealEd/Private/AssetDefinitionDefault.cpp
+++ b/Engine/Source/Editor/UnrealEd/Private/AssetDefinitionDefault.cpp
@@ -5,11 +5,100 @@
 #include "AssetToolsModule.h"
 #include "EditorFramework/AssetImportData.h"
 #include "EditorFramework/ThumbnailInfo.h"
+#include "ISourceControlModule.h"
+#include "ISourceControlProvider.h"
 #include "Settings/EditorLoadingSavingSettings.h"
 #include "Toolkits/SimpleAssetEditor.h"
 
 #define LOCTEXT_NAMESPACE "AssetDefinitionDefault"
 
+namespace UE::AssetDefinitionDefault::Status
+{
+	EVisibility GetDirtyStatusVisibility(const TSharedPtr<IAssetStatusInfoProvider> InAssetStatusInfoProvider)
+	{
+		EVisibility DirtyStatusVisibility = EVisibility::Collapsed;
+		if (!InAssetStatusInfoProvider.IsValid())
+		{
+			return DirtyStatusVisibility;
+		}
+
+		if (const UPackage* Package = InAssetStatusInfoProvider->FindPackage())
+		{
+			DirtyStatusVisibility = Package->IsDirty() ? EVisibility::Visible : EVisibility::Collapsed;
+		}
+		return DirtyStatusVisibility;
+	}
+
+	const FSlateBrush* GetSourceControlStatusBrush(const TSharedPtr<IAssetStatusInfoProvider> InAssetStatusInfoProvider)
+	{
+		const FSlateBrush* SourceControlBrush = FAppStyle::GetNoBrush();
+		if (!InAssetStatusInfoProvider.IsValid())
+		{
+			return SourceControlBrush;
+		}
+
+		if (ISourceControlModule::Get().IsEnabled() && ISourceControlModule::Get().GetProvider().IsAvailable())
+		{
+			const FString FileName = InAssetStatusInfoProvider->TryGetFilename();
+			if (FSourceControlStatePtr SourceControlState = ISourceControlModule::Get().GetProvider().GetState(FileName, EStateCacheUsage::Use))
+			{
+				FSlateIcon SCCIcon = SourceControlState->GetIcon();
+				if (SCCIcon.IsSet())
+				{
+					SourceControlBrush = SCCIcon.GetIcon();
+				}
+			}
+		}
+		return SourceControlBrush;
+	}
+
+	EVisibility GetSourceControlStatusVisibility(const TSharedPtr<IAssetStatusInfoProvider> InAssetStatusInfoProvider)
+	{
+		EVisibility SourceControlStatusVisibility = EVisibility::Collapsed;
+		if (!InAssetStatusInfoProvider.IsValid())
+		{
+			return SourceControlStatusVisibility;
+		}
+
+		if (ISourceControlModule::Get().IsEnabled() && ISourceControlModule::Get().GetProvider().IsAvailable())
+		{
+			const FString FileName = InAssetStatusInfoProvider->TryGetFilename();
+			if (FSourceControlStatePtr SourceControlState = ISourceControlModule::Get().GetProvider().GetState(FileName, EStateCacheUsage::Use))
+			{
+				FSlateIcon SCCIcon = SourceControlState->GetIcon();
+				if (SCCIcon.IsSet())
+				{
+					SourceControlStatusVisibility = EVisibility::Visible;
+				}
+			}
+		}
+		return SourceControlStatusVisibility;
+	}
+
+	FText GetSourceControlStatusDescription(const TSharedPtr<IAssetStatusInfoProvider> InAssetStatusInfoProvider)
+	{
+		FText SourceControlDescription = FText::GetEmpty();
+		if (!InAssetStatusInfoProvider.IsValid())
+		{
+			return SourceControlDescription;
+		}
+
+		if (ISourceControlModule::Get().IsEnabled() && ISourceControlModule::Get().GetProvider().IsAvailable())
+		{
+			const FString FileName = InAssetStatusInfoProvider->TryGetFilename();
+			if (FSourceControlStatePtr SourceControlState = ISourceControlModule::Get().GetProvider().GetState(FileName, EStateCacheUsage::Use))
+			{
+				TOptional<FText> StatusText = SourceControlState->GetStatusText();
+				if (StatusText.IsSet())
+				{
+					SourceControlDescription = StatusText.GetValue();
+				}
+			}
+		}
+		return SourceControlDescription;
+	}
+}
+
 EAssetCommandResult UAssetDefinitionDefault::OpenAssets(const FAssetOpenArgs& OpenArgs) const
 {
 	if (GetAssetOpenSupport(FAssetOpenSupportArgs(OpenArgs.OpenMethod)).IsSupported)
@@ -40,6 +129,25 @@ EAssetCommandResult UAssetDefinitionDefault::PerformAssetDiff(const FAssetDiffAr
 	return EAssetCommandResult::Handled;
 }
 
+#if UE_CONTENTBROWSER_NEW_STYLE
+void UAssetDefinitionDefault::GetAssetStatusInfo(const TSharedPtr<IAssetStatusInfoProvider>& InAssetStatusInfoProvider, FAssetStatusInfo& OutStatusInfo) const
+{
+	FAssetStatus DirtyStatus;
+	DirtyStatus.StatusIcon = FAppStyle::GetBrush("ContentBrowser.ContentDirty");
+	DirtyStatus.Priority = FAssetStatusPriority(EStatusSeverity::Info, 1);
+	DirtyStatus.StatusDescription = LOCTEXT("DirtyAssetTooltip", "Asset has unsaved changes");
+	DirtyStatus.IsVisible = TAttribute<EVisibility>::Create(TAttribute<EVisibility>::FGetter::CreateStatic(&UE::AssetDefinitionDefault::Status::GetDirtyStatusVisibility, InAssetStatusInfoProvider));
+	OutStatusInfo.AssetStatus.Add(DirtyStatus);
+
+	FAssetStatus SCCStatus;
+	SCCStatus.Priority = FAssetStatusPriority(EStatusSeverity::Info, 0);
+	SCCStatus.StatusIcon = TAttribute<const FSlateBrush*>::Create(TAttribute<const FSlateBrush*>::FGetter::CreateStatic(&UE::AssetDefinitionDefault::Status::GetSourceControlStatusBrush, InAssetStatusInfoProvider));
+	SCCStatus.IsVisible = TAttribute<EVisibility>::Create(TAttribute<EVisibility>::FGetter::CreateStatic(&UE::AssetDefinitionDefault::Status::GetSourceControlStatusVisibility, InAssetStatusInfoProvider));
+	SCCStatus.StatusDescription = TAttribute<FText>::Create(TAttribute<FText>::FGetter::CreateStatic(&UE::AssetDefinitionDefault::Status::GetSourceControlStatusDescription, InAssetStatusInfoProvider));
+	OutStatusInfo.AssetStatus.Add(SCCStatus);
+}
+#endif
+
 namespace UE::Editor
 {
 	UThumbnailInfo* FindOrCreateThumbnailInfo(UObject* AssetObject, TSubclassOf<UThumbnailInfo> ThumbnailClass)
diff --git a/Engine/Source/Editor/UnrealEd/Private/AssetThumbnail.cpp b/Engine/Source/Editor/UnrealEd/Private/AssetThumbnail.cpp
--- a/Engine/Source/Editor/UnrealEd/Private/AssetThumbnail.cpp
+++ b/Engine/Source/Editor/UnrealEd/Private/AssetThumbnail.cpp
@@ -1,6 +1,8 @@
 // Copyright Epic Games, Inc. All Rights Reserved.
 
 #include "AssetThumbnail.h"
+#include "AssetDefinitionRegistry.h"
+#include "AssetStatusAssetDataInfoProvider.h"
 #include "Engine/Blueprint.h"
 #include "GameFramework/Actor.h"
 #include "Layout/Margin.h"
@@ -82,6 +84,7 @@ public:
 		, _Padding(0)
 		, _GenericThumbnailSize(64)
 #if UE_CONTENTBROWSER_NEW_STYLE
+		, _AllowAssetStatusThumbnailOverlay(false)
 		, _ShowAssetChip(false)
 		, _AssetChipBorderImageOverride()
 #endif
@@ -104,6 +107,7 @@ public:
 		SLATE_ARGUMENT(FMargin, Padding)
 		SLATE_ATTRIBUTE(int32, GenericThumbnailSize)
 #if UE_CONTENTBROWSER_NEW_STYLE
+		SLATE_ARGUMENT(bool, AllowAssetStatusThumbnailOverlay)
 		SLATE_ARGUMENT(bool, ShowAssetChip)
 		SLATE_ARGUMENT(TAttribute<const FSlateBrush*>, AssetChipBorderImageOverride)
 #endif
@@ -329,11 +333,63 @@ public:
 		];
 #endif
 
+#if UE_CONTENTBROWSER_NEW_STYLE
+		const UAssetDefinition* AssetDefinition = UAssetDefinitionRegistry::Get()->GetAssetDefinitionForClass(Class);
+
+		if (InArgs._AllowAssetStatusThumbnailOverlay && AssetDefinition)
+		{
+			const TSharedPtr<FAssetStatusAssetDataInfoProvider> AssetDataInfoProvider = MakeShared<FAssetStatusAssetDataInfoProvider>(AssetData);
+			AssetDefinition->GetAssetStatusInfo(AssetDataInfoProvider, OverlayInfo);
+			// Sort the list based on Status Priority
+			auto SortStatus = [] (const FAssetStatus& InFirstStatus, const FAssetStatus& InSecondStatus)
+			{
+				if (!InFirstStatus.Priority.IsSet())
+				{
+					return false;
+				}
+
+				if (!InSecondStatus.Priority.IsSet())
+				{
+					return true;
+				}
+
+				return InSecondStatus.Priority.Get() < InFirstStatus.Priority.Get();
+			};
+			OverlayInfo.AssetStatus.Sort(SortStatus);
+
+			const TSharedRef<SHorizontalBox> HorizontalBox = SNew(SHorizontalBox);
+			for (int32 StatusIndex = 0; StatusIndex < OverlayInfo.AssetStatus.Num(); StatusIndex++)
+			{
+				Statuses.Add(CreateStatusWidget(StatusIndex, OverlayInfo.AssetStatus[StatusIndex]));
+				HorizontalBox->AddSlot()
+				[
+					Statuses[StatusIndex].ToSharedRef()
+				];
+			}
+
+			StatusOverflowWidget = CreateStatusOverflowWidget();
+			HorizontalBox->AddSlot()
+			[
+				StatusOverflowWidget.ToSharedRef()
+			];
+
+			OverlayWidget->AddSlot()
+			.HAlign(HAlign_Left)
+			.VAlign(VAlign_Bottom)
+			.Padding(StatusPadding)
+			[
+				SNew(SBorder)
+				.BorderImage(FAppStyle::Get().GetBrush("ContentBrowser.AssetTileItem.AssetThumbnailStatusBar"))
+				.Visibility(this, &SAssetThumbnail::GetStatusBorderVisibility)
+				[
+					HorizontalBox
+				]
+			];
+		}
+#endif
 		if( InArgs._AllowAssetSpecificThumbnailOverlay && AssetTypeActions.IsValid() )
 		{
 #if UE_CONTENTBROWSER_NEW_STYLE
-			const UAssetDefinition* AssetDefinition = UAssetDefinitionRegistry::Get()->GetAssetDefinitionForClass(Class);
-
 			FAssetActionThumbnailOverlayInfo OutThumbnailInfo;
 			if (AssetDefinition && AssetDefinition->GetThumbnailActionOverlay(AssetData, OutThumbnailInfo))
 			{
@@ -498,6 +554,115 @@ public:
 	}
 
 private:
+#if UE_CONTENTBROWSER_NEW_STYLE
+	EVisibility GetStatusBorderVisibility() const
+	{
+		for (const FAssetStatus& AssetStatus : OverlayInfo.AssetStatus)
+		{
+			if (AssetStatus.IsVisible.IsSet() && AssetStatus.IsVisible.Get().IsVisible())
+			{
+				return EVisibility::Visible;
+			}
+		}
+		return EVisibility::Collapsed;
+	}
+
+	TSharedRef<SWidget> CreateStatusWidget(int32 StatusIndex, const FAssetStatus& InStatusInfo) const
+	{
+		return SNew(SBox)
+			.WidthOverride(StatusSize)
+			.HeightOverride(StatusSize)
+			.Visibility(this, &SAssetThumbnail::GetStatusVisibility, StatusIndex)
+			.VAlign(VAlign_Center)
+			[
+				SNew(SImage)
+				.Image(InStatusInfo.StatusIcon)
+			];
+	}
+
+	EVisibility GetStatusVisibility(int32 StatusIndex) const
+	{
+		if (OverlayInfo.AssetStatus.IsValidIndex(StatusIndex))
+		{
+			const FAssetStatus& AssetStatus = OverlayInfo.AssetStatus[StatusIndex];
+			if (AssetStatus.IsVisible.IsSet() && AssetStatus.IsVisible.Get() == EVisibility::Visible)
+			{
+				return GetStatusVisibilityBasedOnGeometry(StatusIndex);
+			}
+		}
+		return EVisibility::Collapsed;
+	}
+
+	TSharedRef<SWidget> CreateStatusOverflowWidget() const
+	{
+		return SNew(SBox)
+			.WidthOverride(StatusSize)
+			.HeightOverride(StatusSize)
+			.HAlign(HAlign_Center)
+			.VAlign(VAlign_Center)
+			.Visibility(this, &SAssetThumbnail::GetStatusOverflowVisibility)
+			[
+				SNew(STextBlock)
+				.Font(FSlateFontInfo( FPaths::EngineContentDir() / TEXT("Slate/Fonts/Roboto-Bold.ttf"), 9))
+				.Text(this, &SAssetThumbnail::GetStatusOverflowText)
+			];
+	}
+
+	EVisibility GetStatusOverflowVisibility() const
+	{
+		const FGeometry ThumbnailGeometry = GetPaintSpaceGeometry();
+		const float ThumbnailWidth = ThumbnailGeometry.GetAbsoluteSize().X - (StatusPadding * 2);
+		const int32 MaxShownStatus = FMath::FloorToInt32(ThumbnailWidth / StatusSize);
+
+		int32 ShownStatus = 0;
+		for (const FAssetStatus& Status : OverlayInfo.AssetStatus)
+		{
+			if (Status.IsVisible.IsSet() && Status.IsVisible.Get().IsVisible())
+			{
+				ShownStatus++;
+			}
+		}
+		return ShownStatus > MaxShownStatus? EVisibility::Visible : EVisibility::Collapsed;
+	}
+
+	FText GetStatusOverflowText() const
+	{
+		FGeometry ThumbnailGeometry = GetPaintSpaceGeometry();
+		const float ThumbnailWidth = ThumbnailGeometry.GetAbsoluteSize().X - (StatusPadding * 2);
+		const int32 MaxShownStatus = FMath::FloorToInt32(ThumbnailWidth / StatusSize);
+
+		int32 ShownStatus = 0;
+		for (const FAssetStatus& Status : OverlayInfo.AssetStatus)
+		{
+			if (Status.IsVisible.IsSet() && Status.IsVisible.Get().IsVisible())
+			{
+				ShownStatus++;
+			}
+		}
+
+		// We need to add 1 to the HiddenStatus since the Overflow "status" will occupy 1 extra slot
+		constexpr int32 OccupiedStatusSlot = 1;
+		const int32 HiddenStatus = ShownStatus - MaxShownStatus + OccupiedStatusSlot;
+		return FText::Format(NSLOCTEXT("AssetThumbnail", "StatusOverflowText", "+{0}"), HiddenStatus);
+	}
+
+	EVisibility GetStatusVisibilityBasedOnGeometry(int32 InStatusIndex) const
+	{
+		int32 CollapsedStatusBeforeThis = 0;
+		for (int32 StatusIndex = 0; StatusIndex < InStatusIndex; StatusIndex++)
+		{
+			CollapsedStatusBeforeThis += Statuses[StatusIndex]->GetVisibility().IsVisible() ? 0 : 1;
+		}
+
+		const FGeometry ThumbnailGeometry = GetPaintSpaceGeometry();
+		const float ThumbnailWidth = ThumbnailGeometry.GetAbsoluteSize().X - (StatusPadding * 2);
+		const int32 StatusIndexConsideringHiddenOnes = (InStatusIndex - CollapsedStatusBeforeThis);
+		const int32 ShownStatus = FMath::FloorToInt32(ThumbnailWidth / StatusSize);
+		const int32 StatusIndexConsideringClippedStatusIfVisible = StatusOverflowWidget.IsValid() && StatusOverflowWidget->GetVisibility().IsVisible() ? StatusIndexConsideringHiddenOnes + 1 : StatusIndexConsideringHiddenOnes;
+		return StatusIndexConsideringClippedStatusIfVisible < ShownStatus ? EVisibility::Visible : EVisibility::Collapsed;
+	}
+#endif
+
 	void OnAssetDataChanged()
 	{
 		if ( GenericLabelTextBlock.IsValid() )
@@ -914,6 +1079,11 @@ private:
 	TAttribute<int32> GenericThumbnailSize;
 	EThumbnailColorStripOrientation ColorStripOrientation;
 #if UE_CONTENTBROWSER_NEW_STYLE
+	TSharedPtr<SWidget> StatusOverflowWidget;
+	const float StatusSize = 16.f;
+	const float StatusPadding = 4.f;
+	FAssetStatusInfo OverlayInfo;
+	TArray<TSharedPtr<SWidget>> Statuses;
 	bool bShowAssetChip;
 	TAttribute<const FSlateBrush*> AssetChipBorderImageOverride;
 #endif
@@ -1056,6 +1226,7 @@ TSharedRef<SWidget> FAssetThumbnail::MakeThumbnailWidget( const FAssetThumbnailC
 		.Padding(InConfig.Padding)
 		.GenericThumbnailSize(InConfig.GenericThumbnailSize)
 #if UE_CONTENTBROWSER_NEW_STYLE
+		.AllowAssetStatusThumbnailOverlay(InConfig.bAllowAssetStatusThumbnailOverlay)
 		.ShowAssetChip(InConfig.bShowAssetChip)
 		.AssetChipBorderImageOverride(InConfig.AssetChipBorderImageOverride)
 #endif
diff --git a/Engine/Source/Editor/UnrealEd/Public/AssetDefinitionDefault.h b/Engine/Source/Editor/UnrealEd/Public/AssetDefinitionDefault.h
--- a/Engine/Source/Editor/UnrealEd/Public/AssetDefinitionDefault.h
+++ b/Engine/Source/Editor/UnrealEd/Public/AssetDefinitionDefault.h
@@ -16,6 +16,9 @@ public:
 	// UAssetDefinition Begin
 	UNREALED_API virtual EAssetCommandResult OpenAssets(const FAssetOpenArgs& OpenArgs) const override;
 	UNREALED_API virtual EAssetCommandResult PerformAssetDiff(const FAssetDiffArgs& DiffArgs) const override;
+#if UE_CONTENTBROWSER_NEW_STYLE
+	UNREALED_API virtual void GetAssetStatusInfo(const TSharedPtr<IAssetStatusInfoProvider>& InAssetStatusInfoProvider, FAssetStatusInfo& OutStatusInfo) const override;
+#endif
 	// UAssetDefinition End
 };
 
diff --git a/Engine/Source/Editor/UnrealEd/Public/AssetThumbnail.h b/Engine/Source/Editor/UnrealEd/Public/AssetThumbnail.h
--- a/Engine/Source/Editor/UnrealEd/Public/AssetThumbnail.h
+++ b/Engine/Source/Editor/UnrealEd/Public/AssetThumbnail.h
@@ -65,6 +65,7 @@ struct FAssetThumbnailConfig
 		, AssetTypeColorOverride()
 		, Padding(0)
 #if UE_CONTENTBROWSER_NEW_STYLE
+		, bAllowAssetStatusThumbnailOverlay( false )
 		, bShowAssetChip(false)
 		, AssetChipBorderImageOverride()
 #endif
@@ -85,6 +86,7 @@ struct FAssetThumbnailConfig
 	TAttribute<int32> GenericThumbnailSize = 64;
 	EThumbnailColorStripOrientation ColorStripOrientation = EThumbnailColorStripOrientation::HorizontalBottomEdge;
 #if UE_CONTENTBROWSER_NEW_STYLE
+	bool bAllowAssetStatusThumbnailOverlay;
 	bool bShowAssetChip;
 	TAttribute<const FSlateBrush*> AssetChipBorderImageOverride;
 #endif
