Gather text from source to include nested macros (Static analysis fix, PVS-Studio)
Invert prepass flag logic from RunNestedMacroPrepass to SkipNestedMacroPrepass, so it runs by default

UE-176755

diff --git a/Engine/Source/Editor/UnrealEd/Classes/Commandlets/GatherTextCommandletBase.h b/Engine/Source/Editor/UnrealEd/Classes/Commandlets/GatherTextCommandletBase.h
--- a/Engine/Source/Editor/UnrealEd/Classes/Commandlets/GatherTextCommandletBase.h
+++ b/Engine/Source/Editor/UnrealEd/Classes/Commandlets/GatherTextCommandletBase.h
@@ -118,7 +118,7 @@ protected:
 	static UNREALED_API const TCHAR* DisableSubmitSwitch;
 	static UNREALED_API const TCHAR* PreviewSwitch;
 	static UNREALED_API const TCHAR* GatherTypeParam;
-	static UNREALED_API const TCHAR* RunNestedMacroPrepassSwitch;
+	static UNREALED_API const TCHAR* SkipNestedMacroPrepassSwitch;
 
 private:
 	UNREALED_API virtual void CreateCustomEngine(const FString& Params) override ; //Disallow other text commandlets to make their own engine.	
diff --git a/Engine/Source/Editor/UnrealEd/Private/Commandlets/GatherTextCommandletBase.cpp b/Engine/Source/Editor/UnrealEd/Private/Commandlets/GatherTextCommandletBase.cpp
--- a/Engine/Source/Editor/UnrealEd/Private/Commandlets/GatherTextCommandletBase.cpp
+++ b/Engine/Source/Editor/UnrealEd/Private/Commandlets/GatherTextCommandletBase.cpp
@@ -28,7 +28,7 @@ const TCHAR* UGatherTextCommandletBase::EnableSourceControlSwitch = TEXT("Enable
 const TCHAR* UGatherTextCommandletBase::DisableSubmitSwitch = TEXT("DisableSCCSubmit");
 const TCHAR* UGatherTextCommandletBase::PreviewSwitch = TEXT("Preview");
 const TCHAR* UGatherTextCommandletBase::GatherTypeParam = TEXT("GatherType");
-const TCHAR* UGatherTextCommandletBase::RunNestedMacroPrepassSwitch = TEXT("RunNestedMacroPrepass");
+const TCHAR* UGatherTextCommandletBase::SkipNestedMacroPrepassSwitch = TEXT("SkipNestedMacroPrepass");
 
 UGatherTextCommandletBase::UGatherTextCommandletBase(const FObjectInitializer& ObjectInitializer)
 	: Super(ObjectInitializer)
diff --git a/Engine/Source/Editor/UnrealEd/Private/Commandlets/GatherTextFromSourceCommandlet.cpp b/Engine/Source/Editor/UnrealEd/Private/Commandlets/GatherTextFromSourceCommandlet.cpp
--- a/Engine/Source/Editor/UnrealEd/Private/Commandlets/GatherTextFromSourceCommandlet.cpp
+++ b/Engine/Source/Editor/UnrealEd/Private/Commandlets/GatherTextFromSourceCommandlet.cpp
@@ -259,10 +259,10 @@ int32 UGatherTextFromSourceCommandlet::Main( const FString& Params )
 	}
 
 	// Prepass for nested macros
-	bool RunNestedMacroPrepass = Switches.Contains(UGatherTextCommandletBase::RunNestedMacroPrepassSwitch);
+	bool SkipNestedMacroPrepass = Switches.Contains(UGatherTextCommandletBase::SkipNestedMacroPrepassSwitch);
 	static TArray<FParsedNestedMacro> PrepassResults;
 	static bool RanPrepassOnce = false;
-	if (RunNestedMacroPrepass && !RanPrepassOnce)
+	if (!SkipNestedMacroPrepass && !RanPrepassOnce)
 	{
 		double StartTime = FPlatformTime::Seconds();
 
@@ -1384,7 +1384,7 @@ void UGatherTextFromSourceCommandlet::PrunePrepassResults(TArray<FParsedNestedMa
 				// We mark them excluded as opposed to removing them, because the regular macro descriptors need to check if they are nested.
 
 				if (Results[i].Filename.EndsWith(TEXT(".h"), ESearchCase::IgnoreCase) ||
-					Results[i].Filename.EndsWith(TEXT(".h"), ESearchCase::IgnoreCase))
+					Results[j].Filename.EndsWith(TEXT(".h"), ESearchCase::IgnoreCase))
 				{
 					++NestedMacroStats.DuplicateExcluded;
 
@@ -2563,6 +2563,7 @@ void UGatherTextFromSourceCommandlet::FNestedMacroDescriptor::TryParse(const FSt
 		if (PosClose < 0)
 		{
 			UE_LOG(LogGatherTextFromSourceCommandlet, Error, TEXT("%s(%d): Missing matching closing bracket in %s macro"), *Context.Filename, Context.LineNumber, *MacroName);
+			delete InnerDescriptor;
 			return;
 		}
 		FString MacroInnerParams = MacroInner.Mid(0, PosClose);						// exclude bracket
