Hide the label on the "Transforms" submenu of the new level editor viewport toolbar

diff --git a/Engine/Source/Editor/UnrealEd/Private/ViewportToolbar/UnrealEdViewportToolbar.cpp b/Engine/Source/Editor/UnrealEd/Private/ViewportToolbar/UnrealEdViewportToolbar.cpp
--- a/Engine/Source/Editor/UnrealEd/Private/ViewportToolbar/UnrealEdViewportToolbar.cpp
+++ b/Engine/Source/Editor/UnrealEd/Private/ViewportToolbar/UnrealEdViewportToolbar.cpp
@@ -665,6 +665,7 @@ FToolMenuEntry CreateViewportToolbarTransformsSection()
 		)
 	);
 
+	Entry.ToolbarLabelOverride = FText();
 	Entry.Icon = FSlateIcon(FAppStyle::GetAppStyleSetName(), "EditorViewport.RelativeCoordinateSystem_Local");
 	return Entry;
 }
