Gather text from source to include nested macros, support escaped quotes
Restore the one escaped quote found in FusionSamplerNode.cpp

diff --git a/Engine/Plugins/Runtime/Harmonix/Source/HarmonixMetasound/Private/HarmonixMetasound/Nodes/FusionSamplerNode.cpp b/Engine/Plugins/Runtime/Harmonix/Source/HarmonixMetasound/Private/HarmonixMetasound/Nodes/FusionSamplerNode.cpp
--- a/Engine/Plugins/Runtime/Harmonix/Source/HarmonixMetasound/Private/HarmonixMetasound/Nodes/FusionSamplerNode.cpp
+++ b/Engine/Plugins/Runtime/Harmonix/Source/HarmonixMetasound/Private/HarmonixMetasound/Nodes/FusionSamplerNode.cpp
@@ -37,7 +37,7 @@ namespace HarmonixMetasound
 		METASOUND_PARAM(FineTuneCents, "Fine Tune Cents", "Adds to the patch's fine tune cents setting if connected AND if the vaue is > -1200 and < 1200.")
 		METASOUND_PARAM(EnableMTFusion, "Multithreaded Rendering", "Turn on to allow Fusion rendering to be done across multiple tasks. NOTE: YOU MUST then connect this node's "
 		                                 "Audio Output AND Render Sync Output to a Fusion Synchronizer Node, and reference THAT NODE'S audio outputs elsewhere in the graph!")
-		METASOUND_PARAM(VoicePoolName, "Voice Pool Name", "The name of a specific voice pool for this node to use. Uses the default pool if the named pool does not exist, or if the name is 'None'")
+		METASOUND_PARAM(VoicePoolName, "Voice Pool Name", "The name of a specific voice pool for this node to use. Uses the default pool if the named pool does not exist, or if the name is \"None\"")
 		METASOUND_PARAM(RenderSync, "Render Sync", "YOU MUST connect this AND this node's audio output to a Fusion Synchronizer Node if you have enabled Multithreaded Rendering for this node!")
 	}
 
diff --git a/Engine/Source/Editor/UnrealEd/Private/Commandlets/GatherTextFromSourceCommandlet.cpp b/Engine/Source/Editor/UnrealEd/Private/Commandlets/GatherTextFromSourceCommandlet.cpp
--- a/Engine/Source/Editor/UnrealEd/Private/Commandlets/GatherTextFromSourceCommandlet.cpp
+++ b/Engine/Source/Editor/UnrealEd/Private/Commandlets/GatherTextFromSourceCommandlet.cpp
@@ -2419,6 +2419,13 @@ void UGatherTextFromSourceCommandlet::FNestedMacroDescriptor::TryParse(const FSt
 		return;
 	}
 
+	// If this nested macro was found in a translation unit (.cpp) then it can only be used in the same file
+	if (Filename.EndsWith(TEXT(".cpp"), ESearchCase::IgnoreCase) &&
+		Filename != Context.Filename)
+	{
+		return;
+	}
+
 	// Ignore partial matches, such as METASOUND_PARAM_EXTERN when we're looking for METASOUND_PARAM
 	FString MacroName = GetToken();
 	int32 Pos = Text.Find(TEXT("("), ESearchCase::CaseSensitive);
@@ -2426,20 +2433,13 @@ void UGatherTextFromSourceCommandlet::FNestedMacroDescriptor::TryParse(const FSt
 	{
 		return;
 	}
-	FString MacroNameCurrent = Text.Mid(0, Pos);	// excludes bracket
+	FString MacroNameCurrent = Text.Mid(0, Pos);		// excludes bracket
 	MacroNameCurrent.TrimEndInline();
 	if (MacroNameCurrent != MacroName)
 	{
 		return;
 	}
 
-	// If this nested macro was found in a translation unit (.cpp) then it can only be used in the same file
-	if (Filename.EndsWith(TEXT(".cpp"), ESearchCase::IgnoreCase) &&
-		Filename != Context.Filename)
-	{
-		return;
-	}
-
 	// Parse outer macro values									MACRONAME("first", "second", "third")
 	TArray<FString> ArgArrayValues;
 	FSourceFileParseContext LocalCtxt1 = Context;		// Local context copy, to avoid changes to the main context
@@ -2603,19 +2603,23 @@ void UGatherTextFromSourceCommandlet::FNestedMacroDescriptor::TryParseArgs(const
 	FString Collect;
 	Collect.Reserve(MacroInnerParams.Len());
 	bool bWithinQuote = false;
+	const TCHAR* CharPrev = *MacroInnerParams;
 	for (const TCHAR* Char = *MacroInnerParams; *Char; ++Char)
 	{
-		if (*Char == TEXT('\"'))
+		if (*Char == TEXT('\"') && *CharPrev != TEXT('\\'))
 		{
+			// Within non-escaped quotes
 			bWithinQuote = !bWithinQuote;
 		}
 		else if (*Char == TEXT(',') && !bWithinQuote)
 		{
 			Params.Add(Collect);
 			Collect.Reset();
+			CharPrev = Char;
 			continue;
 		}
 		Collect += *Char;
+		CharPrev = Char;
 	}
 	Params.Add(Collect);
 
@@ -2630,11 +2634,21 @@ void UGatherTextFromSourceCommandlet::FNestedMacroDescriptor::TryParseArgs(const
 			Collect.Reserve(ParamTrim.Len());
 			bWithinQuote = false;
 			bool bWithinStringification = false;
+			CharPrev = *MacroInnerParams;
 			for (const TCHAR* Char = *ParamTrim; *Char; ++Char)
 			{
+				if (*Char == TEXT('\"') && *CharPrev != TEXT('\\'))
+				{
+					// Within non-escaped quotes
+					bWithinQuote = !bWithinQuote;
+					CharPrev = Char;
+					// Skip quotes, we'll requote
+					continue;
+				}
 				if (*Char == TEXT('#') && !bWithinQuote)
 				{
 					bWithinStringification = true;
+					CharPrev = Char;
 					continue;
 				}
 				if (bWithinStringification)
@@ -2649,18 +2663,15 @@ void UGatherTextFromSourceCommandlet::FNestedMacroDescriptor::TryParseArgs(const
 					if (bIsQuote || bIsSpace)
 					{
 						bWithinStringification = false;
+						CharPrev = Char;
 						continue;
 					}
 				}
-				if (*Char == TEXT('\"'))
-				{
-					bWithinQuote = !bWithinQuote;
-					continue;		// Exclude quotes, we'll requote
-				}
 				if (bWithinStringification || bWithinQuote)
 				{
 					Collect += *Char;
 				}
+				CharPrev = Char;
 			}
 		}
 		else
diff --git a/Engine/Source/Editor/UnrealEd/Private/Commandlets/GatherTextFromSourceTest.cpp b/Engine/Source/Editor/UnrealEd/Private/Commandlets/GatherTextFromSourceTest.cpp
new file mode 100644
--- /dev/null
+++ b/Engine/Source/Editor/UnrealEd/Private/Commandlets/GatherTextFromSourceTest.cpp
@@ -0,0 +1,46 @@
+// Copyright Epic Games, Inc. All Rights Reserved.
+
+#include "Commandlets/GatherTextFromSourceCommandlet.h"
+#include "Containers/UnrealString.h"
+
+
+void UGatherTextFromSourceCommandlet::FNestedMacroDescriptor::TestNestedMacroDescriptorParseArgs()
+{
+	FString MacroInnerParams;
+	FString ParamsNewAll;
+
+	MacroInnerParams = TEXT("\"aaa_\" \"bbb\", \"cc ddd eee\"");
+	TryParseArgs(MacroInnerParams, ParamsNewAll);
+	ensure(ParamsNewAll == TEXT("\"aaa_bbb\", \"cc ddd eee\""));
+
+	// Stringification using hash '#'
+	MacroInnerParams = TEXT("\"aaa \" #bbb \" ccc\"");
+	TryParseArgs(MacroInnerParams, ParamsNewAll);
+	ensure(ParamsNewAll == TEXT("\"aaa bbb ccc\""));
+
+	// Stringification using hash '#'
+	MacroInnerParams = TEXT("#aaa\"bbb\"");
+	TryParseArgs(MacroInnerParams, ParamsNewAll);
+	ensure(ParamsNewAll == TEXT("\"aaabbb\""));
+
+	// String containing hash '#'
+	MacroInnerParams = TEXT("#aaa \"#bbb\" \"ccc ddd\"");
+	TryParseArgs(MacroInnerParams, ParamsNewAll);
+	ensure(ParamsNewAll == TEXT("\"aaa#bbbccc ddd\""));
+
+	// Param containing a comma
+	MacroInnerParams = TEXT("\"aaa\", \"bbb, with comma\"");
+	TryParseArgs(MacroInnerParams, ParamsNewAll);
+	ensure(ParamsNewAll == TEXT("\"aaa\", \"bbb, with comma\""));
+
+	// Param without quotes
+	MacroInnerParams = TEXT("aaa, \"bbb\", ccc");
+	TryParseArgs(MacroInnerParams, ParamsNewAll);
+	ensure(ParamsNewAll == TEXT("\"aaa\", \"bbb\", \"ccc\""));
+
+	// Param with escaped quotes
+	MacroInnerParams = TEXT("aaa, \"bbb \\\"ccc\\\" ddd\", eee");
+	TryParseArgs(MacroInnerParams, ParamsNewAll);
+	ensure(ParamsNewAll == TEXT("\"aaa\", \"bbb \\\"ccc\\\" ddd\", \"eee\""));
+}
+
