Gather text from source to include nested macros; fix for escaped quotes; remove duplicate filepaths in mainpass

diff --git a/Engine/Source/Editor/UnrealEd/Private/Commandlets/GatherTextFromSourceCommandlet.cpp b/Engine/Source/Editor/UnrealEd/Private/Commandlets/GatherTextFromSourceCommandlet.cpp
--- a/Engine/Source/Editor/UnrealEd/Private/Commandlets/GatherTextFromSourceCommandlet.cpp
+++ b/Engine/Source/Editor/UnrealEd/Private/Commandlets/GatherTextFromSourceCommandlet.cpp
@@ -212,6 +212,12 @@ int32 UGatherTextFromSourceCommandlet::Main( const FString& Params )
 
 		return false;
 	});
+	FilesToProcess.Sort([](const FString& LHS, const FString& RHS)
+	{
+		return (LHS < RHS);
+	});
+	// Remove duplicates
+	FilesToProcess.SetNum(Algo::Unique(FilesToProcess));
 	CountFileTypes(FilesToProcess, EGatherSourcePasses::Mainpass);
 	
 	// Return if no source files were found
@@ -2634,7 +2640,7 @@ void UGatherTextFromSourceCommandlet::FNestedMacroDescriptor::TryParseArgs(const
 			Collect.Reserve(ParamTrim.Len());
 			bWithinQuote = false;
 			bool bWithinStringification = false;
-			CharPrev = *MacroInnerParams;
+			CharPrev = *ParamTrim;
 			for (const TCHAR* Char = *ParamTrim; *Char; ++Char)
 			{
 				if (*Char == TEXT('\"') && *CharPrev != TEXT('\\'))
