Raise Surface Snapping to the top-level in the new level editor viewport toolbar

diff --git a/Engine/Source/Editor/UnrealEd/Private/ViewportToolbar/UnrealEdViewportToolbar.cpp b/Engine/Source/Editor/UnrealEd/Private/ViewportToolbar/UnrealEdViewportToolbar.cpp
--- a/Engine/Source/Editor/UnrealEd/Private/ViewportToolbar/UnrealEdViewportToolbar.cpp
+++ b/Engine/Source/Editor/UnrealEd/Private/ViewportToolbar/UnrealEdViewportToolbar.cpp
@@ -136,7 +136,7 @@ FToolMenuEntry CreateSurfaceSnapCheckboxMenu()
 		}
 	);
 
-	return UnrealEd::CreateCheckboxSubmenu(
+	FToolMenuEntry Entry = UnrealEd::CreateCheckboxSubmenu(
 		"SurfaceSnapping",
 		LOCTEXT("SurfaceSnapLabel", "Surface"),
 		FEditorViewportCommands::Get().SurfaceSnapping->MakeTooltip()->GetTextTooltip(),
@@ -157,6 +157,17 @@ FToolMenuEntry CreateSurfaceSnapCheckboxMenu()
 		),
 		MakeMenuDelegate
 	);
+
+	Entry.ToolbarLabelOverride = TAttribute<FText>::CreateLambda(
+		[]()
+		{
+			const ULevelEditorViewportSettings* const Settings = GetMutableDefault<ULevelEditorViewportSettings>();
+			return FText::AsNumber(Settings->SnapToSurface.SnapOffsetExtent);
+		}
+	);
+	Entry.Icon = FSlateIcon(FAppStyle::GetAppStyleSetName(), "EditorViewport.ToggleSurfaceSnapping");
+
+	return Entry;
 }
 
 FToolMenuEntry CreateActorSnapDistanceEntry()
@@ -994,7 +1005,7 @@ FToolMenuEntry CreateViewportToolbarSnappingSubmenu()
 				FToolMenuSection& SnappingSection =
 					Submenu->FindOrAddSection("Snapping", LOCTEXT("SnappingLabel", "Snapping"));
 
-				SnappingSection.AddEntry(Private::CreateSurfaceSnapCheckboxMenu());
+				SnappingSection.AddEntry(Private::CreateSurfaceSnapCheckboxMenu()).SetShowInToolbarTopLevel(true);
 				SnappingSection.AddEntry(Private::CreateLocationSnapCheckboxMenu()).SetShowInToolbarTopLevel(true);
 				SnappingSection.AddEntry(Private::CreateRotationSnapCheckboxMenu()).SetShowInToolbarTopLevel(true);
 				SnappingSection.AddEntry(Private::CreateScaleSnapCheckboxMenu()).SetShowInToolbarTopLevel(true);
