Change icon on the new level editor viewport toolbar Transform menu

diff --git a/Engine/Source/Editor/UnrealEd/Private/ViewportToolbar/UnrealEdViewportToolbar.cpp b/Engine/Source/Editor/UnrealEd/Private/ViewportToolbar/UnrealEdViewportToolbar.cpp
--- a/Engine/Source/Editor/UnrealEd/Private/ViewportToolbar/UnrealEdViewportToolbar.cpp
+++ b/Engine/Source/Editor/UnrealEd/Private/ViewportToolbar/UnrealEdViewportToolbar.cpp
@@ -677,7 +677,7 @@ FToolMenuEntry CreateViewportToolbarTransformsSection()
 	);
 
 	Entry.ToolbarLabelOverride = FText();
-	Entry.Icon = FSlateIcon(FAppStyle::GetAppStyleSetName(), "EditorViewport.RelativeCoordinateSystem_Local");
+	Entry.Icon = FSlateIcon(FAppStyle::GetAppStyleSetName(), "LevelEditor.SelectMode");
 	return Entry;
 }
 
