Collect world and local coordinate systems into a Coordinate Systems submenu in the new level editor viewport toolbar

diff --git a/Engine/Source/Editor/UnrealEd/Private/STransformViewportToolbar.cpp b/Engine/Source/Editor/UnrealEd/Private/STransformViewportToolbar.cpp
--- a/Engine/Source/Editor/UnrealEd/Private/STransformViewportToolbar.cpp
+++ b/Engine/Source/Editor/UnrealEd/Private/STransformViewportToolbar.cpp
@@ -532,20 +532,12 @@ FReply STransformViewportToolBar::OnCycleCoordinateSystem()
 
 FSlateIcon STransformViewportToolBar::GetLocalToWorldIcon() const
 {
-	if( Viewport.IsValid() && Viewport.Pin()->IsCoordSystemActive(COORD_World) )
+	ECoordSystem CoordSystem = ECoordSystem::COORD_Local;
+	if (TSharedPtr<SEditorViewport> PinnedViewport = Viewport.Pin())
 	{
-		static FName WorldIcon("EditorViewport.RelativeCoordinateSystem_World");
-		return FSlateIcon(FAppStyle::GetAppStyleSetName(), WorldIcon);
+		CoordSystem = PinnedViewport->GetViewportClient()->GetWidgetCoordSystemSpace();
 	}
-
-	if( Viewport.IsValid() && Viewport.Pin()->IsCoordSystemActive(COORD_Parent) )
-	{
-		static const FName ParentIcon("Icons.ConstraintManager.ParentHierarchy");
-		return FSlateIcon(FAppStyle::GetAppStyleSetName(), ParentIcon);
-	}
-	
-	static FName LocalIcon("Icons.Transform");
-	return FSlateIcon(FAppStyle::GetAppStyleSetName(), LocalIcon);
+	return UE::UnrealEd::GetIconFromCoordSystem(CoordSystem);
 }
 
 FText STransformViewportToolBar::GetLayer2DLabel() const
diff --git a/Engine/Source/Editor/UnrealEd/Private/ViewportToolbar/UnrealEdViewportToolbar.cpp b/Engine/Source/Editor/UnrealEd/Private/ViewportToolbar/UnrealEdViewportToolbar.cpp
--- a/Engine/Source/Editor/UnrealEd/Private/ViewportToolbar/UnrealEdViewportToolbar.cpp
+++ b/Engine/Source/Editor/UnrealEd/Private/ViewportToolbar/UnrealEdViewportToolbar.cpp
@@ -573,6 +573,25 @@ bool ShowNewViewportToolbars()
 	return Private::CVarToolMenusViewportToolbarsValue >= 1;
 }
 
+FSlateIcon GetIconFromCoordSystem(ECoordSystem InCoordSystem)
+{
+	if (InCoordSystem == COORD_World)
+	{
+		static FName WorldIcon("EditorViewport.RelativeCoordinateSystem_World");
+		return FSlateIcon(FAppStyle::GetAppStyleSetName(), WorldIcon);
+	}
+	else if (InCoordSystem == COORD_Parent)
+	{
+		static const FName ParentIcon("Icons.ConstraintManager.ParentHierarchy");
+		return FSlateIcon(FAppStyle::GetAppStyleSetName(), ParentIcon);
+	}
+	else
+	{
+		static FName LocalIcon("Icons.Transform");
+		return FSlateIcon(FAppStyle::GetAppStyleSetName(), LocalIcon);
+	}
+}
+
 FToolMenuEntry CreateViewportToolbarTransformsSection()
 {
 	FToolMenuEntry Entry = FToolMenuEntry::InitSubMenu(
@@ -606,22 +625,59 @@ FToolMenuEntry CreateViewportToolbarTransformsSection()
 					ScaleMode.UserInterfaceActionType = EUserInterfaceActionType::RadioButton;
 					ScaleMode.SetShowInToolbarTopLevel(true);
 					TransformToolsSection.AddEntry(ScaleMode);
-				}
 
-				{
-					FToolMenuSection& SpacesSection = Submenu->FindOrAddSection("Spaces", LOCTEXT("SpacesLabel", "Spaces"));
-
-					FToolMenuEntry WorldSpace =
-						FToolMenuEntry::InitMenuEntry(FEditorViewportCommands::Get().RelativeCoordinateSystem_World);
-					WorldSpace.UserInterfaceActionType = EUserInterfaceActionType::RadioButton;
-					WorldSpace.SetShowInToolbarTopLevel(true);
-					SpacesSection.AddEntry(WorldSpace);
-
-					FToolMenuEntry LocalSpace =
-						FToolMenuEntry::InitMenuEntry(FEditorViewportCommands::Get().RelativeCoordinateSystem_Local);
-					LocalSpace.UserInterfaceActionType = EUserInterfaceActionType::RadioButton;
-					LocalSpace.SetShowInToolbarTopLevel(true);
-					SpacesSection.AddEntry(LocalSpace);
+					// Build a submenu for selecting the coordinate system to use.
+					{
+						TransformToolsSection.AddSeparator("CoordinateSystemSeparator");
+
+						FToolMenuEntry& CoordianteSystemSubmenu = TransformToolsSection.AddSubMenu(
+							"CoordinateSystem",
+							LOCTEXT("CoordinateSystemLabel", "Coordinate System"),
+							LOCTEXT("CoordinateSystemTooltip", "Select between coordinate systems"),
+							FNewToolMenuDelegate::CreateLambda(
+								[](UToolMenu* InSubmenu)
+								{
+									FToolMenuSection& UnnamedSection = InSubmenu->FindOrAddSection(NAME_None);
+
+									FToolMenuEntry WorldCoords = FToolMenuEntry::InitMenuEntry(
+										FEditorViewportCommands::Get().RelativeCoordinateSystem_World
+									);
+									WorldCoords.UserInterfaceActionType = EUserInterfaceActionType::RadioButton;
+									UnnamedSection.AddEntry(WorldCoords);
+
+									FToolMenuEntry LocalCoords = FToolMenuEntry::InitMenuEntry(
+										FEditorViewportCommands::Get().RelativeCoordinateSystem_Local
+									);
+									LocalCoords.UserInterfaceActionType = EUserInterfaceActionType::RadioButton;
+									UnnamedSection.AddEntry(LocalCoords);
+								}
+							)
+						);
+
+						// Set the icon based on the current coordinate system and fall back to the Local icon.
+						{
+							TWeakPtr<SEditorViewport> WeakViewport;
+							if (UUnrealEdViewportToolbarContext* const Context =
+									Submenu->FindContext<UUnrealEdViewportToolbarContext>())
+							{
+								WeakViewport = Context->Viewport;
+							}
+
+							CoordianteSystemSubmenu.Icon = TAttribute<FSlateIcon>::CreateLambda(
+								[WeakViewport]() -> FSlateIcon
+								{
+									ECoordSystem CoordSystem = ECoordSystem::COORD_Local;
+									if (const TSharedPtr<SEditorViewport> EditorViewport = WeakViewport.Pin())
+									{
+										CoordSystem = EditorViewport->GetViewportClient()->GetWidgetCoordSystemSpace();
+									}
+									return GetIconFromCoordSystem(CoordSystem);
+								}
+							);
+						}
+						CoordianteSystemSubmenu.ToolbarLabelOverride = FText();
+						CoordianteSystemSubmenu.SetShowInToolbarTopLevel(true);
+					}
 				}
 
 				{
diff --git a/Engine/Source/Editor/UnrealEd/Public/ViewportToolbar/UnrealEdViewportToolbar.h b/Engine/Source/Editor/UnrealEd/Public/ViewportToolbar/UnrealEdViewportToolbar.h
--- a/Engine/Source/Editor/UnrealEd/Public/ViewportToolbar/UnrealEdViewportToolbar.h
+++ b/Engine/Source/Editor/UnrealEd/Public/ViewportToolbar/UnrealEdViewportToolbar.h
@@ -7,6 +7,7 @@
 #include "Templates/SharedPointerFwd.h"
 #include "ToolMenuDelegates.h"
 #include "UnrealEdViewportToolbarContext.h"
+#include "UnrealWidgetFwd.h"
 
 class FEditorViewportClient;
 class FText;
@@ -36,6 +37,7 @@ UNREALED_API bool ShowOldViewportToolbars();
 /** The value of this function is controlled by the CVAR "ToolMenusViewportToolbars". */
 UNREALED_API bool ShowNewViewportToolbars();
 
+UNREALED_API FSlateIcon GetIconFromCoordSystem(ECoordSystem InCoordSystem);
 UNREALED_API FToolMenuEntry CreateViewportToolbarTransformsSection();
 
 UNREALED_API FToolMenuEntry CreateViewportToolbarSelectSection();
