[Sequencer] Selectability Feature:
 - Change method of checking "selectable in viewport"
 - Fix marquee selection not being limited

diff --git a/Engine/Source/Editor/Sequencer/Private/Sequencer.cpp b/Engine/Source/Editor/Sequencer/Private/Sequencer.cpp
--- a/Engine/Source/Editor/Sequencer/Private/Sequencer.cpp
+++ b/Engine/Source/Editor/Sequencer/Private/Sequencer.cpp
@@ -12135,8 +12135,6 @@ void FSequencer::SetViewportSelectionLimited(const bool bInSelectionLimited)
 
 bool FSequencer::IsObjectSelectableInViewport(UObject* const InObject)
 {
-	using namespace UE::Sequencer;
-
 	if (!bSelectionLimited)
 	{
 		return true;
@@ -12147,6 +12145,12 @@ bool FSequencer::IsObjectSelectableInViewport(UObject* const InObject)
 		return SelectableObject->IsSelectable();
 	}
 
+	const FGuid ObjectGuid = FindObjectId(*InObject, GetFocusedTemplateID());
+	if (ObjectGuid.IsValid())
+	{
+		return true;
+	}
+
 	UMovieSceneSequence* const FocusedSequence = GetFocusedMovieSceneSequence();
 	if (!IsValid(FocusedSequence))
 	{
diff --git a/Engine/Source/Editor/Sequencer/Private/SequencerEdMode.cpp b/Engine/Source/Editor/Sequencer/Private/SequencerEdMode.cpp
--- a/Engine/Source/Editor/Sequencer/Private/SequencerEdMode.cpp
+++ b/Engine/Source/Editor/Sequencer/Private/SequencerEdMode.cpp
@@ -159,7 +159,8 @@ struct FTrackTransforms
 FSequencerEdMode::FSequencerEdMode()
 {
 	DefaultTool = MakeShared<FSequencerEdModeTool>(this);
-	SelectabilityTool = MakeShared<FSequencerSelectabilityTool>(FOnIsObjectSelectableInViewport::CreateRaw(this, &FSequencerEdMode::IsObjectSelectableInViewport));
+	SelectabilityTool = MakeShared<FSequencerSelectabilityTool>(FOnGetWorld::CreateRaw(this, &FSequencerEdMode::GetWorld)
+		, FOnIsObjectSelectableInViewport::CreateRaw(this, &FSequencerEdMode::IsObjectSelectableInViewport));
 
 	Tools.Add(DefaultTool.Get());
 	Tools.Add(SelectabilityTool.Get());
diff --git a/Engine/Source/Editor/Sequencer/Private/SequencerSelectabilityTool.cpp b/Engine/Source/Editor/Sequencer/Private/SequencerSelectabilityTool.cpp
--- a/Engine/Source/Editor/Sequencer/Private/SequencerSelectabilityTool.cpp
+++ b/Engine/Source/Editor/Sequencer/Private/SequencerSelectabilityTool.cpp
@@ -8,9 +8,8 @@
 
 #define LOCTEXT_NAMESPACE "SequencerSelectabilityTool"
 
-FSequencerSelectabilityTool::FSequencerSelectabilityTool(const FOnIsObjectSelectableInViewport& InOnIsObjectSelectableInViewport)
-	: FModeTool()
-	, FEditorViewportSelectability(InOnIsObjectSelectableInViewport)
+FSequencerSelectabilityTool::FSequencerSelectabilityTool(const FOnGetWorld& InOnGetWorld, const FOnIsObjectSelectableInViewport& InOnIsObjectSelectableInViewport)
+	: FEditorViewportSelectability(InOnGetWorld, InOnIsObjectSelectableInViewport)
 {
 }
 
diff --git a/Engine/Source/Editor/Sequencer/Private/SequencerSelectabilityTool.h b/Engine/Source/Editor/Sequencer/Private/SequencerSelectabilityTool.h
--- a/Engine/Source/Editor/Sequencer/Private/SequencerSelectabilityTool.h
+++ b/Engine/Source/Editor/Sequencer/Private/SequencerSelectabilityTool.h
@@ -9,7 +9,7 @@ class FSequencerSelectabilityTool : public FModeTool, public FEditorViewportSele
 {
 public:
 	FSequencerSelectabilityTool() = delete;
-	FSequencerSelectabilityTool(const FOnIsObjectSelectableInViewport& InOnIsObjectSelectableInViewport);
+	FSequencerSelectabilityTool(const FOnGetWorld& InOnGetWorld, const FOnIsObjectSelectableInViewport& InOnIsObjectSelectableInViewport);
 
 	//~ Begin FModeTool
 	virtual FString GetName() const override
diff --git a/Engine/Source/Editor/UnrealEd/Private/EditorViewportSelectability.cpp b/Engine/Source/Editor/UnrealEd/Private/EditorViewportSelectability.cpp
--- a/Engine/Source/Editor/UnrealEd/Private/EditorViewportSelectability.cpp
+++ b/Engine/Source/Editor/UnrealEd/Private/EditorViewportSelectability.cpp
@@ -16,8 +16,9 @@
 
 const FText FEditorViewportSelectability::DefaultLimitedSelectionText = LOCTEXT("DefaultSelectionLimitedHelp", "Viewport Selection Limited");
 
-FEditorViewportSelectability::FEditorViewportSelectability(const FOnIsObjectSelectableInViewport& InOnIsObjectSelectableInViewport)
-	: OnIsObjectSelectableInViewportDelegate(InOnIsObjectSelectableInViewport)
+FEditorViewportSelectability::FEditorViewportSelectability(const FOnGetWorld& InOnGetWorld, const FOnIsObjectSelectableInViewport& InOnIsObjectSelectableInViewport)
+	: OnGetWorld(InOnGetWorld)
+	, OnIsObjectSelectableInViewportDelegate(InOnIsObjectSelectableInViewport)
 {
 }
 
@@ -223,7 +224,9 @@ void FEditorViewportSelectability::DeselectNonSelectableActors()
 	TArray<AActor*> SelectedActors;
 	ActorSelection->GetSelectedObjects<AActor>(SelectedActors);
 
-	SelectActorsByPredicate(false, false
+	UWorld* const World = OnGetWorld.IsBound() ? OnGetWorld.Execute() : nullptr;
+
+	SelectActorsByPredicate(World, false, false
 		, [this](AActor* const InActor) -> bool
 		{
 			return IsObjectSelectableInViewport(InActor);
@@ -231,13 +234,19 @@ void FEditorViewportSelectability::DeselectNonSelectableActors()
 		, SelectedActors);
 }
 
-bool FEditorViewportSelectability::SelectActorsByPredicate(const bool bInSelect
+bool FEditorViewportSelectability::SelectActorsByPredicate(UWorld* const InWorld
+	, const bool bInSelect
 	, const bool bInClearSelection
 	, const TFunctionRef<bool(AActor*)> InPredicate
 	, const TArray<AActor*>& InActors)
 {
-	USelection* const ActorSelection = GEditor ? GEditor->GetSelectedActors() : nullptr;
-	if (!ActorSelection || ActorSelection->Num() == 0)
+	if (!GEditor || !IsValid(InWorld))
+	{
+		return false;
+	}
+
+	USelection* const ActorSelection = GEditor->GetSelectedActors();
+	if (!ActorSelection)
 	{
 		return false;
 	}
@@ -250,14 +259,14 @@ bool FEditorViewportSelectability::SelectActorsByPredicate(const bool bInSelect
 	ActorSelection->BeginBatchSelectOperation();
 	ActorSelection->Modify();
 
-	if (bInClearSelection)
+	if (bInSelect && bInClearSelection)
 	{
 		ActorSelection->DeselectAll();
 	}
 
 	auto SelectIfPossible = [bInSelect, ActorSelection, &InPredicate, &bSomethingSelected](AActor* const InActor)
 	{
-		if (IsValid(InActor) && !ActorSelection->IsSelected(InActor) && InPredicate(InActor))
+		if (IsValid(InActor) && ActorSelection->IsSelected(InActor) != bInSelect && InPredicate(InActor) == bInSelect)
 		{
 			bSomethingSelected = true;
 			GEditor->SelectActor(InActor, bInSelect, true);
@@ -266,14 +275,10 @@ bool FEditorViewportSelectability::SelectActorsByPredicate(const bool bInSelect
 
 	if (InActors.IsEmpty())
 	{
-		const UWorld* const World = ActorSelection->GetWorld();
-		if (IsValid(World))
+		for (FActorIterator Iter(InWorld); Iter; ++Iter)
 		{
-			for (FActorIterator Iter(World); Iter; ++Iter)
-			{
-				AActor* const Actor = *Iter;
-				SelectIfPossible(Actor);
-			}
+			AActor* const Actor = *Iter;
+			SelectIfPossible(Actor);
 		}
 	}
 	else
@@ -485,7 +490,9 @@ bool FEditorViewportSelectability::BoxSelectWorldActors(FBox& InBox, FEditorView
 	const bool bShiftDown = InEditorViewportClient->Viewport->KeyState(EKeys::LeftShift)
 		|| InEditorViewportClient->Viewport->KeyState(EKeys::RightShift);
 
-	return SelectActorsByPredicate(bInSelect, bShiftDown, Predicate);
+	SelectActorsByPredicate(World, bInSelect, bShiftDown, Predicate);
+
+	return true;
 }
 
 bool FEditorViewportSelectability::FrustumSelectWorldActors(const FConvexVolume& InFrustum, FEditorViewportClient* const InEditorViewportClient, const bool bInSelect)
@@ -515,7 +522,9 @@ bool FEditorViewportSelectability::FrustumSelectWorldActors(const FConvexVolume&
 	const bool bShiftDown = InEditorViewportClient->Viewport->KeyState(EKeys::LeftShift)
 		|| InEditorViewportClient->Viewport->KeyState(EKeys::RightShift);
 
-	return SelectActorsByPredicate(bInSelect, bShiftDown, Predicate);
+	SelectActorsByPredicate(World, bInSelect, bShiftDown, Predicate);
+
+	return true;
 }
 
 void FEditorViewportSelectability::DrawEnabledTextNotice(FCanvas* const InCanvas, const FText& InText)
diff --git a/Engine/Source/Editor/UnrealEd/Public/EditorViewportSelectability.h b/Engine/Source/Editor/UnrealEd/Public/EditorViewportSelectability.h
--- a/Engine/Source/Editor/UnrealEd/Public/EditorViewportSelectability.h
+++ b/Engine/Source/Editor/UnrealEd/Public/EditorViewportSelectability.h
@@ -17,6 +17,7 @@ class HHitProxy;
 struct FConvexVolume;
 struct FViewportClick;
 
+DECLARE_DELEGATE_RetVal(UWorld*, FOnGetWorld);
 DECLARE_DELEGATE_RetVal_OneParam(bool, FOnIsObjectSelectableInViewport, UObject* /*InObject*/);
 
 /**
@@ -49,7 +50,7 @@ public:
 	static void DrawEnabledTextNotice(FCanvas* const InCanvas, const FText& InText);
 
 	FEditorViewportSelectability() = delete;
-	FEditorViewportSelectability(const FOnIsObjectSelectableInViewport& InOnIsObjectSelectableInViewport);
+	FEditorViewportSelectability(const FOnGetWorld& InOnGetWorld, const FOnIsObjectSelectableInViewport& InOnIsObjectSelectableInViewport);
 
 	/** Enables or disables the selectability tool */
 	void EnableLimitedSelection(const bool bInEnabled);
@@ -93,7 +94,8 @@ protected:
 	 * @param bInClearSelection If true, clears the current selection before selecting the new actors
 	 * @return True if atleast one new actor was selected/deselected
 	 */
-	static bool SelectActorsByPredicate(const bool bInSelect
+	static bool SelectActorsByPredicate(UWorld* const InWorld
+		, const bool bInSelect
 		, const bool bInClearSelection
 		, const TFunctionRef<bool(AActor*)> InPredicate
 		, const TArray<AActor*>& InActors = {});
@@ -109,6 +111,8 @@ protected:
 	/** Mouse cursor to display for the viewport when selection is limited. */
 	TOptional<EMouseCursor::Type> MouseCursor;
 
+	FOnGetWorld OnGetWorld;
+
 	/** Delegate used to check if an object is selectable in the viewport */
 	FOnIsObjectSelectableInViewport OnIsObjectSelectableInViewportDelegate;
 };
