Gather text from source to include nested macros;
Ignore matches of the suffix, such as DECLARE_METASOUND_PARAM when we're looking for METASOUND_PARAM

diff --git a/Engine/Source/Editor/UnrealEd/Private/Commandlets/GatherTextFromSourceCommandlet.cpp b/Engine/Source/Editor/UnrealEd/Private/Commandlets/GatherTextFromSourceCommandlet.cpp
--- a/Engine/Source/Editor/UnrealEd/Private/Commandlets/GatherTextFromSourceCommandlet.cpp
+++ b/Engine/Source/Editor/UnrealEd/Private/Commandlets/GatherTextFromSourceCommandlet.cpp
@@ -2432,13 +2432,13 @@ void UGatherTextFromSourceCommandlet::FNestedMacroDescriptor::TryParse(const FSt
 		return;
 	}
 
-	// Ignore partial matches, such as METASOUND_PARAM_EXTERN when we're looking for METASOUND_PARAM
-	FString MacroName = GetToken();
+	// Ignore matches of the prefix, such as METASOUND_PARAM_EXTERN when we're looking for METASOUND_PARAM
 	int32 Pos = Text.Find(TEXT("("), ESearchCase::CaseSensitive);
 	if (Pos < 0)
 	{
 		return;
 	}
+	FString MacroName = GetToken();
 	FString MacroNameCurrent = Text.Mid(0, Pos);		// excludes bracket
 	MacroNameCurrent.TrimEndInline();
 	if (MacroNameCurrent != MacroName)
@@ -2446,6 +2446,17 @@ void UGatherTextFromSourceCommandlet::FNestedMacroDescriptor::TryParse(const FSt
 		return;
 	}
 
+	// Ignore matches of the suffix, such as DECLARE_METASOUND_PARAM when we're looking for METASOUND_PARAM
+	Pos = Context.LineText.Find(MacroName, ESearchCase::CaseSensitive);
+	if (Pos > 0)
+	{
+		TCHAR Char = Context.LineText[Pos - 1];
+		if (!FText::IsWhitespace(Char) && Char != TCHAR('(') && Char != TCHAR('{'))
+		{
+			return;
+		}
+	}
+
 	// Parse outer macro values									MACRONAME("first", "second", "third")
 	TArray<FString> ArgArrayValues;
 	FSourceFileParseContext LocalCtxt1 = Context;		// Local context copy, to avoid changes to the main context
