Rename RelativeCoordinateSystem_World/Local

* Rename RelativeCoordinateSystem_World's FriendlyName from "World-relative Transform" to "World Space".
* Rename RelativeCoordinateSystem_Local's FriendlyName from "Local-relative Transform" to "Local Space".

diff --git a/Engine/Source/Editor/UnrealEd/Private/EditorViewportCommands.cpp b/Engine/Source/Editor/UnrealEd/Private/EditorViewportCommands.cpp
--- a/Engine/Source/Editor/UnrealEd/Private/EditorViewportCommands.cpp
+++ b/Engine/Source/Editor/UnrealEd/Private/EditorViewportCommands.cpp
@@ -154,8 +154,20 @@ void FEditorViewportCommands::RegisterCommands()
 	UI_COMMAND( ShrinkTransformWidget, "Shrink Transform Widget", "Shrink the level editor transform widget", EUserInterfaceActionType::Button, FInputChord(EModifierKey::Alt, EKeys::LeftBracket) );
 	UI_COMMAND( ExpandTransformWidget, "Expand Transform Widget", "Expand the level editor transform widget", EUserInterfaceActionType::Button, FInputChord(EModifierKey::Alt, EKeys::RightBracket) );
 
-	UI_COMMAND( RelativeCoordinateSystem_World, "World-relative Transform", "Move and rotate objects relative to the cardinal world axes", EUserInterfaceActionType::RadioButton, FInputChord() );
-	UI_COMMAND( RelativeCoordinateSystem_Local, "Local-relative Transform", "Move and rotate objects relative to the object's local axes", EUserInterfaceActionType::RadioButton, FInputChord() );
+	UI_COMMAND(
+		RelativeCoordinateSystem_World,
+		"World Space",
+		"Move and rotate objects relative to the cardinal world axes",
+		EUserInterfaceActionType::RadioButton,
+		FInputChord()
+	);
+	UI_COMMAND(
+		RelativeCoordinateSystem_Local,
+		"Local Space",
+		"Move and rotate objects relative to the object's local axes",
+		EUserInterfaceActionType::RadioButton,
+		FInputChord()
+	);
 
 #if PLATFORM_MAC
 	UI_COMMAND( CycleTransformGizmoCoordSystem, "Cycle Transform Coordinate System", "Cycles the transform gizmo coordinate systems between world and local (object) space", EUserInterfaceActionType::Button, FInputChord(EModifierKey::Command, EKeys::Tilde));
