Fbx Legacy Importer: Fix crash when importing multiple skin the same file. The pass parent package get gc after the first sk is imported, we now store the parent package path before importing the first sk.

#rnx

diff --git a/Engine/Source/Editor/UnrealEd/Private/Fbx/FbxFactory.cpp b/Engine/Source/Editor/UnrealEd/Private/Fbx/FbxFactory.cpp
--- a/Engine/Source/Editor/UnrealEd/Private/Fbx/FbxFactory.cpp
+++ b/Engine/Source/Editor/UnrealEd/Private/Fbx/FbxFactory.cpp
@@ -204,6 +204,8 @@ UObject* UFbxFactory::FactoryCreateFile
 
 	GEditor->GetEditorSubsystem<UImportSubsystem>()->BroadcastAssetPreImport(this, Class, InParent, Name, Type);
 
+	FString ParentPackagePath = (InParent && InParent->GetOutermost()) ? FPackageName::GetLongPackagePath(InParent->GetOutermost()->GetName()) : FString(TEXT("/Game"));
+
 	UObject* CreatedObject = NULL;
 	//Look if its a re-import, in that case we must call the re-import factory
 	UObject *ExistingObject = nullptr;
@@ -723,9 +725,9 @@ UObject* UFbxFactory::FactoryCreateFile
 								if (Package == nullptr)
 								{
 									FString NewPackageName;
-									if (InParent != nullptr && InParent->GetOutermost() != nullptr)
+									if (!ParentPackagePath.IsEmpty())
 									{
-										NewPackageName = FPackageName::GetLongPackagePath(InParent->GetOutermost()->GetName()) + TEXT("/") + OutputName.ToString();
+										NewPackageName = ParentPackagePath + TEXT("/") + OutputName.ToString();
 									}
 									else
 									{
