MPCook: CookPackageSplitter: Add RequiresCachedCookedPlatformDataBeforeSplit function for ICookPackageSplitter classes, and add a save substate to support it by waiting for IsCachedCookedPlatformDataLoaded to return true before calling ShouldSplit.

Refactor ESaveState FGenerationHelper to instead be ESaveSubState on FPackageData. This reduces the complexity of maintaining two state machines for the save state. The first state machine was the de facto state machine created out of bools on the FPackageData, and the second was the ESaveState state machine on the FGenerationHelper. They are now merged into the single ESaveSubState state machine.
#rnx

diff --git a/Engine/Source/Editor/UnrealEd/Private/CookOnTheFlyServer.cpp b/Engine/Source/Editor/UnrealEd/Private/CookOnTheFlyServer.cpp
--- a/Engine/Source/Editor/UnrealEd/Private/CookOnTheFlyServer.cpp
+++ b/Engine/Source/Editor/UnrealEd/Private/CookOnTheFlyServer.cpp
@@ -3347,7 +3347,7 @@ UE::Cook::EPollStatus UCookOnTheFlyServer::QueueGeneratedPackages(UE::Cook::FGen
 
 	UPackage* Owner = PackageData.GetPackage();
 	FName OwnerName = Owner->GetFName();
-	if (Info.GetSaveState() <= FCookGenerationInfo::ESaveState::QueueGeneratedPackages)
+	if (PackageData.GetSaveSubState() <= ESaveSubState::Generation_QueueGeneratedPackages)
 	{
 		GenerationHelper.StartQueueGeneratedPackages(*this);
 		TArray<const ITargetPlatform*, TInlineAllocator<ExpectedMaxNumPlatforms>> ReachablePlatforms;
@@ -3366,8 +3366,6 @@ UE::Cook::EPollStatus UCookOnTheFlyServer::QueueGeneratedPackages(UE::Cook::FGen
 				EDiscoveredPlatformSet::CopyFromInstigator, bUrgent, &GenerationHelper);
 		}
 		GenerationHelper.EndQueueGeneratedPackages(*this);
-
-		Info.SetSaveStateComplete(FCookGenerationInfo::ESaveState::QueueGeneratedPackages);
 	}
 	return EPollStatus::Success;
 }
@@ -3386,7 +3384,7 @@ UE::Cook::EPollStatus UCookOnTheFlyServer::PrepareSaveGenerationPackage(UE::Cook
 	}
 	FCookGenerationInfo& Info(*InfoPtr);
 
-	if (Info.GetSaveState() <= FCookGenerationInfo::ESaveState::FinishCachePreMove)
+	if (PackageData.GetSaveSubState() <= ESaveSubState::Generation_PreMoveCookedPlatformData_WaitingForIsLoaded)
 	{
 		// Both Generator packages and Generated packages should wait for all IsCachedCookedPlatformData
 		// to finish before they start BeginCache calls on the objects to move.
@@ -3429,30 +3427,30 @@ UE::Cook::EPollStatus UCookOnTheFlyServer::PrepareSaveGenerationPackage(UE::Cook
 				return EPollStatus::Incomplete;
 			}
 		}
-		Info.SetSaveStateComplete(FCookGenerationInfo::ESaveState::FinishCachePreMove);
+		PackageData.SetSaveSubStateComplete(ESaveSubState::Generation_PreMoveCookedPlatformData_WaitingForIsLoaded);
 	}
 
 	// GeneratedPackagesForPresave is used by multiple steps, recreate it when needed each time we come in to this function
 	TArray<ICookPackageSplitter::FGeneratedPackageForPreSave> GeneratedPackagesForPresave;
-	if (Info.GetSaveState() <= FCookGenerationInfo::ESaveState::FinishCacheObjectsToMove)
+	if (PackageData.GetSaveSubState() <= ESaveSubState::Generation_FinishCacheObjectsToMove)
 	{
-		if (Info.GetSaveState() <= FCookGenerationInfo::ESaveState::BeginCacheObjectsToMove)
+		if (PackageData.GetSaveSubState() <= ESaveSubState::Generation_BeginCacheObjectsToMove)
 		{
 			EPollStatus Result = BeginCacheObjectsToMove(GenerationHelper, Info, Timer, GeneratedPackagesForPresave);
 			if (Result != EPollStatus::Success)
 			{
 				return Result;
 			}
-			Info.SetSaveStateComplete(FCookGenerationInfo::ESaveState::BeginCacheObjectsToMove);
+			PackageData.SetSaveSubStateComplete(ESaveSubState::Generation_BeginCacheObjectsToMove);
 		}
-		check(Info.GetSaveState() <= FCookGenerationInfo::ESaveState::FinishCacheObjectsToMove);
+		check(PackageData.GetSaveSubState() <= ESaveSubState::Generation_FinishCacheObjectsToMove);
 		if (PackageData.GetNumPendingCookedPlatformData() > 0)
 		{
 			return EPollStatus::Incomplete;
 		}
 		bool bFoundNewObjects;
 		EPollStatus Result = Info.RefreshPackageObjects(GenerationHelper, PackageData.GetPackage(), bFoundNewObjects,
-			FCookGenerationInfo::ESaveState::BeginCacheObjectsToMove);
+			ESaveSubState::Generation_BeginCacheObjectsToMove);
 		if (Result != EPollStatus::Success)
 		{
 			return Result;
@@ -3463,10 +3461,10 @@ UE::Cook::EPollStatus UCookOnTheFlyServer::PrepareSaveGenerationPackage(UE::Cook
 			// Note that RefreshPackageObjects checked for too many recursive calls and ErrorExited if so.
 			return PrepareSaveGenerationPackage(GenerationHelper, PackageData, Timer, bPrecaching);
 		}
-		Info.SetSaveStateComplete(FCookGenerationInfo::ESaveState::FinishCacheObjectsToMove);
+		PackageData.SetSaveSubStateComplete(ESaveSubState::Generation_FinishCacheObjectsToMove);
 	}
 
-	if (Info.GetSaveState() <= FCookGenerationInfo::ESaveState::CallPopulate)
+	if (PackageData.GetSaveSubState() <= ESaveSubState::Generation_CallPopulate)
 	{
 		if (bPrecaching)
 		{
@@ -3489,28 +3487,28 @@ UE::Cook::EPollStatus UCookOnTheFlyServer::PrepareSaveGenerationPackage(UE::Cook
 		{
 			return Result;
 		}
-		Info.SetSaveStateComplete(FCookGenerationInfo::ESaveState::CallPopulate);
+		PackageData.SetSaveSubStateComplete(ESaveSubState::Generation_CallPopulate);
 	}
 
-	if (Info.GetSaveState() <= FCookGenerationInfo::ESaveState::FinishCachePostMove)
+	if (PackageData.GetSaveSubState() <= ESaveSubState::LastCookedPlatformData_WaitingForIsLoaded)
 	{
-		if (Info.GetSaveState() <= FCookGenerationInfo::ESaveState::BeginCachePostMove)
+		if (PackageData.GetSaveSubState() <= ESaveSubState::LastCookedPlatformData_CallingBegin)
 		{
 			EPollStatus Result = BeginCachePostMove(GenerationHelper, Info, Timer);
 			if (Result != EPollStatus::Success)
 			{
 				return Result;
 			}
-			Info.SetSaveStateComplete(FCookGenerationInfo::ESaveState::BeginCachePostMove);
+			PackageData.SetSaveSubStateComplete(ESaveSubState::LastCookedPlatformData_CallingBegin);
 		}
-		check(Info.GetSaveState() <= FCookGenerationInfo::ESaveState::FinishCachePostMove);
+		check(PackageData.GetSaveSubState() <= ESaveSubState::LastCookedPlatformData_WaitingForIsLoaded);
 		if (PackageData.GetNumPendingCookedPlatformData() > 0)
 		{
 			return EPollStatus::Incomplete;
 		}
 		bool bFoundNewObjects;
 		EPollStatus Result = Info.RefreshPackageObjects(GenerationHelper, PackageData.GetPackage(), bFoundNewObjects,
-			FCookGenerationInfo::ESaveState::BeginCachePostMove);
+			ESaveSubState::LastCookedPlatformData_CallingBegin);
 		if (Result != EPollStatus::Success)
 		{
 			return Result;
@@ -3522,9 +3520,9 @@ UE::Cook::EPollStatus UCookOnTheFlyServer::PrepareSaveGenerationPackage(UE::Cook
 			return PrepareSaveGenerationPackage(GenerationHelper, PackageData, Timer, bPrecaching);
 		}
 
-		Info.SetSaveStateComplete(FCookGenerationInfo::ESaveState::FinishCachePostMove);
+		PackageData.SetSaveSubStateComplete(ESaveSubState::LastCookedPlatformData_WaitingForIsLoaded);
 	}
-	check(Info.GetSaveState() == FCookGenerationInfo::ESaveState::ReadyForSave);
+	check(PackageData.GetSaveSubState() == ESaveSubState::ReadyForSave);
 
 	return EPollStatus::Success;
 }
@@ -3545,7 +3543,7 @@ UE::Cook::EPollStatus UCookOnTheFlyServer::BeginCacheObjectsToMove(UE::Cook::FGe
 		return EPollStatus::Error;
 	}
 
-	if (Info.GetSaveState() <= FCookGenerationInfo::ESaveState::CallObjectsToMove)
+	if (PackageData.GetSaveSubState() <= ESaveSubState::Generation_CallObjectsToMove)
 	{
 		bool bPopulateSucceeded = false;
 		if (Info.IsGenerator() || GenerationHelper.DoesGeneratedRequireGenerator()
@@ -3578,7 +3576,7 @@ UE::Cook::EPollStatus UCookOnTheFlyServer::BeginCacheObjectsToMove(UE::Cook::FGe
 		}
 
 		Info.TakeOverCachedObjectsAndAddMoved(GenerationHelper, PackageData.GetCachedObjectsInOuter(), ObjectsToMove);
-		Info.SetSaveStateComplete(FCookGenerationInfo::ESaveState::CallObjectsToMove);
+		PackageData.SetSaveSubStateComplete(ESaveSubState::Generation_CallObjectsToMove);
 	}
 
 	EPollStatus Result = CallBeginCacheOnObjects(PackageData, Package, PackageData.GetCachedObjectsInOuter(),
@@ -3678,16 +3676,16 @@ UE::Cook::EPollStatus UCookOnTheFlyServer::BeginCachePostMove(UE::Cook::FGenerat
 		return EPollStatus::Error;
 	}
 
-	if (Info.GetSaveState() <= FCookGenerationInfo::ESaveState::CallGetPostMoveObjects)
+	if (PackageData.GetSaveSubState() <= ESaveSubState::Generation_CallGetPostMoveObjects)
 	{
 		bool bFoundNewObjects;
 		EPollStatus Result = Info.RefreshPackageObjects(GenerationHelper, Package, bFoundNewObjects,
-			FCookGenerationInfo::ESaveState::Last);
+			ESaveSubState::Last);
 		if (Result != EPollStatus::Success)
 		{
 			return Result;
 		}
-		Info.SetSaveStateComplete(FCookGenerationInfo::ESaveState::CallGetPostMoveObjects);
+		PackageData.SetSaveSubStateComplete(ESaveSubState::Generation_CallGetPostMoveObjects);
 	}
 
 	EPollStatus Result = CallBeginCacheOnObjects(PackageData, Package, PackageData.GetCachedObjectsInOuter(),
@@ -3795,7 +3793,7 @@ UE::Cook::EPollStatus UCookOnTheFlyServer::PrepareSave(UE::Cook::FPackageData& P
 	using namespace UE::Cook;
 
 	EPollStatus Result = EPollStatus::Incomplete;
-	if (PackageData.GetCookedPlatformDataComplete())
+	if (PackageData.GetSaveSubState() == ESaveSubState::ReadyForSave)
 	{
 		Result = EPollStatus::Success;
 	}
@@ -3844,9 +3842,9 @@ UE::Cook::EPollStatus UCookOnTheFlyServer::PrepareSaveInternal(UE::Cook::FPackag
 	check(PackageData.GetState() == EPackageState::SaveActive);
 	TRefCountPtr<FGenerationHelper> GenerationHelper;
 
-	if (!PackageData.GetCookedPlatformDataCalled())
+	if (PackageData.GetSaveSubState() < ESaveSubState::CheckForIsGenerated)
 	{
-		if (!PackageData.GetCookedPlatformDataStarted())
+		if (PackageData.GetSaveSubState() <= ESaveSubState::StartSave)
 		{
 			if (PackageData.GetNumPendingCookedPlatformData() > 0)
 			{
@@ -3870,43 +3868,77 @@ UE::Cook::EPollStatus UCookOnTheFlyServer::PrepareSaveInternal(UE::Cook::FPackag
 					return EPollStatus::Error;
 				}
 			}
-			PackageData.SetCookedPlatformDataStarted(true);
+			PackageData.SetSaveSubStateComplete(ESaveSubState::StartSave);
 		}
 
-		PackageData.CreateObjectCache();
-
-		// Note that we cache cooked data for all requested platforms, rather than only for the requested platforms that have not cooked yet.  This allows
-		// us to avoid the complexity of needing to cancel the Save and keep track of the old list of uncooked platforms whenever the cooked platforms change
-		// while PrepareSave is active.
-		// Currently this does not cause significant cost since saving new platforms with some platforms already saved is a rare operation.
+		if (PackageData.GetSaveSubState() <= ESaveSubState::FirstCookedPlatformData_CreateObjectCache)
+		{
+			PackageData.CreateObjectCache();
+			PackageData.SetSaveSubStateComplete(ESaveSubState::FirstCookedPlatformData_CreateObjectCache);
+		}
 
-		int32& CookedPlatformDataNextIndex = PackageData.GetCookedPlatformDataNextIndex();
-		if (CookedPlatformDataNextIndex < 0)
+		if (PackageData.GetSaveSubState() <= ESaveSubState::FirstCookedPlatformData_CallingBegin)
 		{
-			if (!BuildDefinitions->TryRemovePendingBuilds(PackageData.GetPackageName()))
+			// Note that we cache cooked data for all requested platforms, rather than only for the requested platforms that have not cooked yet.  This allows
+			// us to avoid the complexity of needing to cancel the Save and keep track of the old list of uncooked platforms whenever the cooked platforms change
+			// while PrepareSave is active.
+			// Currently this does not cause significant cost since saving new platforms with some platforms already saved is a rare operation.
+
+			int32& CookedPlatformDataNextIndex = PackageData.GetCookedPlatformDataNextIndex();
+			if (CookedPlatformDataNextIndex < 0)
 			{
-				// Builds are in progress; wait for them to complete
-				return EPollStatus::Incomplete;
+				if (!BuildDefinitions->TryRemovePendingBuilds(PackageData.GetPackageName()))
+				{
+					// Builds are in progress; wait for them to complete
+					return EPollStatus::Incomplete;
+				}
+				CookedPlatformDataNextIndex = 0;
 			}
-			CookedPlatformDataNextIndex = 0;
-		}
 
-		TArray<FCachedObjectInOuter>& CachedObjectsInOuter = PackageData.GetCachedObjectsInOuter();
-		EPollStatus Result = CallBeginCacheOnObjects(PackageData, Package, CachedObjectsInOuter,
-			CookedPlatformDataNextIndex, Timer);
-		if (Result != EPollStatus::Success)
-		{
-			return Result;
+			TArray<FCachedObjectInOuter>& CachedObjectsInOuter = PackageData.GetCachedObjectsInOuter();
+			EPollStatus Result = CallBeginCacheOnObjects(PackageData, Package, CachedObjectsInOuter,
+				CookedPlatformDataNextIndex, Timer);
+			if (Result != EPollStatus::Success)
+			{
+				return Result;
+			}
+
+			PackageData.SetSaveSubStateComplete(ESaveSubState::FirstCookedPlatformData_CallingBegin);
 		}
 
-		// Check for whether the Package has a Splitter and initialize its list if so
-		if (!PackageData.HasInitializedGeneratorSave())
+		if (PackageData.GetSaveSubState() <= ESaveSubState::FirstCookedPlatformData_CheckForGeneratorAfterWaitingForIsLoaded)
 		{
+			bool bCookedPlatformDataIsLoaded = PackageData.GetNumPendingCookedPlatformData() == 0;
+			bool bWaitingForIsLoaded = PackageData.GetSaveSubState() > ESaveSubState::FirstCookedPlatformData_CheckForGenerator;
+			if (bWaitingForIsLoaded && !bCookedPlatformDataIsLoaded)
+			{
+				return EPollStatus::Incomplete;
+			}
+
+			// Check for whether the Package has a Splitter and initialize its list if so
 			// The GenerationHelper might have already been created by a child generated package;
 			// or it might have been created and not initialized by iterative cook startup.
 			// If not created or initialized, try looking for it
-			GenerationHelper = PackageData.TryCreateValidGenerationHelper();
-			if (GenerationHelper)
+			bool bNeedWaitForIsLoaded = false;
+			GenerationHelper = PackageData.TryCreateValidGenerationHelper(bCookedPlatformDataIsLoaded, bNeedWaitForIsLoaded);
+			if (!GenerationHelper && bNeedWaitForIsLoaded)
+			{
+				// bNeedWaitForIsLoaded can only be set to true if we pass in !bCookedPlatformDataIsLoaded, and that can only happen
+				// if !bWaitingForIsLoaded, due to the early exit above.
+				check(!bWaitingForIsLoaded);
+				PackageData.SetSaveSubState(ESaveSubState::FirstCookedPlatformData_CheckForGeneratorAfterWaitingForIsLoaded);
+				return EPollStatus::Incomplete;
+			}
+			PackageData.SetSaveSubStateComplete(ESaveSubState::FirstCookedPlatformData_CheckForGeneratorAfterWaitingForIsLoaded);
+		}
+		else
+		{
+			GenerationHelper = PackageData.GetGenerationHelperIfValid();
+		}
+
+		if (GenerationHelper)
+		{
+			if (PackageData.GetSaveSubState() <= ESaveSubState::Generation_TryGenerateList)
 			{
 				// Keep it referenced even if we are only precaching, so we do not recreate it
 				GenerationHelper->SetKeepForGeneratorSave();
@@ -3931,80 +3963,109 @@ UE::Cook::EPollStatus UCookOnTheFlyServer::PrepareSaveInternal(UE::Cook::FPackag
 					{
 						return EPollStatus::Error;
 					}
-					// The earlier exit from SaveState should have reset the progress back to StartPopulate or earlier
-					check(GenerationHelper->GetOwnerInfo().GetSaveState() <= FCookGenerationInfo::ESaveState::StartPopulate);
 					GenerationHelper->StartOwnerSave();
+					PackageData.SetSaveSubStateComplete(ESaveSubState::Generation_TryGenerateList);
 				}
 			}
-			PackageData.SetInitializedGeneratorSave(true);
+
+			if (PackageData.GetSaveSubState() <= ESaveSubState::Generation_QueueGeneratedPackages)
+			{
+				EPollStatus Result = QueueGeneratedPackages(*GenerationHelper, PackageData);
+				if (Result != EPollStatus::Success)
+				{
+					return Result;
+				}
+				PackageData.SetSaveSubStateComplete(ESaveSubState::Generation_QueueGeneratedPackages);
+			}
 		}
 		else
 		{
-			GenerationHelper = PackageData.GetGenerationHelperIfValid();
-		}
-		if (GenerationHelper)
-		{
-			Result = QueueGeneratedPackages(*GenerationHelper, PackageData);
-			if (Result != EPollStatus::Success)
-			{
-				return Result;
-			}
+			PackageData.SetSaveSubState(ESaveSubState::CheckForIsGenerated);
 		}
-
-		PackageData.SetCookedPlatformDataCalled(true);
 	}
 	else
 	{
 		GenerationHelper = PackageData.GetGenerationHelperIfValid();
 	}
 
-	if (GenerationHelper)
+	if (PackageData.GetSaveSubState() < ESaveSubState::ReadyForSave)
 	{
-		EPollStatus Result = PrepareSaveGenerationPackage(*GenerationHelper, PackageData, Timer, bPrecaching);
-		if (Result != EPollStatus::Success)
+		if (GenerationHelper)
 		{
-			return Result;
+			EPollStatus Result = PrepareSaveGenerationPackage(*GenerationHelper, PackageData, Timer, bPrecaching);
+			if (Result != EPollStatus::Success)
+			{
+				return Result;
+			}
 		}
-	}
-	else if (PackageData.IsGenerated())
-	{
-		TRefCountPtr<FGenerationHelper> ParentGenerationHelper = PackageData.GetParentGenerationHelper();
-		if (!ParentGenerationHelper || !ParentGenerationHelper->IsValid())
+		else if (PackageData.IsGenerated())
 		{
-			UE_LOG(LogCook, Error, TEXT("Generated package %s %s ParentGenerator package %s and cannot be saved."),
-				(!ParentGenerationHelper ? TEXT("is missing its") : TEXT("has an invalid")),
-				*PackageData.GetPackageName().ToString(), *PackageData.GetParentGenerator().ToString());
-			return EPollStatus::Error;
-		}
+			TRefCountPtr<FGenerationHelper> ParentGenerationHelper = PackageData.GetParentGenerationHelper();
+			if (!ParentGenerationHelper || !ParentGenerationHelper->IsValid())
+			{
+				UE_LOG(LogCook, Error, TEXT("Generated package %s %s ParentGenerator package %s and cannot be saved."),
+					(!ParentGenerationHelper ? TEXT("is missing its") : TEXT("has an invalid")),
+					*PackageData.GetPackageName().ToString(), *PackageData.GetParentGenerator().ToString());
+				return EPollStatus::Error;
+			}
 
-		EPollStatus Result = PrepareSaveGenerationPackage(*ParentGenerationHelper, PackageData, Timer, bPrecaching);
-		if (Result != EPollStatus::Success)
-		{
-			return Result;
-		}
-	}
-	else
-	{
-		if (PackageData.GetNumPendingCookedPlatformData() > 0)
-		{
-			return EPollStatus::Incomplete;
-		}
-		bool bFoundNewObjects;
-		EPollStatus Result = PackageData.RefreshObjectCache(bFoundNewObjects);
-		if (Result != EPollStatus::Success)
-		{
-			return Result;
+			EPollStatus Result = PrepareSaveGenerationPackage(*ParentGenerationHelper, PackageData, Timer, bPrecaching);
+			if (Result != EPollStatus::Success)
+			{
+				return Result;
+			}
 		}
-		if (bFoundNewObjects)
+		else
 		{
-			// Call this function recursively to reexecute CallBeginCacheOnObjects.
-			// Note that RefreshObjectCache checked for too many recursive calls and ErrorExited if so.
-			return PrepareSaveInternal(PackageData, Timer, bPrecaching);
+			if (PackageData.GetSaveSubState() <= ESaveSubState::CheckForIsGenerated)
+			{
+				// Skip over the LastCookedPlatformData_CallingBegin state; we only need to enter that
+				// state if RefreshObjectCache finds some new objects
+				PackageData.SetSaveSubState(ESaveSubState::LastCookedPlatformData_WaitingForIsLoaded);
+			}
+
+			if (PackageData.GetSaveSubState() <= ESaveSubState::LastCookedPlatformData_CallingBegin)
+			{
+				int32& CookedPlatformDataNextIndex = PackageData.GetCookedPlatformDataNextIndex();
+				TArray<FCachedObjectInOuter>& CachedObjectsInOuter = PackageData.GetCachedObjectsInOuter();
+				EPollStatus Result = CallBeginCacheOnObjects(PackageData, Package, CachedObjectsInOuter,
+					CookedPlatformDataNextIndex, Timer);
+				if (Result != EPollStatus::Success)
+				{
+					return Result;
+				}
+				PackageData.SetSaveSubStateComplete(ESaveSubState::LastCookedPlatformData_CallingBegin);
+			}
+
+			if (PackageData.GetSaveSubState() <= ESaveSubState::LastCookedPlatformData_WaitingForIsLoaded)
+			{
+				if (PackageData.GetNumPendingCookedPlatformData() > 0)
+				{
+					return EPollStatus::Incomplete;
+				}
+				bool bFoundNewObjects;
+				EPollStatus Result = PackageData.RefreshObjectCache(bFoundNewObjects);
+				if (Result != EPollStatus::Success)
+				{
+					return Result;
+				}
+				if (bFoundNewObjects)
+				{
+					PackageData.SetSaveSubState(ESaveSubState::LastCookedPlatformData_CallingBegin);
+					// Call this function recursively to immediately reexecute CallBeginCacheOnObjects.
+					// Note that RefreshObjectCache checked for too many recursive calls and ErrorExited if so.
+					return PrepareSaveInternal(PackageData, Timer, bPrecaching);
+				}
+				else
+				{
+					PackageData.SetSaveSubStateComplete(ESaveSubState::LastCookedPlatformData_WaitingForIsLoaded);
+				}
+			}
 		}
 	}
 
+	check(PackageData.GetSaveSubState() == ESaveSubState::ReadyForSave);
 	check(PackageData.GetNumPendingCookedPlatformData() == 0);
-	PackageData.SetCookedPlatformDataComplete(true);
 	return EPollStatus::Success;
 }
 
@@ -4101,7 +4162,7 @@ void UCookOnTheFlyServer::ReleaseCookedPlatformData(UE::Cook::FPackageData& Pack
 {
 	using namespace UE::Cook;
 
-	if (!PackageData.GetCookedPlatformDataStarted())
+	if (PackageData.GetSaveSubState() == ESaveSubState::StartSave)
 	{
 		PackageData.CheckCookedPlatformDataEmpty();
 		return;
@@ -4195,10 +4256,6 @@ void UCookOnTheFlyServer::ReleaseCookedPlatformData(UE::Cook::FPackageData& Pack
 	if (GenerationInfo)
 	{
 		GenerationHelper->ResetSaveState(*GenerationInfo, PackageData.GetPackage(), ReleaseSaveReason, NewState);
-		if (GenerationInfo->IsGenerator())
-		{
-			PackageData.SetInitializedGeneratorSave(false);
-		}
 	}
 
 	PackageData.ClearCookedPlatformData();
@@ -4216,6 +4273,8 @@ void UCookOnTheFlyServer::ReleaseCookedPlatformData(UE::Cook::FPackageData& Pack
 			}
 		}
 	}
+
+	PackageData.SetSaveSubState(ESaveSubState::StartSave);
 }
 
 void UCookOnTheFlyServer::TickCancels()
@@ -11961,9 +12020,7 @@ void UCookOnTheFlyServer::GetPackagesToRetract(int32 NumToRetract, TArray<FName>
 		}
 		if (FGenerationHelper* GenerationHelper = PackageData->GetGenerationHelper())
 		{
-			if (GenerationHelper->IsInitialized() &&
-				GenerationHelper->GetOwnerInfo().GetSaveState()
-					>= FCookGenerationInfo::ESaveState::QueueGeneratedPackages)
+			if (PackageData->GetSaveSubState() >= ESaveSubState::Generation_QueueGeneratedPackages)
 			{
 				if (GenerationHelper->DoesGeneratedRequireGenerator()
 					>= ICookPackageSplitter::EGeneratedRequiresGenerator::Save
@@ -12028,7 +12085,7 @@ void UCookOnTheFlyServer::GetPackagesToRetract(int32 NumToRetract, TArray<FName>
 	// Send back all packages that have not started saving before sending back any that have
 	for (FPackageData* PackageData : PackageDatas->GetSaveQueue())
 	{
-		if (!PackageData->GetCookedPlatformDataStarted())
+		if (PackageData->GetSaveSubState() <= ESaveSubState::StartSave)
 		{
 			if (AddPackageIfPossibleAndReportDone(PackageData))
 			{
@@ -12038,7 +12095,7 @@ void UCookOnTheFlyServer::GetPackagesToRetract(int32 NumToRetract, TArray<FName>
 	}
 	for (FPackageData* PackageData : PackageDatas->GetSaveQueue())
 	{
-		if (PackageData->GetCookedPlatformDataStarted())
+		if (PackageData->GetSaveSubState() > ESaveSubState::StartSave)
 		{
 			if (AddPackageIfPossibleAndReportDone(PackageData))
 			{
diff --git a/Engine/Source/Editor/UnrealEd/Private/Cooker/CookGenerationHelper.cpp b/Engine/Source/Editor/UnrealEd/Private/Cooker/CookGenerationHelper.cpp
--- a/Engine/Source/Editor/UnrealEd/Private/Cooker/CookGenerationHelper.cpp
+++ b/Engine/Source/Editor/UnrealEd/Private/Cooker/CookGenerationHelper.cpp
@@ -63,11 +63,21 @@ void FGenerationHelper::Initialize()
 	UObject* LocalSplitDataObject;
 	UE::Cook::Private::FRegisteredCookPackageSplitter* LocalRegisteredSplitterType = nullptr;
 	TUniquePtr<ICookPackageSplitter> LocalSplitter;
+
+	// When asked to Initialize for cases outside of the generator's Save state, ignore the
+	// RequiresCachedCookedPlatformDataBeforeSplit requirement before calling ShouldSplit.
+	// MPCOOKTODO: This breaks a contract and we should fix it. We have worked around it for now by requiring
+	// that RequiresCachedCookedPlatformDataBeforeSplit forces EGeneratedRequiresGenerator::Save, so that Initialize
+	// is not called outside of the generator's Save state.
+	constexpr bool bCookedPlatformDataIsLoaded = true;
+	bool bNeedWaitForIsLoaded;
+
 	SearchForRegisteredSplitDataObject(COTFS, OwnerPackageName, LocalOwnerPackage,
 		TOptional<TConstArrayView<FCachedObjectInOuter>>(), LocalSplitDataObject, LocalRegisteredSplitterType,
-		LocalSplitter);
+		LocalSplitter, bCookedPlatformDataIsLoaded, bNeedWaitForIsLoaded);
 	if (!LocalSplitDataObject || !LocalSplitter)
 	{
+		check(!bNeedWaitForIsLoaded);
 		InitializeStatus = EInitializeStatus::Invalid;
 		return;
 	}
@@ -99,6 +109,17 @@ void FGenerationHelper::Initialize(const UObject* InSplitDataObject,
 		CookPackageSplitterInstance->RequiresGeneratorPackageDestructBeforeResplit();
 	DoesGeneratedRequireGeneratorValue =
 		CookPackageSplitterInstance->DoesGeneratedRequireGenerator();
+
+	// Workaround for our current inability to handle RequiresCachedCookedPlatformDataBeforeSplit when
+	// calling Initialize or TryCreateValidParentGenerationHelper. We force EGeneratedRequiresGenerator::Save
+	// in the RequiresCachedCookedPlatformDataBeforeSplit case, so that the generator is always initialized before
+	// we call either of those functions. See the comments in TryCreateValidParentGenerationHelper and
+	// FGenerationHelper::Initialize(void)
+	if (RegisteredSplitterType->RequiresCachedCookedPlatformDataBeforeSplit()
+		&& DoesGeneratedRequireGeneratorValue < ICookPackageSplitter::EGeneratedRequiresGenerator::Save)
+	{
+		DoesGeneratedRequireGeneratorValue = ICookPackageSplitter::EGeneratedRequiresGenerator::Save;
+	}
 }
 
 void FGenerationHelper::InitializeAsInvalid()
@@ -394,8 +415,10 @@ void FGenerationHelper::SearchForRegisteredSplitDataObject(UCookOnTheFlyServer&
 	FName PackageName, UPackage* Package,
 	TOptional<TConstArrayView<FCachedObjectInOuter>> CachedObjectsInOuter,
 	UObject*& OutSplitDataObject, UE::Cook::Private::FRegisteredCookPackageSplitter*& OutRegisteredSplitter,
-	TUniquePtr<ICookPackageSplitter>& OutSplitterInstance)
+	TUniquePtr<ICookPackageSplitter>& OutSplitterInstance, bool bCookedPlatformDataIsLoaded,
+	bool& bOutNeedWaitForIsLoaded)
 {
+	bOutNeedWaitForIsLoaded = false;
 	OutSplitDataObject = nullptr;
 	OutRegisteredSplitter = nullptr;
 	OutSplitterInstance = nullptr;
@@ -405,13 +428,23 @@ void FGenerationHelper::SearchForRegisteredSplitDataObject(UCookOnTheFlyServer&
 	Private::FRegisteredCookPackageSplitter* SplitterType = nullptr;
 	TArray<Private::FRegisteredCookPackageSplitter*> FoundRegisteredSplitters;
 	auto TryLookForSplitterOfObject =
-		[&COTFS, PackageName, &FoundRegisteredSplitters, &SplitterType, &LocalSplitDataObject](UObject* Obj)
+		[&COTFS, PackageName, &FoundRegisteredSplitters, &SplitterType, &LocalSplitDataObject,
+		bCookedPlatformDataIsLoaded, &bOutNeedWaitForIsLoaded](UObject* Obj)
 		{
 			FoundRegisteredSplitters.Reset();
 			COTFS.RegisteredSplitDataClasses.MultiFind(Obj->GetClass(), FoundRegisteredSplitters);
 
 			for (Private::FRegisteredCookPackageSplitter* SplitterForObject : FoundRegisteredSplitters)
 			{
+				if (!SplitterForObject)
+				{
+					continue;
+				}
+				if (SplitterForObject->RequiresCachedCookedPlatformDataBeforeSplit() && !bCookedPlatformDataIsLoaded)
+				{
+					bOutNeedWaitForIsLoaded = true;
+					return false;
+				}
 				if (SplitterForObject && SplitterForObject->ShouldSplitPackage(Obj))
 				{
 					if (!Obj->HasAnyFlags(RF_Public))
@@ -449,7 +482,7 @@ void FGenerationHelper::SearchForRegisteredSplitDataObject(UCookOnTheFlyServer&
 			}
 			if (!TryLookForSplitterOfObject(Obj))
 			{
-				return; // error condition, exit the entire search function
+				return; // unable to complete the search, exit the entire search function
 			}
 		}
 	}
@@ -462,7 +495,7 @@ void FGenerationHelper::SearchForRegisteredSplitDataObject(UCookOnTheFlyServer&
 		{
 			if (!TryLookForSplitterOfObject(Obj))
 			{
-				return; // error condition, exit the entire search function
+				return; // unable to complete the search, exit the entire search function
 			}
 		}
 	}
@@ -1198,7 +1231,7 @@ void FGenerationHelper::ResetSaveState(FCookGenerationInfo& Info, UPackage* Pack
 	// We release references to *this in this function so keep a local reference to avoid deletion during the function.
 	TRefCountPtr<FGenerationHelper> LocalRefCount = this;
 
-	if (Info.GetSaveState() > FCookGenerationInfo::ESaveState::CallPopulate)
+	if (Info.PackageData->GetSaveSubState() > ESaveSubState::Generation_CallPopulate)
 	{
 		UObject* SplitObject = GetWeakSplitDataObject();
 		UPackage* LocalOwnerPackage = Info.IsGenerator() ? Package : GetOwnerPackage();
@@ -1255,8 +1288,6 @@ void FGenerationHelper::ResetSaveState(FCookGenerationInfo& Info, UPackage* Pack
 			Info.PackageData->SetParentGenerationHelper(nullptr);
 		}
 	}
-	Info.SetSaveState(Info.IsGenerator() ? FCookGenerationInfo::ESaveState::StartSave
-		: FCookGenerationInfo::ESaveState::StartPopulate);
 
 	if (Info.HasTakenOverCachedCookedPlatformData())
 	{
@@ -1303,7 +1334,7 @@ bool FGenerationHelper::ShouldRetractionStallRatherThanDemote(FPackageData& Pack
 	{
 		if (PackageData.IsInStateProperty(EPackageStateProperty::Saving))
 		{
-			if (Info->GetSaveState() > FCookGenerationInfo::ESaveState::StartPopulate)
+			if (PackageData.GetSaveSubState() > ESaveSubState::Generation_PreMoveCookedPlatformData_WaitingForIsLoaded)
 			{
 				return true;
 			}
@@ -1402,7 +1433,7 @@ void FGenerationHelper::PreGarbageCollect(const TRefCountPtr<FGenerationHelper>&
 	{
 		// If we don't have a contract to keep the packagedata referenced during GC, don't report
 		// anything to garbage collection, and demote the package if it has progressed too far
-		if (Info.GetSaveState() > FCookGenerationInfo::ESaveState::CallPopulate)
+		if (Info.PackageData->GetSaveSubState() > ESaveSubState::Generation_CallPopulate)
 		{
 			bOutShouldDemote = true;
 		}
@@ -1449,7 +1480,7 @@ void FGenerationHelper::PreGarbageCollect(const TRefCountPtr<FGenerationHelper>&
 
 	// Keep the generator and generated package referenced if we've passed the call to populate, or if we are keeping
 	// any other objects referenced
-	if (bKeepingAnyObjects || Info.GetSaveState() > FCookGenerationInfo::ESaveState::CallPopulate)
+	if (bKeepingAnyObjects || Info.PackageData->GetSaveSubState() > ESaveSubState::Generation_CallPopulate)
 	{
 		bNeedsGeneratorPackage = true;
 		if (&Info != &OwnerInfo)
@@ -1823,7 +1854,6 @@ void FGenerationHelper::UpdateSaveAfterGarbageCollect(const FPackageData& Packag
 
 FCookGenerationInfo::FCookGenerationInfo(FPackageData& InPackageData, bool bInGenerator)
 	: PackageData(&InPackageData)
-	, GeneratorSaveState(bInGenerator ? ESaveState::StartSave : ESaveState::StartPopulate)
 	, bCreateAsMap(false), bHasCreatedPackage(false), bHasSaved(false), bTakenOverCachedCookedPlatformData(false)
 	, bIssuedUndeclaredMovedObjectsWarning(false), bGenerator(bInGenerator), bHasCalledPopulate(false)
 	, bIterativelySkipped(false)
@@ -1833,10 +1863,11 @@ FCookGenerationInfo::FCookGenerationInfo(FPackageData& InPackageData, bool bInGe
 void FCookGenerationInfo::Uninitialize()
 {
 	// Check that we have left the save state first, since other assertions assume we have left the save state
-	checkf(GeneratorSaveState == (bGenerator ? ESaveState::StartSave : ESaveState::StartPopulate),
+	checkf(PackageData->GetSaveSubState() == ESaveSubState::StartSave,
 		TEXT("Cooker bug: Expected FCookGenerationInfo::Uninitialize to not be called for a package still in the save state, ")
-		TEXT("but %s package %s has SaveState %d."),
-		bGenerator ? TEXT("generator") : TEXT("generated"), *GetPackageName(), static_cast<int32>(GeneratorSaveState));
+		TEXT("but %s package %s has SaveSubState %s."),
+		bGenerator ? TEXT("generator") : TEXT("generated"), *GetPackageName(),
+		LexToString(PackageData->GetSaveSubState()));
 
 	PackageHash.Reset();
 	RelativePath.Empty();
@@ -1857,15 +1888,6 @@ void FCookGenerationInfo::Uninitialize()
 	// Keep bIterativelySkipped; it is allowed in the uninitialized state
 }
 
-void FCookGenerationInfo::SetSaveStateComplete(ESaveState CompletedState)
-{
-	GeneratorSaveState = CompletedState;
-	if (GeneratorSaveState < ESaveState::Last)
-	{
-		GeneratorSaveState = static_cast<ESaveState>(static_cast<uint8>(GeneratorSaveState) + 1);
-	}
-}
-
 void FCachedObjectInOuterGeneratorInfo::Initialize(UObject* Object)
 {
 	if (Object)
@@ -1938,7 +1960,7 @@ void FCookGenerationInfo::TakeOverCachedObjectsAndAddMoved(FGenerationHelper& Ge
 }
 
 EPollStatus FCookGenerationInfo::RefreshPackageObjects(FGenerationHelper& GenerationHelper, UPackage* Package,
-	bool& bOutFoundNewObjects, ESaveState DemotionState)
+	bool& bOutFoundNewObjects, ESaveSubState DemotionState)
 {
 	bOutFoundNewObjects = false;
 	TArray<UObject*> CurrentObjectsInOuter;
@@ -1962,9 +1984,9 @@ EPollStatus FCookGenerationInfo::RefreshPackageObjects(FGenerationHelper& Genera
 	}
 	bOutFoundNewObjects = FirstNewObject != nullptr;
 
-	if (FirstNewObject != nullptr && DemotionState != ESaveState::Last)
+	if (FirstNewObject != nullptr && DemotionState != ESaveSubState::Last)
 	{
-		SetSaveState(DemotionState);
+		PackageData->SetSaveSubState(DemotionState);
 		if (++PackageData->GetNumRetriesBeginCacheOnObjects() > FPackageData::GetMaxNumRetriesBeginCacheOnObjects())
 		{
 			UE_LOG(LogCook, Error,
diff --git a/Engine/Source/Editor/UnrealEd/Private/Cooker/CookGenerationHelper.h b/Engine/Source/Editor/UnrealEd/Private/Cooker/CookGenerationHelper.h
--- a/Engine/Source/Editor/UnrealEd/Private/Cooker/CookGenerationHelper.h
+++ b/Engine/Source/Editor/UnrealEd/Private/Cooker/CookGenerationHelper.h
@@ -61,33 +61,8 @@ public:
 struct FCookGenerationInfo
 {
 public:
-	/** State variable for reentrant SplitPackage calls */
-	enum class ESaveState : uint8
-	{
-		StartSave = 0,
-
-		QueueGeneratedPackages = StartSave,
-
-		StartPopulate,
-		FinishCachePreMove = StartPopulate,
-		CallObjectsToMove,
-		BeginCacheObjectsToMove,
-		FinishCacheObjectsToMove,
-		CallPopulate,
-		CallGetPostMoveObjects,
-		BeginCachePostMove,
-		FinishCachePostMove,
-
-		ReadyForSave,
-		Last = ReadyForSave,
-	};
-
 	FCookGenerationInfo(FPackageData& InPackageData, bool bInGenerator);
 
-	ESaveState GetSaveState() const;
-	void SetSaveState(ESaveState InValue);
-	void SetSaveStateComplete(ESaveState CompletedState);
-
 	bool IsCreateAsMap() const;
 	void SetIsCreateAsMap(bool bValue);
 	bool HasCreatedPackage() const;
@@ -114,12 +89,12 @@ public:
 	/**
 	 * Fetch all the objects currently in the package and add them to the list of objects that need
 	 * BeginCacheForCookedPlatformData.
-	 * Reports whether new objects were found. If DemotionState is not ESaveState::Last,
-	 * will SetState back to DemotionState
+	 * Reports whether new objects were found. If DemotionState is not ESaveSubState::Last,
+	 * will SetSaveSubState back to DemotionState
 	 * if new objects were found, and will error exit if this demotion has happened too many times.
 	 */
 	EPollStatus RefreshPackageObjects(FGenerationHelper& GenerationHelper, UPackage* Package,
-		bool& bOutFoundNewObjects, ESaveState DemotionState);
+		bool& bOutFoundNewObjects, ESaveSubState DemotionState);
 
 	void AddKeepReferencedPackages(FGenerationHelper& GenerationHelper, TArray<UPackage*>& InKeepReferencedPackages);
 
@@ -159,7 +134,6 @@ public:
 	TMap<UObject*, FCachedObjectInOuterGeneratorInfo> CachedObjectsInOuterInfo;
 	FWorkerId SavedOnWorker = FWorkerId::Invalid();
 private:
-	ESaveState GeneratorSaveState = ESaveState::StartSave;
 	bool bCreateAsMap : 1;
 	bool bHasCreatedPackage : 1;
 	bool bHasSaved : 1;
@@ -437,7 +411,8 @@ public:
 	static void SearchForRegisteredSplitDataObject(UCookOnTheFlyServer& COTFS, FName PackageName, UPackage* Package,
 		TOptional<TConstArrayView<FCachedObjectInOuter>> CachedObjectsInOuter, UObject*& OutSplitDataObject,
 		UE::Cook::Private::FRegisteredCookPackageSplitter*& OutRegisteredSplitterType,
-		TUniquePtr<ICookPackageSplitter>& OutSplitterInstance);
+		TUniquePtr<ICookPackageSplitter>& OutSplitterInstance, bool bCookedPlatformDataIsLoaded,
+		bool& bOutNeedWaitForIsLoaded);
 	/** Helper function for Initialize and for TryCreateValidGenerationHelper. */
 	static UPackage* FindOrLoadPackage(UCookOnTheFlyServer& COTFS, FPackageData& OwnerPackageData);
 
@@ -510,15 +485,6 @@ private:
 // Inline implementations
 ///////////////////////////////////////////////////////
 
-inline FCookGenerationInfo::ESaveState FCookGenerationInfo::GetSaveState() const
-{
-	return GeneratorSaveState;
-}
-
-inline void FCookGenerationInfo::SetSaveState(ESaveState InValue)
-{
-	GeneratorSaveState = InValue;
-}
 
 inline bool FCookGenerationInfo::IsCreateAsMap() const
 {
diff --git a/Engine/Source/Editor/UnrealEd/Private/Cooker/CookPackageData.cpp b/Engine/Source/Editor/UnrealEd/Private/Cooker/CookPackageData.cpp
--- a/Engine/Source/Editor/UnrealEd/Private/Cooker/CookPackageData.cpp
+++ b/Engine/Source/Editor/UnrealEd/Private/Cooker/CookPackageData.cpp
@@ -96,13 +96,14 @@ FPackageData::FPackageData(FPackageDatas& PackageDatas, const FName& InPackageNa
 	, Instigator(EInstigator::NotYetRequested), bIsUrgent(0), bIsCookLast(0)
 	, bIsVisited(0), bIsPreloadAttempted(0)
 	, bIsPreloaded(0), bHasSaveCache(0), bPrepareSaveFailed(0), bPrepareSaveRequiresGC(0)
-	, bCookedPlatformDataStarted(0), bCookedPlatformDataCalled(0), bCookedPlatformDataComplete(0)
 	, MonitorCookResult((uint8)ECookResult::NotAttempted)
-	, bInitializedGeneratorSave(0), bCompletedGeneration(0), bGenerated(0), bKeepReferencedDuringGC(0)
+	, bCompletedGeneration(0), bGenerated(0), bKeepReferencedDuringGC(0)
 	, bWasCookedThisSession(0)
 	, DoesGeneratedRequireGeneratorValue(static_cast<uint32>(ICookPackageSplitter::EGeneratedRequiresGenerator::None))
 {
 	SetState(EPackageState::Idle);
+	SetSaveSubState(ESaveSubState::StartSave);
+
 	SendToState(EPackageState::Idle, ESendFlags::QueueAdd, EStateChangeReason::Discovered);
 }
 
@@ -1444,8 +1445,7 @@ EPollStatus FPackageData::RefreshObjectCache(bool& bOutFoundNewObjects)
 			CachedObjectsInOuter.Emplace(MoveTemp(ObjectWeakPointer));
 		}
 		// GetCookedPlatformDataNextIndex is already where it should be, pointing at the first of the objects we have
-		// added Change our state back so we know we need to CallBeginCacheOnObjects again 
-		SetCookedPlatformDataCalled(false);
+		// added. Caller is respnsible for changing state back to calling BeginCacheForCookedPlatformData.
 
 		if (++GetNumRetriesBeginCacheOnObjects() > FPackageData::GetMaxNumRetriesBeginCacheOnObjects())
 		{
@@ -1503,21 +1503,31 @@ int32 FPackageData::GetMaxNumRetriesBeginCacheOnObjects()
 	return 10;
 }
 
-void FPackageData::CheckCookedPlatformDataEmpty() const
+void FPackageData::SetSaveSubState(ESaveSubState Value)
 {
-	check(GetCookedPlatformDataNextIndex() <= 0);
-	check(!GetCookedPlatformDataStarted());
-	check(!GetCookedPlatformDataCalled());
-	check(!GetCookedPlatformDataComplete());
-	if (GetGenerationHelper())
+	if (Value != ESaveSubState::StartSave && !IsInStateProperty(EPackageStateProperty::Saving))
 	{
-		check(GetGenerationHelper()->GetOwnerInfo().GetSaveState() <= FCookGenerationInfo::ESaveState::StartPopulate);
+		UE_LOG(LogCook, Error, TEXT("SetSaveSubState(%s) called from invalid PackageState %s. The call will be ignored"),
+			LexToString(Value), LexToString(GetState()));
+		FDebug::DumpStackTraceToLog(ELogVerbosity::Warning);
+		return;
 	}
-	if (ParentGenerationHelper)
+	SaveSubState = static_cast<uint32>(Value);
+}
+
+void FPackageData::SetSaveSubStateComplete(ESaveSubState Value)
+{
+	if (Value < ESaveSubState::Last)
 	{
-		FCookGenerationInfo* Info = ParentGenerationHelper->FindInfo(*this);
-		check(!Info || Info->GetSaveState() <= FCookGenerationInfo::ESaveState::StartPopulate);
+		Value = static_cast<ESaveSubState>(static_cast<uint32>(Value) + 1);
 	}
+	SetSaveSubState(Value);
+}
+
+void FPackageData::CheckCookedPlatformDataEmpty() const
+{
+	check(GetCookedPlatformDataNextIndex() <= 0);
+	check(GetSaveSubState() <= ESaveSubState::StartSave);
 }
 
 void FPackageData::ClearCookedPlatformData()
@@ -1525,9 +1535,7 @@ void FPackageData::ClearCookedPlatformData()
 	CookedPlatformDataNextIndex = -1;
 	NumRetriesBeginCacheOnObject = 0;
 	// Note that GetNumPendingCookedPlatformData is not cleared; it persists across Saves and CookSessions
-	SetCookedPlatformDataStarted(false);
-	SetCookedPlatformDataCalled(false);
-	SetCookedPlatformDataComplete(false);
+	// Caller is responsible for calling SetSaveSubState(ESaveSubState::StartSave);
 }
 
 void FPackageData::OnRemoveSessionPlatform(const ITargetPlatform* Platform)
@@ -1583,7 +1591,10 @@ void FPackageData::UpdateSaveAfterGarbageCollect(bool& bOutDemote)
 
 	// Reexecute PrepareSave if we already completed it; we need to refresh our CachedObjectsInOuter list
 	// and call BeginCacheOnCookedPlatformData on any new objects.
-	SetCookedPlatformDataComplete(false);
+	if (GetSaveSubState() >= ESaveSubState::LastCookedPlatformData_WaitingForIsLoaded)
+	{
+		SetSaveSubState(ESaveSubState::LastCookedPlatformData_WaitingForIsLoaded);
+	}
 
 	if (GetPackage() == nullptr || !GetPackage()->IsFullyLoaded())
 	{
@@ -1685,7 +1696,15 @@ TRefCountPtr<FGenerationHelper> FPackageData::TryCreateValidParentGenerationHelp
 		return nullptr;
 	}
 
-	ParentGenerationHelper = OwnerPackageData->TryCreateValidGenerationHelper();
+	// MPCOOKTODO: We need to support calling BeginCacheForCookedPlatformData/IsCachedCookedPlatformData
+	// on all objects in the generator package if they have not already been called, if 
+	// RequiresCachedCookedPlatformDataBeforeSplit. For now we workaround our inability to do this
+	// by forcing EGeneratedRequiresGenerator::Save.
+	constexpr bool bCookedPlatformDataIsLoaded = true;
+	bool bNeedWaitForIsLoaded;
+	ParentGenerationHelper = OwnerPackageData->TryCreateValidGenerationHelper(bCookedPlatformDataIsLoaded,
+		bNeedWaitForIsLoaded);
+	check(ParentGenerationHelper.IsValid() || !bNeedWaitForIsLoaded);
 
 	return ParentGenerationHelper;
 }
@@ -1701,8 +1720,11 @@ TRefCountPtr<FGenerationHelper> FPackageData::CreateUninitializedGenerationHelpe
 	return Result;
 }
 
-TRefCountPtr<UE::Cook::FGenerationHelper> FPackageData::TryCreateValidGenerationHelper()
+TRefCountPtr<UE::Cook::FGenerationHelper> FPackageData::TryCreateValidGenerationHelper(
+	bool bCookedPlatformDataIsLoaded, bool& bOutNeedWaitForIsLoaded)
 {
+	bOutNeedWaitForIsLoaded = false;
+
 	if (GenerationHelper && GenerationHelper->IsInitialized())
 	{
 		if (!GenerationHelper->IsValid())
@@ -1738,7 +1760,7 @@ TRefCountPtr<UE::Cook::FGenerationHelper> FPackageData::TryCreateValidGeneration
 		}
 		FGenerationHelper::SearchForRegisteredSplitDataObject(COTFS, GetPackageName(),
 			LocalPackage, LocalCachedObjectsInOuter, SplitDataObject, RegisteredSplitterType,
-			CookPackageSplitterInstance);
+			CookPackageSplitterInstance, bCookedPlatformDataIsLoaded, bOutNeedWaitForIsLoaded);
 	}
 	TRefCountPtr<UE::Cook::FGenerationHelper> Result = GenerationHelper;
 	if (!SplitDataObject || !CookPackageSplitterInstance)
diff --git a/Engine/Source/Editor/UnrealEd/Private/Cooker/CookPackageData.h b/Engine/Source/Editor/UnrealEd/Private/Cooker/CookPackageData.h
--- a/Engine/Source/Editor/UnrealEd/Private/Cooker/CookPackageData.h
+++ b/Engine/Source/Editor/UnrealEd/Private/Cooker/CookPackageData.h
@@ -558,24 +558,16 @@ public:
 	/** Get/Set the flag for whether CachedObjectsInOuter is populated. Always false except during save state. */
 	bool GetHasSaveCache() const { return static_cast<bool>(bHasSaveCache); }
 	void SetHasSaveCache(bool Value) { bHasSaveCache = Value != 0; }
-	/** Get/Set the flag for whether CachedObjectsInOuter iteration has started. Always false except during save. */
-	bool GetCookedPlatformDataStarted() const { return static_cast<bool>(bCookedPlatformDataStarted); }
-	void SetCookedPlatformDataStarted(bool Value) { bCookedPlatformDataStarted = (Value != 0); }
-	/**
-	 * Get/Set the flag for whether BeginCacheForCookedPlatformData has been called on every CachedObjectsInOuter.
-	 * Some objects might still be working asynchronously and have not yet returned true from
-	 * IsCachedCookedPlatformDataLoaded, in which case GetNumPendingCookedPlatformData will be non-zero.
-	 * Always false except during the save state.
-	 */
-	bool GetCookedPlatformDataCalled() const { return static_cast<bool>(bCookedPlatformDataCalled); }
-	void SetCookedPlatformDataCalled(bool bValue) { bCookedPlatformDataCalled = bValue != 0; }
+
 	/**
-	 * Get/Set the flag for whether BeginCacheForCookedPlatformData has been called and
-	 * IsCachedCookedPlatformDataLoaded has subsequently returned true for every object in GetCachedObjectsInOuter.
-	 * Always false except during the save state.
+	 * Get/Set the SubState within EPackageState::Save. Outside of EPackageStateProperty::Saving it is always
+	 * ESaveSubState::StartSave, and is ignored and logs an error if attempting to set to any other value.
 	 */
-	bool GetCookedPlatformDataComplete() const { return static_cast<bool>(bCookedPlatformDataComplete); }
-	void SetCookedPlatformDataComplete(bool bValue) { bCookedPlatformDataComplete = bValue != 0; }
+	ESaveSubState GetSaveSubState() const;
+	void SetSaveSubState(ESaveSubState Value);
+	/** Set the SaveSubState to the next state after Value. */
+	void SetSaveSubStateComplete(ESaveSubState Value);
+
 	/**
 	 * Check whether savestate contracts on the PackageData were invalidated by by e.g. garbage collection.
 	 * Request demotion if so unless we have a contract to keep it, in which case it is fixed up.
@@ -626,19 +618,17 @@ public:
 	/**
 	 * Return the GenerationHelper if it already exists, otherwise load the package and check all the objects
 	 * to see whether any of them have a registered CookPackageSplitter. If a split object exists, create the
-	 * GenerationHelper and return it. Otherwise return nullptr.
+	 * GenerationHelper and return it. Otherwise return nullptr. If the splitter requires it, also return nullptr
+	 * if the caller needs to call IsCachedCookedPlatformDataLoaded before creation.
 	 */
-	TRefCountPtr<FGenerationHelper> TryCreateValidGenerationHelper();
+	TRefCountPtr<FGenerationHelper> TryCreateValidGenerationHelper(bool bCookedPlatformDataIsLoaded,
+		bool& bOutNeedWaitForIsLoaded);
 	/** Helper function for ~FGenerationHelper: clear the pointer this->GenerationHelper. */
 	void OnGenerationHelperDestroyed(FGenerationHelper& InGenerationHelper);
 	/** Get whether the PackageData has done any necessary Generator steps and is ready for BeginCache calls. */
 	bool HasCompletedGeneration() const;
 	/** Set whether the PackageData has done any necessary Generator steps and is ready for BeginCache calls. */
 	void SetCompletedGeneration(bool Value);
-	/** Get whether the PackageData has completed the check for whether it is a Generator. */
-	bool HasInitializedGeneratorSave() const;
-	/** Set whether the PackageData has completed the check for whether it is a Generator. */
-	void SetInitializedGeneratorSave(bool Value);
 	/** Return whether the PackageData is a generated package created by a owning generator PackageData. */
 	bool IsGenerated() const;
 	/** Mark that the PackageData is a generated package created by a ParentGenerator. */
@@ -841,6 +831,7 @@ private:
 	FWorkerId WorkerAssignment = FWorkerId::Invalid();
 	FWorkerId WorkerAssignmentConstraint = FWorkerId::Invalid();
 	uint32 State : int32(EPackageState::BitCount);
+	uint32 SaveSubState : int32(ESaveSubState::BitCount);
 	uint32 bIsUrgent : 1;
 	uint32 bIsCookLast : 1;
 	uint32 bIsVisited : 1;
@@ -849,11 +840,7 @@ private:
 	uint32 bHasSaveCache : 1;
 	uint32 bPrepareSaveFailed : 1;
 	uint32 bPrepareSaveRequiresGC : 1;
-	uint32 bCookedPlatformDataStarted : 1;
-	uint32 bCookedPlatformDataCalled : 1;
-	uint32 bCookedPlatformDataComplete : 1;
 	uint32 MonitorCookResult : (int) ECookResult::NumBits;
-	uint32 bInitializedGeneratorSave : 1;
 	uint32 bCompletedGeneration : 1;
 	uint32 bGenerated : 1;
 	uint32 bKeepReferencedDuringGC : 1;
@@ -1524,6 +1511,11 @@ inline void FPackageData::GetReachablePlatforms(ArrayType& OutPlatforms) const
 	}
 }
 
+inline ESaveSubState FPackageData::GetSaveSubState() const
+{
+	return static_cast<ESaveSubState>(SaveSubState);
+}
+
 inline bool FPackageData::HasCompletedGeneration() const
 {
 	return static_cast<bool>(bCompletedGeneration);
@@ -1534,16 +1526,6 @@ inline void FPackageData::SetCompletedGeneration(bool Value)
 	bCompletedGeneration = Value != 0;
 }
 
-inline bool FPackageData::HasInitializedGeneratorSave() const
-{
-	return static_cast<bool>(bInitializedGeneratorSave);
-}
-
-inline void FPackageData::SetInitializedGeneratorSave(bool Value)
-{
-	bInitializedGeneratorSave = Value != 0;
-}
-
 inline bool FPackageData::IsGenerated() const
 {
 	return static_cast<bool>(bGenerated);
diff --git a/Engine/Source/Editor/UnrealEd/Private/Cooker/CookTypes.cpp b/Engine/Source/Editor/UnrealEd/Private/Cooker/CookTypes.cpp
--- a/Engine/Source/Editor/UnrealEd/Private/Cooker/CookTypes.cpp
+++ b/Engine/Source/Editor/UnrealEd/Private/Cooker/CookTypes.cpp
@@ -93,6 +93,32 @@ const TCHAR* LexToString(UE::Cook::EPackageState Reason)
 	}
 }
 
+const TCHAR* LexToString(UE::Cook::ESaveSubState State)
+{
+	static_assert(static_cast<int>(ESaveSubState::Count) == 17);
+	switch (State)
+	{
+	case ESaveSubState::StartSave: return TEXT("StartSave");
+	case ESaveSubState::FirstCookedPlatformData_CreateObjectCache: return TEXT("FirstCookedPlatformData_CreateObjectCache");
+	case ESaveSubState::FirstCookedPlatformData_CallingBegin: return TEXT("FirstCookedPlatformData_CallingBegin");
+	case ESaveSubState::FirstCookedPlatformData_CheckForGenerator: return TEXT("FirstCookedPlatformData_CheckForGenerator");
+	case ESaveSubState::FirstCookedPlatformData_CheckForGeneratorAfterWaitingForIsLoaded: return TEXT("FirstCookedPlatformData_CheckForGeneratorAfterWaitingForIsLoaded");
+	case ESaveSubState::Generation_TryGenerateList: return TEXT("Generation_TryGenerateList");
+	case ESaveSubState::Generation_QueueGeneratedPackages: return TEXT("Generation_QueueGeneratedPackages");
+	case ESaveSubState::CheckForIsGenerated: return TEXT("CheckForIsGenerated");
+	case ESaveSubState::Generation_PreMoveCookedPlatformData_WaitingForIsLoaded: return TEXT("Generation_PreMoveCookedPlatformData_WaitingForIsLoaded");
+	case ESaveSubState::Generation_CallObjectsToMove: return TEXT("Generation_CallObjectsToMove");
+	case ESaveSubState::Generation_BeginCacheObjectsToMove: return TEXT("Generation_BeginCacheObjectsToMove");
+	case ESaveSubState::Generation_FinishCacheObjectsToMove: return TEXT("Generation_FinishCacheObjectsToMove");
+	case ESaveSubState::Generation_CallPopulate: return TEXT("Generation_CallPopulate");
+	case ESaveSubState::Generation_CallGetPostMoveObjects: return TEXT("Generation_CallGetPostMoveObjects");
+	case ESaveSubState::LastCookedPlatformData_CallingBegin: return TEXT("LastCookedPlatformData_CallingBegin");
+	case ESaveSubState::LastCookedPlatformData_WaitingForIsLoaded: return TEXT("LastCookedPlatformData_WaitingForIsLoaded");
+	case ESaveSubState::ReadyForSave: return TEXT("ReadyForSave");
+	default: return TEXT("Invalid");
+	}
+}
+
 EStateChangeReason ConvertToStateChangeReason(ESuppressCookReason Reason)
 {
 	switch (Reason)
diff --git a/Engine/Source/Editor/UnrealEd/Private/Cooker/CookTypes.h b/Engine/Source/Editor/UnrealEd/Private/Cooker/CookTypes.h
--- a/Engine/Source/Editor/UnrealEd/Private/Cooker/CookTypes.h
+++ b/Engine/Source/Editor/UnrealEd/Private/Cooker/CookTypes.h
@@ -244,6 +244,37 @@ namespace UE::Cook
 	};
 	ENUM_CLASS_FLAGS(EPackageStateProperty);
 
+	/** SubState when in a Saving state. */
+	enum class ESaveSubState : uint8
+	{
+		StartSave = 0,
+		FirstCookedPlatformData_CreateObjectCache,
+		FirstCookedPlatformData_CallingBegin,
+		FirstCookedPlatformData_CheckForGenerator,
+		FirstCookedPlatformData_CheckForGeneratorAfterWaitingForIsLoaded,
+		Generation_TryGenerateList,
+		Generation_QueueGeneratedPackages,
+
+		CheckForIsGenerated,
+
+		Generation_PreMoveCookedPlatformData_WaitingForIsLoaded,
+		Generation_CallObjectsToMove,
+		Generation_BeginCacheObjectsToMove,
+		Generation_FinishCacheObjectsToMove,
+		Generation_CallPopulate,
+		Generation_CallGetPostMoveObjects,
+
+		LastCookedPlatformData_CallingBegin,
+		LastCookedPlatformData_WaitingForIsLoaded,
+
+		ReadyForSave,
+		Last = ReadyForSave,
+		Count = Last + 1,
+		/** Number of bits required to store a valid ESaveSubState */
+		BitCount = FPlatformMath::ConstExprCeilLogTwo(Count),
+	};
+	const TCHAR* LexToString(UE::Cook::ESaveSubState State);
+
 	/** Used as a helper to timeslice cooker functions. */
 	struct FCookerTimer
 	{
diff --git a/Engine/Source/Editor/UnrealEd/Public/CookPackageSplitter.h b/Engine/Source/Editor/UnrealEd/Public/CookPackageSplitter.h
--- a/Engine/Source/Editor/UnrealEd/Public/CookPackageSplitter.h
+++ b/Engine/Source/Editor/UnrealEd/Public/CookPackageSplitter.h
@@ -33,6 +33,12 @@ class ICookPackageSplitter
 public:
 	// Static API functions - these static functions are referenced by REGISTER_COOKPACKAGE_SPLITTER
 	// before creating an instance of the class.
+	/**
+	 * Return whether IsCachedCookedPlatformDataLoaded needs to return true for all UObjects in the
+	 * generator package before ShouldSplit or GetGenerateList can be called. If true this slows down
+	 * our ability to parallelize the cook of the generated packages.
+	 */
+	static bool RequiresCachedCookedPlatformDataBeforeSplit() { return false; }
 	/** Return whether the CookPackageSplitter subclass should handle the given SplitDataClass instance. */
 	static bool ShouldSplit(UObject* SplitData) { return false; }
 	/** Return DebugName for this SplitterClass in cook log messages. */
@@ -319,6 +325,7 @@ public:
 	UNREALED_API virtual ~FRegisteredCookPackageSplitter();
 
 	virtual UClass* GetSplitDataClass() const = 0;
+	virtual bool RequiresCachedCookedPlatformDataBeforeSplit() const = 0;
 	virtual bool ShouldSplitPackage(UObject* Object) const = 0;
 	virtual ICookPackageSplitter* CreateInstance(UObject* Object) const = 0;
 	virtual FString GetSplitterDebugName() const = 0;
@@ -338,16 +345,36 @@ private:
  *
  * Example usage:
  *
+ * // In header or cpp
  * class FMyCookPackageSplitter : public ICookPackageSplitter { ... }
+ * 
+ * // In cpp
  * REGISTER_COOKPACKAGE_SPLITTER(FMyCookPackageSplitter, UMySplitDataClass);
  */
 #define REGISTER_COOKPACKAGE_SPLITTER(SplitterClass, SplitDataClass) \
-class PREPROCESSOR_JOIN(PREPROCESSOR_JOIN(SplitterClass, SplitDataClass), _Register) : public UE::Cook::Private::FRegisteredCookPackageSplitter \
+class PREPROCESSOR_JOIN(PREPROCESSOR_JOIN(SplitterClass, SplitDataClass), _Register) \
+	: public UE::Cook::Private::FRegisteredCookPackageSplitter \
 { \
-	virtual UClass* GetSplitDataClass() const override { return SplitDataClass::StaticClass(); } \
-	virtual bool ShouldSplitPackage(UObject* Object) const override { return SplitterClass::ShouldSplit(Object); } \
-	virtual ICookPackageSplitter* CreateInstance(UObject* SplitData) const override { return new SplitterClass(); } \
-	virtual FString GetSplitterDebugName() const override { return SplitterClass::GetSplitterDebugName(); } \
+	virtual UClass* GetSplitDataClass() const override \
+	{ \
+		return SplitDataClass::StaticClass(); \
+	} \
+	virtual bool RequiresCachedCookedPlatformDataBeforeSplit() const override \
+	{ \
+		return SplitterClass::RequiresCachedCookedPlatformDataBeforeSplit(); \
+	} \
+	virtual bool ShouldSplitPackage(UObject* Object) const override \
+	{ \
+		return SplitterClass::ShouldSplit(Object); \
+	} \
+	virtual ICookPackageSplitter* CreateInstance(UObject* SplitData) const override \
+	{ \
+		return new SplitterClass(); \
+	} \
+	virtual FString GetSplitterDebugName() const override \
+	{ \
+		return SplitterClass::GetSplitterDebugName(); \
+	} \
 }; \
 namespace PREPROCESSOR_JOIN(SplitterClass, SplitDataClass) \
 { \
