Gather text from source, exclude more ThirdParty source

diff --git a/Engine/Source/Editor/UnrealEd/Private/Commandlets/GatherTextFromSourceCommandlet.cpp b/Engine/Source/Editor/UnrealEd/Private/Commandlets/GatherTextFromSourceCommandlet.cpp
--- a/Engine/Source/Editor/UnrealEd/Private/Commandlets/GatherTextFromSourceCommandlet.cpp
+++ b/Engine/Source/Editor/UnrealEd/Private/Commandlets/GatherTextFromSourceCommandlet.cpp
@@ -282,6 +282,7 @@ int32 UGatherTextFromSourceCommandlet::Main( const FString& Params )
 		}
 		TArray<FString> ExcludePathFiltersPrepass;
 		ExcludePathFiltersPrepass.Add(TEXT("../Engine/Source/ThirdParty/*"));
+		ExcludePathFiltersPrepass.Add(TEXT("../Engine/Restricted/NotForLicensees/Source/ThirdParty/*"));
 		for (FString& Path : ExcludePathFiltersPrepass)
 		{
 			ResolveLocalizationPath(Path);
