Core: DDC Commandlet now prints out the package that fails to find a soft reference, instead of just the soft reference name.

diff --git a/Engine/Source/Editor/UnrealEd/Private/Commandlets/DerivedDataCacheCommandlet.cpp b/Engine/Source/Editor/UnrealEd/Private/Commandlets/DerivedDataCacheCommandlet.cpp
--- a/Engine/Source/Editor/UnrealEd/Private/Commandlets/DerivedDataCacheCommandlet.cpp
+++ b/Engine/Source/Editor/UnrealEd/Private/Commandlets/DerivedDataCacheCommandlet.cpp
@@ -873,17 +873,18 @@ int32 UDerivedDataCacheCommandlet::Main( const FString& Params )
 
 						if (FPackageName::DoesPackageExist(SoftRefName.ToString(), &SoftRefFilename))
 						{
-							UE_LOG(LogDerivedDataCacheCommandlet, Log, TEXT("Queueing soft reference '%s' for later processing"), *SoftRefName.ToString());
+							UE_LOG(LogDerivedDataCacheCommandlet, Log, TEXT("Package '%s' queueing soft reference '%s' for later processing"), *NewProcessedPackage.ToString(), *SoftRefName.ToString());
 							PackagePaths.Push(TPair<FString, FName>(SoftRefFilename, SoftRefName));
 						}
 						else
 						{
-							UE_LOG(LogDerivedDataCacheCommandlet, Warning, TEXT("Failed to find soft reference '%s'"), *SoftRefName.ToString());
+							UE_LOG(LogDerivedDataCacheCommandlet, Warning, TEXT("Package '%s' failed to find soft reference '%s'"), *NewProcessedPackage.ToString(), *SoftRefName.ToString());
 						}
 					}
 					else
 					{
-						UE_LOG(LogDerivedDataCacheCommandlet, Verbose, TEXT("Skipping soft reference '%s': %s, %s "), 
+						UE_LOG(LogDerivedDataCacheCommandlet, Verbose, TEXT("Package '%s' skipping soft reference '%s': %s, %s "), 
+							*NewProcessedPackage.ToString(),
 							*SoftRefName.ToString(),									 
 							PackagesToProcess.Contains(SoftRefName) ? TEXT("ALREADY QUEUED") : TEXT("NOT QUEUED"),
 							ProcessedPackages.Contains(SoftRefName) ? TEXT("ALREADY PROCESSED") : TEXT("NOT PROCESSED")
