Cooker: Fix crash when creating CookerOpenOrder log in the case that old files from a previous cook are no longer present in current editor depot.
#rnx

diff --git a/Engine/Source/Editor/UnrealEd/Private/Commandlets/AssetRegistryGenerator.cpp b/Engine/Source/Editor/UnrealEd/Private/Commandlets/AssetRegistryGenerator.cpp
--- a/Engine/Source/Editor/UnrealEd/Private/Commandlets/AssetRegistryGenerator.cpp
+++ b/Engine/Source/Editor/UnrealEd/Private/Commandlets/AssetRegistryGenerator.cpp
@@ -1861,6 +1861,12 @@ class FPackageCookerOpenOrderVisitor : public IPlatformFile::FDirectoryVisitor
 	const FString& PlatformSandboxPath;
 	const TSet<FStringView>& ValidExtensions;
 	TMultiMap<FString, FString>& PackageExtensions;
+
+	// Scratch variables
+	FString PackageName;
+	FString AssetSourcePath;
+	FString StandardAssetSourcePath;
+	FString BaseAssetSourcePathBuffer;
 public:
 	FPackageCookerOpenOrderVisitor(
 		const UE::Cook::FCookSandbox& InSandboxFile,
@@ -1889,22 +1895,20 @@ public:
 				return true;
 			}
 
-			FString PackageName;
-			FString AssetSourcePath = SandboxFile.ConvertFromSandboxPathInPlatformRoot(Filename, PlatformSandboxPath);
-			FString StandardAssetSourcePath = FPaths::CreateStandardFilename(AssetSourcePath);
+			AssetSourcePath = SandboxFile.ConvertFromSandboxPathInPlatformRoot(Filename, PlatformSandboxPath);
+			StandardAssetSourcePath = FPaths::CreateStandardFilename(AssetSourcePath);
+			FString* BaseAssetSourcePath = &StandardAssetSourcePath;
 			if (StandardAssetSourcePath.EndsWith(TEXT(".m.ubulk")))
 			{
 				// '.' is an 'invalid' character in a filename; FilenameToLongPackageName will fail.
-				FString BaseAssetSourcePath(StandardAssetSourcePath);
-				BaseAssetSourcePath.RemoveFromEnd(TEXT(".m.ubulk"));
-				PackageName = FPackageName::FilenameToLongPackageName(BaseAssetSourcePath);
+				BaseAssetSourcePathBuffer = StandardAssetSourcePath;
+				BaseAssetSourcePathBuffer.RemoveFromEnd(TEXT(".m.ubulk"));
+				BaseAssetSourcePath = &BaseAssetSourcePathBuffer;
 			}
-			else
+			if (FPackageName::TryConvertFilenameToLongPackageName(*BaseAssetSourcePath, PackageName))
 			{
-				PackageName = FPackageName::FilenameToLongPackageName(StandardAssetSourcePath);
+				PackageExtensions.AddUnique(PackageName, StandardAssetSourcePath);
 			}
-
-			PackageExtensions.AddUnique(PackageName, StandardAssetSourcePath);
 		}
 
 		return true;
