[Sequencer] Allow deselection while selection is being limited when clicking on non selectable object

diff --git a/Engine/Source/Editor/UnrealEd/Private/EditorViewportSelectability.cpp b/Engine/Source/Editor/UnrealEd/Private/EditorViewportSelectability.cpp
--- a/Engine/Source/Editor/UnrealEd/Private/EditorViewportSelectability.cpp
+++ b/Engine/Source/Editor/UnrealEd/Private/EditorViewportSelectability.cpp
@@ -259,11 +259,20 @@ bool FEditorViewportSelectability::SelectActorsByPredicate(UWorld* const InWorld
 	ActorSelection->BeginBatchSelectOperation();
 	ActorSelection->Modify();
 
-	if (bInSelect && bInClearSelection)
+	if (bInClearSelection)
 	{
 		ActorSelection->DeselectAll();
 	}
 
+	// Early out for specific deselect case
+	if (!bInSelect && bInClearSelection)
+	{
+		ActorSelection->EndBatchSelectOperation();
+		GEditor->NoteSelectionChange();
+
+		return true;
+	}
+
 	auto SelectIfPossible = [bInSelect, ActorSelection, &InPredicate, &bSomethingSelected](AActor* const InActor)
 	{
 		if (IsValid(InActor) && ActorSelection->IsSelected(InActor) != bInSelect && InPredicate(InActor) == bInSelect)
@@ -443,10 +452,21 @@ void FEditorViewportSelectability::UpdateHoverFromHitProxy(HHitProxy* const InHi
 
 bool FEditorViewportSelectability::HandleClick(FEditorViewportClient* const InViewportClient, HHitProxy* const InHitProxy, const FViewportClick& InClick)
 {
+	if (!InViewportClient)
+	{
+		return false;
+	}
+
+	UWorld* const World = InViewportClient->GetWorld();
+	if (!IsValid(World))
+	{
+		return false;
+	}
+
 	// Disable actor selection when sequencer is limiting selection
 	const int32 HitX = InViewportClient->Viewport->GetMouseX();
 	const int32 HitY = InViewportClient->Viewport->GetMouseY();
-	HHitProxy* const HitResult = InViewportClient->Viewport->GetHitProxy(HitX, HitY);
+	const HHitProxy* const HitResult = InViewportClient->Viewport->GetHitProxy(HitX, HitY);
 	if (!HitResult)
 	{
 		return false;
@@ -457,17 +477,29 @@ bool FEditorViewportSelectability::HandleClick(FEditorViewportClient* const InVi
 		return false;
 	}
 
-	const HActor* const ActorHitProxy = static_cast<HActor*>(HitResult);
-	return ActorHitProxy
-		&& IsValid(ActorHitProxy->Actor)
-		&& !IsObjectSelectableInViewport(ActorHitProxy->Actor);
+	const HActor* const ActorHitProxy = static_cast<const HActor*>(HitResult);
+	if (!ActorHitProxy || !IsValid(ActorHitProxy->Actor))
+	{
+		return false;
+	}
+
+	const bool bNotSelectable = !IsObjectSelectableInViewport(ActorHitProxy->Actor);
+
+	if (bNotSelectable)
+	{
+		SelectActorsByPredicate(World, false, true
+			, [this](AActor* const InActor) -> bool
+			{
+				return false;
+			});
+	}
+
+	return bNotSelectable;
 }
 
 bool FEditorViewportSelectability::BoxSelectWorldActors(FBox& InBox, FEditorViewportClient* const InEditorViewportClient, const bool bInSelect)
 {
-	if (!GEditor
-		|| !InEditorViewportClient
-		|| InEditorViewportClient->IsInGameView())
+	if (!InEditorViewportClient || InEditorViewportClient->IsInGameView())
 	{
 		return false;
 	}
@@ -497,9 +529,7 @@ bool FEditorViewportSelectability::BoxSelectWorldActors(FBox& InBox, FEditorView
 
 bool FEditorViewportSelectability::FrustumSelectWorldActors(const FConvexVolume& InFrustum, FEditorViewportClient* const InEditorViewportClient, const bool bInSelect)
 {
-	if (!GEditor
-		|| !InEditorViewportClient
-		|| InEditorViewportClient->IsInGameView())
+	if (!InEditorViewportClient || InEditorViewportClient->IsInGameView())
 	{
 		return false;
 	}
