Shader commandlet: throw an error if not all input mateiral we found.

#tests none
#rnx

diff --git a/Engine/Source/Editor/UnrealEd/Private/Commandlets/CompileShadersTestBedCommandlet.cpp b/Engine/Source/Editor/UnrealEd/Private/Commandlets/CompileShadersTestBedCommandlet.cpp
--- a/Engine/Source/Editor/UnrealEd/Private/Commandlets/CompileShadersTestBedCommandlet.cpp
+++ b/Engine/Source/Editor/UnrealEd/Private/Commandlets/CompileShadersTestBedCommandlet.cpp
@@ -109,6 +109,10 @@ int32 UCompileShadersTestBedCommandlet::Main(const FString& Params)
 
 			AssetRegistry.GetAssets(Filter, MaterialList);
 			UE_LOG(LogCompileShadersTestBedCommandlet, Display, TEXT("Found %d/%d requested materials/material instances."), MaterialList.Num() - MaterialsNumBefore, Filter.SoftObjectPaths.Num());
+			if (MaterialList.Num() != MaterialsNumBefore)
+			{
+				UE_LOG(LogCompileShadersTestBedCommandlet, Error, TEXT("Some materials/material instances we not found."));
+			}
 		}
 	}
 
