Make Dump Material Expressions commandlet write out to a csv instead of a plain text file.

- This makes it much easier to work with the data either as a spreadsheet or with python scripts.
- Simplify the code by using the csv file writer.

#rnx

diff --git a/Engine/Source/Editor/UnrealEd/Private/Commandlets/DumpMaterialExpressionsCommandlet.cpp b/Engine/Source/Editor/UnrealEd/Private/Commandlets/DumpMaterialExpressionsCommandlet.cpp
--- a/Engine/Source/Editor/UnrealEd/Private/Commandlets/DumpMaterialExpressionsCommandlet.cpp
+++ b/Engine/Source/Editor/UnrealEd/Private/Commandlets/DumpMaterialExpressionsCommandlet.cpp
@@ -8,6 +8,7 @@
 #include "Containers/Array.h"
 #include "Containers/Map.h"
 #include "Containers/StringConv.h"
+#include "ProfilingDebugging/DiagnosticTable.h"
 #include "HAL/FileManager.h"
 #include "HAL/PlatformCrt.h"
 #include "Internationalization/Text.h"
@@ -74,11 +75,8 @@ void WriteLine(FArchive* FileWriter, const TArray<FString>& FieldNames, const TA
 	FileWriter->Serialize(const_cast<ANSICHAR*>(StringCast<ANSICHAR>(*OutputLine).Get()), OutputLine.Len());
 };
 
-void WriteOutMaterialExpressions(FArchive* FileWriter, const FString& OutputFilePath)
+void WriteOutMaterialExpressions(FDiagnosticTableWriterCSV& CsvFile)
 {
-	FString OutputLine = TEXT("[MATERIAL EXPRESSIONS]\n");
-	FileWriter->Serialize(const_cast<ANSICHAR*>(StringCast<ANSICHAR>(*OutputLine).Get()), OutputLine.Len());
-
 	struct FMaterialExpressionInfo
 	{
 		UMaterialExpression* MaterialExpression;
@@ -95,28 +93,6 @@ void WriteOutMaterialExpressions(FArchive* FileWriter, const FString& OutputFile
 	};
 	TArray<FMaterialExpressionInfo> MaterialExpressionInfos;
 
-	const FString NameField = TEXT("NAME");
-	const FString TypeField = TEXT("TYPE");
-	const FString ClassField = TEXT("CLASS_FLAGS");
-	const FString ShowInCreateMenuField = TEXT("SHOW_IN_CREATE_MENU");
-	const FString KeywordsField = TEXT("KEYWORDS");
-	const FString CreationNameField = TEXT("CREATION_NAME");
-	const FString CreationDescriptionField = TEXT("CREATION_DESCRIPTION");
-	const FString CaptionField = TEXT("CAPTION");
-	const FString DescriptionField = TEXT("DESCRIPTION");
-	const FString TooltipField = TEXT("TOOLTIP");
-
-	int32 MaxNameLength = NameField.Len();
-	int32 MaxTypeLength = TypeField.Len();
-	int32 MaxClassFlagsLength = ClassField.Len();
-	int32 MaxShowInCreateMenuLength = ShowInCreateMenuField.Len();
-	int32 MaxKeywordsLength = KeywordsField.Len();
-	int32 MaxCreationNameLength = CreationNameField.Len();
-	int32 MaxCreationDescriptionLength = CreationDescriptionField.Len();
-	int32 MaxCaptionLength = CaptionField.Len();
-	int32 MaxDescriptionLength = DescriptionField.Len();
-	int32 MaxTooltipLength = TooltipField.Len();
-
 	// Collect all default material expression objects
 	for (TObjectIterator<UClass> It; It; ++It)
 	{
@@ -195,87 +171,41 @@ void WriteOutMaterialExpressions(FArchive* FileWriter, const FString& OutputFile
 				ExpressionInfo.ClassFlags = ClassFlags;
 				ExpressionInfo.ShowInCreateMenu = bShowInCreateMenu ? TEXT("Yes") : TEXT("No");
 				MaterialExpressionInfos.Add(ExpressionInfo);
-
-				MaxNameLength = FMath::Max(MaxNameLength, ExpressionInfo.Name.Len());
-				MaxTypeLength = FMath::Max(MaxTypeLength, ExpressionInfo.Type.Len());
-				MaxClassFlagsLength = FMath::Max(MaxClassFlagsLength, ExpressionInfo.ClassFlags.Len());
-				MaxShowInCreateMenuLength = FMath::Max(MaxShowInCreateMenuLength, ExpressionInfo.ShowInCreateMenu.Len());
-				MaxKeywordsLength = FMath::Max(MaxKeywordsLength, ExpressionInfo.Keywords.Len());
-				MaxCreationNameLength = FMath::Max(MaxCreationNameLength, ExpressionInfo.CreationName.Len());
-				MaxCreationDescriptionLength = FMath::Max(MaxCreationDescriptionLength, ExpressionInfo.CreationDescription.Len());
-				MaxCaptionLength = FMath::Max(MaxCaptionLength, ExpressionInfo.Caption.Len());
-				MaxDescriptionLength = FMath::Max(MaxDescriptionLength, ExpressionInfo.Description.Len());
-				MaxTooltipLength = FMath::Max(MaxTooltipLength, ExpressionInfo.Tooltip.Len());
 			}
 		}
 	}
 
-	// Additional padding for spacing
-	const int32 AdditionalPadding = 3;
-	MaxNameLength += AdditionalPadding;
-	MaxTypeLength += AdditionalPadding;
-	MaxClassFlagsLength += AdditionalPadding;
-	MaxShowInCreateMenuLength += AdditionalPadding;
-	MaxKeywordsLength += AdditionalPadding;
-	MaxCreationNameLength += AdditionalPadding;
-	MaxCreationDescriptionLength += AdditionalPadding;
-	MaxCaptionLength += AdditionalPadding;
-	MaxDescriptionLength += AdditionalPadding;
-	MaxTooltipLength += AdditionalPadding;
-
-	TArray<uint32> MaxFieldLength;
-	MaxFieldLength.Add(MaxNameLength);
-	MaxFieldLength.Add(MaxTypeLength);
-	MaxFieldLength.Add(MaxClassFlagsLength);
-	MaxFieldLength.Add(MaxShowInCreateMenuLength);
-	MaxFieldLength.Add(MaxKeywordsLength);
-	MaxFieldLength.Add(MaxCreationNameLength);
-	MaxFieldLength.Add(MaxCreationDescriptionLength);
-	MaxFieldLength.Add(MaxCaptionLength);
-	MaxFieldLength.Add(MaxDescriptionLength);
-	MaxFieldLength.Add(MaxTooltipLength);
-
 	// Write the material expression list to a text file
-	TArray<FString> FieldNames;
-	FieldNames.Add(NameField);
-	FieldNames.Add(TypeField);
-	FieldNames.Add(ClassField);
-	FieldNames.Add(ShowInCreateMenuField);
-	FieldNames.Add(KeywordsField);
-	FieldNames.Add(CreationNameField);
-	FieldNames.Add(CreationDescriptionField);
-	FieldNames.Add(CaptionField);
-	FieldNames.Add(DescriptionField);
-	FieldNames.Add(TooltipField);
-	WriteLine(FileWriter, FieldNames, MaxFieldLength);
+	CsvFile.AddColumn(TEXT("NAME"));
+	CsvFile.AddColumn(TEXT("TYPE"));
+	CsvFile.AddColumn(TEXT("CLASS_FLAGS"));
+	CsvFile.AddColumn(TEXT("SHOW_IN_CREATE_MENU"));
+	CsvFile.AddColumn(TEXT("KEYWORDS"));
+	CsvFile.AddColumn(TEXT("CREATION_NAME"));
+	CsvFile.AddColumn(TEXT("CREATION_DESCRIPTION"));
+	CsvFile.AddColumn(TEXT("CAPTION"));
+	CsvFile.AddColumn(TEXT("DESCRIPTION"));
+	CsvFile.AddColumn(TEXT("TOOLTIP"));
+	CsvFile.CycleRow();
 
 	for (FMaterialExpressionInfo& ExpressionInfo : MaterialExpressionInfos)
 	{
-		FieldNames.Reset();
-		FieldNames.Add(GetFormattedText(ExpressionInfo.Name));
-		FieldNames.Add(GetFormattedText(ExpressionInfo.Type));
-		FieldNames.Add(GetFormattedText(ExpressionInfo.ClassFlags));
-		FieldNames.Add(GetFormattedText(ExpressionInfo.ShowInCreateMenu));
-		FieldNames.Add(GetFormattedText(ExpressionInfo.Keywords));
-		FieldNames.Add(GetFormattedText(ExpressionInfo.CreationName));
-		FieldNames.Add(GetFormattedText(ExpressionInfo.CreationDescription));
-		FieldNames.Add(GetFormattedText(ExpressionInfo.Caption));
-		FieldNames.Add(GetFormattedText(ExpressionInfo.Description));
-		FieldNames.Add(GetFormattedText(ExpressionInfo.Tooltip));
-		WriteLine(FileWriter, FieldNames, MaxFieldLength);
+		CsvFile.AddColumn(*GetFormattedText(ExpressionInfo.Name));
+		CsvFile.AddColumn(*GetFormattedText(ExpressionInfo.Type));
+		CsvFile.AddColumn(*GetFormattedText(ExpressionInfo.ClassFlags));
+		CsvFile.AddColumn(*GetFormattedText(ExpressionInfo.ShowInCreateMenu));
+		CsvFile.AddColumn(*GetFormattedText(ExpressionInfo.Keywords));
+		CsvFile.AddColumn(*GetFormattedText(ExpressionInfo.CreationName));
+		CsvFile.AddColumn(*GetFormattedText(ExpressionInfo.CreationDescription));
+		CsvFile.AddColumn(*GetFormattedText(ExpressionInfo.Caption));
+		CsvFile.AddColumn(*GetFormattedText(ExpressionInfo.Description));
+		CsvFile.AddColumn(*GetFormattedText(ExpressionInfo.Tooltip));
+		CsvFile.CycleRow();
 	}
-
-	OutputLine = FString::Printf(TEXT("\nTotal %d material expressions found."), MaterialExpressionInfos.Num());
-	FileWriter->Serialize(const_cast<ANSICHAR*>(StringCast<ANSICHAR>(*OutputLine).Get()), OutputLine.Len());
-
-	UE_LOG(LogDumpMaterialExpressionsCommandlet, Log, TEXT("Total %d material expressions are written to %s"), MaterialExpressionInfos.Num(), *OutputFilePath);
 }
 
-void WriteOutMaterialFunctions(FArchive* FileWriter, const FString& OutputFilePath)
+void WriteOutMaterialFunctions(FDiagnosticTableWriterCSV& CsvFile)
 {
-	FString OutputLine = TEXT("[MATERIAL FUNCTIONS]\n");
-	FileWriter->Serialize(const_cast<ANSICHAR*>(StringCast<ANSICHAR>(*OutputLine).Get()), OutputLine.Len());
-
 	struct FMaterialFunctionInfo
 	{
 		FString Name;
@@ -284,14 +214,6 @@ void WriteOutMaterialFunctions(FArchive* FileWriter, const FString& OutputFilePa
 	};
 	TArray<FMaterialFunctionInfo> MaterialFunctionInfos;
 
-	const FString NameField = TEXT("NAME");
-	const FString DescriptionField = TEXT("DESCRIPTION");
-	const FString PathField = TEXT("PATH");
-
-	int32 MaxNameLength = NameField.Len();
-	int32 MaxDescriptionLength = DescriptionField.Len();
-	int32 MaxPathLength = PathField.Len();
-
 	// See UMaterialGraphSchema::GetMaterialFunctionActions for reference
 	TArray<FAssetData> AssetDataList;
 	FAssetRegistryModule& AssetRegistryModule = FModuleManager::LoadModuleChecked<FAssetRegistryModule>(TEXT("AssetRegistry"));
@@ -316,44 +238,22 @@ void WriteOutMaterialFunctions(FArchive* FileWriter, const FString& OutputFilePa
 			FunctionInfo.Description = Description;
 			FunctionInfo.Path = FunctionPathName;
 			MaterialFunctionInfos.Add(FunctionInfo);
-
-			MaxNameLength = FMath::Max(MaxNameLength, FunctionInfo.Name.Len());
-			MaxDescriptionLength = FMath::Max(MaxDescriptionLength, FunctionInfo.Description.Len());
-			MaxPathLength = FMath::Max(MaxPathLength, FunctionInfo.Path.Len());
 		}
 	}
 
-	// Additional padding for spacing
-	const int32 AdditionalPadding = 3;
-	MaxNameLength += AdditionalPadding;
-	MaxDescriptionLength += AdditionalPadding;
-	MaxPathLength += AdditionalPadding;
-
-	TArray<uint32> MaxFieldLength;
-	MaxFieldLength.Add(MaxNameLength);
-	MaxFieldLength.Add(MaxDescriptionLength);
-	MaxFieldLength.Add(MaxPathLength);
-
 	// Write the material expression list to a text file
-	TArray<FString> FieldNames;
-	FieldNames.Add(NameField);
-	FieldNames.Add(DescriptionField);
-	FieldNames.Add(PathField);
-	WriteLine(FileWriter, FieldNames, MaxFieldLength);
+	CsvFile.AddColumn(TEXT("NAME"));
+	CsvFile.AddColumn(TEXT("DESCRIPTION"));
+	CsvFile.AddColumn(TEXT("PATH"));
+	CsvFile.CycleRow();
 
 	for (FMaterialFunctionInfo& FunctionInfo : MaterialFunctionInfos)
 	{
-		FieldNames.Reset();
-		FieldNames.Add(GetFormattedText(FunctionInfo.Name));
-		FieldNames.Add(GetFormattedText(FunctionInfo.Description));
-		FieldNames.Add(GetFormattedText(FunctionInfo.Path));
-		WriteLine(FileWriter, FieldNames, MaxFieldLength);
+		CsvFile.AddColumn(*GetFormattedText(FunctionInfo.Name));
+		CsvFile.AddColumn(*GetFormattedText(FunctionInfo.Description));
+		CsvFile.AddColumn(*GetFormattedText(FunctionInfo.Path));
+		CsvFile.CycleRow();
 	}
-	
-	OutputLine = FString::Printf(TEXT("\nTotal %d (bExposeToLibrary=true) material functions found."), MaterialFunctionInfos.Num());
-	FileWriter->Serialize(const_cast<ANSICHAR*>(StringCast<ANSICHAR>(*OutputLine).Get()), OutputLine.Len());
-
-	UE_LOG(LogDumpMaterialExpressionsCommandlet, Log, TEXT("Total %d material functions are written to %s"), MaterialFunctionInfos.Num(), *OutputFilePath);
 }
 
 UDumpMaterialExpressionsCommandlet::UDumpMaterialExpressionsCommandlet(const FObjectInitializer& ObjectInitializer)
@@ -383,18 +283,21 @@ int32 UDumpMaterialExpressionsCommandlet::Main(const FString& Params)
 		return 0;
 	}
 
-	const FString OutputFilePath = FPaths::Combine(*FPaths::ProjectSavedDir(), TEXT("MaterialEditor"), TEXT("MaterialExpressions.txt"));
-	FArchive* FileWriter = IFileManager::Get().CreateFileWriter(*OutputFilePath);
+	const FString OutputFilePath = FPaths::Combine(*FPaths::ProjectSavedDir(), TEXT("MaterialEditor"), TEXT("MaterialExpressions.csv"));
+	TUniquePtr<FArchive> CSVTableFile = TUniquePtr<FArchive>{ IFileManager::Get().CreateFileWriter(*OutputFilePath) };
+	FDiagnosticTableWriterCSV CsvFile(CSVTableFile.Get());
 
-	WriteOutMaterialExpressions(FileWriter, OutputFilePath);
+	WriteOutMaterialExpressions(CsvFile);
 
-	FString OutputLine = FString::Printf(TEXT("\n\n============================================================================\n\n"));
-	FileWriter->Serialize(const_cast<ANSICHAR*>(StringCast<ANSICHAR>(*OutputLine).Get()), OutputLine.Len());
 
-	WriteOutMaterialFunctions(FileWriter, OutputFilePath);
+	const FString MatFuncOutputFilePath = FPaths::Combine(*FPaths::ProjectSavedDir(), TEXT("MaterialEditor"), TEXT("MaterialFunctions.csv"));
+	TUniquePtr<FArchive> MatFuncCSVTableFile = TUniquePtr<FArchive>{ IFileManager::Get().CreateFileWriter(*MatFuncOutputFilePath) };
+	FDiagnosticTableWriterCSV MatFuncCsvFile(MatFuncCSVTableFile.Get());
+
+	WriteOutMaterialFunctions(MatFuncCsvFile);
 
-	FileWriter->Close();
-	delete FileWriter;
+	UE_LOG(LogDumpMaterialExpressionsCommandlet, Log, TEXT("Results are written to %s"), *OutputFilePath);
+	UE_LOG(LogDumpMaterialExpressionsCommandlet, Log, TEXT("Results are written to %s"), *MatFuncOutputFilePath);
 
 	return 0;
 }
