Cooker: MPCook and CookPackageSplitters: when a package is kicked out of save because it was retracted, do not mark it as saved on the generation helper. Marking it as saved tells the generationhelper the generationhelper can destruct, which is not valid for it to do while we are still waiting for the package's save to complete on another CookWorker.
This fixes  occurrences of the warning "LogCook: Warning: RecordResults received for generated package %s, but its ParentGenerationHelper has already been destructed..."
#rnx

diff --git a/Engine/Source/Editor/UnrealEd/Private/Cooker/CookGenerationHelper.cpp b/Engine/Source/Editor/UnrealEd/Private/Cooker/CookGenerationHelper.cpp
--- a/Engine/Source/Editor/UnrealEd/Private/Cooker/CookGenerationHelper.cpp
+++ b/Engine/Source/Editor/UnrealEd/Private/Cooker/CookGenerationHelper.cpp
@@ -1266,8 +1266,7 @@ void FGenerationHelper::ResetSaveState(FCookGenerationInfo& Info, UPackage* Pack
 		}
 	}
 
-	if (ReleaseSaveReason != EStateChangeReason::RecreateObjectCache &&
-		ReleaseSaveReason != EStateChangeReason::DoneForNow)
+	if (IsTerminalStateChange(ReleaseSaveReason))
 	{
 		// The save is completed and we will not come back to it; set state back to initial
 		// state and drop our reference keeping this GenerationHelper in memory for the save.
@@ -1293,9 +1292,7 @@ void FGenerationHelper::ResetSaveState(FCookGenerationInfo& Info, UPackage* Pack
 	{
 		if (NewState != EPackageState::Idle &&
 			Info.PackageData->GetCachedObjectsInOuter().Num() != 0 && IsUseInternalReferenceToAvoidGarbageCollect() &&
-			(ReleaseSaveReason != EStateChangeReason::Completed && ReleaseSaveReason != EStateChangeReason::DoneForNow
-				&& ReleaseSaveReason != EStateChangeReason::SaveError
-				&& ReleaseSaveReason != EStateChangeReason::CookerShutdown
+			(!IsTerminalStateChange(ReleaseSaveReason) && ReleaseSaveReason != EStateChangeReason::DoneForNow
 				&& ReleaseSaveReason != EStateChangeReason::Retraction))
 		{
 			UE_LOG(LogCook, Error,
diff --git a/Engine/Source/Editor/UnrealEd/Private/Cooker/CookTypes.cpp b/Engine/Source/Editor/UnrealEd/Private/Cooker/CookTypes.cpp
--- a/Engine/Source/Editor/UnrealEd/Private/Cooker/CookTypes.cpp
+++ b/Engine/Source/Editor/UnrealEd/Private/Cooker/CookTypes.cpp
@@ -51,6 +51,18 @@ const TCHAR* LexToString(EStateChangeReason Reason)
 	}
 }
 
+bool IsTerminalStateChange(EStateChangeReason Reason)
+{
+	switch (Reason)
+	{
+	case EStateChangeReason::Completed: return true;
+	case EStateChangeReason::SaveError: return true;
+	case EStateChangeReason::CookerShutdown: return true;
+	case EStateChangeReason::CookSuppressed: return true;
+	default: return false;
+	}
+}
+
 const TCHAR* LexToString(ESuppressCookReason Reason)
 {
 	switch (Reason)
diff --git a/Engine/Source/Editor/UnrealEd/Private/Cooker/CookTypes.h b/Engine/Source/Editor/UnrealEd/Private/Cooker/CookTypes.h
--- a/Engine/Source/Editor/UnrealEd/Private/Cooker/CookTypes.h
+++ b/Engine/Source/Editor/UnrealEd/Private/Cooker/CookTypes.h
@@ -146,6 +146,7 @@ namespace UE::Cook
 		ForceRecook,
 		UrgencyUpdated,
 	};
+	bool IsTerminalStateChange(EStateChangeReason Reason);
 	const TCHAR* LexToString(UE::Cook::EStateChangeReason Reason);
 
 	enum class ESuppressCookReason : uint8
