Rename Culture to Language in UE settings

#rnx

diff --git a/Engine/Source/Developer/DeveloperToolSettings/Classes/Settings/ProjectPackagingSettings.h b/Engine/Source/Developer/DeveloperToolSettings/Classes/Settings/ProjectPackagingSettings.h
--- a/Engine/Source/Developer/DeveloperToolSettings/Classes/Settings/ProjectPackagingSettings.h
+++ b/Engine/Source/Developer/DeveloperToolSettings/Classes/Settings/ProjectPackagingSettings.h
@@ -464,8 +464,8 @@ public:
 	UPROPERTY(config, EditAnywhere, Category=Packaging, AdvancedDisplay, meta=(DisplayName="Internationalization Support"))
 	EProjectPackagingInternationalizationPresets InternationalizationPreset;
 
-	/** Cultures whose data should be cooked, staged, and packaged. */
-	UPROPERTY(config, EditAnywhere, Category=Packaging, AdvancedDisplay, meta=(DisplayName="Localizations to Package"))
+	/** Languages whose data should be cooked, staged, and packaged. */
+	UPROPERTY(config, EditAnywhere, Category=Packaging, AdvancedDisplay, meta=(DisplayName="Languages to Package"))
 	TArray<FString> CulturesToStage;
 
 	/** List of localization targets that should be chunked during cooking (if using chunks) */
diff --git a/Engine/Source/Editor/UnrealEd/Private/Settings/ProjectPackagingSettingsCustomization.h b/Engine/Source/Editor/UnrealEd/Private/Settings/ProjectPackagingSettingsCustomization.h
--- a/Engine/Source/Editor/UnrealEd/Private/Settings/ProjectPackagingSettingsCustomization.h
+++ b/Engine/Source/Editor/UnrealEd/Private/Settings/ProjectPackagingSettingsCustomization.h
@@ -186,7 +186,7 @@ protected:
 
 			PopulateCultureList();
 
-			PackagingCategory.AddCustomRow(LOCTEXT("CulturesToStageLabel", "Cultures To Stage"), true)
+			PackagingCategory.AddCustomRow(LOCTEXT("CulturesToStageLabel", "Languages To Package"), true)
 				.NameContent()
 				.HAlign(HAlign_Fill)
 				.VAlign(VAlign_Top)
@@ -204,7 +204,7 @@ protected:
 					[
 						SNew(SImage)
 						.Image(FAppStyle::GetBrush(TEXT("Icons.Error")))
-						.ToolTipText(LOCTEXT("NoCulturesToStageSelectedError", "At least one culture must be selected or fatal errors may occur when launching games."))
+						.ToolTipText(LOCTEXT("NoCulturesToStageSelectedError", "At least one language must be selected or fatal errors may occur when launching games."))
 						.Visibility(this, &FProjectPackagingSettingsCustomization::HandleNoCulturesErrorIconVisibility)
 					]
 				]
